Metasound plugin initial checkin of data sharing reference graph.

diff --git a/Engine/Plugins/Runtime/Metasound/Metasound.uplugin b/Engine/Plugins/Runtime/Metasound/Metasound.uplugin
new file mode 100644
--- /dev/null
+++ b/Engine/Plugins/Runtime/Metasound/Metasound.uplugin
@@ -0,0 +1,36 @@
+{
+	"FileVersion" : 3,
+	"Version" : 1,
+	"VersionName" : "1.0",
+	"FriendlyName" : "Metasound",
+	"Description" : "",
+	"Category" : "Audio",
+	"CreatedBy" : "Epic Games, Inc.",
+	"CreatedByURL" : "http://epicgames.com",
+	"DocsURL" : "",
+	"MarketplaceURL" : "",
+	"SupportURL" : "",
+	"EnabledByDefault" : false,
+	"CanContainContent" : true,
+	"IsBetaVersion" : true,
+	"Installed" : false,
+	"Modules" :
+	[
+		{
+			"Name" : "MetasoundGraphCore",
+			"Type" : "Runtime",
+			"LoadingPhase" : "PreDefault",
+			"WhitelistPlatforms": [
+				"Win64"
+			]
+		},
+		{
+			"Name" : "MetasoundGenerator",
+			"Type" : "Runtime",
+			"LoadingPhase" : "Default",
+			"WhitelistPlatforms": [
+				"Win64"
+			]
+		}
+	]
+}
diff --git a/Engine/Plugins/Runtime/Metasound/Source/MetasoundGenerator/MetasoundGenerator.Build.cs b/Engine/Plugins/Runtime/Metasound/Source/MetasoundGenerator/MetasoundGenerator.Build.cs
new file mode 100644
--- /dev/null
+++ b/Engine/Plugins/Runtime/Metasound/Source/MetasoundGenerator/MetasoundGenerator.Build.cs
@@ -0,0 +1,32 @@
+// Copyright Epic Games, Inc. All Rights Reserved.
+
+using UnrealBuildTool;
+using System.IO;
+
+namespace UnrealBuildTool.Rules
+{
+	public class MetasoundGenerator : ModuleRules
+	{
+        public MetasoundGenerator(ReadOnlyTargetRules Target) : base(Target)
+		{
+			//OptimizeCode = CodeOptimization.Never;
+
+			PrivateDependencyModuleNames.AddRange(
+				new string[] {
+					"SignalProcessing",
+					"CoreUObject",
+					"Engine",
+					"AudioMixerCore"
+				}
+			);
+
+            PublicDependencyModuleNames.AddRange(
+				new string[] {
+                    "Core",
+					"AudioMixer",
+					"MetasoundGraphCore"
+                }
+            );
+		}
+	}
+}
diff --git a/Engine/Plugins/Runtime/Metasound/Source/MetasoundGenerator/Private/MetasoundGenerator.cpp b/Engine/Plugins/Runtime/Metasound/Source/MetasoundGenerator/Private/MetasoundGenerator.cpp
new file mode 100644
--- /dev/null
+++ b/Engine/Plugins/Runtime/Metasound/Source/MetasoundGenerator/Private/MetasoundGenerator.cpp
@@ -0,0 +1,252 @@
+// Copyright Epic Games, Inc. All Rights Reserved.
+#include "MetasoundGenerator.h"
+
+#include "DSP/Dsp.h"
+#include "MetasoundOperatorInterface.h"
+#include "MetasoundDataReferenceTypes.h"
+#include "MetasoundOscNode.h"
+#include "MetasoundAudioMultiplyNode.h"
+#include "MetasoundADSRNode.h"
+#include "MetasoundPeriodicBopNode.h"
+#include "MetasoundOutputNode.h"
+#include "MetasoundInputNode.h"
+#include "MetasoundOperatorBuilder.h"
+#include "MetasoundGraph.h"
+
+namespace Metasound
+{
+	FMetasoundGenerator::FMetasoundGenerator(FOperatorUniquePtr InOperator, const FAudioReadRef& InOperatorReadBuffer)
+	:	Operator(MoveTemp(InOperator))
+	,	OperatorReadBuffer(InOperatorReadBuffer)
+	,	OperatorReadNum(InOperatorReadBuffer->Num())
+	,	ExecuteOperator(nullptr)
+	{
+		if (Operator.IsValid())
+		{
+			ExecuteOperator = Operator->GetExecuteFunction();
+			FDataReferenceCollection Inputs = Operator->GetInputs();
+
+			if (Inputs.ContainsDataWriteReference<FFloat>(TEXT("Frequency")))
+			{
+				FrequencyRef = Inputs.GetDataWriteReference<FFloat>(TEXT("Frequency"));
+			}
+
+			if (Inputs.ContainsDataWriteReference<FTimeSeconds>(TEXT("BopPeriod")))
+			{
+				BopPeriodRef = Inputs.GetDataWriteReference<FTimeSeconds>(TEXT("BopPeriod"));
+			}
+		}
+	}
+
+	FMetasoundGenerator::~FMetasoundGenerator()
+	{
+	}
+
+	void FMetasoundGenerator::SetFrequency(float InFrequency)
+	{
+		*FrequencyRef = InFrequency;
+	}
+
+	void FMetasoundGenerator::SetBopPeriod(float InPeriodInSeconds)
+	{
+		BopPeriodRef->Seconds = InPeriodInSeconds;
+	}
+
+	int32 FMetasoundGenerator::OnGenerateAudio(float* OutAudio, int32 NumSamples)
+	{
+		if ((NumSamples <= 0) || (nullptr == ExecuteOperator) || (OperatorReadNum <= 0))
+		{
+			return 0;
+		}
+
+		int32 WriteCount = FillWithBuffer(OverflowBuffer, OutAudio, NumSamples);
+
+		if (WriteCount > 0)
+		{
+			NumSamples -= WriteCount;
+			OverflowBuffer.RemoveAtSwap(0 /* Index */, WriteCount /* Count */, false /* bAllowShrinking */);
+		}
+
+
+		while (NumSamples > 0)
+		{
+			// TODO: figure out good wrapper for this.
+			ExecuteOperator(Operator.Get());
+
+			int32 ThisLoopWriteCount = FillWithBuffer(*OperatorReadBuffer, &OutAudio[WriteCount], NumSamples);
+
+			NumSamples -= ThisLoopWriteCount;
+			WriteCount += ThisLoopWriteCount;
+
+			if (0 == NumSamples)
+			{
+				if (ThisLoopWriteCount < OperatorReadBuffer->Num())
+				{
+					int32 OverflowCount = OperatorReadBuffer->Num() - ThisLoopWriteCount;
+
+					OverflowBuffer.Reset();
+					OverflowBuffer.AddUninitialized(OverflowCount);
+
+					FMemory::Memcpy(OverflowBuffer.GetData(), &(OperatorReadBuffer->GetData()[ThisLoopWriteCount]), OverflowCount);
+				}
+			}
+		}
+
+		return WriteCount;
+	}
+
+	int32 FMetasoundGenerator::FillWithBuffer(const Audio::AlignedFloatBuffer& InBuffer, float* OutAudio, int32 MaxNumOutputSamples)
+	{
+		int32 InNum = InBuffer.Num();
+
+		if (InNum > 0)
+		{
+			if (InNum < MaxNumOutputSamples)
+			{
+				FMemory::Memcpy(OutAudio, InBuffer.GetData(), InNum * sizeof(float));
+				return InNum;
+			}
+			else
+			{
+				FMemory::Memcpy(OutAudio, InBuffer.GetData(), MaxNumOutputSamples * sizeof(float));
+				return MaxNumOutputSamples;
+			}
+		}
+
+		return 0;
+	}
+} /* End namespace Metasound */
+
+
+USynthComponentMetasoundGenerator::USynthComponentMetasoundGenerator(const FObjectInitializer& ObjInitializer)
+:	Super(ObjInitializer)
+{
+}
+
+USynthComponentMetasoundGenerator::~USynthComponentMetasoundGenerator()
+{
+}
+
+void USynthComponentMetasoundGenerator::SetBopPeriod(float InBopPeriod)
+{
+	using namespace Metasound;
+
+	BopPeriod = InBopPeriod;
+	if (Generator.IsValid())
+	{
+		FMetasoundGenerator* MetasoundGenerator = static_cast<FMetasoundGenerator*>(Generator.Get());
+		MetasoundGenerator->SetBopPeriod(BopPeriod);
+	}
+}
+
+void USynthComponentMetasoundGenerator::SetFrequency(float InFrequency)
+{
+	using namespace Metasound;
+
+	Frequency = InFrequency;
+	if (Generator.IsValid())
+	{
+		FMetasoundGenerator* MetasoundGenerator = static_cast<FMetasoundGenerator*>(Generator.Get());
+		MetasoundGenerator->SetFrequency(Frequency);
+	}
+}
+
+#if WITH_EDITOR
+void USynthComponentMetasoundGenerator::PostEditChangeProperty(struct FPropertyChangedEvent & PropertyChangedEvent)
+{
+	using namespace Metasound;
+
+	Super::PostEditChangeProperty(PropertyChangedEvent);
+
+	static const FName FrequencyFName = GET_MEMBER_NAME_CHECKED(USynthComponentMetasoundGenerator, Frequency);
+	static const FName BopPeriodFName = GET_MEMBER_NAME_CHECKED(USynthComponentMetasoundGenerator, BopPeriod);
+
+	if (!HasAnyFlags(RF_ClassDefaultObject))
+	{
+		if (Generator.IsValid())
+		{
+			if (FProperty* PropertyThatChanged = PropertyChangedEvent.Property)
+			{
+				const FName& Name = PropertyThatChanged->GetFName();
+				if (Name == FrequencyFName)
+				{
+					FMetasoundGenerator* MetasoundGenerator = static_cast<FMetasoundGenerator*>(Generator.Get());
+					MetasoundGenerator->SetFrequency(Frequency);
+				}
+				else if (Name == BopPeriodFName)
+				{
+					FMetasoundGenerator* MetasoundGenerator = static_cast<FMetasoundGenerator*>(Generator.Get());
+					MetasoundGenerator->SetBopPeriod(BopPeriod);
+				}
+			}
+		}
+	}
+
+}
+#endif
+
+ISoundGeneratorPtr USynthComponentMetasoundGenerator::CreateSoundGenerator(int32 InSampleRate, int32 InNumChannels)
+{
+	using namespace Metasound;
+	using FDataEdge = Metasound::FDataEdge; // Need to make explicit to disambiguate with existing FDataEdge.
+	using FOperatorUniquePtr = TUniquePtr<Metasound::IOperator>;
+
+	TInputNode<FTimeSeconds> BopPeriodInputNode(TEXT("BopPeriod"), 1.f);
+	TInputNode<FFloat> FrequencyInputNode(TEXT("Frequency"), 100.f);
+
+	FPeriodicBopNode BopNode(TEXT("Bop"), 1.f);
+	FOscNode OscNode(TEXT("Osc"), 100.f);
+	FADSRNode ADSRNode(TEXT("ADSR"), 10.f /* AttackMs */, 10.f /* DecayMs */, 50.f /* SustainMs */, 10.f /* ReleaseMs */);
+	FAudioMultiplyNode MultiplyNode(TEXT("Multiply"));
+
+	TOutputNode<FAudioBuffer> OutputNode(TEXT("Audio"));
+
+	FGraph Graph(TEXT("Graph"));
+
+	// Add rate of bop.
+	Graph.AddDataEdge({{&BopPeriodInputNode, BopPeriodInputNode.GetOutputs()[0]}, {&BopNode, BopNode.GetInputs()[0]}});
+	// Add frequency controls
+	Graph.AddDataEdge({{&FrequencyInputNode, FrequencyInputNode.GetOutputs()[0]}, {&OscNode, OscNode.GetInputs()[0]}});
+
+	// TODO: maybe nodes should have name accessors to inputs and outputs instead of an array?
+	Graph.AddDataEdge({{&BopNode, BopNode.GetOutputs()[0]}, {&ADSRNode, ADSRNode.GetInputs()[0]}});
+
+	// Hookup multiply nodes
+	Graph.AddDataEdge({{&OscNode, OscNode.GetOutputs()[0]}, {&MultiplyNode, MultiplyNode.GetInputs()[0]}});
+	Graph.AddDataEdge({{&ADSRNode, ADSRNode.GetOutputs()[0]}, {&MultiplyNode, MultiplyNode.GetInputs()[1]}});
+
+	// Route to output.
+	Graph.AddDataEdge({{&MultiplyNode, MultiplyNode.GetOutputs()[0]}, {&OutputNode, OutputNode.GetInputs()[0]}});
+
+	Graph.AddInputDataVertex({&BopPeriodInputNode, BopPeriodInputNode.GetInputs()[0]});
+	Graph.AddInputDataVertex({&FrequencyInputNode, FrequencyInputNode.GetInputs()[0]});
+	Graph.AddOutputDataVertex({&OutputNode, OutputNode.GetOutputs()[0]});
+
+
+	// Build everything
+	
+	FOperatorSettings Settings(InSampleRate, 256);
+	FOperatorBuilder Builder(Settings);
+
+	TArray<IOperatorBuilder::FBuildErrorPtr> Errors;
+	FOperatorUniquePtr Operator = Builder.BuildGraphOperator(Graph, Errors);
+
+	if (!Operator.IsValid())
+	{
+		return ISoundGeneratorPtr(nullptr);
+	}
+
+	FDataReferenceCollection OutputCollection = Operator->GetOutputs();
+	FAudioBufferReadRef OutAudio = OutputCollection.GetDataReadReference<FAudioBuffer>(TEXT("Audio"));
+
+	Generator = MakeShared<FMetasoundGenerator, ESPMode::ThreadSafe>(MoveTemp(Operator), OutAudio);
+	FMetasoundGenerator* MetasoundGenerator = static_cast<FMetasoundGenerator*>(Generator.Get());
+	MetasoundGenerator->SetFrequency(Frequency);
+	MetasoundGenerator->SetBopPeriod(BopPeriod);
+
+	return Generator;
+}
+
+
+
+
diff --git a/Engine/Plugins/Runtime/Metasound/Source/MetasoundGenerator/Private/MetasoundGeneratorModule.cpp b/Engine/Plugins/Runtime/Metasound/Source/MetasoundGenerator/Private/MetasoundGeneratorModule.cpp
new file mode 100644
--- /dev/null
+++ b/Engine/Plugins/Runtime/Metasound/Source/MetasoundGenerator/Private/MetasoundGeneratorModule.cpp
@@ -0,0 +1,20 @@
+// Copyright Epic Games, Inc. All Rights Reserved.
+
+#include "MetasoundGeneratorModule.h"
+
+#include "CoreMinimal.h"
+#include "Modules/ModuleInterface.h"
+#include "Modules/ModuleManager.h"
+
+DEFINE_LOG_CATEGORY(LogMetasoundGenerator);
+
+namespace Metasound
+{
+	class FMetasoundGeneratorModule : public IMetasoundGeneratorModule
+	{
+	};
+}
+
+IMPLEMENT_MODULE(Metasound::FMetasoundGeneratorModule, MetasoundGenerator);
+
+
diff --git a/Engine/Plugins/Runtime/Metasound/Source/MetasoundGenerator/Public/MetasoundGenerator.h b/Engine/Plugins/Runtime/Metasound/Source/MetasoundGenerator/Public/MetasoundGenerator.h
new file mode 100644
--- /dev/null
+++ b/Engine/Plugins/Runtime/Metasound/Source/MetasoundGenerator/Public/MetasoundGenerator.h
@@ -0,0 +1,84 @@
+// Copyright Epic Games, Inc. All Rights Reserved.
+
+#pragma once
+
+#include "CoreMinimal.h"
+#include "Components/SynthComponent.h"
+#include "DSP/Dsp.h"
+#include "MetasoundOperatorInterface.h"
+#include "MetasoundDataReferenceTypes.h"
+#include "Sound/SoundGenerator.h"
+
+#include "MetasoundGenerator.generated.h"
+
+namespace Metasound
+{
+	class METASOUNDGENERATOR_API FMetasoundGenerator : public ISoundGenerator
+	{
+	public:
+		using FOperatorUniquePtr = TUniquePtr<Metasound::IOperator>;
+		using FAudioReadRef = Metasound::FAudioBufferReadRef;
+
+		FMetasoundGenerator(FOperatorUniquePtr InOperator, const FAudioReadRef& InOperatorReadBuffer);
+		virtual ~FMetasoundGenerator();
+
+		void SetFrequency(float InFrequency);
+		void SetBopPeriod(float InPeriodInSeconds);
+
+		//~ Begin FSoundGenerator
+		virtual int32 OnGenerateAudio(float* OutAudio, int32 NumSamples) override;
+		//~ End FSoundGenerator
+
+	private:
+		int32 FillWithBuffer(const Audio::AlignedFloatBuffer& InBuffer, float* OutAudio, int32 MaxNumOutputSamples); 
+
+		using FExecuteFunction = Metasound::IOperator::FExecuteFunction;
+
+		Audio::AlignedFloatBuffer OverflowBuffer;
+		FOperatorUniquePtr Operator;
+		FAudioReadRef OperatorReadBuffer;
+
+		FFloatWriteRef FrequencyRef;
+		FTimeSecondsWriteRef BopPeriodRef;
+
+		int32 OperatorReadNum;
+		FExecuteFunction ExecuteOperator;
+	};
+}
+
+UCLASS(ClassGroup = Synth, meta = (BlueprintSpawnableComponent))
+class METASOUNDGENERATOR_API USynthComponentMetasoundGenerator : public USynthComponent
+{
+    GENERATED_BODY()
+
+public:
+
+    USynthComponentMetasoundGenerator(const FObjectInitializer& ObjInitializer);
+
+    virtual ~USynthComponentMetasoundGenerator();
+
+	UPROPERTY(EditAnywhere, BlueprintReadWrite, Category=Metasound, meta = (ClampMin="0.0001", ClampMax="10.0"))
+	float BopPeriod;
+
+	UPROPERTY(EditAnywhere, BlueprintReadWrite, Category=Metasound, meta = (ClampMin="20", ClampMax="800"))
+	float Frequency;
+
+	UFUNCTION(BlueprintCallable, Category=Metasound)
+	void SetBopPeriod(float InBopPeriod);
+
+	UFUNCTION(BlueprintCallable, Category = Metasound)
+	void SetFrequency(float InFrequency);
+
+    virtual ISoundGeneratorPtr CreateSoundGenerator(int32 InSampleRate, int32 InNumChannels) override;
+
+#if WITH_EDITOR
+	virtual void PostEditChangeProperty(struct FPropertyChangedEvent & PropertyChangedEvent) override;
+#endif
+
+private:
+	ISoundGeneratorPtr Generator;
+};
+
+
+
+
diff --git a/Engine/Plugins/Runtime/Metasound/Source/MetasoundGenerator/Public/MetasoundGeneratorModule.h b/Engine/Plugins/Runtime/Metasound/Source/MetasoundGenerator/Public/MetasoundGeneratorModule.h
new file mode 100644
--- /dev/null
+++ b/Engine/Plugins/Runtime/Metasound/Source/MetasoundGenerator/Public/MetasoundGeneratorModule.h
@@ -0,0 +1,17 @@
+// Copyright Epic Games, Inc. All Rights Reserved.
+
+#pragma once
+
+#include "CoreMinimal.h"
+#include "Modules/ModuleInterface.h"
+
+DECLARE_LOG_CATEGORY_EXTERN(LogMetasoundGenerator, Log, All);
+
+namespace Metasound
+{
+	class METASOUNDGENERATOR_API IMetasoundGeneratorModule : public IModuleInterface
+	{
+	};
+}
+
+
diff --git a/Engine/Plugins/Runtime/Metasound/Source/MetasoundGraphCore/MetasoundGraphCore.Build.cs b/Engine/Plugins/Runtime/Metasound/Source/MetasoundGraphCore/MetasoundGraphCore.Build.cs
new file mode 100644
--- /dev/null
+++ b/Engine/Plugins/Runtime/Metasound/Source/MetasoundGraphCore/MetasoundGraphCore.Build.cs
@@ -0,0 +1,22 @@
+// Copyright Epic Games, Inc. All Rights Reserved.
+
+using UnrealBuildTool;
+using System.IO;
+
+namespace UnrealBuildTool.Rules
+{
+	public class MetasoundGraphCore : ModuleRules
+	{
+        public MetasoundGraphCore(ReadOnlyTargetRules Target) : base(Target)
+		{
+			//OptimizeCode = CodeOptimization.Never;
+
+            PublicDependencyModuleNames.AddRange(
+				new string[] {
+                    "Core",
+					"SignalProcessing"
+                }
+            );
+		}
+	}
+}
diff --git a/Engine/Plugins/Runtime/Metasound/Source/MetasoundGraphCore/Private/MetasoundADSRNode.cpp b/Engine/Plugins/Runtime/Metasound/Source/MetasoundGraphCore/Private/MetasoundADSRNode.cpp
new file mode 100644
--- /dev/null
+++ b/Engine/Plugins/Runtime/Metasound/Source/MetasoundGraphCore/Private/MetasoundADSRNode.cpp
@@ -0,0 +1,204 @@
+// Copyright Epic Games, Inc. All Rights Reserved.
+
+#include "MetasoundADSRNode.h"
+#include "MetasoundExecutableOperator.h"
+#include "MetasoundDataReferenceTypes.h"
+#include "DSP/Envelope.h"
+
+#define LOCTEXT_NAMESPACE "MetasoundADSRNode"
+
+namespace Metasound
+{
+	struct FADSRDataReferences
+	{
+		FTimeMillisecondsReadRef Attack;
+		FTimeMillisecondsReadRef Decay;
+		FTimeMillisecondsReadRef Sustain;
+		FTimeMillisecondsReadRef Release;
+	};
+
+	class FADSROperator : public TExecutableOperator<FADSROperator>
+	{
+		public:
+			FADSROperator(const FOperatorSettings& InSettings, const FBopReadRef& InBop, const FADSRDataReferences& InADSRData)
+			:	OperatorSettings(InSettings)
+			,	Bop(InBop)
+			,	StopEnvelopePos(-1)
+			,	ADSRDataReferences(InADSRData)
+			,	EnvelopeBuffer(InSettings.FramesPerExecute)
+			{
+				check(EnvelopeBuffer->Num() == InSettings.FramesPerExecute);
+
+				OutputDataReferences.AddDataReadReference(TEXT("Envelope"), FAudioBufferReadRef(EnvelopeBuffer));
+
+				Envelope.Init(OperatorSettings.SampleRate);
+			}
+
+			virtual const FDataReferenceCollection& GetInputs() const override
+			{
+				return InputDataReferences;
+			}
+
+			virtual const FDataReferenceCollection& GetOutputs() const override
+			{
+				return OutputDataReferences;
+			}
+
+			void Execute()
+			{
+				Envelope.SetAttackTime(ADSRDataReferences.Attack->Milliseconds);
+				Envelope.SetDecayTime(ADSRDataReferences.Decay->Milliseconds);
+				Envelope.SetReleaseTime(ADSRDataReferences.Release->Milliseconds);
+
+				float ADSMilliseconds = ADSRDataReferences.Attack->Milliseconds;
+				ADSMilliseconds += ADSRDataReferences.Decay->Milliseconds;
+				ADSMilliseconds += ADSRDataReferences.Sustain->Milliseconds;
+
+				// TODO: Bops need to be sorted.
+				// TODO: Make convenient way to do bop loops like this since it might happen alot.
+				// TODO: Bops need to be done on 4-sample boundaries to allow for SIMD.
+			
+
+				int32 BopIndex = 0;
+				int32 StartPos = 0;
+				int32 NextBop = BopIndex < Bop->Num() ? (*Bop)[BopIndex] : OperatorSettings.FramesPerExecute + 1;
+				int32 EndPos = FMath::Min(OperatorSettings.FramesPerExecute, NextBop);
+				
+				if (StopEnvelopePos > 0)
+				{
+					StopEnvelopePos -= OperatorSettings.FramesPerExecute;
+					EndPos = FMath::Min(StopEnvelopePos, EndPos);
+				}
+
+				float* EnvelopeData = EnvelopeBuffer->GetData();
+
+				// TODO: FramesPerExecute need to be on 4 sample boundaries to allow for SIMD.
+				// TODO: just make FramesPerExecute an int32 instead of a uint32
+				while (StartPos < (int32)OperatorSettings.FramesPerExecute)
+				{
+					int32 i = StartPos;
+					for (; i < EndPos; i++)
+					{
+						EnvelopeData[i] = Envelope.Generate();
+					}
+
+					StartPos = EndPos;
+					
+					if (EndPos == StopEnvelopePos)
+					{
+						// Stop sustain, start release.
+						Envelope.Stop();
+						StopEnvelopePos = -1;
+						EndPos = FMath::Min((int32)OperatorSettings.FramesPerExecute, NextBop);
+					}
+
+					if (EndPos == NextBop)
+					{
+						// Trigger next envelope. 
+						Envelope.Start();
+						StopEnvelopePos = EndPos + FMath::RoundToInt(OperatorSettings.SampleRate * 0.001f * ADSMilliseconds);
+
+						BopIndex++;
+
+						NextBop = BopIndex < Bop->Num() ? (*Bop)[BopIndex] : OperatorSettings.FramesPerExecute + 1;
+						EndPos = FMath::Min(StopEnvelopePos, FMath::Min((int32)OperatorSettings.FramesPerExecute, NextBop));
+					}
+				}
+			}
+
+		private:
+			const FOperatorSettings OperatorSettings;
+
+			// TODO: write envelope gen for metasound more suited to this processing structure. 
+			Audio::FEnvelope Envelope;
+
+			FBopReadRef Bop;
+			int32 StopEnvelopePos;
+			FADSRDataReferences ADSRDataReferences;
+			FAudioBufferWriteRef EnvelopeBuffer;
+
+			FDataReferenceCollection OutputDataReferences;
+			FDataReferenceCollection InputDataReferences;
+	};
+
+	const FName FADSRNode::ClassName = FName(TEXT("ADSR"));
+
+	TUniquePtr<IOperator> FADSRNode::FOperatorFactory::CreateOperator(const INode& InNode, const FOperatorSettings& InOperatorSettings, const FDataReferenceCollection& InInputDataReferences, TArray<TUniquePtr<IOperatorBuildError>>& OutErrors) 
+	{
+		const FADSRNode& ADSRNode = static_cast<const FADSRNode&>(InNode);
+
+		FBopReadRef Bop;
+
+		// TODO: Could no-op this if the bop is not connected.
+		SetReadableRefIfInCollection(TEXT("Bop"), InInputDataReferences, Bop);
+
+		FADSRDataReferences ADSRDataReferences = 
+			{
+				FTimeMillisecondsReadRef(ADSRNode.GetDefaultAttackMs()),
+				FTimeMillisecondsReadRef(ADSRNode.GetDefaultDecayMs()),
+				FTimeMillisecondsReadRef(ADSRNode.GetDefaultSustainMs()),
+				FTimeMillisecondsReadRef(ADSRNode.GetDefaultReleaseMs())
+			};
+
+		// TODO: If none of these are connected, could pregenerate ADSR envelope and return a different operator. 
+		SetReadableRefIfInCollection(TEXT("Attack"), InInputDataReferences, ADSRDataReferences.Attack);
+		SetReadableRefIfInCollection(TEXT("Decay"), InInputDataReferences, ADSRDataReferences.Decay);
+		SetReadableRefIfInCollection(TEXT("Sustain"), InInputDataReferences, ADSRDataReferences.Sustain);
+		SetReadableRefIfInCollection(TEXT("Release"), InInputDataReferences, ADSRDataReferences.Release);
+
+		return MakeUnique<FADSROperator>(InOperatorSettings, Bop, ADSRDataReferences);
+	}
+
+	FADSRNode::FADSRNode(const FString& InName, float InDefaultAttackMs, float InDefaultDecayMs, float InDefaultSustainMs, float InDefaultReleaseMs)
+	:	FNode(InName)
+	,	DefaultAttackMs(InDefaultAttackMs)
+	,	DefaultDecayMs(InDefaultDecayMs)
+	,	DefaultSustainMs(InDefaultSustainMs)
+	,	DefaultReleaseMs(InDefaultReleaseMs)
+	{
+		AddInputDataVertexDescription<FBop>(TEXT("Bop"), LOCTEXT("BopTooltip", "Trigger for envelope."));
+		AddInputDataVertexDescription<FTimeMilliseconds>(TEXT("Attack"), LOCTEXT("AttackTooltip", "Attack time in milliseconds."));
+		AddInputDataVertexDescription<FTimeMilliseconds>(TEXT("Decay"), LOCTEXT("DecayTooltip", "Decay time in milliseconds."));
+		AddInputDataVertexDescription<FTimeMilliseconds>(TEXT("Sustain"), LOCTEXT("SustainTooltip", "Sustain time in milliseconds."));
+		AddInputDataVertexDescription<FTimeMilliseconds>(TEXT("Release"), LOCTEXT("ReleaseTooltip", "Release time in milliseconds."));
+
+		AddOutputDataVertexDescription<FAudioBuffer>(TEXT("Envelope"), LOCTEXT("EnvelopeTooltip", "The output envelope"));
+	}
+
+	FADSRNode::~FADSRNode()
+	{
+	}
+
+
+	float FADSRNode::GetDefaultAttackMs() const
+	{
+		return DefaultAttackMs;
+	}
+
+	float FADSRNode::GetDefaultDecayMs() const
+	{
+		return DefaultDecayMs;
+	}
+
+	float FADSRNode::GetDefaultSustainMs() const
+	{
+		return DefaultSustainMs;
+	}
+
+	float FADSRNode::GetDefaultReleaseMs() const
+	{
+		return DefaultReleaseMs;
+	}
+
+	const FName& FADSRNode::GetClassName() const
+	{
+		return ClassName;
+	}
+
+	IOperatorFactory& FADSRNode::GetDefaultOperatorFactory() 
+	{
+		return Factory;
+	}
+}
+
+#undef LOCTEXT_NAMESPACE
diff --git a/Engine/Plugins/Runtime/Metasound/Source/MetasoundGraphCore/Private/MetasoundAudioMultiplyNode.cpp b/Engine/Plugins/Runtime/Metasound/Source/MetasoundGraphCore/Private/MetasoundAudioMultiplyNode.cpp
new file mode 100644
--- /dev/null
+++ b/Engine/Plugins/Runtime/Metasound/Source/MetasoundGraphCore/Private/MetasoundAudioMultiplyNode.cpp
@@ -0,0 +1,114 @@
+// Copyright Epic Games, Inc. All Rights Reserved.
+
+#include "MetasoundAudioMultiplyNode.h"
+#include "MetasoundExecutableOperator.h"
+#include "MetasoundDataReferenceTypes.h"
+#include "DSP/BufferVectorOperations.h"
+
+#define LOCTEXT_NAMESPACE "MetasoundAudioMultiplyNode"
+
+namespace Metasound
+{
+
+	class FAudioMultiplyOperator : public TExecutableOperator<FAudioMultiplyOperator>
+	{
+		public:
+			FAudioMultiplyOperator(const FOperatorSettings& InSettings, const FAudioBufferReadRef& InBuffer1, const FAudioBufferReadRef& InBuffer2)
+			:	OperatorSettings(InSettings)
+			,	InputBuffer1(InBuffer1)
+			,	InputBuffer2(InBuffer2)
+			,	OutputBuffer(InSettings.FramesPerExecute)
+			{
+				check(OutputBuffer->Num() == InSettings.FramesPerExecute);
+				check(InputBuffer1->Num() == InSettings.FramesPerExecute);
+				check(InputBuffer2->Num() == InSettings.FramesPerExecute);
+
+				OutputDataReferences.AddDataReadReference(TEXT("Audio"), FAudioBufferReadRef(OutputBuffer));
+			}
+
+			virtual const FDataReferenceCollection& GetInputs() const override
+			{
+				return InputDataReferences;
+			}
+
+			virtual const FDataReferenceCollection& GetOutputs() const override
+			{
+				return OutputDataReferences;
+			}
+
+			void Execute()
+			{
+				using namespace Audio;
+
+				// TODO: add buffer vector op to multipy two buffers (not in place). 
+				// TODO: what to do about RESTRICT? This could be something for the builder to keep in mind.
+				FMemory::Memcpy(OutputBuffer->GetData(), InputBuffer1->GetData(), sizeof(float) * OperatorSettings.FramesPerExecute);
+				MultiplyBuffersInPlace(*InputBuffer2, *OutputBuffer);
+			}
+
+		private:
+			const FOperatorSettings OperatorSettings;
+
+			FAudioBufferReadRef InputBuffer1;
+			FAudioBufferReadRef InputBuffer2;
+			FAudioBufferWriteRef OutputBuffer;
+			FDataReferenceCollection OutputDataReferences;
+			FDataReferenceCollection InputDataReferences;
+	};
+
+	const FName FAudioMultiplyNode::ClassName = FName(TEXT("AudioMultiply"));
+
+	TUniquePtr<IOperator> FAudioMultiplyNode::FOperatorFactory::CreateOperator(const INode& InNode, const FOperatorSettings& InOperatorSettings, const FDataReferenceCollection& InInputDataReferences, TArray<TUniquePtr<IOperatorBuildError>>& OutErrors) 
+	{
+		const FAudioMultiplyNode& AudioMultiplyNode = static_cast<const FAudioMultiplyNode&>(InNode);
+
+		FAudioBufferReadRef InputBuffer1;
+		FAudioBufferReadRef InputBuffer2;
+
+		FAudioBuffer Ones(InOperatorSettings.FramesPerExecute);
+
+		// initialize default array to all ones. 
+		float* Data = Ones.GetData();
+		for (int32 i = 0; i < InOperatorSettings.FramesPerExecute; i++)
+		{
+			Data[i] = 1.f;
+		}
+
+		if (!SetReadableRefIfInCollection(TEXT("InputBuffer1"), InInputDataReferences, InputBuffer1))
+		{
+			InputBuffer1 = FAudioBufferReadRef(Ones);
+		}
+
+		if (!SetReadableRefIfInCollection(TEXT("InputBuffer2"), InInputDataReferences, InputBuffer2))
+		{
+			InputBuffer2 = FAudioBufferReadRef(Ones);
+		}
+
+		return MakeUnique<FAudioMultiplyOperator>(InOperatorSettings, InputBuffer1, InputBuffer2);
+	}
+
+	FAudioMultiplyNode::FAudioMultiplyNode(const FString& InName)
+	:	FNode(InName)
+	{
+		AddInputDataVertexDescription<FAudioBuffer>(TEXT("InputBuffer1"), LOCTEXT("InputBuffer1Tooltip", "Input buffer to multiply."));
+		AddInputDataVertexDescription<FAudioBuffer>(TEXT("InputBuffer2"), LOCTEXT("InputBuffer2Tooltip", "Input buffer to multiply."));
+
+		AddOutputDataVertexDescription<FAudioBuffer>(TEXT("Audio"), LOCTEXT("OutpuBufferTooltip", "The output audio."));
+	}
+
+	FAudioMultiplyNode::~FAudioMultiplyNode()
+	{
+	}
+
+	const FName& FAudioMultiplyNode::GetClassName() const
+	{
+		return ClassName;
+	}
+
+	IOperatorFactory& FAudioMultiplyNode::GetDefaultOperatorFactory() 
+	{
+		return Factory;
+	}
+}
+
+#undef LOCTEXT_NAMESPACE //MetasoundAudioMultiplyNode
diff --git a/Engine/Plugins/Runtime/Metasound/Source/MetasoundGraphCore/Private/MetasoundDataReferenceCollection.cpp b/Engine/Plugins/Runtime/Metasound/Source/MetasoundGraphCore/Private/MetasoundDataReferenceCollection.cpp
new file mode 100644
--- /dev/null
+++ b/Engine/Plugins/Runtime/Metasound/Source/MetasoundGraphCore/Private/MetasoundDataReferenceCollection.cpp
@@ -0,0 +1,77 @@
+// Copyright Epic Games, Inc. All Rights Reserved.
+
+
+#include "MetasoundDataReferenceCollection.h"
+
+#include "CoreMinimal.h"
+#include "MetasoundDataReference.h"
+
+namespace Metasound
+{
+	bool FDataReferenceCollection::ContainsDataReadReference(const FString& InName, const FName& InTypeName) const
+	{
+		if (DataReadRefMap.Contains(InName))
+		{
+			const FDataRefWrapper& Wrapper = DataReadRefMap[InName];
+			if (Wrapper.IsValid())
+			{
+				return InTypeName == Wrapper.GetDataReference()->GetDataTypeName();
+			}
+		}
+
+		return false;
+	}
+
+	bool FDataReferenceCollection::ContainsDataWriteReference(const FString& InName, const FName& InTypeName) const
+	{
+		if (DataWriteRefMap.Contains(InName))
+		{
+			const FDataRefWrapper& Wrapper = DataWriteRefMap[InName];
+			if (Wrapper.IsValid())
+			{
+				return InTypeName == Wrapper.GetDataReference()->GetDataTypeName();
+			}
+		}
+
+		return false;
+	}
+
+	bool FDataReferenceCollection::AddDataReadReferenceFrom(const FString& InName, const FDataReferenceCollection& OtherCollection, const FString& OtherName, const FName& OtherTypeName)
+	{
+		if (OtherCollection.ContainsDataReadReference(OtherName, OtherTypeName))
+		{
+			DataReadRefMap.Add(InName, OtherCollection.DataReadRefMap[OtherName]);
+
+			return true;
+		}
+
+		return false;
+	}
+
+	bool FDataReferenceCollection::AddDataWriteReferenceFrom(const FString& InName, const FDataReferenceCollection& OtherCollection, const FString& OtherName, const FName& OtherTypeName)
+	{
+		if (OtherCollection.ContainsDataWriteReference(OtherName, OtherTypeName))
+		{
+			DataWriteRefMap.Add(InName, OtherCollection.DataWriteRefMap[OtherName]);
+
+			return true;
+		}
+
+		return false;
+	}
+
+	const IDataReference* FDataReferenceCollection::GetDataReference(const FDataReferenceMap& InMap, const FString& InName) const
+	{
+		if (InMap.Contains(InName))
+		{
+			return InMap[InName].GetDataReference();
+		}
+
+		return nullptr;
+	}
+
+	void FDataReferenceCollection::AddDataReference(FDataReferenceMap& InMap, const FString& InName, FDataRefWrapper&& InDataRef)
+	{
+		InMap.Add(InName, InDataRef);
+	}
+}
diff --git a/Engine/Plugins/Runtime/Metasound/Source/MetasoundGraphCore/Private/MetasoundGraph.cpp b/Engine/Plugins/Runtime/Metasound/Source/MetasoundGraphCore/Private/MetasoundGraph.cpp
new file mode 100644
--- /dev/null
+++ b/Engine/Plugins/Runtime/Metasound/Source/MetasoundGraphCore/Private/MetasoundGraph.cpp
@@ -0,0 +1,69 @@
+// Copyright Epic Games, Inc. All Rights Reserved.
+
+#include "MetasoundGraph.h"
+
+namespace Metasound
+{
+	const FName FGraph::ClassName = FName(TEXT("Graph"));
+
+	FGraph::FGraph(const FString& InDescription)
+	:	Description(InDescription)
+	{
+	}
+
+	FGraph::~FGraph()
+	{
+	}
+
+	const FString& FGraph::GetDescription() const
+	{
+		return Description;
+	}
+
+	const FName& FGraph::GetClassName() const
+	{
+		return ClassName;
+	}
+
+	const TArray<FInputDataVertexDescription>& FGraph::GetInputs() const
+	{
+		return InputDescriptions;
+	}
+
+	const TArray<FOutputDataVertexDescription>& FGraph::GetOutputs() const
+	{
+		return OutputDescriptions;
+	}
+
+	void FGraph::AddInputDataVertex(const FInputDataVertex& InVertex)
+	{
+		InputVertices.Add(InVertex);
+		InputDescriptions.Add(InVertex.Description);
+	}
+
+	const TArray<FInputDataVertex>& FGraph::GetInputDataVertices() const
+	{
+		return InputVertices;
+	}
+
+	void FGraph::AddOutputDataVertex(const FOutputDataVertex& InVertex)
+	{
+		OutputVertices.Add(InVertex);
+		OutputDescriptions.Add(InVertex.Description);
+	}
+
+	const TArray<FOutputDataVertex>& FGraph::GetOutputDataVertices() const
+	{
+		return OutputVertices;
+	}
+
+	void FGraph::AddDataEdge(const FDataEdge& InEdge)
+	{
+		Edges.Add(InEdge);
+	}
+
+	const TArray<FDataEdge>& FGraph::GetDataEdges() const
+	{
+		return Edges;
+	}
+}
diff --git a/Engine/Plugins/Runtime/Metasound/Source/MetasoundGraphCore/Private/MetasoundGraphCoreModule.cpp b/Engine/Plugins/Runtime/Metasound/Source/MetasoundGraphCore/Private/MetasoundGraphCoreModule.cpp
new file mode 100644
--- /dev/null
+++ b/Engine/Plugins/Runtime/Metasound/Source/MetasoundGraphCore/Private/MetasoundGraphCoreModule.cpp
@@ -0,0 +1,22 @@
+// Copyright Epic Games, Inc. All Rights Reserved.
+
+#pragma once
+
+#include "MetasoundGraphCoreModule.h"
+
+#include "CoreMinimal.h"
+#include "Modules/ModuleInterface.h"
+#include "Modules/ModuleManager.h"
+
+DEFINE_LOG_CATEGORY(LogMetasoundGraphCore);
+
+namespace Metasound 
+{
+	class FMetasoundGraphCoreModule : public IMetasoundGraphCoreModule
+	{
+	};
+}
+
+IMPLEMENT_MODULE(Metasound::FMetasoundGraphCoreModule, MetasoundGraphCore);
+
+
diff --git a/Engine/Plugins/Runtime/Metasound/Source/MetasoundGraphCore/Private/MetasoundGraphOperator.cpp b/Engine/Plugins/Runtime/Metasound/Source/MetasoundGraphCore/Private/MetasoundGraphOperator.cpp
new file mode 100644
--- /dev/null
+++ b/Engine/Plugins/Runtime/Metasound/Source/MetasoundGraphCore/Private/MetasoundGraphOperator.cpp
@@ -0,0 +1,72 @@
+// Copyright Epic Games, Inc. All Rights Reserved.
+
+#include "MetasoundGraphOperator.h"
+
+#include "CoreMinimal.h"
+#include "MetasoundDataReference.h"
+#include "MetasoundDataReferenceCollection.h"
+#include "MetasoundOperatorInterface.h"
+#include "MetasoundExecutableOperator.h"
+
+namespace Metasound
+{
+
+	FGraphOperator::~FGraphOperator()
+	{
+	}
+
+	void FGraphOperator::AppendOperator(FOperatorPtr InOperator)
+	{
+		if (InOperator.IsValid())
+		{
+			if (nullptr != InOperator->GetExecuteFunction())
+			{
+				OperatorStack.Emplace(MoveTemp(InOperator));
+			}
+		}
+	}
+
+	void FGraphOperator::SetInputs(const FDataReferenceCollection& InCollection)
+	{
+		Inputs = InCollection;
+	}
+
+	void FGraphOperator::SetOutputs(const FDataReferenceCollection& InCollection)
+	{
+		Outputs = InCollection;
+	}
+
+	const FDataReferenceCollection& FGraphOperator::GetInputs() const
+	{
+		return Inputs;
+	}
+
+	const FDataReferenceCollection& FGraphOperator::GetOutputs() const
+	{
+		return Outputs;
+	}
+
+	IOperator::FExecuteFunction FGraphOperator::GetExecuteFunction()
+	{
+		return &FGraphOperator::ExecuteFunction;
+	}
+
+	void FGraphOperator::Execute()
+	{
+		FExecutableOperatorWrapper* StackPtr = OperatorStack.GetData();
+		const int32 Num = OperatorStack.Num();
+		for (int32 i = 0; i < Num; i++)
+		{
+			StackPtr[i].Execute();
+		}
+	}
+
+	void FGraphOperator::ExecuteFunction(IOperator* InOperator)
+	{
+		FGraphOperator* Operator = static_cast<FGraphOperator*>(InOperator);
+
+		check(nullptr != Operator);
+
+		Operator->Execute();
+	}
+}
diff --git a/Engine/Plugins/Runtime/Metasound/Source/MetasoundGraphCore/Private/MetasoundNode.cpp b/Engine/Plugins/Runtime/Metasound/Source/MetasoundGraphCore/Private/MetasoundNode.cpp
new file mode 100644
--- /dev/null
+++ b/Engine/Plugins/Runtime/Metasound/Source/MetasoundGraphCore/Private/MetasoundNode.cpp
@@ -0,0 +1,52 @@
+// Copyright Epic Games, Inc. All Rights Reserved.
+
+
+#include "MetasoundNode.h"
+#include "CoreMinimal.h"
+
+namespace Metasound
+{
+	FNode::FNode(const FString& InDescription)
+	:	Description(InDescription)
+	{
+	}
+
+	FNode::~FNode()
+	{
+	}
+
+	const FString& FNode::GetDescription() const
+	{
+		return Description;
+	}
+
+	const TArray<FInputDataVertexDescription>& FNode::GetInputs() const
+	{
+		return Inputs;
+	}
+
+	const TArray<FOutputDataVertexDescription>& FNode::GetOutputs() const
+	{
+		return Outputs;
+	}
+
+	void FNode::AddInputDataVertexDescription(const FInputDataVertexDescription& InDescription)
+	{
+		Inputs.Add(InDescription);
+	}
+
+	void FNode::RemoveInputDataVertexDescription(const FInputDataVertexDescription& InDescription)
+	{
+		Inputs.RemoveSingle(InDescription);
+	}
+
+	void FNode::AddOutputDataVertexDescription(const FOutputDataVertexDescription& InDescription)
+	{
+		Outputs.Add(InDescription);
+	}
+
+	void FNode::RemoveOutputDataVertexDescription(const FOutputDataVertexDescription& InDescription)
+	{
+		Outputs.RemoveSingle(InDescription);
+	}
+}
diff --git a/Engine/Plugins/Runtime/Metasound/Source/MetasoundGraphCore/Private/MetasoundNodeInterface.cpp b/Engine/Plugins/Runtime/Metasound/Source/MetasoundGraphCore/Private/MetasoundNodeInterface.cpp
new file mode 100644
--- /dev/null
+++ b/Engine/Plugins/Runtime/Metasound/Source/MetasoundGraphCore/Private/MetasoundNodeInterface.cpp
@@ -0,0 +1,21 @@
+// Copyright Epic Games, Inc. All Rights Reserved.
+
+#include "MetasoundNodeInterface.h"
+
+namespace Metasound
+{
+	bool operator==(const FDataVertexDescription& LHS, const FDataVertexDescription& RHS)
+	{
+		bool bIsEqual = /*(LHS.Source == RHS.Source) && */
+			(LHS.VertexName == RHS.VertexName) &&
+			(LHS.DataReferenceTypeName == RHS.DataReferenceTypeName) &&
+			LHS.Tooltip.EqualTo(RHS.Tooltip);
+
+		return bIsEqual;
+	}
+
+	bool operator!=(const FDataVertexDescription& LHS, const FDataVertexDescription& RHS)
+	{
+		return !(LHS == RHS);
+	}
+}
diff --git a/Engine/Plugins/Runtime/Metasound/Source/MetasoundGraphCore/Private/MetasoundOperatorBuilder.cpp b/Engine/Plugins/Runtime/Metasound/Source/MetasoundGraphCore/Private/MetasoundOperatorBuilder.cpp
new file mode 100644
--- /dev/null
+++ b/Engine/Plugins/Runtime/Metasound/Source/MetasoundGraphCore/Private/MetasoundOperatorBuilder.cpp
@@ -0,0 +1,312 @@
+// Copyright Epic Games, Inc. All Rights Reserved.
+
+#include "MetasoundOperatorBuilder.h"
+#include "MetasoundBuilderInterface.h"
+#include "MetasoundOperatorInterface.h"
+#include "MetasoundGraphOperator.h"
+
+namespace Metasound
+{
+	FOperatorBuilder::FOperatorBuilder(const FOperatorSettings& InSettings)
+	:	OperatorSettings(InSettings)
+	{
+	}
+
+	FOperatorBuilder::~FOperatorBuilder()
+	{
+	}
+
+	TUniquePtr<IOperator> FOperatorBuilder::BuildGraphOperator(const IGraph& InGraph, TArray<FBuildErrorPtr>& OutErrors)
+	{
+		TArray<INode*> SortedNodes;
+
+		FNodeEdgeMultiMap NodeInputEdges;
+
+		bool bSuccess = true;
+
+
+		bSuccess = GroupInputEdges(InGraph.GetDataEdges(), NodeInputEdges, OutErrors);
+		if (!bSuccess)
+		{
+			// TODO: add error to error stack.
+			return TUniquePtr<IOperator>(nullptr);
+		}
+
+		// TODO: Validate edges.
+
+		// TODO: need to include inputs/outputs vertices here. 
+		bSuccess = TopologicalSort(InGraph.GetDataEdges(), SortedNodes, OutErrors);
+		if (!bSuccess)
+		{
+			// TODO: add error to error stack.
+			return TUniquePtr<IOperator>(nullptr);
+		}
+
+		// TODO: Add graph pruning step.
+
+		TArray<TUniquePtr<IOperator>> Operators;
+		FNodeDataReferenceMap NodeData;
+
+		bSuccess = CreateOperators(SortedNodes, NodeInputEdges, Operators, NodeData, OutErrors);
+		if (!bSuccess)
+		{
+			// TODO: add error to error stack.
+			return TUniquePtr<IOperator>(nullptr);
+		}
+
+		return CreateGraphOperator(InGraph, Operators, NodeData, OutErrors);
+	}
+
+	bool FOperatorBuilder::GroupInputEdges(const TArray<FDataEdge>& InEdges, FNodeEdgeMultiMap& OutNodeInputs, TArray<FBuildErrorPtr>& OutErrors) const
+	{
+		for (const FDataEdge& Edge : InEdges)
+		{
+			OutNodeInputs.Add(Edge.To.Node, &Edge);
+		}
+
+		// TODO: error on duplicate inputs.
+		
+		return true;
+	}
+
+	bool FOperatorBuilder::TopologicalSort(const TArray<FDataEdge>& InEdges, TArray<INode*>& OutNodes, TArray<FBuildErrorPtr>& OutErrors) const
+	{
+		using FNodePair = TPair<INode*, INode*>;
+		using FNodeMultiMap = TMultiMap<INode*, INode*>;
+		// TODO: Implement depth first topological sort to imporove cache coherency.
+		
+		TArray<INode*> UniqueNodes;
+		FNodeMultiMap Dependencies;
+
+		for (const FDataEdge& Edge : InEdges)
+		{
+			if (nullptr == Edge.To.Node)
+			{
+				// TODO: add build error for unconnected edge.
+				continue;
+			}
+			if (nullptr == Edge.From.Node)
+			{
+				// TODO: add build error for unconnected edge.
+				continue;
+			}
+
+			Dependencies.AddUnique(Edge.To.Node, Edge.From.Node);
+			UniqueNodes.AddUnique(Edge.To.Node);
+			UniqueNodes.AddUnique(Edge.From.Node);
+		}
+
+
+		// Sort graph so that nodes with no dependencies always go first.
+		while (UniqueNodes.Num() > 0)
+		{
+			TArray<INode*> IndependentNodes = UniqueNodes.FilterByPredicate([&](INode* InNode) { 
+					return (Dependencies.Num(InNode) == 0);
+			});
+
+			if (0 == IndependentNodes.Num())
+			{
+				// TODO: add build error. likely a cycle in the graph.
+				return false;
+			}
+
+			FNodeMultiMap UpdatedDependencies;
+			
+			// Remove independent nodes from dependency map.
+			for (FNodePair& Element : Dependencies)
+			{
+				if (!IndependentNodes.Contains(Element.Value))
+				{
+					UpdatedDependencies.AddUnique(Element.Key, Element.Value);
+				}
+			}
+			Dependencies = UpdatedDependencies;
+
+			// Remove independent nodes from node list
+			for (INode* Node : IndependentNodes)
+			{
+				UniqueNodes.RemoveSwap(Node);
+			}
+
+			// Add independent nodes to output.
+			OutNodes.Append(IndependentNodes);
+		}
+
+		return true;
+	}
+
+	bool FOperatorBuilder::GatherInputDataReferences(const INode* InNode, const FNodeEdgeMultiMap& InEdgeMap, const FNodeDataReferenceMap& InDataReferenceMap, FDataReferenceCollection& OutCollection, TArray<FBuildErrorPtr>& OutErrors) const
+	{
+		// Find all input edges associated with this node. 
+		TArray<const FDataEdge*> Edges;
+		InEdgeMap.MultiFind(InNode, Edges);
+
+		// Add parameters to collection based on edge description
+		for (const FDataEdge* Edge : Edges)
+		{
+			if (!InDataReferenceMap.Contains(Edge->From.Node))
+			{
+				// TODO: Add to build error.
+				// TODO: This is likely due to a failed topological sort and is more of an internal error
+				// than a user error.
+				return false;
+			}
+
+			
+			const FDataReferenceCollection& FromDataReferenceCollection = InDataReferenceMap[Edge->From.Node].Outputs;
+
+			if (!FromDataReferenceCollection.ContainsDataReadReference(Edge->From.Description.VertexName, Edge->From.Description.DataReferenceTypeName))
+			{
+				// TODO: This is likely a node programming error where the edges reported by the INode interface
+				// did not match the readable parameter refs created by the operators outputs. Or, the edge description is invalid.
+				// TODO: add build error to stack
+				// TODO: consider checking that outputs of node operators match output descriptions of nodes. 
+				// TODO: consider checking the edges are supported by nodes existing parameter descriptions.
+				return false;
+			}
+
+			if (Edge->From.Description.DataReferenceTypeName != Edge->To.Description.DataReferenceTypeName)
+			{
+				// TODO: add build error.
+				// TODO: may want to have this check higher up in the stack. 
+				return false;
+			}
+
+
+
+			bool bSuccess = OutCollection.AddDataReadReferenceFrom(Edge->To.Description.VertexName, FromDataReferenceCollection, Edge->From.Description.VertexName, Edge->From.Description.DataReferenceTypeName);
+
+			if (!bSuccess)
+			{
+				// TODO: add build error.
+				return false;
+			}
+		}
+
+		return true;
+	}
+
+	bool FOperatorBuilder::CreateOperators(const TArray<INode*>& InSortedNodes, FNodeEdgeMultiMap& InNodeInputEdges, TArray<TUniquePtr<IOperator>>& OutOperators, FNodeDataReferenceMap& OutDataReferences, TArray<FBuildErrorPtr>& OutErrors)
+	{
+		for (INode* Node : InSortedNodes)
+		{
+			FDataReferenceCollection InputCollection;
+
+			// Gather the input parameters for this IOperator from the output parameters of already created IOperators. 
+			if (!GatherInputDataReferences(Node, InNodeInputEdges, OutDataReferences, InputCollection, OutErrors))
+			{
+				// TODO Is this a separate error, or the stack of the same error. Perhaps this shouldn't add an error
+				// because it's already sending the error downstream. 
+				// TODO: add to build error
+				return false;
+			}
+
+			IOperatorFactory& Factory = Node->GetDefaultOperatorFactory();
+
+			FOperatorPtr Operator = Factory.CreateOperator(*Node, OperatorSettings, InputCollection, OutErrors);
+
+			if (!Operator.IsValid())
+			{
+				// TODO: add vailed operator creation erorr. 
+				// TODO Is this a separate error, or the stack of the same error. Perhaps this shouldn't add an error
+				// because it's already sending the error downstream. 
+				return false;
+			}
+
+			// Save input and output collection for future use 
+			OutDataReferences.Emplace(Node, FOperatorDataReferences(Operator->GetInputs(), Operator->GetOutputs()));
+
+			// Add operator to operator array
+			OutOperators.Add(MoveTemp(Operator));
+		}
+
+		return true;
+	}
+
+	bool FOperatorBuilder::GatherGraphDataReferences(const IGraph& InGraph, FNodeDataReferenceMap& InNodeDataReferences, FDataReferenceCollection& OutGraphInputs, FDataReferenceCollection& OutGraphOutputs, TArray<FBuildErrorPtr>& OutErrors) const
+	{
+		// Gather inputs
+		for (const FInputDataVertex& InputVertex : InGraph.GetInputDataVertices())
+		{
+			if (!InNodeDataReferences.Contains(InputVertex.Node))
+			{
+				// TODO: build error.  input node does not exist.
+				return false;
+			}
+
+			FDataReferenceCollection& Collection = InNodeDataReferences[InputVertex.Node].Inputs;
+
+			if (!Collection.ContainsDataWriteReference(InputVertex.Description.VertexName, InputVertex.Description.DataReferenceTypeName))
+			{
+				// TODO: build error.  parameter ref does not exist
+				return false;
+			}
+
+			bool bSuccess = OutGraphInputs.AddDataWriteReferenceFrom(InputVertex.Description.VertexName, Collection, InputVertex.Description.VertexName, InputVertex.Description.DataReferenceTypeName);
+
+			if (!bSuccess)
+			{
+				// TODO: add build error
+				// Shouldn't be hitting this error because of earlier "ContainsWritable..." call.
+				return false;
+			}
+		}
+
+		// Gather outputs
+		for (const FOutputDataVertex& OutputVertex : InGraph.GetOutputDataVertices())
+		{
+			if (!InNodeDataReferences.Contains(OutputVertex.Node))
+			{
+				// TODO: build error.  input node does not exist.
+				return false;
+			}
+
+			FDataReferenceCollection& Collection = InNodeDataReferences[OutputVertex.Node].Outputs;
+
+			if (!Collection.ContainsDataReadReference(OutputVertex.Description.VertexName, OutputVertex.Description.DataReferenceTypeName))
+			{
+				// TODO: build error.  parameter ref does not exist
+				return false;
+			}
+
+			bool bSuccess = OutGraphOutputs.AddDataReadReferenceFrom(OutputVertex.Description.VertexName, Collection, OutputVertex.Description.VertexName, OutputVertex.Description.DataReferenceTypeName);
+
+			if (!bSuccess)
+			{
+				// TODO: add build error
+				// Shouldn't be hitting this error because of earlier "ContainsWritable..." call.
+				return false;
+			}
+		}
+	
+		return true;
+	}
+
+	TUniquePtr<IOperator> FOperatorBuilder::CreateGraphOperator(const IGraph& InGraph, TArray<FOperatorPtr>& InOperators, FNodeDataReferenceMap& InNodeDataReferences, TArray<FBuildErrorPtr>& OutErrors) const
+	{
+		FDataReferenceCollection GraphInputs;
+		FDataReferenceCollection GraphOutputs;
+
+		bool bSuccess = GatherGraphDataReferences(InGraph, InNodeDataReferences, GraphInputs, GraphOutputs, OutErrors);
+
+		if (!bSuccess)
+		{
+			// TODO: add build error.
+			return TUniquePtr<IOperator>(nullptr);
+		}
+
+		TUniquePtr<FGraphOperator> GraphOperator = MakeUnique<FGraphOperator>();
+
+		GraphOperator->SetInputs(GraphInputs);
+		GraphOperator->SetOutputs(GraphOutputs);
+
+		for (FOperatorPtr& Ptr : InOperators)
+		{
+			GraphOperator->AppendOperator(MoveTemp(Ptr));
+		}
+
+		InOperators.Reset();
+
+		return MoveTemp(GraphOperator);
+	}
+}
diff --git a/Engine/Plugins/Runtime/Metasound/Source/MetasoundGraphCore/Private/MetasoundOscNode.cpp b/Engine/Plugins/Runtime/Metasound/Source/MetasoundGraphCore/Private/MetasoundOscNode.cpp
new file mode 100644
--- /dev/null
+++ b/Engine/Plugins/Runtime/Metasound/Source/MetasoundGraphCore/Private/MetasoundOscNode.cpp
@@ -0,0 +1,104 @@
+// Copyright Epic Games, Inc. All Rights Reserved.
+
+#include "MetasoundOscNode.h"
+#include "MetasoundExecutableOperator.h"
+#include "MetasoundDataReferenceTypes.h"
+
+#define LOCTEXT_NAMESPACE "MetasoundOscNode"
+
+namespace Metasound
+{
+	class FOscOperator : public TExecutableOperator<FOscOperator>
+	{
+		public:
+			FOscOperator(const FOperatorSettings& InSettings, const FFloatReadRef& InFrequency)
+			:	OperatorSettings(InSettings)
+			,	TwoPi(2.f * PI)
+			,	Phase(0.f)
+			,	Frequency(InFrequency)
+			,	AudioBuffer(InSettings.FramesPerExecute)
+			{
+				check(AudioBuffer->Num() == InSettings.FramesPerExecute);
+
+				OutputDataReferences.AddDataReadReference(TEXT("Audio"), FAudioBufferReadRef(AudioBuffer));
+			}
+
+			virtual const FDataReferenceCollection& GetInputs() const override
+			{
+				return InputDataReferences;
+			}
+
+			virtual const FDataReferenceCollection& GetOutputs() const override
+			{
+				return OutputDataReferences;
+			}
+
+			void Execute()
+			{
+				const float PhaseDelta = *Frequency * 2 * PI / OperatorSettings.SampleRate;
+				float* Data = AudioBuffer->GetData();
+
+				for (int32 i = 0; i < OperatorSettings.FramesPerExecute; i++)
+				{
+					Data[i] = FMath::Sin(Phase);
+					Phase += PhaseDelta;
+				}
+
+				Phase -= FMath::FloorToFloat(Phase / TwoPi) * TwoPi;
+			}
+
+		private:
+			const FOperatorSettings OperatorSettings;
+			const float TwoPi;
+			float Phase;
+
+			FFloatReadRef Frequency;
+			FAudioBufferWriteRef AudioBuffer;
+
+			FDataReferenceCollection InputDataReferences;
+			FDataReferenceCollection OutputDataReferences;
+	};
+
+	const FName FOscNode::ClassName = FName(TEXT("Osc"));
+
+	TUniquePtr<IOperator> FOscNode::FOperatorFactory::CreateOperator(const INode& InNode, const FOperatorSettings& InOperatorSettings, const FDataReferenceCollection& InInputDataReferences, TArray<TUniquePtr<IOperatorBuildError>>& OutErrors) 
+	{
+		const FOscNode& OscNode = static_cast<const FOscNode&>(InNode);
+		FFloatReadRef Frequency(OscNode.GetDefaultFrequency());
+
+		if (InInputDataReferences.ContainsDataReadReference<FFloat>(TEXT("Frequency")))
+		{
+			Frequency = InInputDataReferences.GetDataReadReference<FFloat>(TEXT("Frequency"));
+		}
+
+		return MakeUnique<FOscOperator>(InOperatorSettings, Frequency);
+	}
+
+	FOscNode::FOscNode(const FString& InName, float InDefaultFrequency)
+	:	FNode(InName)
+	,	DefaultFrequency(InDefaultFrequency)
+	{
+		AddInputDataVertexDescription<FFloat>(TEXT("Frequency"), LOCTEXT("FrequencyTooltip", "The frequency of oscillator in Hz."));
+		AddOutputDataVertexDescription<FAudioBuffer>(TEXT("Audio"), LOCTEXT("AudioTooltip", "The output audio"));
+	}
+
+	FOscNode::~FOscNode()
+	{
+	}
+
+	float FOscNode::GetDefaultFrequency() const
+	{
+		return DefaultFrequency;
+	}
+
+	const FName& FOscNode::GetClassName() const
+	{
+		return ClassName;
+	}
+
+	IOperatorFactory& FOscNode::GetDefaultOperatorFactory() 
+	{
+		return Factory;
+	}
+}
+#undef LOCTEXT_NAMESPACE //MetasoundOscNode
diff --git a/Engine/Plugins/Runtime/Metasound/Source/MetasoundGraphCore/Private/MetasoundPeriodicBopNode.cpp b/Engine/Plugins/Runtime/Metasound/Source/MetasoundGraphCore/Private/MetasoundPeriodicBopNode.cpp
new file mode 100644
--- /dev/null
+++ b/Engine/Plugins/Runtime/Metasound/Source/MetasoundGraphCore/Private/MetasoundPeriodicBopNode.cpp
@@ -0,0 +1,119 @@
+// Copyright Epic Games, Inc. All Rights Reserved.
+
+#include "MetasoundPeriodicBopNode.h"
+
+#include "CoreMinimal.h"
+#include "MetasoundNode.h"
+#include "MetasoundOperatorInterface.h"
+#include "MetasoundBuilderInterface.h"
+#include "MetasoundNodeInterface.h"
+#include "MetasoundDataReferenceCollection.h"
+#include "MetasoundDataReferenceTypes.h"
+#include "MetasoundExecutableOperator.h"
+
+#define LOCTEXT_NAMESPACE "MetasoundPeriodicBopNode"
+
+namespace Metasound
+{
+	class FPeriodicBopOperator : public TExecutableOperator<FPeriodicBopOperator>
+	{
+		public:
+			static constexpr float MinimumPeriodSeconds = 0.001f;
+			static constexpr float MinimumPeriodSamples = 20.f;
+
+			FPeriodicBopOperator(const FOperatorSettings& InSettings, const FTimeSecondsReadRef& InPeriod)
+			:	OperatorSettings(InSettings)
+			,	Period(InPeriod)
+			,	ExecuteDurationInSamples(InSettings.FramesPerExecute)
+			,	SampleCountdown(0.f)
+			{
+				OutputDataReferences.AddDataReadReference(TEXT("Bop"), FBopReadRef(Bop));
+			}
+
+			virtual const FDataReferenceCollection& GetInputs() const override
+			{
+				return InputDataReferences;
+			}
+
+			virtual const FDataReferenceCollection& GetOutputs() const override
+			{
+				return OutputDataReferences;
+			}
+
+			void Execute()
+			{
+				// Advance internal counter to get rid of old bops.
+				Bop->Advance(OperatorSettings.FramesPerExecute);
+
+				float PeriodInSamples = FMath::Max(Period->Seconds, MinimumPeriodSeconds) * OperatorSettings.SampleRate;
+
+				PeriodInSamples = FMath::Max(PeriodInSamples, MinimumPeriodSamples);
+
+				while ((SampleCountdown - ExecuteDurationInSamples) <= 0.f)
+				{
+					uint32 Frame = FMath::RoundToInt(SampleCountdown);
+					Bop->BopFrame(Frame);
+					SampleCountdown += PeriodInSamples;
+				}
+
+				SampleCountdown -= ExecuteDurationInSamples;
+			}
+
+		private:
+
+			FOperatorSettings OperatorSettings;
+			FBopWriteRef Bop;
+			FTimeSecondsReadRef Period;
+
+			float ExecuteDurationInSamples;
+			float SampleCountdown;
+
+			FDataReferenceCollection InputDataReferences;
+			FDataReferenceCollection OutputDataReferences;
+	};
+
+
+	TUniquePtr<IOperator> FPeriodicBopNode::FOperatorFactory::CreateOperator(const INode& InNode, const FOperatorSettings& InOperatorSettings, const FDataReferenceCollection& InInputDataReferences, TArray<TUniquePtr<IOperatorBuildError>>& OutErrors)
+	{
+		const FPeriodicBopNode& PeriodicBopNode = static_cast<const FPeriodicBopNode&>(InNode);
+		FTimeSecondsReadRef Period(PeriodicBopNode.GetDefaultPeriodInSeconds());
+
+		if (InInputDataReferences.ContainsDataReadReference<FTimeSeconds>(TEXT("Period")))
+		{
+			Period = InInputDataReferences.GetDataReadReference<FTimeSeconds>(TEXT("Period"));
+		}
+
+		return MakeUnique<FPeriodicBopOperator>(InOperatorSettings, Period);
+	}
+
+	const FName FPeriodicBopNode::ClassName = FName(TEXT("PeriodicBop"));
+
+	FPeriodicBopNode::FPeriodicBopNode(const FString& InName, float InDefaultPeriodInSeconds)
+	:	FNode(InName)
+	,	DefaultPeriod(InDefaultPeriodInSeconds)
+	{
+		AddInputDataVertexDescription<FTimeSeconds>(TEXT("Period"), LOCTEXT("PeriodTooltip", "The period of the bops in seconds."));
+		AddOutputDataVertexDescription<FBop>(TEXT("Bop"), LOCTEXT("BopTooltip", "The output bop"));
+	}
+
+	FPeriodicBopNode::~FPeriodicBopNode()
+	{
+	}
+
+	float FPeriodicBopNode::GetDefaultPeriodInSeconds() const
+	{
+		return DefaultPeriod;
+	}
+
+	const FName& FPeriodicBopNode::GetClassName() const
+	{
+		return ClassName;
+	}
+
+	IOperatorFactory& FPeriodicBopNode::GetDefaultOperatorFactory()
+	{
+		return Factory;
+	}
+}
+
+#undef LOCTEXT_NAMESPACE //MetasoundPeriodicBopNode
diff --git a/Engine/Plugins/Runtime/Metasound/Source/MetasoundGraphCore/Public/MetasoundADSRNode.h b/Engine/Plugins/Runtime/Metasound/Source/MetasoundGraphCore/Public/MetasoundADSRNode.h
new file mode 100644
--- /dev/null
+++ b/Engine/Plugins/Runtime/Metasound/Source/MetasoundGraphCore/Public/MetasoundADSRNode.h
@@ -0,0 +1,77 @@
+// Copyright Epic Games, Inc. All Rights Reserved.
+
+#pragma once
+
+#include "CoreMinimal.h"
+#include "MetasoundNode.h"
+#include "MetasoundOperatorInterface.h"
+#include "MetasoundBuilderInterface.h"
+#include "MetasoundDataReferenceCollection.h"
+
+namespace Metasound
+{
+	/** FADSRNode
+	 *
+	 *  Creates an Attack, Decay Sustain, Release audio processer node. 
+	 */
+	class METASOUNDGRAPHCORE_API FADSRNode : public FNode
+	{
+			// The operator factory for this node.
+			class FOperatorFactory : public IOperatorFactory
+			{
+				public:
+					virtual TUniquePtr<IOperator> CreateOperator(const INode& InNode, const FOperatorSettings& InOperatorSettings, const FDataReferenceCollection& InInputDataReferences, TArray<TUniquePtr<IOperatorBuildError>>& OutErrors) override;
+
+					template<typename DataType>
+					void SetReadableRefIfInCollection(const FString& InName, const FDataReferenceCollection& InCollection, TDataReadReference<DataType>& DataRef)
+					{
+						// TODO: add a helper function to FDataReferenceCollection to do a SetDataIfContains.
+						if (InCollection.ContainsDataReadReference<DataType>(InName))
+						{
+							DataRef = InCollection.GetDataReadReference<DataType>(InName);
+						}
+					}
+
+			};
+
+		public:
+			/** Class name for FADSRNode */
+			static const FName ClassName;
+
+			/** FADSR node constructor.
+			 *
+			 * @param InName - Name of this node.
+			 * @param InDefaultAttackMs - Default attack in milliseconds.
+			 * @param InDefaultDecayMs - Default decay in milliseconds.
+			 * @param InDefaultSustainMs - Default sustain in milliseconds.
+			 * @param InDefaultReleaseMs - Default release in milliseconds.
+			 */
+			FADSRNode(const FString& InName, float InDefaultAttackMs, float InDefaultDecayMs, float InDefaultSustainMs, float InDefaultReleaseMs);
+
+			virtual ~FADSRNode();
+
+			/** Return default attack time in milliseconds */
+			float GetDefaultAttackMs() const;
+			/** Return default attack time in milliseconds */
+			float GetDefaultDecayMs() const;
+			/** Return default attack time in milliseconds */
+			float GetDefaultSustainMs() const;
+			/** Return default attack time in milliseconds */
+			float GetDefaultReleaseMs() const;
+
+			/** Returns the type name of the FADSRNode */
+			virtual const FName& GetClassName() const override;
+
+			/** Return a factory for building a metasound operator. */
+			virtual IOperatorFactory& GetDefaultOperatorFactory() override;
+
+		private:
+
+			float DefaultAttackMs;
+			float DefaultDecayMs;
+			float DefaultSustainMs;
+			float DefaultReleaseMs;
+
+			FOperatorFactory Factory;
+	};
+}
diff --git a/Engine/Plugins/Runtime/Metasound/Source/MetasoundGraphCore/Public/MetasoundAudioBuffer.h b/Engine/Plugins/Runtime/Metasound/Source/MetasoundGraphCore/Public/MetasoundAudioBuffer.h
new file mode 100644
--- /dev/null
+++ b/Engine/Plugins/Runtime/Metasound/Source/MetasoundGraphCore/Public/MetasoundAudioBuffer.h
@@ -0,0 +1,134 @@
+// Copyright Epic Games, Inc. All Rights Reserved.
+
+#pragma once
+
+#include "CoreMinimal.h"
+#include "DSP/BufferVectorOperations.h"
+#include "MetasoundGraphCoreModule.h"
+#include "MetasoundDataReference.h"
+
+/** Define which determines whether to check that the size of the audio buffer has not changed since initialization */
+#define METASOUNDGRAPHCORE_CHECKAUDIONUM !(UE_BUILD_SHIPPING || UE_BUILD_TEST)
+
+namespace Metasound 
+{
+	/**  FAudioBuffer
+	 *
+	 * FAudioBuffer is the default buffer for passing audio data between nodes.  It should not be resized. 
+	 * It is compatible with functions accepting const references to AlignedFloatBuffer (const AlignedFloatBuffer&) 
+	 * arguments via an implicit conversion operator which exposes the underlying AlignedFloatBuffer container.
+	 *
+	 */
+	class FAudioBuffer
+	{
+		public:
+			/** Create an FAudioBuffer with a specific number of samples. 
+			 *
+			 * @param InNumSamples - Number of samples in buffer.
+			 */
+			FAudioBuffer(int32 InNumSamples)
+			{
+				Buffer.AddUninitialized(InNumSamples);
+
+#if METASOUNDGRAPHCORE_CHECKAUDIONUM
+				// InitialNum should not change during the life of an FAudioBuffer
+				InitialNum = Buffer.Num();
+#endif				
+			}
+
+			/** Allow all other constructors of Audio::AlignedFloatBuffer to be forwarded to
+			 * underlying container */
+			template <
+				typename... ArgTypes,
+				typename = typename TEnableIf<
+					TAndValue<
+						sizeof...(ArgTypes) != 0,
+						TOrValue<
+							sizeof...(ArgTypes) != 1,
+							TNot<
+								TIsIntegral<
+									typename TDecay<
+										typename TNthTypeFromParameterPack<0, ArgTypes...>::Type
+									>::Type
+								>
+							>
+						>
+					>::Value
+				>::Type
+			>
+			FAudioBuffer(ArgTypes&&... Args)
+			:	Buffer(Forward<ArgTypes>(Args)...)
+			{
+#if METASOUNDGRAPHCORE_CHECKAUDIONUM
+				// InitialNum should not change during the life of an FAudioBuffer
+				InitialNum = Buffer.Num();
+#endif				
+			}
+
+			FAudioBuffer()
+			{
+#if METASOUNDGRAPHCORE_CHECKAUDIONUM
+				// InitialNum should not change during the life of an FAudioBuffer
+				InitialNum = Buffer.Num();
+#endif				
+			}
+
+			/** Return a pointer to the audio float data. */
+			FORCEINLINE const float* GetData() const
+			{
+#if METASOUNDGRAPHCORE_CHECKAUDIONUM
+				UE_CLOG(InitialNum != Buffer.Num(), LogMetasoundGraphCore, Error, TEXT("Metasound audio buffer size change detected.  Audio Buffers should not be resized."));
+#endif				
+				return Buffer.GetData();
+			}
+
+			/** Return a pointer to the audio float data. */
+			FORCEINLINE float* GetData()
+			{
+#if METASOUNDGRAPHCORE_CHECKAUDIONUM
+				UE_CLOG(InitialNum != Buffer.Num(), LogMetasoundGraphCore, Error, TEXT("Metasound audio buffer size change detected.  Audio Buffers should not be resized."));
+#endif				
+				return Buffer.GetData();
+			}
+
+			/** Return the number of samples in the audio buffer. */
+			FORCEINLINE int32 Num() const
+			{
+#if METASOUNDGRAPHCORE_CHECKAUDIONUM
+				UE_CLOG(InitialNum != Buffer.Num(), LogMetasoundGraphCore, Error, TEXT("Metasound audio buffer size change detected.  Audio Buffers should not be resized."));
+#endif				
+				return Buffer.Num();
+			}
+
+			/** Implicit conversion to Audio::AlignedFloatBuffer */
+			FORCEINLINE operator const Audio::AlignedFloatBuffer& () const
+			{
+#if METASOUNDGRAPHCORE_CHECKAUDIONUM
+				UE_CLOG(InitialNum != Buffer.Num(), LogMetasoundGraphCore, Error, TEXT("Metasound audio buffer size change detected.  Audio Buffers should not be resized."));
+#endif				
+				return Buffer;
+			}
+
+			/** Implicit conversion to Audio::AlignedFloatBuffer 
+			 *
+			 * WARNING: if the buffer is resized, it will cause errors. 
+			 */
+			FORCEINLINE operator Audio::AlignedFloatBuffer& ()
+			{
+#if METASOUNDGRAPHCORE_CHECKAUDIONUM
+				UE_CLOG(InitialNum != Buffer.Num(), LogMetasoundGraphCore, Error, TEXT("Metasound audio buffer size change detected.  Audio Buffers should not be resized."));
+#endif				
+				return Buffer;
+			}
+
+
+		private:
+
+			Audio::AlignedFloatBuffer Buffer;
+
+#if METASOUNDGRAPHCORE_CHECKAUDIONUM
+			int32 InitialNum;
+#endif				
+	};
+
+};
diff --git a/Engine/Plugins/Runtime/Metasound/Source/MetasoundGraphCore/Public/MetasoundAudioFormats.h b/Engine/Plugins/Runtime/Metasound/Source/MetasoundGraphCore/Public/MetasoundAudioFormats.h
new file mode 100644
--- /dev/null
+++ b/Engine/Plugins/Runtime/Metasound/Source/MetasoundGraphCore/Public/MetasoundAudioFormats.h
@@ -0,0 +1,337 @@
+// Copyright Epic Games, Inc. All Rights Reserved.
+
+#pragma once
+
+#include "CoreMinimal.h"
+#include "MetasoundDataReference.h"
+#include "MetasoundAudioBuffer.h"
+
+namespace Metasound
+{
+	// TODO:
+	// The 3 flavors of audio
+	//
+	// FVariableChannelAudio - Number of channels can change whenever! Format identifier included.
+	// FMultiChannelAudio - Number of channels set on constructor!
+	// TStaticChannelAudio - Number of channels set on declaration. Used for static channel configurations "Stereo", "FiveDotOne".
+
+	// FAudio
+	// 	- Union of all constant sized 
+	//
+	//
+	// So lets say we want to write our node with generic audio.
+	//
+	//   GenericAudioIn -> Node -> GenericAudioOut -> CastToSpecificAudio.
+	//   	We would have to have a failure mode and/or a specific mapping for each input type. 
+	//   	Not too crazy. This would be an upmix/downmix node that would have to have a basic
+	//   	approach for up/down mixing. 
+	//
+	//
+	//   	Notes from email to self: Type erasure to define a union subtype interface with sharedptr passing semantics.  
+	//
+	//   	Audioref can support many types. All types must support format name and downcasting to self.  And monochannel forwarding. And channel count. And duplicating.  Operator implementations can use multimono convenience template, or specialize on audio subtype. Can also fail, but prolly shouldnt take audio union then. 
+	//
+	
+	class FAudioUnion
+	{
+		class FAudioBase
+		{
+			public:
+				virtual ~FAudioBase() {}
+
+				virtual TArray<const FAudioBufferReadRef>& GetStorage() const = 0;
+
+				virtual TArray<const FAudioBufferWriteRef>& GetStorage() = 0;
+
+				virtual int32 GetMaxNumChannels() const = 0;
+
+		};
+
+		template<typename AudioType>
+		class FAudioWrapper : public FAudioBase
+		{
+			public:
+				FAudioWrapper(AudioType&& InType)
+				:	Audio(MoveTemp(InType))
+				{}
+
+				virtual TArray<const FAudioBufferReadRef>& GetStorage() const override
+				{
+					return Audio.GetStorage();
+				}
+
+
+				virtual TArray<const FAudioBufferWriteRef>& GetStorage() override
+				{
+					return Audio.GetStorage();
+				}
+
+				virtual int32 GetMaxNumChannels() const override
+				{
+					return Audio.GetMaxNumChannels();
+				}
+
+			private:
+				AudioType Audio;
+		};
+
+		public:
+			template<typename AudioType>
+			FAudioUnion(AudioType&& InType)
+			:	AudioPtr(MoveTemp(InType))
+			{
+			}
+
+			TArray<const FAudioBufferReadRef>& GetStorage() const
+			{
+				return AudioPtr->GetStorage();
+			}
+
+			TArray<const FAudioBufferWriteRef>& GetStorage()
+			{
+				return AudioPtr->GetStorage();
+			}
+
+			int32 GetMaxNumChannels() const
+			{
+				return AudioPTr->GetMaxNumChannels();
+			}
+
+			FName GetDataTypeName() const
+			{
+				// TODO: this should just be a call to the owning reference type's existing type information.
+				return AudioPtr->GetDataTypeName();
+			}
+
+			template<typename AudioReferenceType>
+			TDataReadReference<AudioReferenceType> GetReadableRefChecked() const
+			{
+				// TODO: check typename matches audioptr
+				// TODO: think about how this cast actually works. Should it be a clone?
+				// Where's that actual object stored? not clear yet from this implementation.
+				return TDataReadReference<AudioReferenceType>(ParameterReferenceCast<AudioReferenceType>());
+			}
+
+		private:
+			
+			// TODO: This is holding onto the actual object, when it needs to hold on to the reference to the object. 
+			TUniquePtr<FAudioBase> AudioPrt;
+	};
+
+
+
+	// If we are using this format, that we expect the channel count to change during runtime.
+	// TODO: create another one that is sized on construction but cannot change channel count.
+	class FVariableChannelAudio
+	{
+		public:
+			FVariableChannelAudio(int32 InNumFrames, int32 InMaxNumChannels, int32 InNumChannels)
+			:	NumFrames(InNumFrames)
+			,	MaxNumChannels(InMaxNumChannels)
+			,	NumChannels(0)
+			{
+				NumFrames = FMath::Max(NumFrames, 1);
+				MaxNumChannels = FMath::Max(MaxNumChannels, 0);
+
+				for (int32 i = 0; i < MaxNumChannels; i++)
+				{
+					FAudioBufferWriteRef Audio(NumFrames);
+					// TODO: SetZeroed() should be a member func
+					FMemory::Memset(Audio.GetData(), 0, sizeof(float) * NumFrames);
+
+					WritableBufferStorage.Add(Audio);
+					ReadableBufferStorage.Add(Audio);
+				}
+
+				SetNumChannels(InNumChannels)
+			}
+
+
+			FVariableChannelAudio(const TArray<FAudioBufferWriteRef>& InStorage)
+			{
+				// TODO:
+			}
+
+			int32 SetNumChannels(int32 InNumChannels)
+			{
+				// TODO: should we just add more channels if they exist? 
+				NumChannels = FMath::Max(0, FMath::Min(InNumChannels, MaxNumChannels));
+
+				ReadableBuffers = TArrayView<const FReadableAudioBuffer>(ReadableBufferStorage.GetData(), NumChannels);
+				WritableBuffers = TArrayView<const FWritableAudioBuffer>(WritableBufferStorage.GetData(), NumChannels);
+			}
+
+			int32 GetNumChannels() const
+			{
+				return NumChannels;
+			}
+
+			// TODO: where do we store formats?
+			//EFormat EFormat;
+			//
+
+			TArrayView<const FAudioBufferReadRef>& GetBuffers() const
+			{
+				return ReadableBuffers;
+			}
+
+			TArrayView<const FAudioBufferWriteRef>& GetBuffers() 
+			{
+				return WritableBuffers;
+			}
+
+			const TArray<const FAudioBufferReadRef>& GetStorage() const
+			{
+				return ReadableBufferStorage;
+			}
+
+			const TArray<const FAudioBufferWriteRef>& GetStorage()
+			{
+				return WritableBufferStorage;
+			}
+		protected:
+			friend class TDataReadReference<FVariableChannelAudio>;
+
+			FVariableChannelAudio(const TArray<FAudioBufferReadRef>& InStorage)
+			{
+				// TODO: What do we do about write storage in this scenario? We cannot return write data because
+				// we have only been given a read reference. Maybe we could have a DataType specialization
+				// which allows this constructor to exist for readable data types, and not for writable data types. 
+				// Maybe a protected thing with a friend class? Maybe readable parameter has to do a WritableCast<>
+				//
+				// The static sized n channel templates should do something like that as well.
+				//
+				// FVariableChannelAudioReadRef ReadRef(OtherReadRef)
+				// 		
+				//
+				// FVariableChannelAudioWriteRef WriteRef(OtherReadRef) // <- Not allowed!  
+				//
+				// 
+			}
+			
+		private:
+
+
+			int32 NumChannels;
+
+			TArrayView<const FAudioBufferReadRef> ReadableBuffers;
+			TArrayView<const FAudioBufferReadRef> WritableBuffers;
+
+			// The reference storage is const so that the mono audio buffer 
+			// references can be passed forward when the format is changed (as
+			// opposed to making a copy of the reference). This requires that 
+			// this object always be referencing the same underlying audio
+			// buffer objects, hence the const protection to ensure that it 
+			// never gets changed.
+			TArray<const FAudioBufferReadRef> ReadableBufferStorage;
+			TArray<const FAudioBufferWriteRef> WritableBufferStorage;
+	};
+
+
+	
+	template<int32 NumChannels>
+	class TStaticChannelAudio
+	{
+		public:
+			TStaticChannelAudio(int32 InNumFrames)
+			:	NumFrames(InNumFrames)
+			{
+				static_assert(NumChannels > 0, "NumChannels must be greater than zero");
+
+				NumFrames = FMath::Max(NumFrames, 1);
+
+				for (int32 i = 0; i < NumChannels; i++)
+				{
+					FAudioBufferWriteRef Audio(NumFrames);
+					// TODO: SetZeroed() should be a member func
+					FMemory::Memset(Audio.GetData(), 0, sizeof(float) * NumFrames);
+
+					WritableBufferStorage.Add(Audio);
+					ReadableBufferStorage.Add(Audio);
+				}
+
+				WritableBuffers = WritableBufferStorage;
+				ReadableBuffers = ReadableBufferStorage;
+			}
+
+			/* TODO:
+			// Make a constructor that can take individual arguments of buffers.
+			// TStaticChannelAudio<2>(FAudioWriteRef Val1, FAudioWriteRef Val2)
+			template<
+				typename... ArgTypes,
+				typename = typename TEnableIf<
+					TAndValue<
+						sizeof...(ArgTypes) == NumChannels,
+						TAnd< TIsSameType<FAudioWriteRef, typename TDecay<ArgTypes>... >
+					>
+				>
+			>
+			//TStaticChannelAudio(ArgType&&....)
+			*/
+
+			int32 GetNumChannels() const
+			{
+				return NumChannels;
+			}
+
+			template<int32 Index>
+			FAudioBufferReadRef GetBuffer() const
+			{
+				static_assert(Index >= 0, "Index must be within range of channels");
+				static_assert(Index < NumChannels, "Index must be within range of channels");
+
+				return ReadableBuffers.GetData()[Index];
+			}
+
+			template<int32 Index>
+			FAudioBufferWriteRef GetBuffer() 
+			{
+				static_assert(Index >= 0, "Index must be within range of channels");
+				static_assert(Index < NumChannels, "Index must be within range of channels");
+
+				return WritableBuffers.GetData()[Index];
+			}
+
+			TArrayView<const FAudioBufferReadRef> GetBuffers() const
+			{
+				return ReadableBuffers;
+			}
+
+			TArrayView<const FAudioBufferWriteRef> GetBuffers() 
+			{
+				return WritableBuffers;
+			}
+
+			TArrayView<const FAudioBufferReadRef> GetStorage() const
+			{
+				return ReadableBufferStorage;
+			}
+
+			TArrayView<const FAudioBufferWriteRef> GetStorage()
+			{
+				return WritableBufferStorage;
+			}
+
+		private:
+
+			TArray<const FAudioBufferWriteRef> WritableBufferStorage;
+			TArray<const FAudioBufferReadRef> ReadableBufferStorage;
+
+	};
+
+	class FMonoAudio : public TStaticAudio<1>
+	{
+		public:
+			using TStaticAudio<1>::TStaticAudio<1>;
+			using Center = TStaticAudio<1>::GetBuffer<0>;
+	};
+
+	class FStereoAudio : public TStaticAudio<2>
+	{
+		public:
+			using TStaticAudio<2>::TStaticAudio<2>;
+			using Left = TStaticAudio<2>::GetBuffer<0>;
+			using Right = TStaticAudio<2>::GetBuffer<1>;
+	};
+
+}
diff --git a/Engine/Plugins/Runtime/Metasound/Source/MetasoundGraphCore/Public/MetasoundAudioMultiplyNode.h b/Engine/Plugins/Runtime/Metasound/Source/MetasoundGraphCore/Public/MetasoundAudioMultiplyNode.h
new file mode 100644
--- /dev/null
+++ b/Engine/Plugins/Runtime/Metasound/Source/MetasoundGraphCore/Public/MetasoundAudioMultiplyNode.h
@@ -0,0 +1,50 @@
+// Copyright Epic Games, Inc. All Rights Reserved.
+
+#pragma once
+
+#include "CoreMinimal.h"
+#include "MetasoundNode.h"
+#include "MetasoundOperatorInterface.h"
+#include "MetasoundBuilderInterface.h"
+#include "MetasoundDataReferenceCollection.h"
+
+namespace Metasound
+{
+	class METASOUNDGRAPHCORE_API FAudioMultiplyNode : public FNode
+	{
+			class FOperatorFactory : public IOperatorFactory
+			{
+				public:
+					virtual TUniquePtr<IOperator> CreateOperator(const INode& InNode, const FOperatorSettings& InOperatorSettings, const FDataReferenceCollection& InInputDataReferences, TArray<TUniquePtr<IOperatorBuildError>>& OutErrors) override;
+
+					template<typename ParamType>
+					bool SetReadableRefIfInCollection(const FString& InParamName, const FDataReferenceCollection& InCollection, TDataReadReference<ParamType>& ParamRef)
+					{
+						// TODO: add a helper function to FDataReferenceCollection to do a SetParamIfContains.
+						if (InCollection.ContainsDataReadReference<ParamType>(InParamName))
+						{
+							ParamRef = InCollection.GetDataReadReference<ParamType>(InParamName);
+							return true;
+						}
+
+						return false;
+					}
+
+			};
+
+		public:
+			static const FName ClassName;
+
+			FAudioMultiplyNode(const FString& InName);
+
+			virtual ~FAudioMultiplyNode();
+
+			virtual const FName& GetClassName() const override;
+
+			virtual IOperatorFactory& GetDefaultOperatorFactory() override;
+
+		private:
+
+			FOperatorFactory Factory;
+	};
+}
diff --git a/Engine/Plugins/Runtime/Metasound/Source/MetasoundGraphCore/Public/MetasoundBuilderInterface.h b/Engine/Plugins/Runtime/Metasound/Source/MetasoundGraphCore/Public/MetasoundBuilderInterface.h
new file mode 100644
--- /dev/null
+++ b/Engine/Plugins/Runtime/Metasound/Source/MetasoundGraphCore/Public/MetasoundBuilderInterface.h
@@ -0,0 +1,91 @@
+// Copyright Epic Games, Inc. All Rights Reserved.
+
+#pragma once
+
+#include "CoreMinimal.h"
+#include "MetasoundNodeInterface.h"
+#include "MetasoundOperatorInterface.h"
+
+namespace Metasound
+{
+	/** IOperatorBuildError
+	 *
+	 * This interface is intended for errors encountered when building an GraphOperator.
+	 */
+	class IOperatorBuildError
+	{
+		public:
+			virtual ~IOperatorBuildError() {}
+
+			/** Returns the INode associated with the error. */
+			virtual const INode* GetNode() const = 0;
+
+			/** Returns teh FDataEdge associated with the error. */
+			virtual const FDataEdge& GetEdge() const = 0;
+
+			/** Returns the FInputDataVertex associated with the error. */
+			virtual const FInputDataVertex& GetInputDataVertex() const = 0;
+
+			/** Returns the FOutputDataVertex associated with the error. */
+			virtual const FOutputDataVertex& GetOutputDataVertex() const = 0;
+
+			/** Returns the parameter type name associated with the error. */
+			virtual const FString& GetDataReferenceTypeName() const = 0;
+
+			/** Returns the parameter name associated with the error. */
+			virtual const FString& GetDataReferenceName() const = 0;
+
+			/** Returns teh type of error. */
+			virtual const FString& GetErrorType() const = 0; // TODO: FName?
+
+			/** Returns a human readable error description. */
+			virtual const FText& GetErrorDescription() const = 0;
+	};
+
+	/** IOperatorFactory
+	 *
+	 * IOperatorFactory defines an interface for building an IOperator from an INode.  In practice,
+	 * each INode returns its own IOperatorFactory through the INode::GetDefaultOperatorFactory() 
+	 * member function.
+	 */
+	class IOperatorFactory
+	{
+		public:
+			virtual ~IOperatorFactory() {}
+
+			/** Create a new IOperator.
+			 *
+			 * @param InNode - The node associated with this factory and the desired IOperator.
+			 * @param FOperatorSettings - General operator settings for the graph.
+			 * @param InInputCollection - Collection of input parameters available for to this IOperator.
+			 * @param OutErrors - An array of errors. Errors can be added if issues occur while creating the IOperator.
+			 *
+			 * @return A unique pointer to an IOperator. 
+			 */
+			virtual TUniquePtr<IOperator> CreateOperator(const INode& InNode, const FOperatorSettings& InOperatorSettings, const FDataReferenceCollection& InInputDataReferences, TArray<TUniquePtr<IOperatorBuildError>>& OutErrors) = 0;
+	};
+
+
+	/** IOperatorBuilder
+	 *
+	 * Defines an interface for building a graph of operators from a graph of nodes. 
+	 */
+	class IOperatorBuilder 
+	{
+		public:
+			/** A TUniquePtr of an IOperatorBuildError */
+			using FBuildErrorPtr = TUniquePtr<IOperatorBuildError>;
+
+			virtual ~IOperatorBuilder() {}
+
+			/** Build a graph operator from a graph. 
+			 *
+			 * @params InGraph - The input graph object containing edges, input vertices and output vertices.
+			 * @param OutErrors - An array of errors. Errors can be added if issues occur while creating the IOperator.
+			 *
+			 * @return A unique pointer to an IOperator. 
+			 */
+			virtual TUniquePtr<IOperator> BuildGraphOperator(const IGraph& InGraph, TArray<FBuildErrorPtr>& OutErrors) = 0;
+	};
+}
+
diff --git a/Engine/Plugins/Runtime/Metasound/Source/MetasoundGraphCore/Public/MetasoundDataReference.h b/Engine/Plugins/Runtime/Metasound/Source/MetasoundGraphCore/Public/MetasoundDataReference.h
new file mode 100644
--- /dev/null
+++ b/Engine/Plugins/Runtime/Metasound/Source/MetasoundGraphCore/Public/MetasoundDataReference.h
@@ -0,0 +1,342 @@
+// Copyright Epic Games, Inc. All Rights Reserved.
+
+#pragma once
+
+#include "CoreMinimal.h"
+
+
+/** Macros to make declaring a metasound parameter simple.  */
+
+// Declares a metasound parameter type by
+// - Adding typedefs for commonly used template types.
+// - Defining parameter type traits.
+#define DECLARE_METASOUND_DATA_REFERENCE_TYPES(DataType, DataTypeName, DataTypeMagicNumber, DataTypeInfoTypeName, DataReadReferenceTypeName, DataWriteReferenceTypeName) \
+	template<> \
+	struct TDataReferenceTypeInfo<DataType> \
+	{ \
+		static constexpr const TCHAR* TypeName = TEXT(DataTypeName); \
+		static constexpr FDataTypeMagicNumber MagicNumber = (DataTypeMagicNumber); \
+	}; \
+	\
+	typedef TDataReferenceTypeInfo<DataType> DataTypeInfoTypeName; \
+	\
+	typedef TDataReadReference<DataType> DataReadReferenceTypeName; \
+	typedef TDataWriteReference<DataType> DataWriteReferenceTypeName; \
+
+
+
+namespace Metasound
+{
+	/** ID for parameter type. Should be unique for different C++ types. */
+	typedef int32 FDataTypeMagicNumber;
+
+	/** A Data Reference Interface.
+	 *
+	 * A parameter references provides information and access to a shared object in the graph. 
+	 */
+	class IDataReference
+	{
+		public:
+
+		virtual ~IDataReference() 
+		{
+		}
+
+		/** Returns the name of the data type. */
+		virtual const FName& GetDataTypeName() const = 0;
+
+		/** Returns the ID of the parameter type. */
+		virtual FDataTypeMagicNumber GetDataTypeMagicNumber() const = 0;
+
+		/** Creates a copy of the parameter type. */
+		virtual TUniquePtr<IDataReference> Clone() const = 0;
+	};
+
+	/** Info for templated data reference types help perform runtime type 
+	 * verification. 
+	 */
+	template<typename DataType>
+	struct TDataReferenceTypeInfo
+	{
+		/** Name of the data type. Internally, these are converted to FNames.
+		 * Users should be aware of FName case-insensitivity and other 
+		 * limitiations. 
+		 */
+		static constexpr const TCHAR* TypeName = TEXT("");
+
+		/** Magic number used to check the data type when casting. */
+		static constexpr const FDataTypeMagicNumber MagicNumber = -1;
+	};
+	
+	/** Template class for a paramter reference. 
+	 *
+	 * This fulfills the IParamterRef interface, utilizing TDataReferenceTypeInfo to
+	 * define the the TypeName and TypeId of the parameter. 
+	 */
+	template <typename DataType>
+	class TDataReference : public IDataReference
+	{
+		public:
+			typedef TSharedRef<DataType, ESPMode::NotThreadSafe> FRefType;
+
+			typedef TDataReferenceTypeInfo<DataType> FInfoType;
+
+			/** This constructor forwards arguments to the constructor of the underlying DataType.
+			 *
+			 * SFINAE is used to disable copy and move constructors as those are defined separately. 
+			 */
+			template <
+				typename... ArgTypes,
+				typename = typename TEnableIf<
+					TAndValue<
+						sizeof...(ArgTypes) != 0,
+						TOrValue<
+							sizeof...(ArgTypes) != 1,
+							TNot< TIsDerivedFrom< typename TDecay< typename TNthTypeFromParameterPack< 0, ArgTypes... >::Type >::Type, /*typename*/ TDataReference<DataType> > >
+						>
+					>::Value
+				>::Type
+			>
+			TDataReference(ArgTypes&&... Args)
+			:	ObjectReference(MakeShared<DataType, ESPMode::NotThreadSafe>(Forward<ArgTypes>(Args)...))
+			{
+			}
+
+	
+			/** Construct operator with no arguments if the DataType has a default constructor.  */
+			template< typename = typename TEnableIf< TIsConstructible<DataType>::Value >::Type >
+			TDataReference()
+			:	ObjectReference(MakeShared<DataType>())
+			{
+			}
+
+			/** Enable copy constructor */
+			TDataReference(const TDataReference<DataType>& Other) = default;
+
+			/** Enable move constructor */
+			TDataReference(TDataReference<DataType>&& Other) = default;
+
+			/** Enable copy operator */
+			TDataReference<DataType>& operator=(const TDataReference<DataType>& Other) = default;
+
+			/** Enable move operator */
+			TDataReference<DataType>& operator=(TDataReference<DataType>&& Other) = default;
+
+			/** Return the name of the underlying type. */
+			virtual const FName& GetDataTypeName() const override
+			{
+				static const FName TypeName = FName(FInfoType::TypeName);
+				return TypeName;
+			}
+
+			/** Return the ID of the underlying type. */
+			virtual FDataTypeMagicNumber GetDataTypeMagicNumber() const override
+			{
+				return FInfoType::MagicNumber;
+			}
+
+			/** Create a clone of this object. 
+			virtual TUniquePtr<IDataReference> Clone() const override
+			{
+				typedef TDataReference<DataType> FDataReference;
+
+				return MakeUnique< FDataReference >(*this);
+			}
+			*/
+
+		protected:
+
+			// Protected object reference is utilized by subclasses which define what
+			// access is provided to the ObjectReference. 
+			FRefType ObjectReference;
+	};
+
+
+	// Forward declare
+	template <typename DataType>
+	class TDataReadReference;
+
+	/** TDataWriteReference provides write access to a shared parameter reference. */
+	template <typename DataType>
+	class TDataWriteReference : public TDataReference<DataType>
+	{
+		public:
+			typedef TDataReference<DataType> FDataReference;
+
+			/** Construct operator type with arguments to "DataType" constructor. */
+			template<
+				typename... ArgTypes,
+				typename = typename TEnableIf<
+					TAndValue<
+						sizeof...(ArgTypes) != 0,
+						TOrValue<
+							sizeof...(ArgTypes) != 1,
+							TNot< TIsDerivedFrom< typename TDecay< typename TNthTypeFromParameterPack< 0, ArgTypes... >::Type >::Type, /*typename*/ TDataReference<DataType> > >
+						>
+					>::Value
+				>::Type
+			>
+			explicit TDataWriteReference(ArgTypes&&... Args)
+			:	FDataReference(Forward<ArgTypes>(Args)...)
+			{
+			}
+
+	
+			/** Construct operator with no arguments if the DataType has a default constructor.  */
+			template< typename = typename TEnableIf< TIsConstructible<DataType>::Value >::Type >
+			TDataWriteReference()
+			:	FDataReference()
+			{
+			}
+
+			/** Enable copy constructor */
+			TDataWriteReference(const TDataWriteReference<DataType>& Other) = default;
+
+			/** Enable move constructor */
+			TDataWriteReference(TDataWriteReference<DataType>&& Other) = default;
+
+			// TODO: Need to think about whether this protection is needed.
+			/** Disable copy operator. Changes to the internal reference object will break link with
+			 * readable references. */
+			//TDataWriteReference<DataType>& operator=(const TDataWriteReference<DataType>& Other) = delete;
+			TDataWriteReference<DataType>& operator=(const TDataWriteReference<DataType>& Other) = default;
+
+			// TODO: Need to think about whether this protection is needed.
+			/** Disable move operator. Changes to the internal reference object will break link with
+			 * readable references. */
+			//TDataWriteReference<DataType>& operator=(TDataWriteReference<DataType>&& Other) = delete;
+			TDataWriteReference<DataType>& operator=(TDataWriteReference<DataType>&& Other) = default;
+
+			/** Implicit conversion to a readable parameter. */
+			operator TDataReadReference<DataType>() const
+			{
+				return TDataReadReference<DataType>(*this);
+			}
+
+			/** Non-const access to the underlying parameter object. */
+			FORCEINLINE DataType& operator*() const
+			{
+				return *TDataReference<DataType>::ObjectReference;
+				//return *ObjectReference;
+			}
+
+			/** Non-const access to the underlying parameter object. */
+			FORCEINLINE DataType* operator->() const
+			{
+				return TDataReference<DataType>::ObjectReference.operator->();
+				//return ObjectReference.operator->();
+			}
+
+			/** Create a clone of this parameter reference. */
+			virtual TUniquePtr<IDataReference> Clone() const override
+			{
+				typedef TDataWriteReference<DataType> FDataWriteReference;
+
+				return MakeUnique< FDataWriteReference >(*this);
+			}
+
+			// Provide access to ObjectReference when converting from Write to Read.
+			friend class TDataReadReference<DataType>;
+
+			/** Cast a TDataReadReference to a TDataWriteReference. 
+			 *
+			 * In general TDataReadReferences should not be converted into TDataWriteReferences unless the caller
+			 * can be certain that no other TDataWriteReference exists for the underlying parameter. Having multiple 
+			 * TDataWriteReferences to the same parameter can cause confusion behavior as values are overwritten in
+			 * an underterministic fashion.
+			 */
+			friend TDataWriteReference<DataType> WriteCast(const TDataReadReference<DataType>& InReadableRef)
+			{
+				// Friend because it calls protected constructor
+				return TDataWriteReference<DataType>(static_cast<const FDataReference&>(InReadableRef));
+			}
+
+		protected:
+
+			/** Create a writable ref from a blank parameter ref. Should be done with care and understanding
+			 * of side-effects of converting a readable ref to a writable ref. 
+			 */
+			TDataWriteReference<DataType>(const FDataReference& InDataReference)
+			:	FDataReference(InDataReference)
+			{
+			}
+	};
+
+
+	/** TDataReadReference provides read access to a shared parameter reference. */
+	template <typename DataType>
+	class TDataReadReference : public TDataReference<DataType>
+	{
+		public:
+			typedef TDataReference<DataType> FDataReference;
+
+			/** Construct operator type with arguments to "DataType" constructor. */
+			template<
+				typename... ArgTypes,
+				typename = typename TEnableIf<
+					TAndValue<
+						sizeof...(ArgTypes) != 0,
+						TOrValue<
+							sizeof...(ArgTypes) != 1,
+							TNot< TIsDerivedFrom< typename TDecay< typename TNthTypeFromParameterPack< 0, ArgTypes... >::Type >::Type, /*typename*/ TDataReference<DataType> > >
+						>
+					>::Value
+				>::Type
+			>
+			explicit TDataReadReference(ArgTypes&&... Args)
+			:	FDataReference(Forward<ArgTypes>(Args)...)
+			{
+			}
+
+			/** Construct operator with no arguments if the DataType has a default constructor.  */
+			template< typename = typename TEnableIf< TIsConstructible<DataType>::Value >::Type >
+			TDataReadReference()
+			:	FDataReference()
+			{
+			}
+
+			TDataReadReference(const TDataReadReference& Other) = default;
+
+			/** Construct a readable parameter ref from a writable parameter ref. */
+			explicit TDataReadReference(const TDataWriteReference<DataType>& WritableRef)
+			:	FDataReference(WritableRef)
+			{
+			}
+
+			/** Assign a readable parameter ref from a writable parameter ref. */
+			TDataReadReference<DataType>& operator=(const TDataWriteReference<DataType>& Other)
+			{
+				TDataReference<DataType>::ObjectReference = Other.ObjectReference;
+				//ObjectReference = Other.ObjectReference;
+				return *this;
+			}
+
+			/** Enable copy operator */
+			TDataReadReference<DataType>& operator=(const TDataReadReference<DataType>& Other) = default;
+
+			/** Enable move operator */
+			TDataReadReference<DataType>& operator=(TDataReadReference<DataType>&& Other) = default;
+
+			/** Const access to the underlying parameter object. */
+			FORCEINLINE const DataType& operator*() const
+			{
+				return *TDataReference<DataType>::ObjectReference;
+				//return *ObjectReference;
+			}
+
+			/** Const access to the underlying parameter object. */
+			FORCEINLINE const DataType* operator->() const
+			{
+				return TDataReference<DataType>::ObjectReference.operator->();
+				//return ObjectReference.operator->();
+			}
+
+			/** Create a clone of this parameter reference. */
+			virtual TUniquePtr<IDataReference> Clone() const override
+			{
+				typedef TDataReadReference<DataType> FDataReadReference;
+
+				return MakeUnique< FDataReadReference >(*this);
+			}
+	};
+}
diff --git a/Engine/Plugins/Runtime/Metasound/Source/MetasoundGraphCore/Public/MetasoundDataReferenceCollection.h b/Engine/Plugins/Runtime/Metasound/Source/MetasoundGraphCore/Public/MetasoundDataReferenceCollection.h
new file mode 100644
--- /dev/null
+++ b/Engine/Plugins/Runtime/Metasound/Source/MetasoundGraphCore/Public/MetasoundDataReferenceCollection.h
@@ -0,0 +1,252 @@
+// Copyright Epic Games, Inc. All Rights Reserved.
+
+#pragma once
+
+#include "CoreMinimal.h"
+#include "MetasoundDataReference.h"
+
+namespace Metasound
+{
+	class METASOUNDGRAPHCORE_API FDataReferenceCollection
+	{
+		// Convenience class to handle copying and moving of data refs
+		// without knowing data refs derived type.  
+		//
+		// Uses Clone() function as a virtual copy constructor.
+		class FDataRefWrapper 
+		{
+			public:
+				// Anything that implements a IDataReference interface
+				template<typename DataRefType>
+				FDataRefWrapper(const DataRefType& InDataReference)
+				:	DataRefPtr(InDataReference.Clone())
+				{
+				}
+
+				// Override copy constructor to use underlying Clone method.
+				FDataRefWrapper(const FDataRefWrapper& Other)
+				:	DataRefPtr(Other.CloneDataReference())
+				{
+				}
+
+				// Override equal operator to use underlying Clone method.
+				FDataRefWrapper& operator=(const FDataRefWrapper& Other)
+				{
+					DataRefPtr = Other.CloneDataReference();
+					return *this;
+				}
+
+				bool IsValid() const
+				{
+					return DataRefPtr.IsValid();
+				}
+
+				// Get underlying data ref
+				const IDataReference* GetDataReference() const
+				{
+					return DataRefPtr.Get();
+				}
+
+				// Clone underlying data ref
+				TUniquePtr<IDataReference> CloneDataReference() const
+				{
+					if (DataRefPtr.IsValid())
+					{
+						return DataRefPtr->Clone();
+					}
+				
+					return TUniquePtr<IDataReference>();
+				}
+
+			private:
+
+				TUniquePtr<IDataReference> DataRefPtr;
+		};
+
+
+		using FDataReferenceMap = TMap<FString, FDataRefWrapper>;
+
+		public:
+
+			/** Add a readable data reference to the collection.
+			 *
+			 * @param InName - Name of data reference.
+			 * @InDataReference - Readable data reference. 
+			 */
+			template<typename DataType>
+			void AddDataReadReference(const FString& InName, const TDataReadReference<DataType>& InDataReference)
+			{
+				FDataRefWrapper Wrapper(InDataReference);
+
+				AddDataReference(DataReadRefMap, InName, MoveTemp(Wrapper));
+			}
+
+			/** Add a readable data reference to the collection.
+			 *
+			 * @param InName - Name of data reference.
+			 * @InDataReference - Writable data reference. 
+			 */
+			template<typename DataType>
+			void AddDataReadReference(const FString& InName, const TDataWriteReference<DataType>& InDataReference)
+			{
+				TDataReadReference<DataType> Readable(InDataReference);
+
+				AddDataReadReference(InName, Readable);
+			}
+
+			/** Add a writable data reference to the collection.
+			 *
+			 * @param InName - Name of data reference.
+			 * @InDataReference - Writable data reference. 
+			 */
+			template<typename DataType>
+			void AddDataWriteReference(const FString& InName, const TDataWriteReference<DataType>& InDataReference)
+			{
+				FDataRefWrapper Wrapper(InDataReference);
+
+				AddDataReference(DataWriteRefMap, InName, MoveTemp(Wrapper));
+			}
+
+			/** Add a readable data reference from another collection.
+			 *
+			 * @param InLocalName - Name to give data added to this collection.
+			 * @param InOtherCollection - Collection which contains the desired data reference. 
+			 * @param InOtherName - Name of desired data in other collection.
+			 * @param InOtherType - TypeName of desired data.
+			 *
+			 * @return True if the data was successfully added. False otherwise. 
+			 */
+			bool AddDataReadReferenceFrom(const FString& InLocalName, const FDataReferenceCollection& InOtherCollection, const FString& InOtherName, const FName& InOtherTypeName);
+
+			/** Add a writable data reference from another collection.
+			 *
+			 * @param InLocalName - Name to give data added to this collection.
+			 * @param InOtherCollection - Collection which contains the desired data reference. 
+			 * @param InOtherName - Name of desired data in other collection.
+			 * @param InOtherType - TypeName of desired data.
+			 *
+			 * @return True if the data was successfully added. False otherwise. 
+			 */
+			bool AddDataWriteReferenceFrom(const FString& InLocalName, const FDataReferenceCollection& InOtherCollection, const FString& InOtherName, const FName& InOtherTypeName);
+
+			/** Query whether a readable data reference is within the collection. 
+			 *
+			 * @param InName - Name of the data.
+			 * @param InTypeName - TypeName of teh data.
+			 *
+			 * @return True if the data exists in the collection. False otherwise. 
+			 */
+			bool ContainsDataReadReference(const FString& InName, const FName& InTypeName) const;
+
+			/** Query whether a writable data reference is within the collection. 
+			 *
+			 * @param InName - Name of the data.
+			 * @param InTypeName - TypeName of teh data.
+			 *
+			 * @return True if the data exists in the collection. False otherwise. 
+			 */
+			bool ContainsDataWriteReference(const FString& InName, const FName& InTypeName) const;
+
+
+			/** Query whether a readable data reference is within the collection. 
+			 *
+			 * @param InName - Name of the data.
+			 *
+			 * @return True if the data exists in the collection. False otherwise. 
+			 */
+			template<typename DataType>
+			bool ContainsDataReadReference(const FString& InName) const
+			{
+				return Contains<DataType>(DataReadRefMap, InName);
+			}
+
+			/** Query whether a writable data reference is within the collection. 
+			 *
+			 * @param InName - Name of the data.
+			 *
+			 * @return True if the data exists in the collection. False otherwise. 
+			 */
+			template<typename DataType>
+			bool ContainsDataWriteReference(const FString& InName) const
+			{
+				return Contains<DataType>(DataWriteRefMap, InName);
+			}
+
+			/** Returns a readable data ref from the collection.
+			 *
+			 * @param InName - Name of readable data.
+			 *
+			 * @return A readable data reference.
+			 */
+			template<typename DataType>
+			TDataReadReference<DataType> GetDataReadReference(const FString& InName) const
+			{
+				typedef TDataReadReference<DataType> FDataRefType;
+
+				const IDataReference* DataRefPtr = GetDataReference(DataReadRefMap, InName);
+
+				check(nullptr != DataRefPtr);
+
+				check(HasMatchingTypeInfo<DataType>(*DataRefPtr));
+
+				return FDataRefType(*static_cast<const FDataRefType*>(DataRefPtr));
+			}
+
+			/** Returns a writable data ref from the collection.
+			 *
+			 * @param InName - Name of readable data.
+			 *
+			 * @return A readable data reference.
+			 */
+			template<typename DataType>
+			TDataWriteReference<DataType> GetDataWriteReference(const FString& InName) const
+			{
+				typedef TDataWriteReference<DataType> FDataRefType;
+
+				const IDataReference* DataRefPtr = GetDataReference(DataWriteRefMap, InName);
+
+				check(nullptr != DataRefPtr);
+
+				check(HasMatchingTypeInfo<DataType>(*DataRefPtr));
+
+				return FDataRefType(*static_cast<const FDataRefType*>(DataRefPtr));
+			}
+
+
+		private:
+
+			template<typename DataType>
+			bool Contains(const FDataReferenceMap& InMap, const FString& InName) const
+			{
+				if (InMap.Contains(InName))
+				{
+					const IDataReference* DataRefPtr = InMap[InName].GetDataReference();
+
+					if (nullptr != DataRefPtr)
+					{
+						return HasMatchingTypeInfo<DataType>(*DataRefPtr);
+					}
+				}
+
+				return false;
+			}
+
+			template<typename DataType>
+			bool HasMatchingTypeInfo(const IDataReference& InRef) const
+			{
+				static const FName DataTypeName = FName(TDataReferenceTypeInfo<DataType>::TypeName);
+
+				bool bTypeNameMatch = InRef.GetDataTypeName() == DataTypeName;
+				bool bTypeIdMatch = InRef.GetDataTypeMagicNumber() == TDataReferenceTypeInfo<DataType>::MagicNumber;
+
+				return bTypeNameMatch && bTypeIdMatch;
+			}
+
+			const IDataReference* GetDataReference(const FDataReferenceMap& InMap, const FString& InName) const;
+
+			void AddDataReference(FDataReferenceMap& InMap, const FString& InName, FDataRefWrapper&& InDataReference);
+
+			FDataReferenceMap DataReadRefMap;
+			FDataReferenceMap DataWriteRefMap;
+	};
+}
diff --git a/Engine/Plugins/Runtime/Metasound/Source/MetasoundGraphCore/Public/MetasoundDataReferenceTypes.h b/Engine/Plugins/Runtime/Metasound/Source/MetasoundGraphCore/Public/MetasoundDataReferenceTypes.h
new file mode 100644
--- /dev/null
+++ b/Engine/Plugins/Runtime/Metasound/Source/MetasoundGraphCore/Public/MetasoundDataReferenceTypes.h
@@ -0,0 +1,99 @@
+// Copyright Epic Games, Inc. All Rights Reserved.
+
+#pragma once
+
+#include "MetasoundDataReference.h"
+#include "MetasoundAudioBuffer.h"
+
+namespace Metasound
+{
+	class FBop
+	{
+		public:
+			void BopFrame(uint32 InFrameToBop)
+			{
+				// TODO: need to maintain a sorted array. 
+				// Can likely implement a faster version.
+				bool bNeedsSort = false;
+
+				if (Frames.Num() > 0)
+				{
+					bNeedsSort = InFrameToBop < Frames.Last();
+				}
+
+				Frames.Add(InFrameToBop);
+				Frames.Sort();
+			}
+
+			void Advance(uint32 InNumFrames)
+			{
+				Frames.RemoveAllSwap([InNumFrames](uint32 Frame) { return Frame < InNumFrames; }, false /* bAllowShrinking */);
+
+				for (uint32& Frame : Frames)
+				{
+					Frame -= InNumFrames;
+				}
+			}
+
+			int32 Num() const
+			{
+				return Frames.Num();
+			}
+
+			uint32 operator[](int32 Index) const
+			{
+				return Frames[Index];
+			}
+
+			operator bool() const 
+			{
+				return Frames.Num() >= 0;
+			}
+
+			void Reset()
+			{
+				Frames.Reset();
+			}
+
+		private:
+			TArray<uint32> Frames;
+	};
+
+	struct FTimeSeconds
+	{
+		FTimeSeconds()
+		:	Seconds(0.f)
+		{
+		}
+
+		FTimeSeconds(float InSeconds)
+		:	Seconds(InSeconds)
+		{
+		}
+
+		float Seconds;
+	};
+
+	struct FTimeMilliseconds
+	{
+		FTimeMilliseconds()
+		:	Milliseconds(0.f)
+		{
+		}
+
+		FTimeMilliseconds(float InMilliseconds)
+		:	Milliseconds(InMilliseconds)
+		{
+		}
+
+		float Milliseconds;
+	};
+
+	typedef float FFloat;
+
+	DECLARE_METASOUND_DATA_REFERENCE_TYPES(FAudioBuffer, "AudioBuffer", 1, FAudioBufferTypeInfo, FAudioBufferReadRef, FAudioBufferWriteRef);
+	DECLARE_METASOUND_DATA_REFERENCE_TYPES(FBop, "Bop", 2, FBopTypeInfo, FBopReadRef, FBopWriteRef);
+	DECLARE_METASOUND_DATA_REFERENCE_TYPES(FTimeSeconds, "Time:Seconds", 3, FTimeSecondsTypeInfo, FTimeSecondsReadRef, FTimeSecondsWriteRef);
+	DECLARE_METASOUND_DATA_REFERENCE_TYPES(FTimeMilliseconds, "Time:Milliseconds", 4, FTimeMillisecondsTypeInfo, FTimeMillisecondsReadRef, FTimeMillisecondsWriteRef);
+	DECLARE_METASOUND_DATA_REFERENCE_TYPES(FFloat, "Primitive:Float", 5, FFloatTypeInfo, FFloatReadRef, FFloatWriteRef);
+}
diff --git a/Engine/Plugins/Runtime/Metasound/Source/MetasoundGraphCore/Public/MetasoundExecutableOperator.h b/Engine/Plugins/Runtime/Metasound/Source/MetasoundGraphCore/Public/MetasoundExecutableOperator.h
new file mode 100644
--- /dev/null
+++ b/Engine/Plugins/Runtime/Metasound/Source/MetasoundGraphCore/Public/MetasoundExecutableOperator.h
@@ -0,0 +1,97 @@
+// Copyright Epic Games, Inc. All Rights Reserved.
+
+#pragma once
+
+#include "CoreMinimal.h"
+#include "MetasoundOperatorInterface.h"
+
+namespace Metasound
+{
+	/** Convenience class for supporting the IOperator interface's GetExecuteFunction virtual member function.
+	 *
+	 * Dervied classes should inherit from this template class as well as implement a void Execute() member
+	 * function. 
+	 *
+	 * 	class MyOperator : public TExecutableOperator<MyOperator>
+	 * 	{
+	 * 	  public:
+	 * 	  ...
+	 * 	  void Execute()
+	 * 	  {
+	 * 	     ...
+	 * 	  }
+	 * 	};
+	 */
+	template<class DerivedOperatorType>
+	class TExecutableOperator : public IOperator
+	{
+		public:
+
+			virtual ~TExecutableOperator() {}
+
+			virtual FExecuteFunction GetExecuteFunction() override
+			{
+				return &TExecutableOperator<DerivedOperatorType>::ExecuteFunction;
+			}
+
+		private:
+
+			static void ExecuteFunction(IOperator* InOperator)
+			{
+				DerivedOperatorType* DerivedOperator = static_cast<DerivedOperatorType*>(InOperator);
+
+				check(nullptr != DerivedOperator);
+
+				DerivedOperator->Execute();
+			}
+	};
+
+	/** FExecutableOperatorWrapper 
+	 *
+	 * Wraps an IOperator and provides an Execute() member function.
+	 */
+	class FExecutableOperatorWrapper 
+	{
+		public:
+			using FOperatorPtr = TUniquePtr<IOperator>;
+			using FExecuteFunction = IOperator::FExecuteFunction;
+
+			FExecutableOperatorWrapper(FOperatorPtr InOperator)
+			:	Operator(MoveTemp(InOperator))
+			,	ExecuteFunction(&FExecutableOperatorWrapper::NoOp)
+			{
+				if (Operator.IsValid())
+				{
+					FExecuteFunction Func = Operator->GetExecuteFunction();
+
+					if (nullptr != Func)
+					{
+						ExecuteFunction = Func;
+					}
+				}
+			}
+
+			virtual ~FExecutableOperatorWrapper()
+			{
+			}
+
+			void Execute()
+			{
+				ExecuteFunction(Operator.Get());
+			}
+
+			bool IsNoOp()
+			{
+				return (ExecuteFunction == &FExecutableOperatorWrapper::NoOp);
+			}
+
+		private:
+			FOperatorPtr Operator;
+
+			FExecuteFunction ExecuteFunction;
+
+			static void NoOp(IOperator*)
+			{
+			}
+	};
+}
diff --git a/Engine/Plugins/Runtime/Metasound/Source/MetasoundGraphCore/Public/MetasoundGraph.h b/Engine/Plugins/Runtime/Metasound/Source/MetasoundGraphCore/Public/MetasoundGraph.h
new file mode 100644
--- /dev/null
+++ b/Engine/Plugins/Runtime/Metasound/Source/MetasoundGraphCore/Public/MetasoundGraph.h
@@ -0,0 +1,61 @@
+// Copyright Epic Games, Inc. All Rights Reserved.
+
+#pragma once
+
+#include "CoreMinimal.h"
+#include "MetasoundNode.h"
+
+namespace Metasound
+{
+	class METASOUNDGRAPHCORE_API FGraph : public IGraph
+	{
+		public:
+			static const FName ClassName;
+
+			FGraph(const FString& InDescription);
+			virtual ~FGraph();
+
+			/** Return the name of this graph. */
+			virtual const FString& GetDescription() const override;
+
+			/** Return the type name of this graph. */
+			virtual const FName& GetClassName() const override;
+
+			/** Return an array of input parameter descriptions for this graph. */
+			virtual const TArray<FInputDataVertexDescription>& GetInputs() const override;
+
+			/** Return an array of output parameter descriptions for this graph. */
+			virtual const TArray<FOutputDataVertexDescription>& GetOutputs() const override;
+
+			/** Retrieve all the edges associated with a graph. */
+			virtual const TArray<FDataEdge>& GetDataEdges() const override;
+
+			/** Get vertices which contain input parameters. */
+			virtual const TArray<FInputDataVertex>& GetInputDataVertices() const override;
+
+			/** Get vertices which contain output parameters. */
+			virtual const TArray<FOutputDataVertex>& GetOutputDataVertices() const override;
+
+			void AddDataEdge(const FDataEdge& InEdge);
+
+			void AddInputDataVertex(const FInputDataVertex& InVertex);
+
+			void AddOutputDataVertex(const FOutputDataVertex& InVertex);
+
+
+			// TODO: Add ability to remove things.
+			
+
+
+		private:
+			FString Description;
+
+			TArray<FInputDataVertexDescription> InputDescriptions;
+			TArray<FOutputDataVertexDescription> OutputDescriptions;
+
+			TArray<FDataEdge> Edges;
+
+			TArray<FInputDataVertex> InputVertices;
+			TArray<FOutputDataVertex> OutputVertices;
+	};
+}
diff --git a/Engine/Plugins/Runtime/Metasound/Source/MetasoundGraphCore/Public/MetasoundGraphCoreModule.h b/Engine/Plugins/Runtime/Metasound/Source/MetasoundGraphCore/Public/MetasoundGraphCoreModule.h
new file mode 100644
--- /dev/null
+++ b/Engine/Plugins/Runtime/Metasound/Source/MetasoundGraphCore/Public/MetasoundGraphCoreModule.h
@@ -0,0 +1,17 @@
+// Copyright Epic Games, Inc. All Rights Reserved.
+
+#pragma once
+
+#include "CoreMinimal.h"
+#include "Modules/ModuleInterface.h"
+
+METASOUNDGRAPHCORE_API DECLARE_LOG_CATEGORY_EXTERN(LogMetasoundGraphCore, Log, All);
+
+namespace Metasound 
+{
+	class METASOUNDGRAPHCORE_API IMetasoundGraphCoreModule : public IModuleInterface
+	{
+	};
+}
+
+
diff --git a/Engine/Plugins/Runtime/Metasound/Source/MetasoundGraphCore/Public/MetasoundGraphOperator.h b/Engine/Plugins/Runtime/Metasound/Source/MetasoundGraphCore/Public/MetasoundGraphOperator.h
new file mode 100644
--- /dev/null
+++ b/Engine/Plugins/Runtime/Metasound/Source/MetasoundGraphCore/Public/MetasoundGraphOperator.h
@@ -0,0 +1,50 @@
+// Copyright Epic Games, Inc. All Rights Reserved.
+
+#pragma once
+
+#include "CoreMinimal.h"
+#include "MetasoundDataReference.h"
+#include "MetasoundDataReferenceCollection.h"
+#include "MetasoundOperatorInterface.h"
+#include "MetasoundExecutableOperator.h"
+
+namespace Metasound
+{
+	class METASOUNDGRAPHCORE_API FGraphOperator : public IOperator
+	{
+		public:
+			using FOperatorPtr = TUniquePtr<IOperator>;
+			using FExecuteFunction = IOperator::FExecuteFunction;
+
+			FGraphOperator() = default;
+
+			virtual ~FGraphOperator();
+
+			void AppendOperator(FOperatorPtr InOperator);
+
+			void SetInputs(const FDataReferenceCollection& InCollection);
+
+			void SetOutputs(const FDataReferenceCollection& InCollection);
+
+			virtual const FDataReferenceCollection& GetInputs() const override;
+
+			virtual const FDataReferenceCollection& GetOutputs() const override;
+
+			virtual FExecuteFunction GetExecuteFunction() override;
+
+			void Execute();
+
+		private:
+			// Delete copy operator because underlying types cannot be copied. 
+			FGraphOperator& operator=(const FGraphOperator&) = delete;
+			FGraphOperator(const FGraphOperator&) = delete;
+			
+
+			static void ExecuteFunction(IOperator* InOperator);
+
+			TArray<FExecutableOperatorWrapper> OperatorStack;
+
+			FDataReferenceCollection Inputs;
+			FDataReferenceCollection Outputs;
+	};
+}
diff --git a/Engine/Plugins/Runtime/Metasound/Source/MetasoundGraphCore/Public/MetasoundInputNode.h b/Engine/Plugins/Runtime/Metasound/Source/MetasoundGraphCore/Public/MetasoundInputNode.h
new file mode 100644
--- /dev/null
+++ b/Engine/Plugins/Runtime/Metasound/Source/MetasoundGraphCore/Public/MetasoundInputNode.h
@@ -0,0 +1,137 @@
+// Copyright Epic Games, Inc. All Rights Reserved.
+
+#pragma once
+
+#include "CoreMinimal.h"
+#include "MetasoundOperatorInterface.h"
+#include "MetasoundNodeInterface.h"
+
+namespace Metasound
+{
+	template<typename DataType>
+	class TInputNode : public INode
+	{
+			class FInputOperator : public IOperator
+			{
+				public:
+					using FDataWriteReference = TDataWriteReference<DataType>;
+
+					FInputOperator(const FString& InDataReferenceName, FDataWriteReference InDataReference)
+					{
+						Inputs.AddDataWriteReference<DataType>(InDataReferenceName, InDataReference);
+						Outputs.AddDataReadReference<DataType>(InDataReferenceName, InDataReference);
+					}
+
+					virtual ~FInputOperator() {}
+
+					virtual const FDataReferenceCollection& GetInputs() const override
+					{
+						return Inputs;
+					}
+
+					virtual const FDataReferenceCollection& GetOutputs() const override
+					{
+						return Outputs;
+					}
+
+					virtual FExecuteFunction GetExecuteFunction() override
+					{
+						return nullptr;
+					}
+
+				private:
+					FDataReferenceCollection Outputs;
+					FDataReferenceCollection Inputs;
+			};
+
+			// TODO: Would be nice if we didn't utilize the copy constructor by default here. But it's a nice catch-all
+			// for our use case.  Could store constructor parameters of the paramtype, but run the risk of holding rvalues
+			// or pointers to objects which are no longer valid. 
+			class FCopyOperatorFactory : public IOperatorFactory
+			{
+				public:
+					using FDataWriteReference = TDataWriteReference<DataType>;
+
+					template<typename... ArgTypes, typename = typename TEnableIf< TIsConstructible<DataType, ArgTypes...>::Value >::Type >
+					FCopyOperatorFactory(ArgTypes&&... Args)
+					:	Data(Forward<ArgTypes>(Args)...)
+					{
+					}
+
+					template<typename = typename TEnableIf< TIsConstructible<DataType>::Value >::Type >
+					FCopyOperatorFactory()
+					:	Data()
+					{
+					}
+
+					virtual TUniquePtr<IOperator> CreateOperator(const INode& InNode, const FOperatorSettings& InOperatorSettings, const FDataReferenceCollection& InInputDataReferences, TArray<TUniquePtr<IOperatorBuildError>>& OutErrors) override
+					{
+						// Use copy constructor to create a new paramter reference.
+						FDataWriteReference DataRef(Data);
+						// TODO: Write special version of this for audio types since they will need to be constructed based upon the operator settings. 
+						// TODO: Need to ponder how inputs are initialized, but template specialization might do the trick. 
+
+
+						return MakeUnique<FInputOperator>(InNode.GetDescription(), DataRef);
+					}
+
+				private:
+					DataType Data;
+			};
+
+		public:
+
+
+			template<typename = typename TEnableIf< TIsConstructible<DataType>::Value >::Type >
+			TInputNode(const FString& InDescription)
+			:	Description(InDescription)
+			{
+				Outputs.Add(MakeOutputDataVertexDescription<DataType>(Description, FText::GetEmpty()));
+			}
+
+			template<typename... ArgTypes, typename = typename TEnableIf< TIsConstructible<DataType, ArgTypes...>::Value >::Type >
+			TInputNode(const FString& InDescription, ArgTypes&&... Args)
+			:	Description(InDescription)
+			,	Factory(Forward<ArgTypes>(Args)...)
+			{
+				Inputs.Add(MakeInputDataVertexDescription<DataType>(Description, FText::GetEmpty()));
+				Outputs.Add(MakeOutputDataVertexDescription<DataType>(Description, FText::GetEmpty()));
+			}
+
+			virtual const FString& GetDescription() const override
+			{
+				return Description;
+			}
+
+			virtual const FName& GetClassName() const override
+			{
+				static const FName ClassName = FName(FString(TEXT("Input_")) + FString(TDataReferenceTypeInfo<DataType>::TypeName));
+
+				return ClassName;
+			}
+
+			virtual const TArray<FInputDataVertexDescription>& GetInputs() const override
+			{
+				return Inputs;
+			}
+
+			virtual const TArray<FOutputDataVertexDescription>& GetOutputs() const override
+			{
+				return Outputs;
+			}
+
+			virtual IOperatorFactory& GetDefaultOperatorFactory() override
+			{
+				return Factory;
+			}
+
+
+		private:
+			FString Description;
+
+			FCopyOperatorFactory Factory;
+
+			TArray<FInputDataVertexDescription> Inputs;
+			TArray<FOutputDataVertexDescription> Outputs;
+	};
+}
diff --git a/Engine/Plugins/Runtime/Metasound/Source/MetasoundGraphCore/Public/MetasoundNode.h b/Engine/Plugins/Runtime/Metasound/Source/MetasoundGraphCore/Public/MetasoundNode.h
new file mode 100644
--- /dev/null
+++ b/Engine/Plugins/Runtime/Metasound/Source/MetasoundGraphCore/Public/MetasoundNode.h
@@ -0,0 +1,51 @@
+// Copyright Epic Games, Inc. All Rights Reserved.
+
+#pragma once
+
+#include "CoreMinimal.h"
+#include "MetasoundNodeInterface.h"
+
+namespace Metasound
+{
+	//DECLARE_MULTICAST_DELEGATE_TwoParams(FParameterChangedDelegate, const INodeBase*, const FDataVertexDescription&);
+
+	class METASOUNDGRAPHCORE_API FNode : public INode
+	{
+		public:
+			FNode(const FString& InDescription);
+			virtual ~FNode();
+
+			virtual const FString& GetDescription() const override;
+
+			virtual const TArray<FInputDataVertexDescription>& GetInputs() const override;
+			virtual const TArray<FOutputDataVertexDescription>& GetOutputs() const override;
+
+		protected:
+
+
+			template<typename DataType>
+			void AddInputDataVertexDescription(const FString& InVertexName, const FText& InVertexTooltip)
+			{
+				AddInputDataVertexDescription(MakeInputDataVertexDescription<DataType>(InVertexName, InVertexTooltip));
+			}
+
+			void AddInputDataVertexDescription(const FInputDataVertexDescription& InDescription);
+			void RemoveInputDataVertexDescription(const FInputDataVertexDescription& InDescription);
+
+			template<typename DataType>
+			void AddOutputDataVertexDescription(const FString& InVertexName, const FText& InVertexTooltip)
+			{
+				AddOutputDataVertexDescription(MakeOutputDataVertexDescription<DataType>(InVertexName, InVertexTooltip));
+			}
+
+			void AddOutputDataVertexDescription(const FOutputDataVertexDescription& InDescription);
+			void RemoveOutputDataVertexDescription(const FOutputDataVertexDescription& InDescription);
+
+		private:
+
+			FString Description;
+
+			TArray<FInputDataVertexDescription> Inputs;
+			TArray<FOutputDataVertexDescription> Outputs;
+	};
+}
diff --git a/Engine/Plugins/Runtime/Metasound/Source/MetasoundGraphCore/Public/MetasoundNodeInterface.h b/Engine/Plugins/Runtime/Metasound/Source/MetasoundGraphCore/Public/MetasoundNodeInterface.h
new file mode 100644
--- /dev/null
+++ b/Engine/Plugins/Runtime/Metasound/Source/MetasoundGraphCore/Public/MetasoundNodeInterface.h
@@ -0,0 +1,210 @@
+// Copyright Epic Games, Inc. All Rights Reserved.
+
+#pragma once
+
+#include "CoreMinimal.h"
+#include "MetasoundOperatorInterface.h"
+#include "MetasoundDataReference.h"
+
+namespace Metasound
+{
+	/** FDataVertexDescription
+	 *
+	 *  This describes a data vertex on a Node.
+	 */
+	struct METASOUNDGRAPHCORE_API FDataVertexDescription
+	{
+		/** Name of parameteer. */
+		FString VertexName;
+
+		/** Type name of parmaeter. */
+		FName DataReferenceTypeName;
+
+		/** Tooltip for parameter. */
+		FText Tooltip;
+
+		FDataVertexDescription()
+		:	Tooltip(FText::GetEmpty())
+		{
+		}
+
+		FDataVertexDescription(const FString& InVertexName, const FName& InDataReferenceTypeName, const FText& InTooltip)
+		:	VertexName(InVertexName)
+		,	DataReferenceTypeName(InDataReferenceTypeName)
+		,	Tooltip(InTooltip)
+		{
+		}
+
+		/** Check if two parameter descriptions are equal. */
+		friend bool METASOUNDGRAPHCORE_API operator==(const FDataVertexDescription& LHS, const FDataVertexDescription& RHS);
+
+		/** Check if two parameter descriptions are unequal. */
+		friend bool METASOUNDGRAPHCORE_API operator!=(const FDataVertexDescription& LHS, const FDataVertexDescription& RHS);
+	};
+
+	struct FInputDataVertexDescription : FDataVertexDescription
+	{
+		using FDataVertexDescription::FDataVertexDescription;
+	};
+
+	struct FOutputDataVertexDescription : FDataVertexDescription
+	{
+		using FDataVertexDescription::FDataVertexDescription;
+	};
+
+
+	template<typename DataType>
+	FInputDataVertexDescription MakeInputDataVertexDescription(const FString& InVertexName, const FText& InTooltip)
+	{
+		static const FName DataTypeName = FName(TDataReferenceTypeInfo<DataType>::TypeName);
+
+		return FInputDataVertexDescription(InVertexName, DataTypeName, InTooltip);
+	}
+
+	template<typename DataType>
+	FOutputDataVertexDescription MakeOutputDataVertexDescription(const FString& InVertexName, const FText& InTooltip)
+	{
+		static const FName DataTypeName = FName(TDataReferenceTypeInfo<DataType>::TypeName);
+
+		return FOutputDataVertexDescription(InVertexName, DataTypeName, InTooltip);
+	}
+
+
+	// Forward declare.
+	class INode;
+
+	/** Represents a specific input data vertex from a specific node. */
+	struct METASOUNDGRAPHCORE_API FInputDataVertex
+	{
+		/** Pointer to the node associated with the vertex. */
+		INode* Node;
+
+		/** The description of the Node's vertex. */
+		FInputDataVertexDescription Description;
+
+		FInputDataVertex()
+		:	Node(nullptr)
+		{
+		}
+
+		FInputDataVertex(INode* InNode, const FInputDataVertexDescription& InDescription)
+		:	Node(InNode)
+		,	Description(InDescription)
+		{
+		}
+	};
+
+	/** Represents a specific output data vertex from a specific node. */
+	struct METASOUNDGRAPHCORE_API FOutputDataVertex
+	{
+		/** Pointer to the node associated with the vertex. */
+		INode* Node;
+
+		/** The description of the Node's vertex. */
+		FOutputDataVertexDescription Description;
+
+		FOutputDataVertex()
+		:	Node(nullptr)
+		{
+		}
+
+		FOutputDataVertex(INode* InNode, const FOutputDataVertexDescription& InDescription)
+		:	Node(InNode)
+		,	Description(InDescription)
+		{
+		}
+	};
+
+	/** FDataEdge
+	 *
+	 * An edge describes a connection between two Nodes. 
+	 */
+	struct METASOUNDGRAPHCORE_API FDataEdge
+	{
+		/** Vertex producing the parameter. */
+		FOutputDataVertex From;
+
+		/** Vertex consuming the parameter. */
+		FInputDataVertex To;
+
+		FDataEdge()
+		{
+		}
+
+		FDataEdge(const FOutputDataVertex& InFrom, const FInputDataVertex& InTo)
+		:	From(InFrom)
+		,	To(InTo)
+		{
+		}
+	};
+
+	/** INodeBase
+	 * 
+	 * Interface for all nodes that can describe their name, type, inputs and outputs.
+	 */
+	class INodeBase
+	{
+		public:
+			virtual ~INodeBase() {}
+
+			/** Return the name of this node. */
+			virtual const FString& GetDescription() const = 0;
+
+			/** Return the type name of this node. */
+			virtual const FName& GetClassName() const = 0;
+
+			/** Return an array of input parameter descriptions for this node. */
+			virtual const TArray<FInputDataVertexDescription>& GetInputs() const = 0;
+
+			/** Return an array of output parameter descriptions for this node. */
+			virtual const TArray<FOutputDataVertexDescription>& GetOutputs() const = 0;
+	};
+
+	// Forward declare
+	class IOperatorFactory;
+
+	/** INode 
+	 * 
+	 * Interface for all nodes that can create operators. 
+	 */
+	class INode : public INodeBase
+	{
+		public:
+			virtual ~INode() {}
+
+			/** Return a reference to the default operator factory. */
+			virtual IOperatorFactory& GetDefaultOperatorFactory() = 0;
+
+			/* Future implementations may support additional factory types and interfaces
+				virtual bool DoesSupportOperatorFactory(const FString& InFractoryName) const = 0;
+				virtual IOperatorFactory* GetOperatorFactory(const FString& InFactoryName) = 0;
+				virtual ISpecialFactory* GetSpecialFactory() { return nullptr; }
+			*/
+	};
+
+	/** Interface for graph of nodes. */
+	class IGraph : public INodeBase
+	{
+		public:
+			virtual ~IGraph() {}
+
+			/** Retrieve all the edges associated with a graph. */
+			virtual const TArray<FDataEdge>& GetDataEdges() const = 0;
+
+			/** Get vertices which contain input parameters. */
+			virtual const TArray<FInputDataVertex>& GetInputDataVertices() const = 0;
+
+			/** Get vertices which contain output parameters. */
+			virtual const TArray<FOutputDataVertex>& GetOutputDataVertices() const = 0;
+	};
+
+
+	// This one is not yet implemented. It's design will likely be driven by UX and object management requirements. 
+	class INodeFactory
+	{
+		public:
+			virtual ~INodeFactory() {}
+			virtual TUniquePtr<INode> CreateNode(const FString& InJsonString) = 0;
+	};
+
+}
diff --git a/Engine/Plugins/Runtime/Metasound/Source/MetasoundGraphCore/Public/MetasoundNodeRegister.h b/Engine/Plugins/Runtime/Metasound/Source/MetasoundGraphCore/Public/MetasoundNodeRegister.h
new file mode 100644
--- /dev/null
+++ b/Engine/Plugins/Runtime/Metasound/Source/MetasoundGraphCore/Public/MetasoundNodeRegister.h
@@ -0,0 +1,17 @@
+// Copyright Epic Games, Inc. All Rights Reserved.
+
+#pragma once
+
+#include "CoreMinimal.h"
+
+namespace Metasound
+{
+	class METASOUNDGRAPHCORE_API FNodeFactoryRegister
+	{
+		public:
+			void RegisterNodeFactory(const FString& InNodeTypeName, TUniquePtr<INodeFactory> InNodeFactory);
+
+		private:
+			TMap<FString, TUniquePtr<INodeFactory>> NodeFactories;
+	};
+}
diff --git a/Engine/Plugins/Runtime/Metasound/Source/MetasoundGraphCore/Public/MetasoundOperatorBuilder.h b/Engine/Plugins/Runtime/Metasound/Source/MetasoundGraphCore/Public/MetasoundOperatorBuilder.h
new file mode 100644
--- /dev/null
+++ b/Engine/Plugins/Runtime/Metasound/Source/MetasoundGraphCore/Public/MetasoundOperatorBuilder.h
@@ -0,0 +1,58 @@
+// Copyright Epic Games, Inc. All Rights Reserved.
+
+#pragma once
+
+#include "CoreMinimal.h"
+#include "MetasoundNodeInterface.h"
+#include "MetasoundOperatorInterface.h"
+#include "MetasoundBuilderInterface.h"
+
+namespace Metasound
+{
+	// TODO: Builder needs a type registry to create inputs. Or can have specific input/output nodes and avoid registry.
+	class METASOUNDGRAPHCORE_API FOperatorBuilder : public IOperatorBuilder
+	{
+		public:
+			FOperatorBuilder(const FOperatorSettings& InSettings);
+
+			virtual ~FOperatorBuilder();
+
+			virtual TUniquePtr<IOperator> BuildGraphOperator(const IGraph& InGraph, TArray<FBuildErrorPtr>& OutErrors) override;
+			
+
+		private:
+			struct FOperatorDataReferences
+			{
+				FOperatorDataReferences()
+				{
+				}
+
+				FOperatorDataReferences(const FDataReferenceCollection& InInputs, const FDataReferenceCollection& InOutputs)
+				:	Inputs(InInputs)
+				,	Outputs(InOutputs)
+				{
+				}
+
+				FDataReferenceCollection Inputs;
+				FDataReferenceCollection Outputs;
+			};
+			using FNodeEdgeMultiMap = TMultiMap<const INode*, const FDataEdge*>;
+			using FNodeDataReferenceMap = TMap<INode*, FOperatorDataReferences>;
+			using FOperatorPtr = TUniquePtr<IOperator>;
+
+			bool GroupInputEdges(const TArray<FDataEdge>& InEdges, FNodeEdgeMultiMap& OutNodeInputs, TArray<FBuildErrorPtr>& OutErrors) const;
+
+			bool TopologicalSort(const TArray<FDataEdge>& InEdges, TArray<INode*>& OutNodes, TArray<FBuildErrorPtr>& OutErrors) const;
+
+			bool GatherInputDataReferences(const INode* InNode, const FNodeEdgeMultiMap& InEdgeMap, const FNodeDataReferenceMap& InDataReferenceMap, FDataReferenceCollection& OutCollection, TArray<FBuildErrorPtr>& OutErrors) const;
+
+			bool GatherGraphDataReferences(const IGraph& InGraph, FNodeDataReferenceMap& InNodeDataReferences, FDataReferenceCollection& OutGraphInputs, FDataReferenceCollection& OutGraphOutputs, TArray<FBuildErrorPtr>& OutErrors) const;
+
+			bool CreateOperators(const TArray<INode*>& InSortedNodes, FNodeEdgeMultiMap& InNodeInputEdges, TArray<FOperatorPtr>& OutOperators, FNodeDataReferenceMap& OutDataReferences, TArray<FBuildErrorPtr>& OutErrors);
+
+			TUniquePtr<IOperator> CreateGraphOperator(const IGraph& InGraph, TArray<FOperatorPtr>& InOperators, FNodeDataReferenceMap& InNodeDataReferences, TArray<FBuildErrorPtr>& OutErrors) const;
+
+			FOperatorSettings OperatorSettings;
+	};
+}
+
diff --git a/Engine/Plugins/Runtime/Metasound/Source/MetasoundGraphCore/Public/MetasoundOperatorInterface.h b/Engine/Plugins/Runtime/Metasound/Source/MetasoundGraphCore/Public/MetasoundOperatorInterface.h
new file mode 100644
--- /dev/null
+++ b/Engine/Plugins/Runtime/Metasound/Source/MetasoundGraphCore/Public/MetasoundOperatorInterface.h
@@ -0,0 +1,82 @@
+// Copyright Epic Games, Inc. All Rights Reserved.
+
+#pragma once
+
+#include "CoreMinimal.h"
+#include "MetasoundDataReference.h"
+#include "MetasoundDataReferenceCollection.h"
+
+namespace Metasound
+{
+	/** IOperator
+	 *
+	 *  IOperator defines the interface for render time operations.  IOperators are created using an INodeOperatorFactory.
+	 */
+	class IOperator
+	{
+		public:
+			/** Pointer to execute function for an operator. 
+			 *
+			 * @param IOperator* - The operator associated with the function pointer.
+			 */
+			typedef void(*FExecuteFunction)(IOperator*);
+
+			virtual ~IOperator() {}
+
+			/** Return the input parameters associated with this IOperator. 
+			 *
+			 * Implementations of IOperator should populate and return their input parameters
+			 * if they want to enable callers to query and write to their parameters. Most IOperator
+			 * implementations will return an empty FDataReferenceCollection because thier inputs are 
+			 * set during the IOperator's construction and do not need to be updated afterwards. The
+			 * exceptions are input operators and graph operators which need to interface with external
+			 * systems even after the operator is created.
+			 */
+			virtual const FDataReferenceCollection& GetInputs() const = 0;
+
+			/** Return the output parameters associated with this IOperator.
+			 *
+			 * Implementations of IOperator should return a collection of their output parameters 
+			 * which other nodes can read.
+			 */
+			virtual const FDataReferenceCollection& GetOutputs() const = 0;
+
+			/** Return the execution function to call during graph execution.
+			 * 
+			 * The IOperator* argument to the FExecutionFunction will be the same IOperator instance
+			 * which returned the execution function. 
+			 *
+			 * nullptr return values are valid and signal an IOperator which does not need to be 
+			 * executed.
+			 */
+			virtual FExecuteFunction GetExecuteFunction() = 0;
+	};
+
+	/** FOperatorSettings
+	 *
+	 * Audio settings for operators. 
+	 */
+	struct FOperatorSettings
+	{
+		/** The number of samples per a second of audio. */
+		float SampleRate = 0.f;
+
+		/** The number of audio frames produced during a single call to IOperators execution function. */
+		int32 FramesPerExecute = 0;
+
+		FOperatorSettings(float InSampleRate, uint32 InFramesPerExecute)
+		:	SampleRate(InSampleRate)
+		,	FramesPerExecute(InFramesPerExecute)
+		{
+		}
+	};
+
+	/*
+	class IGraphOperator : public IOperator
+	{
+		public:
+
+			virtual ~IGraphOperator() {}
+	};
+	*/
+}
diff --git a/Engine/Plugins/Runtime/Metasound/Source/MetasoundGraphCore/Public/MetasoundOscNode.h b/Engine/Plugins/Runtime/Metasound/Source/MetasoundGraphCore/Public/MetasoundOscNode.h
new file mode 100644
--- /dev/null
+++ b/Engine/Plugins/Runtime/Metasound/Source/MetasoundGraphCore/Public/MetasoundOscNode.h
@@ -0,0 +1,38 @@
+// Copyright Epic Games, Inc. All Rights Reserved.
+
+#pragma once
+
+#include "CoreMinimal.h"
+#include "MetasoundNode.h"
+#include "MetasoundOperatorInterface.h"
+#include "MetasoundBuilderInterface.h"
+#include "MetasoundDataReferenceCollection.h"
+
+namespace Metasound
+{
+	class METASOUNDGRAPHCORE_API FOscNode : public FNode
+	{
+			class FOperatorFactory : public IOperatorFactory
+			{
+				virtual TUniquePtr<IOperator> CreateOperator(const INode& InNode, const FOperatorSettings& InOperatorSettings, const FDataReferenceCollection& InInputDataReferences, TArray<TUniquePtr<IOperatorBuildError>>& OutErrors) override;
+			};
+
+		public:
+			static const FName ClassName;
+
+			FOscNode(const FString& InName, float InDefaultFrequency);
+
+			virtual ~FOscNode();
+
+			float GetDefaultFrequency() const;
+
+			virtual const FName& GetClassName() const override;
+
+			virtual IOperatorFactory& GetDefaultOperatorFactory() override;
+
+		private:
+
+			float DefaultFrequency;
+			FOperatorFactory Factory;
+	};
+}
diff --git a/Engine/Plugins/Runtime/Metasound/Source/MetasoundGraphCore/Public/MetasoundOutputNode.h b/Engine/Plugins/Runtime/Metasound/Source/MetasoundGraphCore/Public/MetasoundOutputNode.h
new file mode 100644
--- /dev/null
+++ b/Engine/Plugins/Runtime/Metasound/Source/MetasoundGraphCore/Public/MetasoundOutputNode.h
@@ -0,0 +1,133 @@
+// Copyright Epic Games, Inc. All Rights Reserved.
+
+#pragma once
+
+#include "CoreMinimal.h"
+#include "MetasoundOperatorInterface.h"
+#include "MetasoundNodeInterface.h"
+
+#define LOCTEXT_NAMESPACE "MetasoundOutputNode"
+
+namespace Metasound
+{
+	template<typename DataType>
+	class TOutputNode : public INode
+	{
+			class FOutputOperator : public IOperator
+			{
+				public:
+					using FDataReadReference = TDataReadReference<DataType>;
+
+					FOutputOperator(const FString& InDataReferenceName, FDataReadReference InDataReference)
+					{
+						Outputs.AddDataReadReference<DataType>(InDataReferenceName, InDataReference);
+					}
+
+					virtual ~FOutputOperator() {}
+
+					virtual const FDataReferenceCollection& GetInputs() const override
+					{
+						return Inputs;
+					}
+
+					virtual const FDataReferenceCollection& GetOutputs() const override
+					{
+						return Outputs;
+					}
+
+					virtual FExecuteFunction GetExecuteFunction() override
+					{
+						return nullptr;
+					}
+
+				private:
+					FDataReferenceCollection Inputs;
+					FDataReferenceCollection Outputs;
+			};
+
+			class FOutputOperatorFactory : public IOperatorFactory
+			{
+				public:
+					FOutputOperatorFactory(const FString& InDataReferenceName)
+					:	DataReferenceName(InDataReferenceName)
+					{
+					}
+
+					virtual TUniquePtr<IOperator> CreateOperator(const INode& InNode, const FOperatorSettings& InOperatorSettings, const FDataReferenceCollection& InInputDataReferences, TArray<TUniquePtr<IOperatorBuildError>>& OutErrors) override
+					{
+						if (!InInputDataReferences.ContainsDataReadReference<DataType>(DataReferenceName))
+						{
+							// TODO: Add build error.
+							return TUniquePtr<IOperator>(nullptr);
+						}
+
+						return MakeUnique<FOutputOperator>(InNode.GetDescription(), InInputDataReferences.GetDataReadReference<DataType>(DataReferenceName));
+					}
+
+				private:
+					FString DataReferenceName;
+			};
+
+		public:
+
+
+			template<typename = typename TEnableIf< TIsConstructible<DataType>::Value >::Type >
+			TOutputNode(const FString& InDescription)
+			:	Description(InDescription)
+			,	Factory(InDescription)
+			{
+				Outputs.Add(MakeOutputDataVertexDescription<DataType>(Description, FText::GetEmpty()));
+				Inputs.Add(MakeInputDataVertexDescription<DataType>(Description, FText::GetEmpty()));
+			}
+
+			template<typename... ArgTypes, typename = typename TEnableIf< TIsConstructible<DataType, ArgTypes...>::Value >::Type >
+			TOutputNode(const FString& InDescription, ArgTypes&&... Args)
+			:	Description(InDescription)
+			,	Factory(InDescription)
+			,	Data(Forward<ArgTypes>(Args)...)
+			{
+				Outputs.Add(MakeOutputDataVertexDescription<DataType>(Description, FText::GetEmpty()));
+				Inputs.Add(MakeInputDataVertexDescription<DataType>(Description, FText::GetEmpty()));
+			}
+
+			virtual const FString& GetDescription() const override
+			{
+				return Description;
+			}
+
+			virtual const FName& GetClassName() const override
+			{
+				// TODO: Any special formatting for these node type names?
+				static const FName ClassName = FName(FString(TEXT("Output_")) + FString(TDataReferenceTypeInfo<DataType>::TypeName));
+
+				return ClassName;
+			}
+
+			virtual const TArray<FInputDataVertexDescription>& GetInputs() const override
+			{
+				return Inputs;
+			}
+
+			virtual const TArray<FOutputDataVertexDescription>& GetOutputs() const override
+			{
+				return Outputs;
+			}
+
+			virtual IOperatorFactory& GetDefaultOperatorFactory() override
+			{
+				return Factory;
+			}
+
+		private:
+			FString Description;
+
+			FOutputOperatorFactory Factory;
+
+			DataType Data;
+
+			TArray<FInputDataVertexDescription> Inputs;
+			TArray<FOutputDataVertexDescription> Outputs;
+	};
+}
+
+#undef LOCTEXT_NAMESPACE //MetasoundOutputNode
diff --git a/Engine/Plugins/Runtime/Metasound/Source/MetasoundGraphCore/Public/MetasoundPeriodicBopNode.h b/Engine/Plugins/Runtime/Metasound/Source/MetasoundGraphCore/Public/MetasoundPeriodicBopNode.h
new file mode 100644
--- /dev/null
+++ b/Engine/Plugins/Runtime/Metasound/Source/MetasoundGraphCore/Public/MetasoundPeriodicBopNode.h
@@ -0,0 +1,39 @@
+// Copyright Epic Games, Inc. All Rights Reserved.
+
+#pragma once
+
+#include "CoreMinimal.h"
+#include "MetasoundNode.h"
+#include "MetasoundBuilderInterface.h"
+#include "MetasoundOperatorInterface.h"
+#include "MetasoundNodeInterface.h"
+#include "MetasoundDataReferenceCollection.h"
+#include "MetasoundDataReferenceTypes.h"
+
+namespace Metasound
+{
+	class METASOUNDGRAPHCORE_API FPeriodicBopNode : public FNode
+	{
+			class FOperatorFactory : public IOperatorFactory
+			{
+				virtual TUniquePtr<IOperator> CreateOperator(const INode& InNode, const FOperatorSettings& InOperatorSettings, const FDataReferenceCollection& InInputDataReferences, TArray<TUniquePtr<IOperatorBuildError>>& OutErrors) override;
+			};
+
+		public:
+			static const FName ClassName;
+
+			FPeriodicBopNode(const FString& InName, float InDefaultPeriodInSeconds);
+
+			virtual ~FPeriodicBopNode();
+
+			float GetDefaultPeriodInSeconds() const;
+
+			virtual const FName& GetClassName() const override;
+
+			virtual IOperatorFactory& GetDefaultOperatorFactory() override;
+
+		private:
+			float DefaultPeriod;
+			FOperatorFactory Factory;
+	};
+}
