adding Metasound audio format definitions.

diff --git a/Engine/Plugins/Runtime/Metasound/Source/MetasoundGraphCore/Private/MetasoundAudioFormats.cpp b/Engine/Plugins/Runtime/Metasound/Source/MetasoundGraphCore/Private/MetasoundAudioFormats.cpp
new file mode 100644
--- /dev/null
+++ b/Engine/Plugins/Runtime/Metasound/Source/MetasoundGraphCore/Private/MetasoundAudioFormats.cpp
@@ -0,0 +1,105 @@
+// Copyright Epic Games, Inc. All Rights Reserved.
+
+#include "MetasoundAudioFormats.h"
+
+#include "CoreMinimal.h"
+#include "MetasoundDataReference.h"
+#include "MetasoundAudioBuffer.h"
+
+namespace Metasound
+{
+	/* FUnformattedAudio */
+
+	FUnformattedAudio::FUnformattedAudio(int32 InNumFrames, int32 InNumChannels, int32 InMaxNumChannels)
+	:	NumFrames(InNumFrames)
+	,	NumChannels(0)
+	,	MaxNumChannels(InMaxNumChannels)
+	{
+		NumFrames = FMath::Max(NumFrames, 0);
+		MaxNumChannels = FMath::Max(MaxNumChannels, 0);
+
+		for (int32 i = 0; i < MaxNumChannels; i++)
+		{
+			FAudioBufferWriteRef Audio(NumFrames);
+			Audio->Zero();
+
+			WritableBufferStorage.Add(Audio);
+			ReadableBufferStorage.Add(Audio);
+		}
+
+		SetNumChannels(InNumChannels);
+	}
+
+	int32 FUnformattedAudio::SetNumChannels(int32 InNumChannels)
+	{
+		NumChannels = FMath::Max(0, FMath::Min(InNumChannels, MaxNumChannels));
+
+		ReadableBuffers = TArrayView<const FAudioBufferReadRef>(ReadableBufferStorage.GetData(), NumChannels);
+		WritableBuffers = TArrayView<const FAudioBufferWriteRef>(WritableBufferStorage.GetData(), NumChannels);
+
+		return NumChannels;
+	}
+
+
+	/* FMultichannelAudioFormat */
+
+	FMultichannelAudioFormat::FMultichannelAudioFormat(int32 InNumFrames, int32 InNumChannels)
+	:	NumChannels(InNumChannels)
+	{
+		NumChannels = FMath::Max(0, NumChannels);
+		InNumFrames = FMath::Max(0, InNumFrames);
+
+		for (int32 i = 0; i < NumChannels; i++)
+		{
+			FAudioBufferWriteRef Audio(InNumFrames);
+			Audio->Zero();
+
+			WritableBufferStorage.Add(Audio);
+			ReadableBufferStorage.Add(Audio);
+		}
+
+		WritableBuffers = WritableBufferStorage;
+		ReadableBuffers = ReadableBufferStorage;
+	}
+
+	FMultichannelAudioFormat::FMultichannelAudioFormat(TArrayView<const FAudioBufferWriteRef> InWriteRefs)
+	:	NumChannels(InWriteRefs.Num())
+	{
+		if (NumChannels > 0)
+		{
+			const int32 NumFrames = InWriteRefs[0]->Num();
+
+			for (const FAudioBufferWriteRef& Ref : InWriteRefs)
+			{
+				checkf(NumFrames == Ref->Num(), TEXT("All buffers must have same number of frames (%d != %d)"), NumFrames, Ref->Num());
+
+				WritableBufferStorage.Add(Ref);
+				ReadableBufferStorage.Add(Ref);
+			}
+
+			WritableBuffers = WritableBufferStorage;
+			ReadableBuffers = ReadableBufferStorage;
+		}
+	}
+
+	FMultichannelAudioFormat::FMultichannelAudioFormat(TArrayView<const FAudioBufferReadRef> InReadRefs)
+	:	NumChannels(InReadRefs.Num())
+	{
+		if (NumChannels > 0)
+		{
+			const int32 NumFrames = InReadRefs[0]->Num();
+
+			for (const FAudioBufferReadRef& Ref : InReadRefs)
+			{
+				checkf(NumFrames == Ref->Num(), TEXT("All buffers must have same number of frames (%d != %d)"), NumFrames, Ref->Num());
+
+				WritableBufferStorage.Add(WriteCast(Ref));
+				ReadableBufferStorage.Add(Ref);
+			}
+
+			WritableBuffers = WritableBufferStorage;
+			ReadableBuffers = ReadableBufferStorage;
+		}
+	}
+
+}
diff --git a/Engine/Plugins/Runtime/Metasound/Source/MetasoundGraphCore/Public/MetasoundAudioBuffer.h b/Engine/Plugins/Runtime/Metasound/Source/MetasoundGraphCore/Public/MetasoundAudioBuffer.h
--- a/Engine/Plugins/Runtime/Metasound/Source/MetasoundGraphCore/Public/MetasoundAudioBuffer.h
+++ b/Engine/Plugins/Runtime/Metasound/Source/MetasoundGraphCore/Public/MetasoundAudioBuffer.h
@@ -121,6 +121,13 @@ namespace Metasound
 				return Buffer;
 			}
 
+			FORCEINLINE void Zero()
+			{
+				if (Buffer.Num() > 0)
+				{
+					FMemory::Memset(Buffer.GetData(), 0, sizeof(float) * Buffer.Num());
+				}
+			}
 
 		private:
 
@@ -131,4 +138,5 @@ namespace Metasound
 #endif				
 	};
 
+	DECLARE_METASOUND_DATA_REFERENCE_TYPES(FAudioBuffer, "AudioBuffer", 0x1b02dac2 , FAudioBufferTypeInfo, FAudioBufferReadRef, FAudioBufferWriteRef);
 };
diff --git a/Engine/Plugins/Runtime/Metasound/Source/MetasoundGraphCore/Public/MetasoundAudioFormats.h b/Engine/Plugins/Runtime/Metasound/Source/MetasoundGraphCore/Public/MetasoundAudioFormats.h
--- a/Engine/Plugins/Runtime/Metasound/Source/MetasoundGraphCore/Public/MetasoundAudioFormats.h
+++ b/Engine/Plugins/Runtime/Metasound/Source/MetasoundGraphCore/Public/MetasoundAudioFormats.h
@@ -8,330 +8,418 @@
 
 namespace Metasound
 {
-	// TODO:
-	// The 3 flavors of audio
-	//
-	// FVariableChannelAudio - Number of channels can change whenever! Format identifier included.
-	// FMultiChannelAudio - Number of channels set on constructor!
-	// TStaticChannelAudio - Number of channels set on declaration. Used for static channel configurations "Stereo", "FiveDotOne".
-
-	// FAudio
-	// 	- Union of all constant sized 
-	//
-	//
-	// So lets say we want to write our node with generic audio.
-	//
-	//   GenericAudioIn -> Node -> GenericAudioOut -> CastToSpecificAudio.
-	//   	We would have to have a failure mode and/or a specific mapping for each input type. 
-	//   	Not too crazy. This would be an upmix/downmix node that would have to have a basic
-	//   	approach for up/down mixing. 
-	//
-	//
-	//   	Notes from email to self: Type erasure to define a union subtype interface with sharedptr passing semantics.  
-	//
-	//   	Audioref can support many types. All types must support format name and downcasting to self.  And monochannel forwarding. And channel count. And duplicating.  Operator implementations can use multimono convenience template, or specialize on audio subtype. Can also fail, but prolly shouldnt take audio union then. 
-	//
-	
-	class FAudioUnion
+	/** FUnformattedAudio
+	 *
+	 * FUnformattedAudio represents deinterleaved multichannel audio which can dynamically change
+	 * the number of channels of audio. In possible, using FMultichannelAudioFormat is preferred over
+	 * using FUnformattedAudio when the number of audio channels is known at construction.
+	 *
+	 * The audio buffers in FUnformattedAudio are shared data references which can be accessed outside
+	 * of the FUnformattedAudio. All audio buffers within a FUnformattedAudio object must contain the 
+	 * same number of audio frames.
+	 */
+	class FUnformattedAudio
 	{
-		class FAudioBase
-		{
-			public:
-				virtual ~FAudioBase() {}
-
-				virtual TArray<const FAudioBufferReadRef>& GetStorage() const = 0;
-
-				virtual TArray<const FAudioBufferWriteRef>& GetStorage() = 0;
-
-				virtual int32 GetMaxNumChannels() const = 0;
-
-		};
-
-		template<typename AudioType>
-		class FAudioWrapper : public FAudioBase
-		{
-			public:
-				FAudioWrapper(AudioType&& InType)
-				:	Audio(MoveTemp(InType))
-				{}
-
-				virtual TArray<const FAudioBufferReadRef>& GetStorage() const override
-				{
-					return Audio.GetStorage();
-				}
-
+		public:
+			/** FUnformattedAudio Constructor.
+			 * 
+			 * @param InNumFrames - Number of frames to hold in each buffer.
+			 * @param InNumChannels - Initial number of audio channels.
+			 * @param InMaxNumChannels - Maximum number of audio channels.
+			 */
+			FUnformattedAudio(int32 InNumFrames, int32 InNumChannels, int32 InMaxNumChannels);
+
+			/** Sets the number of audio channels.
+			 * 
+			 * The requested number of channels is internally clamped to the supported range 
+			 * of [0, GetMaxNumChannels()].
+			 *
+			 * @parma InNumChannels - The desired number of channels.
+			 *
+			 * @return The actual number of channels. This may differ from the requested number 
+			 * 		   of channels if the requested number exceeded the maximum number of channels.
+			 */
+			int32 SetNumChannels(int32 InNumChannels);
+
+			/** Return the number of channels. */
+			int32 GetNumChannels() const { return NumChannels; }
+
+			/** Return the maximum number of channels. */
+			int32 GetMaxNumChannels() const { return MaxNumChannels; }
+
+			/** Return an array view of the readable buffer references.
+			 *
+			 * This array will have GetNumChannels() elements.
+			 */
+			const TArrayView<const FAudioBufferReadRef>& GetBuffers() const { return ReadableBuffers; }
+
+			/** Return an array view of the writable buffer references.
+			 *
+			 * This array will have GetNumChannels() elements.
+			 */
+			const TArrayView<const FAudioBufferWriteRef>& GetBuffers() { return WritableBuffers; }
+
+			/** Return an array of the readable buffer reference storage.
+			 *
+			 * This array will have GetMaxNumChannels() elements.
+			 */
+			const TArray<FAudioBufferReadRef>& GetStorage() const { return ReadableBufferStorage; }
+
+			/** Return an array of the writable buffer reference storage.
+			 *
+			 * This array will have GetMaxNumChannels() elements.
+			 */
+			const TArray<FAudioBufferWriteRef>& GetStorage() { return WritableBufferStorage; }
+			
+		private:
+			int32 NumFrames;
+			int32 NumChannels;
+			int32 MaxNumChannels;
 
-				virtual TArray<const FAudioBufferWriteRef>& GetStorage() override
-				{
-					return Audio.GetStorage();
-				}
+			TArrayView<const FAudioBufferReadRef> ReadableBuffers;
+			TArrayView<const FAudioBufferWriteRef> WritableBuffers;
 
-				virtual int32 GetMaxNumChannels() const override
-				{
-					return Audio.GetMaxNumChannels();
-				}
+			TArray<FAudioBufferReadRef> ReadableBufferStorage;
+			TArray<FAudioBufferWriteRef> WritableBufferStorage;
+	};
 
-			private:
-				AudioType Audio;
-		};
 
+	/** FMultichannelAudioFormat
+	 *
+	 * FMultichannelAudioFormat represents deinterleaved multichannel audio which supports a constant
+	 * number of channels for the lifetime of the object. 
+	 *
+	 * The audio buffers in FMultichannelAudioFormat are shared data references which can be accessed outside
+	 * of the FMultichannelAudioFormat. All audio buffers within a FMultichannelAudioFormat object must contain the 
+	 * same number of audio frames.
+	 */
+	class FMultichannelAudioFormat
+	{
 		public:
-			template<typename AudioType>
-			FAudioUnion(AudioType&& InType)
-			:	AudioPtr(MoveTemp(InType))
-			{
-			}
+			/** FMultichannelAudioFormat Constructor.
+			 *
+			 * @param InNumFrames - The number of frames per an audio buffer.
+			 * @param InNumChannels - The number of audio channels.
+			 */
+			FMultichannelAudioFormat(int32 InNumFrames, int32 InNumChannels);
+
+			/** FMultichannelAudioFormat Constructor.
+			 *
+			 * This constructor accepts an array of writable audio buffer references. Each 
+			 * buffer in the array must contain equal number of frames. 
+			 *
+			 * @param InWriteRefs - An array of writable audio buffer references. 
+			 */
+			FMultichannelAudioFormat(TArrayView<const FAudioBufferWriteRef> InWriteRefs);
+
+			// Enable the copy constructor.
+			FMultichannelAudioFormat(const FMultichannelAudioFormat& InOther) = default;
+
+			// Disable move constructor as incoming object should not be altered. 
+			FMultichannelAudioFormat(FMultichannelAudioFormat&& InOther) = delete;
+
+			// Disable equal operator so channel count does not change
+			FMultichannelAudioFormat& operator=(const FMultichannelAudioFormat& Other) = delete;
+
+			// Disable move operator so channel count does not change
+			FMultichannelAudioFormat& operator=(FMultichannelAudioFormat&& Other) = delete;
+
+			/** Return the number of audio channels. */
+			int32 GetNumChannels() const { return NumChannels; }
+
+			/** Return an array view of the readable buffer references.
+			 *
+			 * This array will have GetNumChannels() elements.
+			 */
+			const TArrayView<const FAudioBufferReadRef>& GetBuffers() const { return ReadableBuffers; }
+
+			/** Return an array view of the writable buffer references.
+			 *
+			 * This array will have GetNumChannels() elements.
+			 */
+			const TArrayView<const FAudioBufferWriteRef>& GetBuffers() { return WritableBuffers; }
+
+			/** Return an array of the readable buffer reference storage.
+			 *
+			 * This array will have GetNumChannels() elements.
+			 */
+			const TArray<FAudioBufferReadRef>& GetStorage() const { return ReadableBufferStorage; }
+
+			/** Return an array of the writable buffer reference storage.
+			 *
+			 * This array will have GetNumChannels() elements.
+			 */
+			const TArray<FAudioBufferWriteRef>& GetStorage() { return WritableBufferStorage; }
 
-			TArray<const FAudioBufferReadRef>& GetStorage() const
-			{
-				return AudioPtr->GetStorage();
-			}
-
-			TArray<const FAudioBufferWriteRef>& GetStorage()
-			{
-				return AudioPtr->GetStorage();
-			}
-
-			int32 GetMaxNumChannels() const
-			{
-				return AudioPTr->GetMaxNumChannels();
-			}
+		private:
+			// Friendship with the data reference class gives it access to the protected constructor
+			// for the scenario where a TDataReadReference<FMulitchannelAudioFormat> is constructed
+			// with FAudioBufferReadRefs. The constructor cannot be public as it would provide writable
+			// access to the passed in audio buffers, even though the passed buffers explicitly were
+			// read references.
+			friend class TDataReadReference<FMultichannelAudioFormat>;
+
+			// Construct a FMultichannelAudioFormat with an array of readable buffers. 
+			// 
+			// This constructor should only be used when it can be assured that the constructed object
+			// will not provide writable access to the passed in audio buffers. 
+			FMultichannelAudioFormat(TArrayView<const FAudioBufferReadRef> InReadRefs);
 
-			FName GetDataTypeName() const
-			{
-				// TODO: this should just be a call to the owning reference type's existing type information.
-				return AudioPtr->GetDataTypeName();
-			}
+			int32 NumChannels;
 
-			template<typename AudioReferenceType>
-			TDataReadReference<AudioReferenceType> GetReadableRefChecked() const
-			{
-				// TODO: check typename matches audioptr
-				// TODO: think about how this cast actually works. Should it be a clone?
-				// Where's that actual object stored? not clear yet from this implementation.
-				return TDataReadReference<AudioReferenceType>(ParameterReferenceCast<AudioReferenceType>());
-			}
+			TArrayView<const FAudioBufferWriteRef> WritableBuffers;
+			TArrayView<const FAudioBufferReadRef> ReadableBuffers;
 
-		private:
-			
-			// TODO: This is holding onto the actual object, when it needs to hold on to the reference to the object. 
-			TUniquePtr<FAudioBase> AudioPrt;
+			TArray<FAudioBufferWriteRef> WritableBufferStorage;
+			TArray<FAudioBufferReadRef> ReadableBufferStorage;
 	};
 
-
-
-	// If we are using this format, that we expect the channel count to change during runtime.
-	// TODO: create another one that is sized on construction but cannot change channel count.
-	class FVariableChannelAudio
+	
+	/** A TStaticChannelAudioFormat represents deinterleaved multichannel audio 
+	 * where the number of channels is known at compile time. This is primarily 
+	 * useful to define such cases as Stereo, Mono, Qaud, 5.1, etc. 
+	 *
+	 * The audio buffers in FMultichannelAudioFormat are shared data references 
+	 * which can be accessed outside of the FMultichannelAudioFormat. All audio 
+	 * buffers within a FMultichannelAudioFormat object must contain the same 
+	 * number of audio frames.
+	 */
+	template<int32 NumChannels>
+	class TStaticChannelAudioFormat
 	{
 		public:
-			FVariableChannelAudio(int32 InNumFrames, int32 InMaxNumChannels, int32 InNumChannels)
-			:	NumFrames(InNumFrames)
-			,	MaxNumChannels(InMaxNumChannels)
-			,	NumChannels(0)
+			/** TStaticChannelAudioFormat Constructor
+			 *
+			 * @param InNumFrames - The number of frames per an audio buffer.
+			 */
+			TStaticChannelAudioFormat(int32 InNumFrames)
 			{
-				NumFrames = FMath::Max(NumFrames, 1);
-				MaxNumChannels = FMath::Max(MaxNumChannels, 0);
+				static_assert(NumChannels > 0, "NumChannels must be greater than zero");
 
-				for (int32 i = 0; i < MaxNumChannels; i++)
+				InNumFrames = FMath::Max(InNumFrames, 0);
+
+				for (int32 i = 0; i < NumChannels; i++)
 				{
-					FAudioBufferWriteRef Audio(NumFrames);
-					// TODO: SetZeroed() should be a member func
-					FMemory::Memset(Audio.GetData(), 0, sizeof(float) * NumFrames);
+					FAudioBufferWriteRef Audio(InNumFrames);
+					Audio->Zero();
 
 					WritableBufferStorage.Add(Audio);
 					ReadableBufferStorage.Add(Audio);
 				}
 
-				SetNumChannels(InNumChannels)
+				WritableBuffers = WritableBufferStorage;
+				ReadableBuffers = ReadableBufferStorage;
 			}
 
-
-			FVariableChannelAudio(const TArray<FAudioBufferWriteRef>& InStorage)
+			/** Return the number of audio channels. */
+			int32 GetNumChannels() const
 			{
-				// TODO:
+				return NumChannels;
 			}
 
-			int32 SetNumChannels(int32 InNumChannels)
+			/** Return an readable buffer reference for a specific channel. */
+			template<int32 ChannelIndex>
+			FAudioBufferReadRef GetBuffer() const
 			{
-				// TODO: should we just add more channels if they exist? 
-				NumChannels = FMath::Max(0, FMath::Min(InNumChannels, MaxNumChannels));
+				static_assert(ChannelIndex >= 0, "Index must be within range of channels");
+				static_assert(ChannelIndex < NumChannels, "Index must be within range of channels");
 
-				ReadableBuffers = TArrayView<const FReadableAudioBuffer>(ReadableBufferStorage.GetData(), NumChannels);
-				WritableBuffers = TArrayView<const FWritableAudioBuffer>(WritableBufferStorage.GetData(), NumChannels);
+				return ReadableBuffers.GetData()[ChannelIndex];
 			}
 
-			int32 GetNumChannels() const
+			/** Return an writable buffer reference for a specific channel. */
+			template<int32 ChannelIndex>
+			FAudioBufferWriteRef GetBuffer() 
 			{
-				return NumChannels;
-			}
+				static_assert(ChannelIndex >= 0, "Index must be within range of channels");
+				static_assert(ChannelIndex < NumChannels, "Index must be within range of channels");
 
-			// TODO: where do we store formats?
-			//EFormat EFormat;
-			//
+				return WritableBuffers.GetData()[ChannelIndex];
+			}
 
-			TArrayView<const FAudioBufferReadRef>& GetBuffers() const
+			/** Return an array view of the readable buffer references.
+			 *
+			 * This array will have GetNumChannels() elements.
+			 */
+			const TArrayView<const FAudioBufferReadRef> GetBuffers() const
 			{
 				return ReadableBuffers;
 			}
 
-			TArrayView<const FAudioBufferWriteRef>& GetBuffers() 
+			/** Return an array view of the writable buffer references.
+			 *
+			 * This array will have GetNumChannels() elements.
+			 */
+			const TArrayView<const FAudioBufferWriteRef> GetBuffers() 
 			{
 				return WritableBuffers;
 			}
 
-			const TArray<const FAudioBufferReadRef>& GetStorage() const
+			/** Return an array of the readable buffer references.
+			 *
+			 * This array will have GetNumChannels() elements.
+			 */
+			const TArray<FAudioBufferReadRef> GetStorage() const
 			{
 				return ReadableBufferStorage;
 			}
 
-			const TArray<const FAudioBufferWriteRef>& GetStorage()
+			/** Return an array of the writable buffer references.
+			 *
+			 * This array will have GetNumChannels() elements.
+			 */
+			const TArray<FAudioBufferWriteRef> GetStorage()
 			{
 				return WritableBufferStorage;
 			}
-		protected:
-			friend class TDataReadReference<FVariableChannelAudio>;
-
-			FVariableChannelAudio(const TArray<FAudioBufferReadRef>& InStorage)
-			{
-				// TODO: What do we do about write storage in this scenario? We cannot return write data because
-				// we have only been given a read reference. Maybe we could have a DataType specialization
-				// which allows this constructor to exist for readable data types, and not for writable data types. 
-				// Maybe a protected thing with a friend class? Maybe readable parameter has to do a WritableCast<>
-				//
-				// The static sized n channel templates should do something like that as well.
-				//
-				// FVariableChannelAudioReadRef ReadRef(OtherReadRef)
-				// 		
-				//
-				// FVariableChannelAudioWriteRef WriteRef(OtherReadRef) // <- Not allowed!  
-				//
-				// 
-			}
-			
-		private:
-
-
-			int32 NumChannels;
-
-			TArrayView<const FAudioBufferReadRef> ReadableBuffers;
-			TArrayView<const FAudioBufferReadRef> WritableBuffers;
-
-			// The reference storage is const so that the mono audio buffer 
-			// references can be passed forward when the format is changed (as
-			// opposed to making a copy of the reference). This requires that 
-			// this object always be referencing the same underlying audio
-			// buffer objects, hence the const protection to ensure that it 
-			// never gets changed.
-			TArray<const FAudioBufferReadRef> ReadableBufferStorage;
-			TArray<const FAudioBufferWriteRef> WritableBufferStorage;
-	};
 
+		protected:
 
-	
-	template<int32 NumChannels>
-	class TStaticChannelAudio
-	{
-		public:
-			TStaticChannelAudio(int32 InNumFrames)
-			:	NumFrames(InNumFrames)
+			// TStaticChannelAudioFormat constructor with an array of writable buffers. 
+			TStaticChannelAudioFormat(const FAudioBufferWriteRef(&InBuffers)[NumChannels])
 			{
-				static_assert(NumChannels > 0, "NumChannels must be greater than zero");
-
-				NumFrames = FMath::Max(NumFrames, 1);
+				int32 NumFrames = 0;
+				if (NumChannels > 0)
+				{
+					NumFrames = InBuffers[0]->Num();
+				}
 
 				for (int32 i = 0; i < NumChannels; i++)
 				{
-					FAudioBufferWriteRef Audio(NumFrames);
-					// TODO: SetZeroed() should be a member func
-					FMemory::Memset(Audio.GetData(), 0, sizeof(float) * NumFrames);
+					checkf(NumFrames == InBuffers[i]->Num(), TEXT("All buffers must have same number of frames (%d != %d)"), NumFrames, InBuffers[i]->Num());
 
-					WritableBufferStorage.Add(Audio);
-					ReadableBufferStorage.Add(Audio);
+					WritableBufferStorage.Add(InBuffers[i]);
+					ReadableBufferStorage.Add(InBuffers[i]);
 				}
 
 				WritableBuffers = WritableBufferStorage;
 				ReadableBuffers = ReadableBufferStorage;
 			}
 
-			/* TODO:
-			// Make a constructor that can take individual arguments of buffers.
-			// TStaticChannelAudio<2>(FAudioWriteRef Val1, FAudioWriteRef Val2)
-			template<
-				typename... ArgTypes,
-				typename = typename TEnableIf<
-					TAndValue<
-						sizeof...(ArgTypes) == NumChannels,
-						TAnd< TIsSameType<FAudioWriteRef, typename TDecay<ArgTypes>... >
-					>
-				>
-			>
-			//TStaticChannelAudio(ArgType&&....)
-			*/
+		private:
+			TArrayView<const FAudioBufferWriteRef> WritableBuffers;
+			TArrayView<const FAudioBufferReadRef> ReadableBuffers;
 
-			int32 GetNumChannels() const
-			{
-				return NumChannels;
-			}
+			TArray<FAudioBufferWriteRef> WritableBufferStorage;
+			TArray<FAudioBufferReadRef> ReadableBufferStorage;
+	};
 
-			template<int32 Index>
-			FAudioBufferReadRef GetBuffer() const
+	/** FMonoAudioFormat represents mono audio containing one channel of audio.
+	 *
+	 * The audio buffer is a shared data references which can be accessed 
+	 * outside of the FMonoAudioFormat. 
+	 */
+	class FMonoAudioFormat : public TStaticChannelAudioFormat<1>
+	{
+		public:
+			using Super = TStaticChannelAudioFormat<1>;
+
+			// Inherit constructors of base class.
+			using Super::Super;
+
+			/** FMonoAudioFormat Construtor
+			 *
+			 * Construct with a single writable audio buffer reference.
+			 *
+			 * @param InAudio - A writable audio buffer reference.
+			 */
+			FMonoAudioFormat(const FAudioBufferWriteRef& InAudio)
+			:	Super({InAudio})
 			{
-				static_assert(Index >= 0, "Index must be within range of channels");
-				static_assert(Index < NumChannels, "Index must be within range of channels");
-
-				return ReadableBuffers.GetData()[Index];
 			}
 
-			template<int32 Index>
-			FAudioBufferWriteRef GetBuffer() 
-			{
-				static_assert(Index >= 0, "Index must be within range of channels");
-				static_assert(Index < NumChannels, "Index must be within range of channels");
+			/** Return writable audio buffer reference of center channel. */
+			FAudioBufferWriteRef GetCenter() { return GetBuffer<0>(); }
 
-				return WritableBuffers.GetData()[Index];
-			}
+			/** Return readable audio buffer reference of center channel. */
+			FAudioBufferReadRef GetCenter() const { return GetBuffer<0>(); }
 
-			TArrayView<const FAudioBufferReadRef> GetBuffers() const
+		private:
+			// Friendship with the data reference class gives it access to the 
+			// protected constructor for the scenario where a 
+			// TDataReadReference<*> is constructed with FAudioBufferReadRefs. 
+			// The constructor cannot be public as it would provide writable
+			// access to the passed in audio buffers, even though the passed 
+			// buffers explicitly were read references.
+			friend class TDataReadReference<FMonoAudioFormat>;
+
+			// Construct a FMonoAudioFormat a readable buffer reference.
+			// 
+			// This constructor should only be used when it can be assured 
+			// that the constructed object will not provide writable access 
+			// to the passed in audio buffers. 
+			FMonoAudioFormat(const FAudioBufferReadRef& InAudio)
+			:	TStaticChannelAudioFormat<1>({WriteCast(InAudio)})
 			{
-				return ReadableBuffers;
 			}
+	};
 
-			TArrayView<const FAudioBufferWriteRef> GetBuffers() 
+	/** FStereoAudioFormat represents stereo audio containing two channels of 
+	 * audio.
+	 *
+	 * The audio buffers are shared data references which can be accessed 
+	 * outside of the FStereoAudioFormat. 
+	 */
+	class FStereoAudioFormat : public TStaticChannelAudioFormat<2>
+	{
+		public:
+			using Super = TStaticChannelAudioFormat<2>;
+
+			// Inherit constructors of base class.
+			using Super::Super;
+
+			/** FStereoAudioFormat Construtor
+			 *
+			 * Construct with a two writable audio buffer reference.
+			 *
+			 * @param InLeftAudio - A writable audio buffer reference for the left channel.
+			 * @param InRightAudio - A writable audio buffer reference for the right channel.
+			 */
+			FStereoAudioFormat(const FAudioBufferWriteRef& InLeftAudio, const FAudioBufferWriteRef& InRightAudio)
+			:	Super({InLeftAudio, InRightAudio})
 			{
-				return WritableBuffers;
 			}
 
-			TArrayView<const FAudioBufferReadRef> GetStorage() const
-			{
-				return ReadableBufferStorage;
-			}
+			/** Return writable audio buffer reference of left channel. */
+			FAudioBufferWriteRef GetLeft() { return GetBuffer<0>(); }
+
+			/** Return readable audio buffer reference of left channel. */
+			FAudioBufferReadRef GetLeft() const { return GetBuffer<0>(); }
 
-			TArrayView<const FAudioBufferWriteRef> GetStorage()
+			/** Return writable audio buffer reference of right channel. */
+			FAudioBufferWriteRef GetRight() { return GetBuffer<1>(); }
+
+			/** Return readable audio buffer reference of right channel. */
+			FAudioBufferReadRef GetRight() const { return GetBuffer<1>(); }
+
+		protected:
+
+			// Friendship with the data reference class gives it access to the 
+			// protected constructor for the scenario where a 
+			// TDataReadReference<*> is constructed with FAudioBufferReadRefs. 
+			// The constructor cannot be public as it would provide writable
+			// access to the passed in audio buffers, even though the passed 
+			// buffers explicitly were read references.
+			friend class TDataReadReference<FStereoAudioFormat>;
+
+			// Construct a FStereoAudioFormat a readable buffer reference.
+			// 
+			// This constructor should only be used when it can be assured 
+			// that the constructed object will not provide writable access 
+			// to the passed in audio buffers. 
+			FStereoAudioFormat(const FAudioBufferReadRef& InLeftAudio, const FAudioBufferReadRef& InRightAudio)
+			:	Super({WriteCast(InLeftAudio), WriteCast(InRightAudio)})
 			{
-				return WritableBufferStorage;
 			}
+	};
 
-		private:
 
-			TArray<const FAudioBufferWriteRef> WritableBufferStorage;
-			TArray<const FAudioBufferReadRef> ReadableBufferStorage;
+	DECLARE_METASOUND_DATA_REFERENCE_TYPES(FUnformattedAudio, "Audio:Unformatted", 0xd78a3ed1 , FUnformattedAudioTypeInfo, FUnformattedAudioReadRef, FUnformattedAudioWriteRef);
 
-	};
+	DECLARE_METASOUND_DATA_REFERENCE_TYPES(FMultichannelAudioFormat, "Audio:Multichannel", 0x56bdcbe0 , FMultichannelAudioFormatTypeInfo, FMultichannelAudioFormatReadRef, FMultichannelAudioFormatWriteRef);
 
-	class FMonoAudio : public TStaticAudio<1>
-	{
-		public:
-			using TStaticAudio<1>::TStaticAudio<1>;
-			using Center = TStaticAudio<1>::GetBuffer<0>;
-	};
+	DECLARE_METASOUND_DATA_REFERENCE_TYPES(FMonoAudioFormat, "Audio:Mono", 0x6f468c8c, FMonoAudioFormatTypeInfo, FMonoAudioFormatReadRef, FMonoAudioFormatWriteRef);
 
-	class FStereoAudio : public TStaticAudio<2>
-	{
-		public:
-			using TStaticAudio<2>::TStaticAudio<2>;
-			using Left = TStaticAudio<2>::GetBuffer<0>;
-			using Right = TStaticAudio<2>::GetBuffer<1>;
-	};
+	DECLARE_METASOUND_DATA_REFERENCE_TYPES(FStereoAudioFormat, "Audio:Stereo", 0xb55304e2 , FStereoAudioFormatTypeInfo, FStereoAudioFormatReadRef, FStereoAudioFormatWriteRef);
 
 }
diff --git a/Engine/Plugins/Runtime/Metasound/Source/MetasoundGraphCore/Public/MetasoundDataReference.h b/Engine/Plugins/Runtime/Metasound/Source/MetasoundGraphCore/Public/MetasoundDataReference.h
--- a/Engine/Plugins/Runtime/Metasound/Source/MetasoundGraphCore/Public/MetasoundDataReference.h
+++ b/Engine/Plugins/Runtime/Metasound/Source/MetasoundGraphCore/Public/MetasoundDataReference.h
@@ -195,16 +195,10 @@ namespace Metasound
 			/** Enable move constructor */
 			TDataWriteReference(TDataWriteReference<DataType>&& Other) = default;
 
-			// TODO: Need to think about whether this protection is needed.
-			/** Disable copy operator. Changes to the internal reference object will break link with
-			 * readable references. */
-			//TDataWriteReference<DataType>& operator=(const TDataWriteReference<DataType>& Other) = delete;
+			/** Enable assignment operator. */
 			TDataWriteReference<DataType>& operator=(const TDataWriteReference<DataType>& Other) = default;
 
-			// TODO: Need to think about whether this protection is needed.
-			/** Disable move operator. Changes to the internal reference object will break link with
-			 * readable references. */
-			//TDataWriteReference<DataType>& operator=(TDataWriteReference<DataType>&& Other) = delete;
+			/** Enable move operator. */
 			TDataWriteReference<DataType>& operator=(TDataWriteReference<DataType>&& Other) = default;
 
 			/** Implicit conversion to a readable parameter. */
@@ -238,18 +232,9 @@ namespace Metasound
 			// Provide access to ObjectReference when converting from Write to Read.
 			friend class TDataReadReference<DataType>;
 
-			/** Cast a TDataReadReference to a TDataWriteReference. 
-			 *
-			 * In general TDataReadReferences should not be converted into TDataWriteReferences unless the caller
-			 * can be certain that no other TDataWriteReference exists for the underlying parameter. Having multiple 
-			 * TDataWriteReferences to the same parameter can cause confusion behavior as values are overwritten in
-			 * an underterministic fashion.
-			 */
-			friend TDataWriteReference<DataType> WriteCast(const TDataReadReference<DataType>& InReadableRef)
-			{
-				// Friend because it calls protected constructor
-				return TDataWriteReference<DataType>(static_cast<const FDataReference&>(InReadableRef));
-			}
+			// Friend because it calls protected constructor
+			template <typename T>
+			friend TDataWriteReference<T> WriteCast(const TDataReadReference<T>& InReadableRef);
 
 		protected:
 
@@ -262,6 +247,21 @@ namespace Metasound
 			}
 	};
 
+	/** Cast a TDataReadReference to a TDataWriteReference. 
+	 *
+	 * In general TDataReadReferences should not be converted into TDataWriteReferences unless the caller
+	 * can be certain that no other TDataWriteReference exists for the underlying parameter. Having multiple 
+	 * TDataWriteReferences to the same parameter can cause confusion behavior as values are overwritten in
+	 * an underterministic fashion.
+	 */
+	template <typename DataType>
+	TDataWriteReference<DataType> WriteCast(const TDataReadReference<DataType>& InReadableRef)
+	{
+		const TDataReference<DataType>& Ref = static_cast<const TDataReference<DataType>&>(InReadableRef); 
+		return TDataWriteReference<DataType>(Ref);
+	}
+
+
 
 	/** TDataReadReference provides read access to a shared parameter reference. */
 	template <typename DataType>
diff --git a/Engine/Plugins/Runtime/Metasound/Source/MetasoundGraphCore/Public/MetasoundDataReferenceTypes.h b/Engine/Plugins/Runtime/Metasound/Source/MetasoundGraphCore/Public/MetasoundDataReferenceTypes.h
--- a/Engine/Plugins/Runtime/Metasound/Source/MetasoundGraphCore/Public/MetasoundDataReferenceTypes.h
+++ b/Engine/Plugins/Runtime/Metasound/Source/MetasoundGraphCore/Public/MetasoundDataReferenceTypes.h
@@ -91,7 +91,6 @@ namespace Metasound
 
 	typedef float FFloat;
 
-	DECLARE_METASOUND_DATA_REFERENCE_TYPES(FAudioBuffer, "AudioBuffer", 1, FAudioBufferTypeInfo, FAudioBufferReadRef, FAudioBufferWriteRef);
 	DECLARE_METASOUND_DATA_REFERENCE_TYPES(FBop, "Bop", 2, FBopTypeInfo, FBopReadRef, FBopWriteRef);
 	DECLARE_METASOUND_DATA_REFERENCE_TYPES(FTimeSeconds, "Time:Seconds", 3, FTimeSecondsTypeInfo, FTimeSecondsReadRef, FTimeSecondsWriteRef);
 	DECLARE_METASOUND_DATA_REFERENCE_TYPES(FTimeMilliseconds, "Time:Milliseconds", 4, FTimeMillisecondsTypeInfo, FTimeMillisecondsReadRef, FTimeMillisecondsWriteRef);
