Minor build fixes to allow metasounds on all platforms

diff --git a/Engine/Plugins/Runtime/Metasound/Metasound.uplugin b/Engine/Plugins/Runtime/Metasound/Metasound.uplugin
--- a/Engine/Plugins/Runtime/Metasound/Metasound.uplugin
+++ b/Engine/Plugins/Runtime/Metasound/Metasound.uplugin
@@ -19,18 +19,12 @@
 		{
 			"Name" : "MetasoundGraphCore",
 			"Type" : "Runtime",
-			"LoadingPhase" : "PreDefault",
-			"WhitelistPlatforms": [
-				"Win64"
-			]
+			"LoadingPhase" : "PreDefault"
 		},
 		{
 			"Name" : "MetasoundGenerator",
 			"Type" : "Runtime",
-			"LoadingPhase" : "Default",
-			"WhitelistPlatforms": [
-				"Win64"
-			]
+			"LoadingPhase" : "Default"
 		}
 	]
 }
diff --git a/Engine/Plugins/Runtime/Metasound/Source/MetasoundGraphCore/Private/MetasoundGraphCoreModule.cpp b/Engine/Plugins/Runtime/Metasound/Source/MetasoundGraphCore/Private/MetasoundGraphCoreModule.cpp
--- a/Engine/Plugins/Runtime/Metasound/Source/MetasoundGraphCore/Private/MetasoundGraphCoreModule.cpp
+++ b/Engine/Plugins/Runtime/Metasound/Source/MetasoundGraphCore/Private/MetasoundGraphCoreModule.cpp
@@ -1,7 +1,5 @@
 // Copyright Epic Games, Inc. All Rights Reserved.
 
-#pragma once
-
 #include "MetasoundGraphCoreModule.h"
 
 #include "CoreMinimal.h"
diff --git a/Engine/Plugins/Runtime/Metasound/Source/MetasoundGraphCore/Public/MetasoundDataReference.h b/Engine/Plugins/Runtime/Metasound/Source/MetasoundGraphCore/Public/MetasoundDataReference.h
--- a/Engine/Plugins/Runtime/Metasound/Source/MetasoundGraphCore/Public/MetasoundDataReference.h
+++ b/Engine/Plugins/Runtime/Metasound/Source/MetasoundGraphCore/Public/MetasoundDataReference.h
@@ -106,7 +106,7 @@ namespace Metasound
 			/** Construct operator with no arguments if the DataType has a default constructor.  */
 			template< typename = typename TEnableIf< TIsConstructible<DataType>::Value >::Type >
 			TDataReference()
-			:	ObjectReference(MakeShared<DataType>())
+			:	ObjectReference(MakeShared<DataType, ESPMode::NotThreadSafe>())
 			{
 			}
 
