Adding TTime and FFrequency metasound classes.
Adding several primative types.

diff --git a/Engine/Plugins/Runtime/Metasound/Source/MetasoundGenerator/Private/MetasoundGenerator.cpp b/Engine/Plugins/Runtime/Metasound/Source/MetasoundGenerator/Private/MetasoundGenerator.cpp
--- a/Engine/Plugins/Runtime/Metasound/Source/MetasoundGenerator/Private/MetasoundGenerator.cpp
+++ b/Engine/Plugins/Runtime/Metasound/Source/MetasoundGenerator/Private/MetasoundGenerator.cpp
@@ -12,6 +12,7 @@
 #include "MetasoundInputNode.h"
 #include "MetasoundOperatorBuilder.h"
 #include "MetasoundGraph.h"
+#include "MetasoundFrequency.h"
 
 namespace Metasound
 {
@@ -26,14 +27,14 @@ namespace Metasound
 			ExecuteOperator = Operator->GetExecuteFunction();
 			FDataReferenceCollection Inputs = Operator->GetInputs();
 
-			if (Inputs.ContainsDataWriteReference<FFloat>(TEXT("Frequency")))
+			if (Inputs.ContainsDataWriteReference<FFrequency>(TEXT("Frequency")))
 			{
-				FrequencyRef = Inputs.GetDataWriteReference<FFloat>(TEXT("Frequency"));
+				FrequencyRef = Inputs.GetDataWriteReference<FFrequency>(TEXT("Frequency"));
 			}
 
-			if (Inputs.ContainsDataWriteReference<FTimeSeconds>(TEXT("BopPeriod")))
+			if (Inputs.ContainsDataWriteReference<FFloatTime>(TEXT("BopPeriod")))
 			{
-				BopPeriodRef = Inputs.GetDataWriteReference<FTimeSeconds>(TEXT("BopPeriod"));
+				BopPeriodRef = Inputs.GetDataWriteReference<FFloatTime>(TEXT("BopPeriod"));
 			}
 		}
 	}
@@ -44,12 +45,12 @@ namespace Metasound
 
 	void FMetasoundGenerator::SetFrequency(float InFrequency)
 	{
-		*FrequencyRef = InFrequency;
+		FrequencyRef->SetHertz(InFrequency);
 	}
 
 	void FMetasoundGenerator::SetBopPeriod(float InPeriodInSeconds)
 	{
-		BopPeriodRef->Seconds = InPeriodInSeconds;
+		BopPeriodRef->SetSeconds(InPeriodInSeconds);
 	}
 
 	int32 FMetasoundGenerator::OnGenerateAudio(float* OutAudio, int32 NumSamples)
@@ -191,8 +192,8 @@ ISoundGeneratorPtr USynthComponentMetasoundGenerator::CreateSoundGenerator(int32
 	using FDataEdge = Metasound::FDataEdge; // Need to make explicit to disambiguate with existing FDataEdge.
 	using FOperatorUniquePtr = TUniquePtr<Metasound::IOperator>;
 
-	TInputNode<FTimeSeconds> BopPeriodInputNode(TEXT("BopPeriod"), 1.f);
-	TInputNode<FFloat> FrequencyInputNode(TEXT("Frequency"), 100.f);
+	TInputNode<FFloatTime> BopPeriodInputNode(TEXT("BopPeriod"), 1.f, ETimeResolution::Seconds);
+	TInputNode<FFrequency> FrequencyInputNode(TEXT("Frequency"), 100.f, EFrequencyResolution::Hertz);
 
 	FPeriodicBopNode BopNode(TEXT("Bop"), 1.f);
 	FOscNode OscNode(TEXT("Osc"), 100.f);
diff --git a/Engine/Plugins/Runtime/Metasound/Source/MetasoundGenerator/Public/MetasoundGenerator.h b/Engine/Plugins/Runtime/Metasound/Source/MetasoundGenerator/Public/MetasoundGenerator.h
--- a/Engine/Plugins/Runtime/Metasound/Source/MetasoundGenerator/Public/MetasoundGenerator.h
+++ b/Engine/Plugins/Runtime/Metasound/Source/MetasoundGenerator/Public/MetasoundGenerator.h
@@ -7,6 +7,8 @@
 #include "DSP/Dsp.h"
 #include "MetasoundOperatorInterface.h"
 #include "MetasoundDataReferenceTypes.h"
+#include "MetasoundTime.h"
+#include "MetasoundFrequency.h"
 #include "Sound/SoundGenerator.h"
 
 #include "MetasoundGenerator.generated.h"
@@ -38,8 +40,8 @@ namespace Metasound
 		FOperatorUniquePtr Operator;
 		FAudioReadRef OperatorReadBuffer;
 
-		FFloatWriteRef FrequencyRef;
-		FTimeSecondsWriteRef BopPeriodRef;
+		FFrequencyWriteRef FrequencyRef;
+		FFloatTimeWriteRef BopPeriodRef;
 
 		int32 OperatorReadNum;
 		FExecuteFunction ExecuteOperator;
diff --git a/Engine/Plugins/Runtime/Metasound/Source/MetasoundGraphCore/Private/MetasoundADSRNode.cpp b/Engine/Plugins/Runtime/Metasound/Source/MetasoundGraphCore/Private/MetasoundADSRNode.cpp
--- a/Engine/Plugins/Runtime/Metasound/Source/MetasoundGraphCore/Private/MetasoundADSRNode.cpp
+++ b/Engine/Plugins/Runtime/Metasound/Source/MetasoundGraphCore/Private/MetasoundADSRNode.cpp
@@ -2,6 +2,7 @@
 
 #include "MetasoundADSRNode.h"
 #include "MetasoundExecutableOperator.h"
+#include "MetasoundTime.h"
 #include "MetasoundDataReferenceTypes.h"
 #include "DSP/Envelope.h"
 
@@ -11,10 +12,10 @@ namespace Metasound
 {
 	struct FADSRDataReferences
 	{
-		FTimeMillisecondsReadRef Attack;
-		FTimeMillisecondsReadRef Decay;
-		FTimeMillisecondsReadRef Sustain;
-		FTimeMillisecondsReadRef Release;
+		FFloatTimeReadRef Attack;
+		FFloatTimeReadRef Decay;
+		FFloatTimeReadRef Sustain;
+		FFloatTimeReadRef Release;
 	};
 
 	class FADSROperator : public TExecutableOperator<FADSROperator>
@@ -46,13 +47,13 @@ namespace Metasound
 
 			void Execute()
 			{
-				Envelope.SetAttackTime(ADSRDataReferences.Attack->Milliseconds);
-				Envelope.SetDecayTime(ADSRDataReferences.Decay->Milliseconds);
-				Envelope.SetReleaseTime(ADSRDataReferences.Release->Milliseconds);
+				Envelope.SetAttackTime(ADSRDataReferences.Attack->GetMilliseconds());
+				Envelope.SetDecayTime(ADSRDataReferences.Decay->GetMilliseconds());
+				Envelope.SetReleaseTime(ADSRDataReferences.Release->GetMilliseconds());
 
-				float ADSMilliseconds = ADSRDataReferences.Attack->Milliseconds;
-				ADSMilliseconds += ADSRDataReferences.Decay->Milliseconds;
-				ADSMilliseconds += ADSRDataReferences.Sustain->Milliseconds;
+				float ADSMilliseconds = ADSRDataReferences.Attack->GetMilliseconds();
+				ADSMilliseconds += ADSRDataReferences.Decay->GetMilliseconds();
+				ADSMilliseconds += ADSRDataReferences.Sustain->GetMilliseconds();
 
 				// TODO: Bops need to be sorted.
 				// TODO: Make convenient way to do bop loops like this since it might happen alot.
@@ -134,10 +135,10 @@ namespace Metasound
 
 		FADSRDataReferences ADSRDataReferences = 
 			{
-				FTimeMillisecondsReadRef(ADSRNode.GetDefaultAttackMs()),
-				FTimeMillisecondsReadRef(ADSRNode.GetDefaultDecayMs()),
-				FTimeMillisecondsReadRef(ADSRNode.GetDefaultSustainMs()),
-				FTimeMillisecondsReadRef(ADSRNode.GetDefaultReleaseMs())
+				FFloatTimeReadRef(ADSRNode.GetDefaultAttackMs(), ETimeResolution::Milliseconds),
+				FFloatTimeReadRef(ADSRNode.GetDefaultDecayMs(), ETimeResolution::Milliseconds),
+				FFloatTimeReadRef(ADSRNode.GetDefaultSustainMs(), ETimeResolution::Milliseconds),
+				FFloatTimeReadRef(ADSRNode.GetDefaultReleaseMs(), ETimeResolution::Milliseconds)
 			};
 
 		// TODO: If none of these are connected, could pregenerate ADSR envelope and return a different operator. 
@@ -157,10 +158,10 @@ namespace Metasound
 	,	DefaultReleaseMs(InDefaultReleaseMs)
 	{
 		AddInputDataVertexDescription<FBop>(TEXT("Bop"), LOCTEXT("BopTooltip", "Trigger for envelope."));
-		AddInputDataVertexDescription<FTimeMilliseconds>(TEXT("Attack"), LOCTEXT("AttackTooltip", "Attack time in milliseconds."));
-		AddInputDataVertexDescription<FTimeMilliseconds>(TEXT("Decay"), LOCTEXT("DecayTooltip", "Decay time in milliseconds."));
-		AddInputDataVertexDescription<FTimeMilliseconds>(TEXT("Sustain"), LOCTEXT("SustainTooltip", "Sustain time in milliseconds."));
-		AddInputDataVertexDescription<FTimeMilliseconds>(TEXT("Release"), LOCTEXT("ReleaseTooltip", "Release time in milliseconds."));
+		AddInputDataVertexDescription<FFloatTime>(TEXT("Attack"), LOCTEXT("AttackTooltip", "Attack time in milliseconds."));
+		AddInputDataVertexDescription<FFloatTime>(TEXT("Decay"), LOCTEXT("DecayTooltip", "Decay time in milliseconds."));
+		AddInputDataVertexDescription<FFloatTime>(TEXT("Sustain"), LOCTEXT("SustainTooltip", "Sustain time in milliseconds."));
+		AddInputDataVertexDescription<FFloatTime>(TEXT("Release"), LOCTEXT("ReleaseTooltip", "Release time in milliseconds."));
 
 		AddOutputDataVertexDescription<FAudioBuffer>(TEXT("Envelope"), LOCTEXT("EnvelopeTooltip", "The output envelope"));
 	}
diff --git a/Engine/Plugins/Runtime/Metasound/Source/MetasoundGraphCore/Private/MetasoundOscNode.cpp b/Engine/Plugins/Runtime/Metasound/Source/MetasoundGraphCore/Private/MetasoundOscNode.cpp
--- a/Engine/Plugins/Runtime/Metasound/Source/MetasoundGraphCore/Private/MetasoundOscNode.cpp
+++ b/Engine/Plugins/Runtime/Metasound/Source/MetasoundGraphCore/Private/MetasoundOscNode.cpp
@@ -2,7 +2,7 @@
 
 #include "MetasoundOscNode.h"
 #include "MetasoundExecutableOperator.h"
-#include "MetasoundDataReferenceTypes.h"
+#include "MetasoundPrimitives.h"
 
 #define LOCTEXT_NAMESPACE "MetasoundOscNode"
 
@@ -11,7 +11,7 @@ namespace Metasound
 	class FOscOperator : public TExecutableOperator<FOscOperator>
 	{
 		public:
-			FOscOperator(const FOperatorSettings& InSettings, const FFloatReadRef& InFrequency)
+			FOscOperator(const FOperatorSettings& InSettings, const FFrequencyReadRef& InFrequency)
 			:	OperatorSettings(InSettings)
 			,	TwoPi(2.f * PI)
 			,	Phase(0.f)
@@ -35,7 +35,7 @@ namespace Metasound
 
 			void Execute()
 			{
-				const float PhaseDelta = *Frequency * 2 * PI / OperatorSettings.SampleRate;
+				const float PhaseDelta = Frequency->GetRadiansPerSample(OperatorSettings.SampleRate);
 				float* Data = AudioBuffer->GetData();
 
 				for (int32 i = 0; i < OperatorSettings.FramesPerExecute; i++)
@@ -52,7 +52,7 @@ namespace Metasound
 			const float TwoPi;
 			float Phase;
 
-			FFloatReadRef Frequency;
+			FFrequencyReadRef Frequency;
 			FAudioBufferWriteRef AudioBuffer;
 
 			FDataReferenceCollection InputDataReferences;
@@ -64,11 +64,11 @@ namespace Metasound
 	TUniquePtr<IOperator> FOscNode::FOperatorFactory::CreateOperator(const INode& InNode, const FOperatorSettings& InOperatorSettings, const FDataReferenceCollection& InInputDataReferences, TArray<TUniquePtr<IOperatorBuildError>>& OutErrors) 
 	{
 		const FOscNode& OscNode = static_cast<const FOscNode&>(InNode);
-		FFloatReadRef Frequency(OscNode.GetDefaultFrequency());
+		FFrequencyReadRef Frequency(OscNode.GetDefaultFrequency(), EFrequencyResolution::Hertz);
 
-		if (InInputDataReferences.ContainsDataReadReference<FFloat>(TEXT("Frequency")))
+		if (InInputDataReferences.ContainsDataReadReference<FFrequency>(TEXT("Frequency")))
 		{
-			Frequency = InInputDataReferences.GetDataReadReference<FFloat>(TEXT("Frequency"));
+			Frequency = InInputDataReferences.GetDataReadReference<FFrequency>(TEXT("Frequency"));
 		}
 
 		return MakeUnique<FOscOperator>(InOperatorSettings, Frequency);
@@ -78,7 +78,7 @@ namespace Metasound
 	:	FNode(InName)
 	,	DefaultFrequency(InDefaultFrequency)
 	{
-		AddInputDataVertexDescription<FFloat>(TEXT("Frequency"), LOCTEXT("FrequencyTooltip", "The frequency of oscillator in Hz."));
+		AddInputDataVertexDescription<FFrequency>(TEXT("Frequency"), LOCTEXT("FrequencyTooltip", "The frequency of oscillator."));
 		AddOutputDataVertexDescription<FAudioBuffer>(TEXT("Audio"), LOCTEXT("AudioTooltip", "The output audio"));
 	}
 
diff --git a/Engine/Plugins/Runtime/Metasound/Source/MetasoundGraphCore/Private/MetasoundPeriodicBopNode.cpp b/Engine/Plugins/Runtime/Metasound/Source/MetasoundGraphCore/Private/MetasoundPeriodicBopNode.cpp
--- a/Engine/Plugins/Runtime/Metasound/Source/MetasoundGraphCore/Private/MetasoundPeriodicBopNode.cpp
+++ b/Engine/Plugins/Runtime/Metasound/Source/MetasoundGraphCore/Private/MetasoundPeriodicBopNode.cpp
@@ -8,8 +8,8 @@
 #include "MetasoundBuilderInterface.h"
 #include "MetasoundNodeInterface.h"
 #include "MetasoundDataReferenceCollection.h"
-#include "MetasoundDataReferenceTypes.h"
 #include "MetasoundExecutableOperator.h"
+#include "MetasoundPrimitives.h"
 
 #define LOCTEXT_NAMESPACE "MetasoundPeriodicBopNode"
 
@@ -21,7 +21,7 @@ namespace Metasound
 			static constexpr float MinimumPeriodSeconds = 0.001f;
 			static constexpr float MinimumPeriodSamples = 20.f;
 
-			FPeriodicBopOperator(const FOperatorSettings& InSettings, const FTimeSecondsReadRef& InPeriod)
+			FPeriodicBopOperator(const FOperatorSettings& InSettings, const FFloatTimeReadRef& InPeriod)
 			:	OperatorSettings(InSettings)
 			,	Period(InPeriod)
 			,	ExecuteDurationInSamples(InSettings.FramesPerExecute)
@@ -45,7 +45,7 @@ namespace Metasound
 				// Advance internal counter to get rid of old bops.
 				Bop->Advance(OperatorSettings.FramesPerExecute);
 
-				float PeriodInSamples = FMath::Max(Period->Seconds, MinimumPeriodSeconds) * OperatorSettings.SampleRate;
+				float PeriodInSamples = FMath::Max(Period->GetSeconds(), MinimumPeriodSeconds) * OperatorSettings.SampleRate;
 
 				PeriodInSamples = FMath::Max(PeriodInSamples, MinimumPeriodSamples);
 
@@ -63,7 +63,7 @@ namespace Metasound
 
 			FOperatorSettings OperatorSettings;
 			FBopWriteRef Bop;
-			FTimeSecondsReadRef Period;
+			FFloatTimeReadRef Period;
 
 			float ExecuteDurationInSamples;
 			float SampleCountdown;
@@ -76,11 +76,11 @@ namespace Metasound
 	TUniquePtr<IOperator> FPeriodicBopNode::FOperatorFactory::CreateOperator(const INode& InNode, const FOperatorSettings& InOperatorSettings, const FDataReferenceCollection& InInputDataReferences, TArray<TUniquePtr<IOperatorBuildError>>& OutErrors)
 	{
 		const FPeriodicBopNode& PeriodicBopNode = static_cast<const FPeriodicBopNode&>(InNode);
-		FTimeSecondsReadRef Period(PeriodicBopNode.GetDefaultPeriodInSeconds());
+		FFloatTimeReadRef Period(PeriodicBopNode.GetDefaultPeriodInSeconds(), ETimeResolution::Seconds);
 
-		if (InInputDataReferences.ContainsDataReadReference<FTimeSeconds>(TEXT("Period")))
+		if (InInputDataReferences.ContainsDataReadReference<FFloatTime>(TEXT("Period")))
 		{
-			Period = InInputDataReferences.GetDataReadReference<FTimeSeconds>(TEXT("Period"));
+			Period = InInputDataReferences.GetDataReadReference<FFloatTime>(TEXT("Period"));
 		}
 
 		return MakeUnique<FPeriodicBopOperator>(InOperatorSettings, Period);
@@ -92,7 +92,7 @@ namespace Metasound
 	:	FNode(InName)
 	,	DefaultPeriod(InDefaultPeriodInSeconds)
 	{
-		AddInputDataVertexDescription<FTimeSeconds>(TEXT("Period"), LOCTEXT("PeriodTooltip", "The period of the bops in seconds."));
+		AddInputDataVertexDescription<FFloatTime>(TEXT("Period"), LOCTEXT("PeriodTooltip", "The period of the bops in seconds."));
 		AddOutputDataVertexDescription<FBop>(TEXT("Bop"), LOCTEXT("BopTooltip", "The output bop"));
 	}
 
diff --git a/Engine/Plugins/Runtime/Metasound/Source/MetasoundGraphCore/Public/MetasoundDataReferenceTypes.h b/Engine/Plugins/Runtime/Metasound/Source/MetasoundGraphCore/Public/MetasoundDataReferenceTypes.h
--- a/Engine/Plugins/Runtime/Metasound/Source/MetasoundGraphCore/Public/MetasoundDataReferenceTypes.h
+++ b/Engine/Plugins/Runtime/Metasound/Source/MetasoundGraphCore/Public/MetasoundDataReferenceTypes.h
@@ -59,40 +59,5 @@ namespace Metasound
 			TArray<uint32> Frames;
 	};
 
-	struct FTimeSeconds
-	{
-		FTimeSeconds()
-		:	Seconds(0.f)
-		{
-		}
-
-		FTimeSeconds(float InSeconds)
-		:	Seconds(InSeconds)
-		{
-		}
-
-		float Seconds;
-	};
-
-	struct FTimeMilliseconds
-	{
-		FTimeMilliseconds()
-		:	Milliseconds(0.f)
-		{
-		}
-
-		FTimeMilliseconds(float InMilliseconds)
-		:	Milliseconds(InMilliseconds)
-		{
-		}
-
-		float Milliseconds;
-	};
-
-	typedef float FFloat;
-
 	DECLARE_METASOUND_DATA_REFERENCE_TYPES(FBop, "Bop", 2, FBopTypeInfo, FBopReadRef, FBopWriteRef);
-	DECLARE_METASOUND_DATA_REFERENCE_TYPES(FTimeSeconds, "Time:Seconds", 3, FTimeSecondsTypeInfo, FTimeSecondsReadRef, FTimeSecondsWriteRef);
-	DECLARE_METASOUND_DATA_REFERENCE_TYPES(FTimeMilliseconds, "Time:Milliseconds", 4, FTimeMillisecondsTypeInfo, FTimeMillisecondsReadRef, FTimeMillisecondsWriteRef);
-	DECLARE_METASOUND_DATA_REFERENCE_TYPES(FFloat, "Primitive:Float", 5, FFloatTypeInfo, FFloatReadRef, FFloatWriteRef);
 }
diff --git a/Engine/Plugins/Runtime/Metasound/Source/MetasoundGraphCore/Public/MetasoundFrequency.h b/Engine/Plugins/Runtime/Metasound/Source/MetasoundGraphCore/Public/MetasoundFrequency.h
new file mode 100644
--- /dev/null
+++ b/Engine/Plugins/Runtime/Metasound/Source/MetasoundGraphCore/Public/MetasoundFrequency.h
@@ -0,0 +1,147 @@
+// Copyright Epic Games, Inc. All Rights Reserved.
+
+#pragma once
+
+#include "MetasoundDataReference.h"
+
+namespace Metasound
+{
+	/** Resolution of frequency value */
+	enum class EFrequencyResolution : uint8
+	{
+		/** Frequency value given in hertz. */
+		Hertz,
+
+		/** Frequency value given in kilohertz. */
+		Kilohertz,
+
+		/** Frequency value given in megahertz. */
+		Megahertz
+	};
+
+	/** FFrequency represents a frequency value. It provides clearly defined 
+	 * getters and setters as well as several convenience functions for 
+	 * converting frequencies to radians per sample..
+	 */
+	class FFrequency
+	{
+		public:
+			/** Construct default frequency of 0 hz. */
+			FFrequency()
+			:	FrequencyInHz(0.f)
+			{
+			}
+
+			/** Construct a frequency with a given value and resolution.
+			 *
+			 * @param InValue - The value of the frequency at the given 
+			 * 					resolution.
+			 * @param InResolution - The resolution of the given value.
+			 */
+			FFrequency(float InValue, EFrequencyResolution InResolution)
+			:	FrequencyInHz(0.f)
+			{
+				switch (InResolution)
+				{
+					case EFrequencyResolution::Megahertz:
+						SetMegahertz(InValue);
+						break;
+
+					case EFrequencyResolution::Kilohertz:
+						SetKilohertz(InValue);
+						break;
+
+					case EFrequencyResolution::Hertz:
+					default:
+						SetHertz(InValue);
+						break;
+				}
+			}
+
+			/** Set the frequency in hertz. */
+			void SetHertz(float InHz)
+			{
+				FrequencyInHz = InHz;	
+			}
+
+			/** Set the frequency in kilohertz. */
+			void SetKilohertz(float InKHz)
+			{
+				FrequencyInHz = InKHz * 1e3f;
+			}
+
+			/** Set the frequency in megahertz. */
+			void SetMegahertz(float InMHz)
+			{
+				FrequencyInHz = InMHz * 1e6f; 
+			}
+
+			/** Set the frequency in radians per a sample.
+			 *
+			 * @param InRadians - The number of radians per a sample.
+			 * @param InSampleRate - The number of samples per a second.
+			 */
+			void SetRadiansPerSample(float InRadians, float InSampleRate)
+			{
+				FrequencyInHz = InSampleRate * InRadians / (2.f * PI);
+			}
+
+			/** Return the frequency in hertz. */
+			float GetHertz() const
+			{
+				return FrequencyInHz;
+			}
+
+			/** Return the frequency in kilohertz. */
+			float GetKilohertz() const
+			{
+				return FrequencyInHz * 1e-3f;
+			}
+
+			/** Return the frequency in megahertz. */
+			float GetMegahertz() const
+			{
+				return FrequencyInHz * 1e-6f;
+			}
+
+			/** Return the radians per a sample.
+			 *
+			 * @param InSampleRate - The number of samples per a second.
+			 *
+			 * @return The phase increment as radians per a sample.
+			 */
+			float GetRadiansPerSample(float InSampleRate) const
+			{
+				if (!ensure(InSampleRate > 0.f))
+				{
+					InSampleRate = SMALL_NUMBER;
+				}
+
+				return (FrequencyInHz * 2.f * PI) / InSampleRate;
+			}
+
+			/** Return the number of radians for a given number of samples.
+			 *
+			 * @param InNumSamples - The number of samples.
+			 * @param InSampleRate - The number of samples per a second.
+			 *
+			 * @return The number of radians for the given samples.
+			 */
+			float GetRadians(float InNumSamples, float InSampleRate) const
+			{
+				if (!ensure(InSampleRate > 0.f))
+				{
+					InSampleRate = SMALL_NUMBER;
+				}
+
+				return (InNumSamples * FrequencyInHz * 2.f * PI) / InSampleRate;
+			}
+		  
+		private:
+
+			float FrequencyInHz;
+	};
+
+	DECLARE_METASOUND_DATA_REFERENCE_TYPES(FFrequency, "Primitive:Frequency", 0xf1c430a4 , FFrequencyTypeInfo, FFrequencyReadRef, FFrequencyWriteRef);
+
+}
diff --git a/Engine/Plugins/Runtime/Metasound/Source/MetasoundGraphCore/Public/MetasoundPrimitives.h b/Engine/Plugins/Runtime/Metasound/Source/MetasoundGraphCore/Public/MetasoundPrimitives.h
new file mode 100644
--- /dev/null
+++ b/Engine/Plugins/Runtime/Metasound/Source/MetasoundGraphCore/Public/MetasoundPrimitives.h
@@ -0,0 +1,25 @@
+// Copyright Epic Games, Inc. All Rights Reserved.
+
+#pragma once
+
+#include "CoreMinimal.h"
+#include "MetasoundDataReference.h"
+
+#include "MetasoundTime.h"
+#include "MetasoundFrequency.h"
+#include "MetasoundAudioBuffer.h"
+#include "MetasoundAudioFormats.h"
+
+namespace Metasound
+{
+	DECLARE_METASOUND_DATA_REFERENCE_TYPES(bool, "Primitive:Bool", 0xcc9d715e, FBoolTypeInfo, FBoolReadRef, FBoolWriteRef);
+
+	DECLARE_METASOUND_DATA_REFERENCE_TYPES(int32, "Primitive:Int32", 0xf9920034, FInt32TypeInfo, FInt32ReadRef, FInt32WriteRef);
+
+	DECLARE_METASOUND_DATA_REFERENCE_TYPES(int64, "Primitive:Int64", 0xfc8f8ada, FInt64TypeInfo, FInt64ReadRef, FInt64WriteRef);
+
+	DECLARE_METASOUND_DATA_REFERENCE_TYPES(float, "Primitive:Float", 0xe7adc4db, FFloatTypeInfo, FFloatReadRef, FFloatWriteRef);
+
+	DECLARE_METASOUND_DATA_REFERENCE_TYPES(double, "Primitive:Double", 0x8f48a053, FDoubleTypeInfo, FDoubleReadRef, FDoubleWriteRef);
+}
+
diff --git a/Engine/Plugins/Runtime/Metasound/Source/MetasoundGraphCore/Public/MetasoundTime.h b/Engine/Plugins/Runtime/Metasound/Source/MetasoundGraphCore/Public/MetasoundTime.h
new file mode 100644
--- /dev/null
+++ b/Engine/Plugins/Runtime/Metasound/Source/MetasoundGraphCore/Public/MetasoundTime.h
@@ -0,0 +1,623 @@
+// Copyright Epic Games, Inc. All Rights Reserved.
+
+#pragma once
+
+#include "MetasoundDataReference.h"
+
+namespace Metasound
+{ 
+	/** Resolution of time value */
+	enum class ETimeResolution : uint8
+	{
+		/** Time value given in seconds. */
+		Seconds,
+
+		/** Time value given in milliseconds. */
+		Milliseconds,
+
+		/** Time value given in microseconds. */
+		Microseconds
+	};
+
+	/** TTime represents a time value. It provides basic mathematical 
+	 * operations, clearly defined getters and setters, and interoperability 
+	 * with TTime objects utilizing different template arguments. 
+	 *
+	 * QuantizationType declares the type of object which stores time 
+	 * information.
+	 *
+	 * TimeType declares the type of object which represent time values 
+	 * returned from TTime's getters and accepted by TTime's setters.
+	 */
+	template<typename QuantizationType, typename TimeType>
+	class TTime
+	{
+			static_assert(TIsFloatingPoint<TimeType>::Value, "TimeType must be floating point");
+
+		public:
+			/** Default constructor */
+			TTime()
+			:	Value(0)
+			{}
+			
+			/** Initialize with a time value and time resolution.
+			 *
+			 * @param InValue - The initial value of the time object for the 
+			 * 					given resolution.
+			 * @param InResolution - The Resolution of the given InValue.
+			 */
+			TTime(QuantizationType InValue, ETimeResolution InResolution)
+			:	Value(0)
+			{
+				// Different setters called based upon resolution.
+				switch (InResolution)
+				{
+					case ETimeResolution::Microseconds:
+						SetMicroseconds(InValue);
+						break;
+
+					case ETimeResolution::Milliseconds:
+						SetMilliseconds(InValue);
+						break;
+
+					case ETimeResolution::Seconds:
+					default:
+						SetSeconds(InValue);
+				}
+			}
+
+			/** Returns the time as seconds. */
+			TimeType GetSeconds() const 
+			{ 
+				return static_cast<TimeType>(Value);
+			}
+
+			/** Returns the time as milliseconds. */
+			TimeType GetMilliseconds() const 
+			{ 
+				return static_cast<TimeType>(Value * QuantizationType(1e3));
+			}
+
+			/** Returns the time as microseconds. */
+			TimeType GetMicroseconds() const 
+			{ 
+				return static_cast<TimeType>(Value * QuantizationType(1e6));
+			}
+
+			/** Returns the number of samples which represent the current time 
+			 * value.
+			 *
+			 * @param InSampleRate - The sample rate to use when calculating
+			 * 						 the number of samples. 
+			 *
+			 * @return The number of samples. 
+			 */
+			QuantizationType GetNumSamples(TimeType InSampleRate) const
+			{
+				return static_cast<QuantizationType>(Value * InSampleRate);
+			}
+
+			/** Set the time value in seconds. */
+			void SetSeconds(TimeType InSeconds) 
+			{ 
+				Value = static_cast<QuantizationType>(InSeconds);
+			}
+
+			/** Set the time value in milliseconds. */
+			void SetMilliseconds(TimeType InMilliseconds) 
+			{ 
+				Value = static_cast<QuantizationType>(InMilliseconds * TimeType(1e-3));
+			}
+
+			/** Set the time value in microseconds. */
+			void SetMicroseconds(TimeType InMicroseconds) 
+			{ 
+				Value = static_cast<QuantizationType>(InMicroseconds * TimeType(1e-6));
+			}
+
+			/** Set the number of samples to represent. 
+			 *
+			 * @param InNumSamples - The number of samples for the TTime object 
+			 * 						 to represent.
+			 * @param InSampleRate - The sample rate to use when calculating
+			 * 						 the number of samples. 
+			 */
+			void SetNumSamples(QuantizationType InNumSamples, TimeType InSampleRate)
+			{ 
+				SetSeconds(InNumSamples / FMath::Max(InSampleRate, static_cast<TimeType>(SMALL_NUMBER)));
+			}
+
+			/** Assignment operator. */
+			template<typename OtherQuantizationType, typename OtherTimeType>
+			TTime& operator=(const TTime<OtherQuantizationType, OtherTimeType>& InOther)
+			{
+				Value = static_cast<QuantizationType>(InOther.GetSeconds());
+				return *this;
+			}
+
+			/** Addition assignment operator. 
+			 * 
+			 * Adds the Other's time to this time.
+			 */
+			template<typename OtherQuantizationType, typename OtherTimeType>
+			TTime& operator+=(const TTime<OtherQuantizationType, OtherTimeType>& InOther)
+			{
+				Value += static_cast<QuantizationType>(InOther.GetSeconds());
+				return *this;
+			}
+
+			/** Subtraction assignment operator. 
+			 * 
+			 * Subtracts the Other's time from this time.
+			 */
+			template<typename OtherQuantizationType, typename OtherTimeType>
+			TTime& operator-=(const TTime<OtherQuantizationType, OtherTimeType>& InOther)
+			{
+				Value -= static_cast<QuantizationType>(InOther.GetSeconds());
+				return *this;
+			}
+
+			/* Multiplication assignment operator. 
+			 *
+			 * Multiplies this time by the value. 
+			 */
+			template<typename OtherType>
+			TTime& operator*=(OtherType InValue)
+			{
+				static_assert(TIsArithmetic<OtherType>::Value, "TTime can only be multiplied with arithmetic types");
+
+				Value *= InValue;
+				return *this;
+			}
+
+			/* Division assignment operator. 
+			 *
+			 * Divides this time by the value. 
+			 */
+			template<typename OtherType>
+			TTime& operator/=(OtherType InValue)
+			{
+				static_assert(TIsArithmetic<OtherType>::Value, "TTime can only be divided with arithmetic types");
+
+				Value /= InValue;
+				return *this;
+			}
+
+			/* Comparison operator global friend functions */
+
+			template<typename LHSQuantizationType, typename LHSTimeType, typename RHSQuantizationType, typename RHSTimeType>
+			friend bool operator<(const TTime<LHSQuantizationType, LHSTimeType>& LHS, const TTime<RHSQuantizationType, RHSTimeType>& RHS);
+
+			template<typename LHSQuantizationType, typename LHSTimeType, typename RHSQuantizationType, typename RHSTimeType>
+			friend bool operator>(const TTime<LHSQuantizationType, LHSTimeType>& LHS, const TTime<RHSQuantizationType, RHSTimeType>& RHS);
+
+			template<typename LHSQuantizationType, typename LHSTimeType, typename RHSQuantizationType, typename RHSTimeType>
+			friend bool operator<=(const TTime<LHSQuantizationType, LHSTimeType>& LHS, const TTime<RHSQuantizationType, RHSTimeType>& RHS);
+
+			template<typename LHSQuantizationType, typename LHSTimeType, typename RHSQuantizationType, typename RHSTimeType>
+			friend bool operator>=(const TTime<LHSQuantizationType, LHSTimeType>& LHS, const TTime<RHSQuantizationType, RHSTimeType>& RHS);
+
+			template<typename LHSQuantizationType, typename LHSTimeType, typename RHSQuantizationType, typename RHSTimeType>
+			friend bool operator==(const TTime<LHSQuantizationType, LHSTimeType>& LHS, const TTime<RHSQuantizationType, RHSTimeType>& RHS);
+
+			template<typename LHSQuantizationType, typename LHSTimeType, typename RHSQuantizationType, typename RHSTimeType>
+			friend bool operator!=(const TTime<LHSQuantizationType, LHSTimeType>& LHS, const TTime<RHSQuantizationType, RHSTimeType>& RHS);
+
+			/* Arithmetic operator global friend functions */
+
+			template<typename LHSQuantizationType, typename LHSTimeType, typename RHSQuantizationType, typename RHSTimeType>
+			friend TTime<LHSQuantizationType, LHSTimeType> operator+(TTime<LHSQuantizationType, LHSTimeType> LHS, const TTime<RHSQuantizationType, RHSTimeType>& RHS);
+
+			template<typename LHSQuantizationType, typename LHSTimeType, typename RHSQuantizationType, typename RHSTimeType>
+			friend TTime<LHSQuantizationType, LHSTimeType> operator-(TTime<LHSQuantizationType, LHSTimeType> LHS, const TTime<RHSQuantizationType, RHSTimeType>& RHS);
+
+			template<typename LHSQuantizationType, typename LHSTimeType, typename OtherType>
+			friend TTime<LHSQuantizationType, LHSTimeType> operator*(TTime<LHSQuantizationType, LHSTimeType> LHS, OtherType RHS);
+
+			template<typename RHSQuantizationType, typename RHSTimeType, typename OtherType>
+			friend TTime<RHSQuantizationType, RHSTimeType> operator*(OtherType LHS, const TTime<RHSQuantizationType, RHSTimeType>& RHS);
+
+			template<typename LHSQuantizationType, typename LHSTimeType, typename OtherType>
+			friend TTime<LHSQuantizationType, LHSTimeType> operator/(TTime<LHSQuantizationType, LHSTimeType> LHS, OtherType RHS);
+
+
+		private:
+
+			QuantizationType Value;
+	};
+
+
+	// holds on to a count of samples.
+	struct FTimeSampleCounter
+	{
+		// defines the type for counting samples.
+		typedef int32 FCountType;
+		FCountType Num;
+	};
+
+	/** TTime sample counting template specialization holds time values which
+	 * are quantized to sample boundaries. It provides basic mathematical 
+	 * operations, clearly defined getters and setters, and interoperability 
+	 * with generic TTime objects utilizing different template arguments. 
+	 *
+	 * This specialization differs distinctly by storing the number of samples
+	 * rather than the number of seconds in the quantization type. Additionally,
+	 * instances have an associated sample rate member which is used to convert
+	 * between TimeType and the sample count.
+	 *
+	 * TimeType declares the type of object which represent time values 
+	 * returned from TTime's getters.
+	 */
+	template<typename TimeType>
+	class TTime<FTimeSampleCounter, TimeType>
+	{
+		public:
+			static_assert(TIsFloatingPoint<TimeType>::Value, "TimeType must be floating point");
+
+			// using types from namespace.
+			using FCountType = FTimeSampleCounter::FCountType;
+			
+			/** Construct a sample count based TTime with a given time value, 
+			 * time resolution, and sample rate.
+			 *
+			 * @param InValue - The initial time value.
+			 * @param InResolution - The time resolution of InValue.
+			 * @param InSampleRate - The sample rate associated with this time 
+			 * 						 object.
+			 */
+			TTime(TimeType InValue, ETimeResolution InResolution, TimeType InSampleRate)
+			{
+				// Need to set sample rate before setting time or else the 
+				// resulting sample count will be incorrect.
+				SetSampleRate(InSampleRate);
+
+				switch (InResolution)
+				{
+					case ETimeResolution::Microseconds:
+						SetMicroseconds(InValue);
+						break;
+
+					case ETimeResolution::Milliseconds:
+						SetMilliseconds(InValue);
+						break;
+
+					case ETimeResolution::Seconds:
+					default:
+						SetSeconds(InValue);
+				}
+			}
+
+			/** Construct a sample count based TTime using a number of samples
+			 * and a sample rate.
+			 *
+			 * @param InNumSamples - The initial number of samples.
+			 * @param InSampleRate - The sample rate of this instance.
+			 */
+			TTime(FCountType InNumSamples, TimeType InSampleRate)
+			{
+				SetSampleRate(InSampleRate);
+
+				SetNumSamples(InNumSamples);
+			}
+
+			/** Set the sample rate of this object without changing the number
+			 * of samples. This will result in the object representing a 
+			 * different amount of time. 
+			 */
+			void SetSampleRate(TimeType InSampleRate)
+			{
+				if (!ensure(InSampleRate > 0))
+				{
+					InSampleRate = SMALL_NUMBER;
+				}
+
+				SampleRate = InSampleRate;
+			}
+
+			/** Set the sample rate of this object and change the number
+			 * of samples. This will result in the object representing 
+			 * approximately the same amount of time as it did before the 
+			 * sample rate was changed.
+			 */
+			void SetSampleRateAndUpdateSampleCount(TimeType InSampleRate)
+			{
+				if (!ensure(InSampleRate > 0))
+				{
+					InSampleRate = SMALL_NUMBER;
+				}
+
+				Value.Num = GetNumSamples(InSampleRate);
+				
+				SampleRate = InSampleRate;
+			}
+
+			/** Return the sample rate. */
+			TimeType GetSampleRate() const
+			{
+				return SampleRate;
+			}
+
+			/** Return the time as seconds. */
+			TimeType GetSeconds() const 
+			{ 
+				return static_cast<TimeType>(Value.Num / SampleRate);
+			}
+
+			/** Return the time as milliseconds. */
+			TimeType GetMilliseconds() const 
+			{ 
+				return static_cast<TimeType>(Value.Num * 1e3 / SampleRate);
+			}
+
+			/** Return the time as microseconds. */
+			TimeType GetMicroseconds() const 
+			{ 
+				return static_cast<TimeType>(Value.Num * 1e6 / SampleRate);
+			}
+
+			/** Return the number of samples. */
+			FCountType GetNumSamples() const
+			{ 
+				return Value.Num;
+			}
+
+			/** Return the number of samples which represent this time duration
+			 * using a different sample rate. 
+			 *
+			 * @param InOtherSampleRate - Sample rate to use when calculating 
+			 * 							  the number of samples.  
+			 */
+			FCountType GetNumSamples(TimeType InOtherSampleRate) const
+			{ 
+				if (InOtherSampleRate == SampleRate)
+				{
+					return Value.Num;
+				}
+				else
+				{
+					return FMath::RoundToInt(GetSeconds() * InOtherSampleRate);
+				}
+			}
+
+			/** Set the time value in seconds. */
+			void SetSeconds(TimeType InSeconds) 
+			{ 
+				Value.Num = FMath::RoundToInt(InSeconds * SampleRate);
+			}
+
+			/** Set the time value in milliseconds. */
+			void SetMilliseconds(TimeType InMilliseconds) 
+			{ 
+				Value.Num = FMath::RoundToInt(InMilliseconds * SampleRate * TimeType(1e-3));
+			}
+
+			/** Set the time value in microseconds. */
+			void SetMicroseconds(TimeType InMicroseconds) 
+			{ 
+				Value.Num = FMath::RoundToInt(InMicroseconds * SampleRate * TimeType(1e-6));
+			}
+
+			/** Set the number of samples. */
+			void SetNumSamples(FCountType InNumSamples)
+			{ 
+				Value.Num = InNumSamples;
+			}
+
+			/** Assignment operator 
+			 *
+			 * This assignment preserves this objects sample rate when being 
+			 * assigned to TTime objects of different quantization types.
+			 * */
+			template<typename OtherQuantizationType, typename OtherTimeType>
+			TTime& operator=(const TTime<OtherQuantizationType, OtherTimeType>& Other)
+			{
+				Value.Num = Other.GetNumSamples(SampleRate);
+
+				return *this;
+			}
+
+			/** Assignment operator 
+			 *
+			 * This assignment copies the sample rate and sample count of TTime
+			 * objects of equal quantization types.
+			 */
+			template<typename OtherTimeType>
+			TTime& operator=(const TTime<FTimeSampleCounter, OtherTimeType>& Other)
+			{
+				Value.Num = Other.Value.Num;
+				SampleRate = Other.SampleRate;
+
+				return *this;
+			}
+
+			/** Addition assignment operator. */
+			template<typename OtherQuantizationType, typename OtherTimeType>
+			TTime& operator+=(const TTime<OtherQuantizationType, OtherTimeType>& Other)
+			{
+				Value.Num += Other.GetNumSamples(SampleRate);
+
+				return *this;
+			}
+
+			/** Subtraction assignment operator. */
+			template<typename OtherQuantizationType, typename OtherTimeType>
+			TTime& operator-=(const TTime<OtherQuantizationType, OtherTimeType>& Other)
+			{
+				Value.Num -= Other.GetNumSamples(SampleRate);
+
+				return *this;
+			}
+
+			/** Multiplication assignment operator. */
+			template<typename OtherType>
+			TTime& operator*=(OtherType Other)
+			{
+				static_assert(TIsArithmetic<OtherType>::Value, "TTime can only be multiplied with arithmetic types");
+
+				Value.Num *= Other;
+				return *this;
+			}
+
+			/** Division assignment operator. */
+			template<typename OtherType>
+			TTime& operator/=(OtherType Other)
+			{
+				static_assert(TIsArithmetic<OtherType>::Value, "TTime can only be divided with arithmetic types");
+
+				Value.Num /= Other;
+				return *this;
+			}
+
+			/* Comparison operator global friend functions */
+			
+			template<typename LHSQuantizationType, typename LHSTimeType, typename RHSQuantizationType, typename RHSTimeType>
+			friend bool operator<(const TTime<LHSQuantizationType, LHSTimeType>& LHS, const TTime<RHSQuantizationType, RHSTimeType>& RHS);
+
+			template<typename LHSQuantizationType, typename LHSTimeType, typename RHSQuantizationType, typename RHSTimeType>
+			friend bool operator>(const TTime<LHSQuantizationType, LHSTimeType>& LHS, const TTime<RHSQuantizationType, RHSTimeType>& RHS);
+
+			template<typename LHSQuantizationType, typename LHSTimeType, typename RHSQuantizationType, typename RHSTimeType>
+			friend bool operator<=(const TTime<LHSQuantizationType, LHSTimeType>& LHS, const TTime<RHSQuantizationType, RHSTimeType>& RHS);
+
+			template<typename LHSQuantizationType, typename LHSTimeType, typename RHSQuantizationType, typename RHSTimeType>
+			friend bool operator>=(const TTime<LHSQuantizationType, LHSTimeType>& LHS, const TTime<RHSQuantizationType, RHSTimeType>& RHS);
+
+			template<typename LHSQuantizationType, typename LHSTimeType, typename RHSQuantizationType, typename RHSTimeType>
+			friend bool operator==(const TTime<LHSQuantizationType, LHSTimeType>& LHS, const TTime<RHSQuantizationType, RHSTimeType>& RHS);
+
+			template<typename LHSQuantizationType, typename LHSTimeType, typename RHSQuantizationType, typename RHSTimeType>
+			friend bool operator!=(const TTime<LHSQuantizationType, LHSTimeType>& LHS, const TTime<RHSQuantizationType, RHSTimeType>& RHS);
+
+			/* Arithmetic operator global friend functions */
+
+			template<typename LHSQuantizationType, typename LHSTimeType, typename RHSQuantizationType, typename RHSTimeType>
+			friend TTime<LHSQuantizationType, LHSTimeType> operator+(TTime<LHSQuantizationType, LHSTimeType> LHS, const TTime<RHSQuantizationType, RHSTimeType>& RHS);
+
+			template<typename LHSQuantizationType, typename LHSTimeType, typename RHSQuantizationType, typename RHSTimeType>
+			friend TTime<LHSQuantizationType, LHSTimeType> operator-(TTime<LHSQuantizationType, LHSTimeType> LHS, const TTime<RHSQuantizationType, RHSTimeType>& RHS);
+
+			template<typename LHSQuantizationType, typename LHSTimeType, typename OtherType>
+			friend TTime<LHSQuantizationType, LHSTimeType> operator*(TTime<LHSQuantizationType, LHSTimeType> LHS, OtherType RHS);
+
+			template<typename RHSQuantizationType, typename RHSTimeType, typename OtherType>
+			friend TTime<RHSQuantizationType, RHSTimeType> operator*(OtherType LHS, const TTime<RHSQuantizationType, RHSTimeType>& RHS);
+
+			template<typename LHSQuantizationType, typename LHSTimeType, typename OtherType>
+			friend TTime<LHSQuantizationType, LHSTimeType> operator/(TTime<LHSQuantizationType, LHSTimeType> LHS, OtherType RHS);
+
+		private:
+
+			FTimeSampleCounter Value = {0};
+			TimeType SampleRate = static_cast<TimeType>(SMALL_NUMBER);
+	};
+
+
+	/** Less-than comparison of TTime */
+	template<typename LHSQuantizationType, typename LHSTimeType, typename RHSQuantizationType, typename RHSTimeType>
+	bool operator<(const TTime<LHSQuantizationType, LHSTimeType>& LHS, const TTime<RHSQuantizationType, RHSTimeType>& RHS)
+	{ 
+		return LHS.GetSeconds() < RHS.GetSeconds();
+	}
+
+	/** Greater-than comparison of TTime */
+	template<typename LHSQuantizationType, typename LHSTimeType, typename RHSQuantizationType, typename RHSTimeType>
+	bool operator>(const TTime<LHSQuantizationType, LHSTimeType>& LHS, const TTime<RHSQuantizationType, RHSTimeType>& RHS)
+	{ 
+		return RHS < LHS; 
+	}
+
+	/** Less-than-or-equal-to comparison of TTime */
+	template<typename LHSQuantizationType, typename LHSTimeType, typename RHSQuantizationType, typename RHSTimeType>
+	bool operator<=(const TTime<LHSQuantizationType, LHSTimeType>& LHS, const TTime<RHSQuantizationType, RHSTimeType>& RHS)
+	{ 
+		return !(LHS > RHS); 
+	}
+
+	/** Greater-than-or-equal-to comparison of TTime */
+	template<typename LHSQuantizationType, typename LHSTimeType, typename RHSQuantizationType, typename RHSTimeType>
+	bool operator>=(const TTime<LHSQuantizationType, LHSTimeType>& LHS, const TTime<RHSQuantizationType, RHSTimeType>& RHS)
+	{ 
+		return !(LHS < RHS); 
+	}
+
+	/** Equal comparison of TTime */
+	template<typename LHSQuantizationType, typename LHSTimeType, typename RHSQuantizationType, typename RHSTimeType>
+	bool operator==(const TTime<LHSQuantizationType, LHSTimeType>& LHS, const TTime<RHSQuantizationType, RHSTimeType>& RHS)
+	{
+		return LHS.GetSeconds() == RHS.GetSeconds();
+	}
+
+	/** Not-equal comparison of TTime */
+	template<typename LHSQuantizationType, typename LHSTimeType, typename RHSQuantizationType, typename RHSTimeType>
+	bool operator!=(const TTime<LHSQuantizationType, LHSTimeType>& LHS, const TTime<RHSQuantizationType, RHSTimeType>& RHS)
+	{ 
+		return !(LHS == RHS); 
+	}
+
+	/** Addition operator of TTime */
+	template<typename LHSQuantizationType, typename LHSTimeType, typename RHSQuantizationType, typename RHSTimeType>
+	TTime<LHSQuantizationType, LHSTimeType> operator+(TTime<LHSQuantizationType, LHSTimeType> LHS, const TTime<RHSQuantizationType, RHSTimeType>& RHS)
+	{
+		LHS += RHS;
+		return LHS;
+	}
+
+	/** Subtraction operator of TTime */
+	template<typename LHSQuantizationType, typename LHSTimeType, typename RHSQuantizationType, typename RHSTimeType>
+	TTime<LHSQuantizationType, LHSTimeType> operator-(TTime<LHSQuantizationType, LHSTimeType> LHS, const TTime<RHSQuantizationType, RHSTimeType>& RHS)
+	{
+		LHS -= RHS;
+		return LHS;
+	}
+
+	/** Multiplication operator of TTime */
+	template<typename LHSQuantizationType, typename LHSTimeType, typename OtherType>
+	TTime<LHSQuantizationType, LHSTimeType> operator*(TTime<LHSQuantizationType, LHSTimeType> LHS, OtherType RHS)
+	{
+		static_assert(TIsArithmetic<OtherType>::Value, "TTime can only be multiplied with arithmetic types");
+
+		LHS *= RHS;
+		return LHS;
+	}
+
+	/** Division operator of TTime */
+	template<typename LHSQuantizationType, typename LHSTimeType, typename OtherType>
+	TTime<LHSQuantizationType, LHSTimeType> operator/(TTime<LHSQuantizationType, LHSTimeType> LHS, OtherType RHS)
+	{
+		static_assert(TIsArithmetic<OtherType>::Value, "TTime can only be divided with arithmetic types");
+
+		LHS /= RHS;
+		return LHS;
+	}
+
+	/** Multiplication operator of TTime */
+	template<typename RHSQuantizationType, typename RHSTimeType, typename OtherType>
+	TTime<RHSQuantizationType, RHSTimeType> operator*(OtherType LHS, const TTime<RHSQuantizationType, RHSTimeType>& RHS)
+	{
+		static_assert(TIsArithmetic<OtherType>::Value, "TTime can only be multiplied with arithmetic types");
+
+		return RHS * LHS;
+	}
+	
+
+	/** FFloatTime stores and represents all time values as 32bit floating 
+	 * point values. 
+	 */
+	typedef TTime<float, float> FFloatTime;
+
+	DECLARE_METASOUND_DATA_REFERENCE_TYPES(FFloatTime, "Primitive:Time", 0x2490848e, FFloatTimeTypeInfo, FFloatTimeReadRef, FFloatTimeWriteRef);
+
+	/** FDoubleTime stores and represents all time values as 64bit floating
+	 * point values.
+	 */
+	typedef TTime<double, double> FDoubleTime;
+
+	DECLARE_METASOUND_DATA_REFERENCE_TYPES(FDoubleTime, "Primitive:Time:HighResolution", 0x407fcff1, FDoubleTimeTypeInfo, FDoubleTimeReadRef, FDoubleTimeWriteRef);
+
+	/** FSampleTime stores time quantized to sample boundaries. It returns time
+	 * as 32bit floating point values.
+	 */
+	typedef TTime<FTimeSampleCounter, float> FSampleTime;
+
+	DECLARE_METASOUND_DATA_REFERENCE_TYPES(FSampleTime, "Primitive:Time:SampleResolution", 0xf1de8f2d, FSampleTimeTypeInfo, FSampleTimeReadRef, FSampleTimeWriteRef);
+}
