Update nodes to use TMaps to describe input/output vertices.

diff --git a/Engine/Plugins/Runtime/Metasound/Source/MetasoundGenerator/Private/MetasoundGenerator.cpp b/Engine/Plugins/Runtime/Metasound/Source/MetasoundGenerator/Private/MetasoundGenerator.cpp
--- a/Engine/Plugins/Runtime/Metasound/Source/MetasoundGenerator/Private/MetasoundGenerator.cpp
+++ b/Engine/Plugins/Runtime/Metasound/Source/MetasoundGenerator/Private/MetasoundGenerator.cpp
@@ -192,36 +192,36 @@ ISoundGeneratorPtr USynthComponentMetasoundGenerator::CreateSoundGenerator(int32
 	using FDataEdge = Metasound::FDataEdge; // Need to make explicit to disambiguate with existing FDataEdge.
 	using FOperatorUniquePtr = TUniquePtr<Metasound::IOperator>;
 
-	TInputNode<FFloatTime> BopPeriodInputNode(TEXT("BopPeriod"), 1.f, ETimeResolution::Seconds);
-	TInputNode<FFrequency> FrequencyInputNode(TEXT("Frequency"), 100.f, EFrequencyResolution::Hertz);
+	TInputNode<FFloatTime> BopPeriodInputNode(TEXT("Input Bop Period Node"), TEXT("BopPeriod"), 1.f, ETimeResolution::Seconds);
+	TInputNode<FFrequency> FrequencyInputNode(TEXT("Input Frequency Node"), TEXT("Frequency"), 100.f, EFrequencyResolution::Hertz);
 
 	FPeriodicBopNode BopNode(TEXT("Bop"), 1.f);
 	FOscNode OscNode(TEXT("Osc"), 100.f);
 	FADSRNode ADSRNode(TEXT("ADSR"), 10.f /* AttackMs */, 10.f /* DecayMs */, 50.f /* SustainMs */, 10.f /* ReleaseMs */);
 	FAudioMultiplyNode MultiplyNode(TEXT("Multiply"));
 
-	TOutputNode<FAudioBuffer> OutputNode(TEXT("Audio"));
+	TOutputNode<FAudioBuffer> OutputNode(TEXT("Output Audio Node"), TEXT("Audio"));
 
 	FGraph Graph(TEXT("Graph"));
 
 	// Add rate of bop.
-	Graph.AddDataEdge({{&BopPeriodInputNode, BopPeriodInputNode.GetOutputs()[0]}, {&BopNode, BopNode.GetInputs()[0]}});
+	Graph.AddDataEdge(BopPeriodInputNode, TEXT("BopPeriod"), BopNode, TEXT("Period"));
 	// Add frequency controls
-	Graph.AddDataEdge({{&FrequencyInputNode, FrequencyInputNode.GetOutputs()[0]}, {&OscNode, OscNode.GetInputs()[0]}});
+	Graph.AddDataEdge(FrequencyInputNode, TEXT("Frequency"), OscNode, TEXT("Frequency"));
 
 	// TODO: maybe nodes should have name accessors to inputs and outputs instead of an array?
-	Graph.AddDataEdge({{&BopNode, BopNode.GetOutputs()[0]}, {&ADSRNode, ADSRNode.GetInputs()[0]}});
+	Graph.AddDataEdge(BopNode, TEXT("Bop"), ADSRNode, TEXT("Bop"));
 
 	// Hookup multiply nodes
-	Graph.AddDataEdge({{&OscNode, OscNode.GetOutputs()[0]}, {&MultiplyNode, MultiplyNode.GetInputs()[0]}});
-	Graph.AddDataEdge({{&ADSRNode, ADSRNode.GetOutputs()[0]}, {&MultiplyNode, MultiplyNode.GetInputs()[1]}});
+	Graph.AddDataEdge(OscNode, TEXT("Audio"), MultiplyNode, TEXT("InputBuffer1"));
+	Graph.AddDataEdge(ADSRNode, TEXT("Envelope"), MultiplyNode, TEXT("InputBuffer2"));
 
 	// Route to output.
-	Graph.AddDataEdge({{&MultiplyNode, MultiplyNode.GetOutputs()[0]}, {&OutputNode, OutputNode.GetInputs()[0]}});
+	Graph.AddDataEdge(MultiplyNode, TEXT("Audio"), OutputNode, TEXT("Audio"));
 
-	Graph.AddInputDataVertex({&BopPeriodInputNode, BopPeriodInputNode.GetInputs()[0]});
-	Graph.AddInputDataVertex({&FrequencyInputNode, FrequencyInputNode.GetInputs()[0]});
-	Graph.AddOutputDataVertex({&OutputNode, OutputNode.GetOutputs()[0]});
+	Graph.AddInputDataDestination(BopPeriodInputNode, TEXT("BopPeriod"));
+	Graph.AddInputDataDestination(FrequencyInputNode, TEXT("Frequency"));
+	Graph.AddOutputDataSource(OutputNode, TEXT("Audio"));
 
 
 	// Build everything
diff --git a/Engine/Plugins/Runtime/Metasound/Source/MetasoundGraphCore/Private/MetasoundADSRNode.cpp b/Engine/Plugins/Runtime/Metasound/Source/MetasoundGraphCore/Private/MetasoundADSRNode.cpp
--- a/Engine/Plugins/Runtime/Metasound/Source/MetasoundGraphCore/Private/MetasoundADSRNode.cpp
+++ b/Engine/Plugins/Runtime/Metasound/Source/MetasoundGraphCore/Private/MetasoundADSRNode.cpp
@@ -157,13 +157,13 @@ namespace Metasound
 	,	DefaultSustainMs(InDefaultSustainMs)
 	,	DefaultReleaseMs(InDefaultReleaseMs)
 	{
-		AddInputDataVertexDescription<FBop>(TEXT("Bop"), LOCTEXT("BopTooltip", "Trigger for envelope."));
-		AddInputDataVertexDescription<FFloatTime>(TEXT("Attack"), LOCTEXT("AttackTooltip", "Attack time in milliseconds."));
-		AddInputDataVertexDescription<FFloatTime>(TEXT("Decay"), LOCTEXT("DecayTooltip", "Decay time in milliseconds."));
-		AddInputDataVertexDescription<FFloatTime>(TEXT("Sustain"), LOCTEXT("SustainTooltip", "Sustain time in milliseconds."));
-		AddInputDataVertexDescription<FFloatTime>(TEXT("Release"), LOCTEXT("ReleaseTooltip", "Release time in milliseconds."));
+		AddInputDataVertex<FBop>(TEXT("Bop"), LOCTEXT("BopTooltip", "Trigger for envelope."));
+		AddInputDataVertex<FFloatTime>(TEXT("Attack"), LOCTEXT("AttackTooltip", "Attack time in milliseconds."));
+		AddInputDataVertex<FFloatTime>(TEXT("Decay"), LOCTEXT("DecayTooltip", "Decay time in milliseconds."));
+		AddInputDataVertex<FFloatTime>(TEXT("Sustain"), LOCTEXT("SustainTooltip", "Sustain time in milliseconds."));
+		AddInputDataVertex<FFloatTime>(TEXT("Release"), LOCTEXT("ReleaseTooltip", "Release time in milliseconds."));
 
-		AddOutputDataVertexDescription<FAudioBuffer>(TEXT("Envelope"), LOCTEXT("EnvelopeTooltip", "The output envelope"));
+		AddOutputDataVertex<FAudioBuffer>(TEXT("Envelope"), LOCTEXT("EnvelopeTooltip", "The output envelope"));
 	}
 
 	FADSRNode::~FADSRNode()
diff --git a/Engine/Plugins/Runtime/Metasound/Source/MetasoundGraphCore/Private/MetasoundAudioMultiplyNode.cpp b/Engine/Plugins/Runtime/Metasound/Source/MetasoundGraphCore/Private/MetasoundAudioMultiplyNode.cpp
--- a/Engine/Plugins/Runtime/Metasound/Source/MetasoundGraphCore/Private/MetasoundAudioMultiplyNode.cpp
+++ b/Engine/Plugins/Runtime/Metasound/Source/MetasoundGraphCore/Private/MetasoundAudioMultiplyNode.cpp
@@ -90,10 +90,10 @@ namespace Metasound
 	FAudioMultiplyNode::FAudioMultiplyNode(const FString& InName)
 	:	FNode(InName)
 	{
-		AddInputDataVertexDescription<FAudioBuffer>(TEXT("InputBuffer1"), LOCTEXT("InputBuffer1Tooltip", "Input buffer to multiply."));
-		AddInputDataVertexDescription<FAudioBuffer>(TEXT("InputBuffer2"), LOCTEXT("InputBuffer2Tooltip", "Input buffer to multiply."));
+		AddInputDataVertex<FAudioBuffer>(TEXT("InputBuffer1"), LOCTEXT("InputBuffer1Tooltip", "Input buffer to multiply."));
+		AddInputDataVertex<FAudioBuffer>(TEXT("InputBuffer2"), LOCTEXT("InputBuffer2Tooltip", "Input buffer to multiply."));
 
-		AddOutputDataVertexDescription<FAudioBuffer>(TEXT("Audio"), LOCTEXT("OutpuBufferTooltip", "The output audio."));
+		AddOutputDataVertex<FAudioBuffer>(TEXT("Audio"), LOCTEXT("OutpuBufferTooltip", "The output audio."));
 	}
 
 	FAudioMultiplyNode::~FAudioMultiplyNode()
diff --git a/Engine/Plugins/Runtime/Metasound/Source/MetasoundGraphCore/Private/MetasoundGraph.cpp b/Engine/Plugins/Runtime/Metasound/Source/MetasoundGraphCore/Private/MetasoundGraph.cpp
--- a/Engine/Plugins/Runtime/Metasound/Source/MetasoundGraphCore/Private/MetasoundGraph.cpp
+++ b/Engine/Plugins/Runtime/Metasound/Source/MetasoundGraphCore/Private/MetasoundGraph.cpp
@@ -25,36 +25,64 @@ namespace Metasound
 		return ClassName;
 	}
 
-	const TArray<FInputDataVertexDescription>& FGraph::GetInputs() const
+	const FInputDataVertexCollection& FGraph::GetInputDataVertices() const
 	{
-		return InputDescriptions;
+		return InputVertices;
 	}
 
-	const TArray<FOutputDataVertexDescription>& FGraph::GetOutputs() const
+	const FOutputDataVertexCollection& FGraph::GetOutputDataVertices() const
 	{
-		return OutputDescriptions;
+		return OutputVertices;
 	}
 
-	void FGraph::AddInputDataVertex(const FInputDataVertex& InVertex)
+	bool FGraph::AddInputDataDestination(INode& InNode, const FString& InVertexName)
 	{
-		InputVertices.Add(InVertex);
-		InputDescriptions.Add(InVertex.Description);
+		if (!InNode.GetInputDataVertices().Contains(InVertexName))
+		{
+			return false;
+		}
+
+		FInputDataDestination Destination(InNode, InNode.GetInputDataVertices()[InVertexName]);
+
+		AddInputDataDestination(Destination);
+
+		return true;
 	}
 
-	const TArray<FInputDataVertex>& FGraph::GetInputDataVertices() const
+	void FGraph::AddInputDataDestination(const FInputDataDestination& InDestination)
 	{
-		return InputVertices;
+		InputVertices.Add(MakeDataVertexKey(InDestination.Vertex), InDestination.Vertex);
+		InputDestinations.Add(MakeDestinationDataVertexKey(InDestination), InDestination);
 	}
 
-	void FGraph::AddOutputDataVertex(const FOutputDataVertex& InVertex)
+	const FInputDataDestinationCollection& FGraph::GetInputDataDestinations() const
 	{
-		OutputVertices.Add(InVertex);
-		OutputDescriptions.Add(InVertex.Description);
+		return InputDestinations;
 	}
 
-	const TArray<FOutputDataVertex>& FGraph::GetOutputDataVertices() const
+	bool FGraph::AddOutputDataSource(INode& InNode, const FString& InVertexName)
 	{
-		return OutputVertices;
+		if (!InNode.GetOutputDataVertices().Contains(InVertexName))
+		{
+			return false;
+		}
+
+		FOutputDataSource Source(InNode, InNode.GetOutputDataVertices()[InVertexName]);
+
+		AddOutputDataSource(Source);
+
+		return true;
+	}
+
+	void FGraph::AddOutputDataSource(const FOutputDataSource& InSource)
+	{
+		OutputVertices.Add(MakeDataVertexKey(InSource.Vertex), InSource.Vertex);
+		OutputSources.Add(MakeSourceDataVertexKey(InSource), InSource);
+	}
+
+	const FOutputDataSourceCollection& FGraph::GetOutputDataSources() const
+	{
+		return OutputSources;
 	}
 
 	void FGraph::AddDataEdge(const FDataEdge& InEdge)
@@ -62,6 +90,35 @@ namespace Metasound
 		Edges.Add(InEdge);
 	}
 
+	// TODO: can we make node references const?  Would need to make operator factory accessing const as well.
+	bool FGraph::AddDataEdge(INode& FromNode, const FString& FromName, INode& ToNode, const FString ToName)
+	{
+		
+		if (!FromNode.GetOutputDataVertices().Contains(FromName))
+		{
+			return false;
+		}
+
+		if (!ToNode.GetInputDataVertices().Contains(ToName))
+		{
+			return false;
+		}
+
+		const FOutputDataVertex& FromVertex = FromNode.GetOutputDataVertices()[FromName];
+		const FInputDataVertex& ToVertex = ToNode.GetInputDataVertices()[ToName];
+
+		if (FromVertex.DataReferenceTypeName != ToVertex.DataReferenceTypeName)
+		{
+			return false;
+		}
+
+		FDataEdge Edge(FOutputDataSource(FromNode, FromVertex), FInputDataDestination(ToNode, ToVertex));
+
+		AddDataEdge(Edge);
+
+		return true;
+	}
+
 	const TArray<FDataEdge>& FGraph::GetDataEdges() const
 	{
 		return Edges;
diff --git a/Engine/Plugins/Runtime/Metasound/Source/MetasoundGraphCore/Private/MetasoundNode.cpp b/Engine/Plugins/Runtime/Metasound/Source/MetasoundGraphCore/Private/MetasoundNode.cpp
--- a/Engine/Plugins/Runtime/Metasound/Source/MetasoundGraphCore/Private/MetasoundNode.cpp
+++ b/Engine/Plugins/Runtime/Metasound/Source/MetasoundGraphCore/Private/MetasoundNode.cpp
@@ -20,33 +20,33 @@ namespace Metasound
 		return Description;
 	}
 
-	const TArray<FInputDataVertexDescription>& FNode::GetInputs() const
+	const FInputDataVertexCollection& FNode::GetInputDataVertices() const
 	{
 		return Inputs;
 	}
 
-	const TArray<FOutputDataVertexDescription>& FNode::GetOutputs() const
+	const FOutputDataVertexCollection& FNode::GetOutputDataVertices() const
 	{
 		return Outputs;
 	}
 
-	void FNode::AddInputDataVertexDescription(const FInputDataVertexDescription& InDescription)
+	void FNode::AddInputDataVertex(const FInputDataVertex& InVertex)
 	{
-		Inputs.Add(InDescription);
+		Inputs.Add(MakeDataVertexKey(InVertex), InVertex);
 	}
 
-	void FNode::RemoveInputDataVertexDescription(const FInputDataVertexDescription& InDescription)
+	void FNode::RemoveInputDataVertex(const FInputDataVertex& InVertex)
 	{
-		Inputs.RemoveSingle(InDescription);
+		Inputs.Remove(MakeDataVertexKey(InVertex));
 	}
 
-	void FNode::AddOutputDataVertexDescription(const FOutputDataVertexDescription& InDescription)
+	void FNode::AddOutputDataVertex(const FOutputDataVertex& InVertex)
 	{
-		Outputs.Add(InDescription);
+		Outputs.Add(MakeDataVertexKey(InVertex), InVertex);
 	}
 
-	void FNode::RemoveOutputDataVertexDescription(const FOutputDataVertexDescription& InDescription)
+	void FNode::RemoveOutputDataVertex(const FOutputDataVertex& InVertex)
 	{
-		Outputs.RemoveSingle(InDescription);
+		Outputs.Remove(MakeDataVertexKey(InVertex));
 	}
 }
diff --git a/Engine/Plugins/Runtime/Metasound/Source/MetasoundGraphCore/Private/MetasoundNodeInterface.cpp b/Engine/Plugins/Runtime/Metasound/Source/MetasoundGraphCore/Private/MetasoundNodeInterface.cpp
--- a/Engine/Plugins/Runtime/Metasound/Source/MetasoundGraphCore/Private/MetasoundNodeInterface.cpp
+++ b/Engine/Plugins/Runtime/Metasound/Source/MetasoundGraphCore/Private/MetasoundNodeInterface.cpp
@@ -4,17 +4,17 @@
 
 namespace Metasound
 {
-	bool operator==(const FDataVertexDescription& LHS, const FDataVertexDescription& RHS)
+	bool operator==(const FDataVertex& LHS, const FDataVertex& RHS)
 	{
 		bool bIsEqual = /*(LHS.Source == RHS.Source) && */
 			(LHS.VertexName == RHS.VertexName) &&
 			(LHS.DataReferenceTypeName == RHS.DataReferenceTypeName) &&
-			LHS.Tooltip.EqualTo(RHS.Tooltip);
+			LHS.Description.EqualTo(RHS.Description);
 
 		return bIsEqual;
 	}
 
-	bool operator!=(const FDataVertexDescription& LHS, const FDataVertexDescription& RHS)
+	bool operator!=(const FDataVertex& LHS, const FDataVertex& RHS)
 	{
 		return !(LHS == RHS);
 	}
diff --git a/Engine/Plugins/Runtime/Metasound/Source/MetasoundGraphCore/Private/MetasoundOperatorBuilder.cpp b/Engine/Plugins/Runtime/Metasound/Source/MetasoundGraphCore/Private/MetasoundOperatorBuilder.cpp
--- a/Engine/Plugins/Runtime/Metasound/Source/MetasoundGraphCore/Private/MetasoundOperatorBuilder.cpp
+++ b/Engine/Plugins/Runtime/Metasound/Source/MetasoundGraphCore/Private/MetasoundOperatorBuilder.cpp
@@ -155,7 +155,7 @@ namespace Metasound
 			
 			const FDataReferenceCollection& FromDataReferenceCollection = InDataReferenceMap[Edge->From.Node].Outputs;
 
-			if (!FromDataReferenceCollection.ContainsDataReadReference(Edge->From.Description.VertexName, Edge->From.Description.DataReferenceTypeName))
+			if (!FromDataReferenceCollection.ContainsDataReadReference(Edge->From.Vertex.VertexName, Edge->From.Vertex.DataReferenceTypeName))
 			{
 				// TODO: This is likely a node programming error where the edges reported by the INode interface
 				// did not match the readable parameter refs created by the operators outputs. Or, the edge description is invalid.
@@ -165,7 +165,7 @@ namespace Metasound
 				return false;
 			}
 
-			if (Edge->From.Description.DataReferenceTypeName != Edge->To.Description.DataReferenceTypeName)
+			if (Edge->From.Vertex.DataReferenceTypeName != Edge->To.Vertex.DataReferenceTypeName)
 			{
 				// TODO: add build error.
 				// TODO: may want to have this check higher up in the stack. 
@@ -174,7 +174,7 @@ namespace Metasound
 
 
 
-			bool bSuccess = OutCollection.AddDataReadReferenceFrom(Edge->To.Description.VertexName, FromDataReferenceCollection, Edge->From.Description.VertexName, Edge->From.Description.DataReferenceTypeName);
+			bool bSuccess = OutCollection.AddDataReadReferenceFrom(Edge->To.Vertex.VertexName, FromDataReferenceCollection, Edge->From.Vertex.VertexName, Edge->From.Vertex.DataReferenceTypeName);
 
 			if (!bSuccess)
 			{
@@ -225,24 +225,29 @@ namespace Metasound
 
 	bool FOperatorBuilder::GatherGraphDataReferences(const IGraph& InGraph, FNodeDataReferenceMap& InNodeDataReferences, FDataReferenceCollection& OutGraphInputs, FDataReferenceCollection& OutGraphOutputs, TArray<FBuildErrorPtr>& OutErrors) const
 	{
+		using FDestinationElement = FInputDataDestinationCollection::ElementType;
+		using FSourceElement = FOutputDataSourceCollection::ElementType;
 		// Gather inputs
-		for (const FInputDataVertex& InputVertex : InGraph.GetInputDataVertices())
+
+		for (const FDestinationElement& Element : InGraph.GetInputDataDestinations())
 		{
-			if (!InNodeDataReferences.Contains(InputVertex.Node))
+			const FInputDataDestination& InputDestination = Element.Value;
+
+			if (!InNodeDataReferences.Contains(InputDestination.Node))
 			{
 				// TODO: build error.  input node does not exist.
 				return false;
 			}
 
-			FDataReferenceCollection& Collection = InNodeDataReferences[InputVertex.Node].Inputs;
+			FDataReferenceCollection& Collection = InNodeDataReferences[InputDestination.Node].Inputs;
 
-			if (!Collection.ContainsDataWriteReference(InputVertex.Description.VertexName, InputVertex.Description.DataReferenceTypeName))
+			if (!Collection.ContainsDataWriteReference(InputDestination.Vertex.VertexName, InputDestination.Vertex.DataReferenceTypeName))
 			{
 				// TODO: build error.  parameter ref does not exist
 				return false;
 			}
 
-			bool bSuccess = OutGraphInputs.AddDataWriteReferenceFrom(InputVertex.Description.VertexName, Collection, InputVertex.Description.VertexName, InputVertex.Description.DataReferenceTypeName);
+			bool bSuccess = OutGraphInputs.AddDataWriteReferenceFrom(InputDestination.Vertex.VertexName, Collection, InputDestination.Vertex.VertexName, InputDestination.Vertex.DataReferenceTypeName);
 
 			if (!bSuccess)
 			{
@@ -253,23 +258,25 @@ namespace Metasound
 		}
 
 		// Gather outputs
-		for (const FOutputDataVertex& OutputVertex : InGraph.GetOutputDataVertices())
+		for (const FSourceElement& Element : InGraph.GetOutputDataSources())
 		{
-			if (!InNodeDataReferences.Contains(OutputVertex.Node))
+			const FOutputDataSource& OutputSource = Element.Value;
+
+			if (!InNodeDataReferences.Contains(OutputSource.Node))
 			{
 				// TODO: build error.  input node does not exist.
 				return false;
 			}
 
-			FDataReferenceCollection& Collection = InNodeDataReferences[OutputVertex.Node].Outputs;
+			FDataReferenceCollection& Collection = InNodeDataReferences[OutputSource.Node].Outputs;
 
-			if (!Collection.ContainsDataReadReference(OutputVertex.Description.VertexName, OutputVertex.Description.DataReferenceTypeName))
+			if (!Collection.ContainsDataReadReference(OutputSource.Vertex.VertexName, OutputSource.Vertex.DataReferenceTypeName))
 			{
 				// TODO: build error.  parameter ref does not exist
 				return false;
 			}
 
-			bool bSuccess = OutGraphOutputs.AddDataReadReferenceFrom(OutputVertex.Description.VertexName, Collection, OutputVertex.Description.VertexName, OutputVertex.Description.DataReferenceTypeName);
+			bool bSuccess = OutGraphOutputs.AddDataReadReferenceFrom(OutputSource.Vertex.VertexName, Collection, OutputSource.Vertex.VertexName, OutputSource.Vertex.DataReferenceTypeName);
 
 			if (!bSuccess)
 			{
diff --git a/Engine/Plugins/Runtime/Metasound/Source/MetasoundGraphCore/Private/MetasoundOscNode.cpp b/Engine/Plugins/Runtime/Metasound/Source/MetasoundGraphCore/Private/MetasoundOscNode.cpp
--- a/Engine/Plugins/Runtime/Metasound/Source/MetasoundGraphCore/Private/MetasoundOscNode.cpp
+++ b/Engine/Plugins/Runtime/Metasound/Source/MetasoundGraphCore/Private/MetasoundOscNode.cpp
@@ -78,8 +78,8 @@ namespace Metasound
 	:	FNode(InName)
 	,	DefaultFrequency(InDefaultFrequency)
 	{
-		AddInputDataVertexDescription<FFrequency>(TEXT("Frequency"), LOCTEXT("FrequencyTooltip", "The frequency of oscillator."));
-		AddOutputDataVertexDescription<FAudioBuffer>(TEXT("Audio"), LOCTEXT("AudioTooltip", "The output audio"));
+		AddInputDataVertex<FFrequency>(TEXT("Frequency"), LOCTEXT("FrequencyTooltip", "The frequency of oscillator."));
+		AddOutputDataVertex<FAudioBuffer>(TEXT("Audio"), LOCTEXT("AudioTooltip", "The output audio"));
 	}
 
 	FOscNode::~FOscNode()
diff --git a/Engine/Plugins/Runtime/Metasound/Source/MetasoundGraphCore/Private/MetasoundPeriodicBopNode.cpp b/Engine/Plugins/Runtime/Metasound/Source/MetasoundGraphCore/Private/MetasoundPeriodicBopNode.cpp
--- a/Engine/Plugins/Runtime/Metasound/Source/MetasoundGraphCore/Private/MetasoundPeriodicBopNode.cpp
+++ b/Engine/Plugins/Runtime/Metasound/Source/MetasoundGraphCore/Private/MetasoundPeriodicBopNode.cpp
@@ -92,8 +92,8 @@ namespace Metasound
 	:	FNode(InName)
 	,	DefaultPeriod(InDefaultPeriodInSeconds)
 	{
-		AddInputDataVertexDescription<FFloatTime>(TEXT("Period"), LOCTEXT("PeriodTooltip", "The period of the bops in seconds."));
-		AddOutputDataVertexDescription<FBop>(TEXT("Bop"), LOCTEXT("BopTooltip", "The output bop"));
+		AddInputDataVertex<FFloatTime>(TEXT("Period"), LOCTEXT("PeriodTooltip", "The period of the bops in seconds."));
+		AddOutputDataVertex<FBop>(TEXT("Bop"), LOCTEXT("BopTooltip", "The output bop"));
 	}
 
 	FPeriodicBopNode::~FPeriodicBopNode()
diff --git a/Engine/Plugins/Runtime/Metasound/Source/MetasoundGraphCore/Public/MetasoundDataReference.h b/Engine/Plugins/Runtime/Metasound/Source/MetasoundGraphCore/Public/MetasoundDataReference.h
--- a/Engine/Plugins/Runtime/Metasound/Source/MetasoundGraphCore/Public/MetasoundDataReference.h
+++ b/Engine/Plugins/Runtime/Metasound/Source/MetasoundGraphCore/Public/MetasoundDataReference.h
@@ -52,6 +52,13 @@ namespace Metasound
 		virtual TUniquePtr<IDataReference> Clone() const = 0;
 	};
 
+	/** Helper class to enforce specialization of TDataReferenceTypeInfo */
+	template<typename DataType>
+	struct TSpecializationHelper 
+	{
+		enum { Value = false };
+	};
+
 	/** Info for templated data reference types help perform runtime type 
 	 * verification. 
 	 */
@@ -62,10 +69,14 @@ namespace Metasound
 		 * Users should be aware of FName case-insensitivity and other 
 		 * limitiations. 
 		 */
-		static constexpr const TCHAR* TypeName = TEXT("");
+		static constexpr const TCHAR* TypeName = nullptr;
 
 		/** Magic number used to check the data type when casting. */
 		static constexpr const FDataTypeMagicNumber MagicNumber = -1;
+
+		// This static assert is triggered if TDataReferenceTypeInfo is used 
+		// without specialization.
+		static_assert(TSpecializationHelper<DataType>::Value, "TDataReferenceTypeInfo must be specialized.  Use macro DECLARE_METASOUND_DATA_REFERENCE_TYPES");
 	};
 	
 	/** Template class for a paramter reference. 
diff --git a/Engine/Plugins/Runtime/Metasound/Source/MetasoundGraphCore/Public/MetasoundGraph.h b/Engine/Plugins/Runtime/Metasound/Source/MetasoundGraphCore/Public/MetasoundGraph.h
--- a/Engine/Plugins/Runtime/Metasound/Source/MetasoundGraphCore/Public/MetasoundGraph.h
+++ b/Engine/Plugins/Runtime/Metasound/Source/MetasoundGraphCore/Public/MetasoundGraph.h
@@ -7,6 +7,11 @@
 
 namespace Metasound
 {
+	// TODO: consider making nodes TSharedPtr<INode>
+	/** FGraph contains the edges between nodes as well as input and output 
+	 * vertices.  FGraph does not maintain ownership over any node. Nodes used
+	 * within the graph must be valid for the lifetime of the graph. 
+	 */
 	class METASOUNDGRAPHCORE_API FGraph : public IGraph
 	{
 		public:
@@ -21,41 +26,80 @@ namespace Metasound
 			/** Return the type name of this graph. */
 			virtual const FName& GetClassName() const override;
 
-			/** Return an array of input parameter descriptions for this graph. */
-			virtual const TArray<FInputDataVertexDescription>& GetInputs() const override;
+			/** Return a collection of input parameter descriptions for this graph. */
+			virtual const FInputDataVertexCollection& GetInputDataVertices() const override;
 
-			/** Return an array of output parameter descriptions for this graph. */
-			virtual const TArray<FOutputDataVertexDescription>& GetOutputs() const override;
+			/** Return a collection of output parameter descriptions for this graph. */
+			virtual const FOutputDataVertexCollection& GetOutputDataVertices() const override;
 
 			/** Retrieve all the edges associated with a graph. */
 			virtual const TArray<FDataEdge>& GetDataEdges() const override;
 
 			/** Get vertices which contain input parameters. */
-			virtual const TArray<FInputDataVertex>& GetInputDataVertices() const override;
+			virtual const FInputDataDestinationCollection& GetInputDataDestinations() const override;
 
 			/** Get vertices which contain output parameters. */
-			virtual const TArray<FOutputDataVertex>& GetOutputDataVertices() const override;
+			virtual const FOutputDataSourceCollection& GetOutputDataSources() const override;
 
+			/** Add an edge to the graph. */
 			void AddDataEdge(const FDataEdge& InEdge);
 
-			void AddInputDataVertex(const FInputDataVertex& InVertex);
-
-			void AddOutputDataVertex(const FOutputDataVertex& InVertex);
-
+			/** Add an edge to the graph, connecting two vertices from two 
+			 * nodes. 
+			 *
+			 * @param FromNode - Node which contains the output vertex.
+			 * @param FromVertexName - Name of the vertex in the FromNode.
+			 * @param ToNode - Node which contains the input vertex.
+			 * @param ToVertexName - Name of the vertex in the ToNode.
+			 *
+			 * @return True if the edge was successfully added. False otherwise.
+			 */
+			bool AddDataEdge(INode& FromNode, const FString& FromVertexName, INode& ToNode, const FString ToVertexName);
+
+			/** Add an input data destination to describe how data provided 
+			 * outside this graph should be routed internally.
+			 *
+			 * @param InNode - Node which receives the data.
+			 * @param InVertexName - Name of input vertex on InNode.
+			 *
+			 * @return True if the destination was successfully added. False 
+			 * otherwise.
+			 */
+			bool AddInputDataDestination(INode& InNode, const FString& InVertexName);
+
+
+			/** Add an input data destination to describe how data provided 
+			 * outside this graph should be routed internally.
+			 */
+			void AddInputDataDestination(const FInputDataDestination& InDestination);
+
+			/** Add an output data source which describes routing of data which is 
+			 * owned this graph and exposed externally.
+			 *
+			 * @param InNode - Node which produces the data.
+			 * @param InVertexName - Name of output vertex on InNode.
+			 *
+			 * @return True if the source was successfully added. False 
+			 * otherwise.
+			 */
+			bool AddOutputDataSource(INode& InNode, const FString& InVertexName);
+
+			/** Add an output data source which describes routing of data which is 
+			 * owned this graph and exposed externally.
+			 */
+			void AddOutputDataSource(const FOutputDataSource& InSource);
 
 			// TODO: Add ability to remove things.
-			
-
 
 		private:
 			FString Description;
 
-			TArray<FInputDataVertexDescription> InputDescriptions;
-			TArray<FOutputDataVertexDescription> OutputDescriptions;
+			FInputDataVertexCollection InputVertices;
+			FOutputDataVertexCollection OutputVertices;
 
 			TArray<FDataEdge> Edges;
 
-			TArray<FInputDataVertex> InputVertices;
-			TArray<FOutputDataVertex> OutputVertices;
+			FInputDataDestinationCollection InputDestinations;
+			FOutputDataSourceCollection OutputSources;
 	};
 }
diff --git a/Engine/Plugins/Runtime/Metasound/Source/MetasoundGraphCore/Public/MetasoundInputNode.h b/Engine/Plugins/Runtime/Metasound/Source/MetasoundGraphCore/Public/MetasoundInputNode.h
--- a/Engine/Plugins/Runtime/Metasound/Source/MetasoundGraphCore/Public/MetasoundInputNode.h
+++ b/Engine/Plugins/Runtime/Metasound/Source/MetasoundGraphCore/Public/MetasoundInputNode.h
@@ -11,6 +11,7 @@ namespace Metasound
 	template<typename DataType>
 	class TInputNode : public INode
 	{
+		
 			class FInputOperator : public IOperator
 			{
 				public:
@@ -51,6 +52,7 @@ namespace Metasound
 			{
 				public:
 					using FDataWriteReference = TDataWriteReference<DataType>;
+					using FInputNodeType = TInputNode<DataType>;
 
 					template<typename... ArgTypes, typename = typename TEnableIf< TIsConstructible<DataType, ArgTypes...>::Value >::Type >
 					FCopyOperatorFactory(ArgTypes&&... Args)
@@ -66,13 +68,14 @@ namespace Metasound
 
 					virtual TUniquePtr<IOperator> CreateOperator(const INode& InNode, const FOperatorSettings& InOperatorSettings, const FDataReferenceCollection& InInputDataReferences, TArray<TUniquePtr<IOperatorBuildError>>& OutErrors) override
 					{
-						// Use copy constructor to create a new paramter reference.
+						// Use copy constructor to create a new parameter reference.
 						FDataWriteReference DataRef(Data);
 						// TODO: Write special version of this for audio types since they will need to be constructed based upon the operator settings. 
 						// TODO: Need to ponder how inputs are initialized, but template specialization might do the trick. 
 
+						const FInputNodeType& InputNode = static_cast<const FInputNodeType&>(InNode);
 
-						return MakeUnique<FInputOperator>(InNode.GetDescription(), DataRef);
+						return MakeUnique<FInputOperator>(InputNode.GetVertexName(), DataRef);
 					}
 
 				private:
@@ -83,24 +86,43 @@ namespace Metasound
 
 
 			template<typename = typename TEnableIf< TIsConstructible<DataType>::Value >::Type >
-			TInputNode(const FString& InDescription)
-			:	Description(InDescription)
+			TInputNode(const FString& InNodeDescription, const FString& InVertexName)
+			:	NodeDescription(InNodeDescription)
+			,	VertexName(InVertexName)
 			{
-				Outputs.Add(MakeOutputDataVertexDescription<DataType>(Description, FText::GetEmpty()));
+				FOutputDataVertex OutputVertex = MakeOutputDataVertex<DataType>(VertexName, FText::GetEmpty());
+				FDataVertexKey OutputVertexKey = MakeDataVertexKey(OutputVertex);
+
+				Outputs.Add(OutputVertexKey, OutputVertex);
+
+				FInputDataVertex InputVertex = MakeInputDataVertex<DataType>(VertexName, FText::GetEmpty());
+				FDataVertexKey InputVertexKey = MakeDataVertexKey(InputVertex);
+
+				Inputs.Add(InputVertexKey, InputVertex);
+
 			}
 
 			template<typename... ArgTypes, typename = typename TEnableIf< TIsConstructible<DataType, ArgTypes...>::Value >::Type >
-			TInputNode(const FString& InDescription, ArgTypes&&... Args)
-			:	Description(InDescription)
+			TInputNode(const FString& InNodeDescription, const FString& InVertexName, ArgTypes&&... Args)
+			:	NodeDescription(InNodeDescription)
+			,	VertexName(InVertexName)
 			,	Factory(Forward<ArgTypes>(Args)...)
 			{
-				Inputs.Add(MakeInputDataVertexDescription<DataType>(Description, FText::GetEmpty()));
-				Outputs.Add(MakeOutputDataVertexDescription<DataType>(Description, FText::GetEmpty()));
+				FOutputDataVertex OutputVertex = MakeOutputDataVertex<DataType>(VertexName, FText::GetEmpty());
+				FDataVertexKey OutputVertexKey = MakeDataVertexKey(OutputVertex);
+
+				Outputs.Add(OutputVertexKey, OutputVertex);
+
+				FInputDataVertex InputVertex = MakeInputDataVertex<DataType>(VertexName, FText::GetEmpty());
+				FDataVertexKey InputVertexKey = MakeDataVertexKey(InputVertex);
+
+				Inputs.Add(InputVertexKey, InputVertex);
+
 			}
 
 			virtual const FString& GetDescription() const override
 			{
-				return Description;
+				return NodeDescription;
 			}
 
 			virtual const FName& GetClassName() const override
@@ -110,12 +132,17 @@ namespace Metasound
 				return ClassName;
 			}
 
-			virtual const TArray<FInputDataVertexDescription>& GetInputs() const override
+			const FString& GetVertexName() const
+			{
+				return VertexName;
+			}
+
+			virtual const FInputDataVertexCollection& GetInputDataVertices() const override
 			{
 				return Inputs;
 			}
 
-			virtual const TArray<FOutputDataVertexDescription>& GetOutputs() const override
+			virtual const FOutputDataVertexCollection& GetOutputDataVertices() const override
 			{
 				return Outputs;
 			}
@@ -127,11 +154,12 @@ namespace Metasound
 
 
 		private:
-			FString Description;
+			FString NodeDescription;
+			FString VertexName;
 
 			FCopyOperatorFactory Factory;
 
-			TArray<FInputDataVertexDescription> Inputs;
-			TArray<FOutputDataVertexDescription> Outputs;
+			FInputDataVertexCollection Inputs;
+			FOutputDataVertexCollection Outputs;
 	};
 }
diff --git a/Engine/Plugins/Runtime/Metasound/Source/MetasoundGraphCore/Public/MetasoundNode.h b/Engine/Plugins/Runtime/Metasound/Source/MetasoundGraphCore/Public/MetasoundNode.h
--- a/Engine/Plugins/Runtime/Metasound/Source/MetasoundGraphCore/Public/MetasoundNode.h
+++ b/Engine/Plugins/Runtime/Metasound/Source/MetasoundGraphCore/Public/MetasoundNode.h
@@ -7,8 +7,6 @@
 
 namespace Metasound
 {
-	//DECLARE_MULTICAST_DELEGATE_TwoParams(FParameterChangedDelegate, const INodeBase*, const FDataVertexDescription&);
-
 	class METASOUNDGRAPHCORE_API FNode : public INode
 	{
 		public:
@@ -17,35 +15,35 @@ namespace Metasound
 
 			virtual const FString& GetDescription() const override;
 
-			virtual const TArray<FInputDataVertexDescription>& GetInputs() const override;
-			virtual const TArray<FOutputDataVertexDescription>& GetOutputs() const override;
+			virtual const FInputDataVertexCollection& GetInputDataVertices() const override;
+			virtual const FOutputDataVertexCollection& GetOutputDataVertices() const override;
 
 		protected:
 
 
 			template<typename DataType>
-			void AddInputDataVertexDescription(const FString& InVertexName, const FText& InVertexTooltip)
+			void AddInputDataVertex(const FString& InVertexName, const FText& InVertexDescription)
 			{
-				AddInputDataVertexDescription(MakeInputDataVertexDescription<DataType>(InVertexName, InVertexTooltip));
+				AddInputDataVertex(MakeInputDataVertex<DataType>(InVertexName, InVertexDescription));
 			}
 
-			void AddInputDataVertexDescription(const FInputDataVertexDescription& InDescription);
-			void RemoveInputDataVertexDescription(const FInputDataVertexDescription& InDescription);
+			void AddInputDataVertex(const FInputDataVertex& InVertex);
+			void RemoveInputDataVertex(const FInputDataVertex& InVertex);
 
 			template<typename DataType>
-			void AddOutputDataVertexDescription(const FString& InVertexName, const FText& InVertexTooltip)
+			void AddOutputDataVertex(const FString& InVertexName, const FText& InVertexDescription)
 			{
-				AddOutputDataVertexDescription(MakeOutputDataVertexDescription<DataType>(InVertexName, InVertexTooltip));
+				AddOutputDataVertex(MakeOutputDataVertex<DataType>(InVertexName, InVertexDescription));
 			}
 
-			void AddOutputDataVertexDescription(const FOutputDataVertexDescription& InDescription);
-			void RemoveOutputDataVertexDescription(const FOutputDataVertexDescription& InDescription);
+			void AddOutputDataVertex(const FOutputDataVertex& InVertex);
+			void RemoveOutputDataVertex(const FOutputDataVertex& InVertex);
 
 		private:
 
 			FString Description;
 
-			TArray<FInputDataVertexDescription> Inputs;
-			TArray<FOutputDataVertexDescription> Outputs;
+			FInputDataVertexCollection Inputs;
+			FOutputDataVertexCollection Outputs;
 	};
 }
diff --git a/Engine/Plugins/Runtime/Metasound/Source/MetasoundGraphCore/Public/MetasoundNodeInterface.h b/Engine/Plugins/Runtime/Metasound/Source/MetasoundGraphCore/Public/MetasoundNodeInterface.h
--- a/Engine/Plugins/Runtime/Metasound/Source/MetasoundGraphCore/Public/MetasoundNodeInterface.h
+++ b/Engine/Plugins/Runtime/Metasound/Source/MetasoundGraphCore/Public/MetasoundNodeInterface.h
@@ -8,135 +8,109 @@
 
 namespace Metasound
 {
-	/** FDataVertexDescription
+	/** FDataVertex
 	 *
 	 *  This describes a data vertex on a Node.
 	 */
-	struct METASOUNDGRAPHCORE_API FDataVertexDescription
+	struct FDataVertex
 	{
-		/** Name of parameteer. */
+		/** Name of vertex. */
 		FString VertexName;
 
-		/** Type name of parmaeter. */
+		/** Type name of data reference. */
 		FName DataReferenceTypeName;
 
-		/** Tooltip for parameter. */
-		FText Tooltip;
+		/** Description of the vertex. */
+		FText Description;
 
-		FDataVertexDescription()
-		:	Tooltip(FText::GetEmpty())
+		FDataVertex()
+		:	Description(FText::GetEmpty())
 		{
 		}
 
-		FDataVertexDescription(const FString& InVertexName, const FName& InDataReferenceTypeName, const FText& InTooltip)
+		FDataVertex(const FString& InVertexName, const FName& InDataReferenceTypeName, const FText& InDescription)
 		:	VertexName(InVertexName)
 		,	DataReferenceTypeName(InDataReferenceTypeName)
-		,	Tooltip(InTooltip)
+		,	Description(InDescription)
 		{
 		}
 
-		/** Check if two parameter descriptions are equal. */
-		friend bool METASOUNDGRAPHCORE_API operator==(const FDataVertexDescription& LHS, const FDataVertexDescription& RHS);
+		/** Check if two data vertex specs are equal. */
+		friend bool METASOUNDGRAPHCORE_API operator==(const FDataVertex& LHS, const FDataVertex& RHS);
 
-		/** Check if two parameter descriptions are unequal. */
-		friend bool METASOUNDGRAPHCORE_API operator!=(const FDataVertexDescription& LHS, const FDataVertexDescription& RHS);
+		/** Check if two data vertex specs are unequal. */
+		friend bool METASOUNDGRAPHCORE_API operator!=(const FDataVertex& LHS, const FDataVertex& RHS);
 	};
 
-	struct FInputDataVertexDescription : FDataVertexDescription
+	struct FInputDataVertex: FDataVertex
 	{
-		using FDataVertexDescription::FDataVertexDescription;
+		using FDataVertex::FDataVertex;
 	};
 
-	struct FOutputDataVertexDescription : FDataVertexDescription
+	struct FOutputDataVertex : FDataVertex
 	{
-		using FDataVertexDescription::FDataVertexDescription;
+		using FDataVertex::FDataVertex;
 	};
 
 
+	/** Create an input data vertex for a DataType. DataType refers to the C++
+	 * type used to specialize the TMetasoundDataTypeInfo<> struct.
+	 *
+	 * @param InVertexName - Name of the input data vertex.
+	 * @param InDescription - Description of the input data vertex.
+	 *
+	 * @return A new FInputDataVertex.
+	 */
 	template<typename DataType>
-	FInputDataVertexDescription MakeInputDataVertexDescription(const FString& InVertexName, const FText& InTooltip)
+	FInputDataVertex MakeInputDataVertex(const FString& InVertexName, const FText& InDescription)
 	{
 		static const FName DataTypeName = FName(TDataReferenceTypeInfo<DataType>::TypeName);
 
-		return FInputDataVertexDescription(InVertexName, DataTypeName, InTooltip);
+		return FInputDataVertex(InVertexName, DataTypeName, InDescription);
 	}
 
+	/** Create an output data vertex for a DataType. DataType refers to the C++
+	 * type used to specialize the TMetasoundDataTypeInfo<> struct.
+	 *
+	 * @param InVertexName - Name of the output data vertex.
+	 * @param InDescription - Description of the output data vertex.
+	 *
+	 * @return A new FOutputDataVertex.
+	 */
 	template<typename DataType>
-	FOutputDataVertexDescription MakeOutputDataVertexDescription(const FString& InVertexName, const FText& InTooltip)
+	FOutputDataVertex MakeOutputDataVertex(const FString& InVertexName, const FText& InDescription)
 	{
 		static const FName DataTypeName = FName(TDataReferenceTypeInfo<DataType>::TypeName);
 
-		return FOutputDataVertexDescription(InVertexName, DataTypeName, InTooltip);
+		return FOutputDataVertex(InVertexName, DataTypeName, InDescription);
 	}
 
+	/** Key type for an FInputDataVertexColletion or 
+	 * FOutputDataVertexCollection. 
+	 */
+	typedef FString FDataVertexKey;
 
-	// Forward declare.
-	class INode;
-
-	/** Represents a specific input data vertex from a specific node. */
-	struct METASOUNDGRAPHCORE_API FInputDataVertex
-	{
-		/** Pointer to the node associated with the vertex. */
-		INode* Node;
-
-		/** The description of the Node's vertex. */
-		FInputDataVertexDescription Description;
-
-		FInputDataVertex()
-		:	Node(nullptr)
-		{
-		}
+	/** FInputDataVertexCollection contains multiple FInputDataVertexes mapped
+	 * by FDataVertexKeys.
+	 */
+	typedef TMap<FDataVertexKey, FInputDataVertex> FInputDataVertexCollection;
 
-		FInputDataVertex(INode* InNode, const FInputDataVertexDescription& InDescription)
-		:	Node(InNode)
-		,	Description(InDescription)
-		{
-		}
-	};
+	/** FOutputDataVertexCollection contains multiple FOutputDataVertexes mapped
+	 * by FDataVertexKeys.
+	 */
+	typedef TMap<FDataVertexKey, FOutputDataVertex> FOutputDataVertexCollection;
 
-	/** Represents a specific output data vertex from a specific node. */
-	struct METASOUNDGRAPHCORE_API FOutputDataVertex
+	/** Create an FDataVertexKey from a FInputDataVertex. */
+	FORCEINLINE FDataVertexKey MakeDataVertexKey(const FInputDataVertex& InVertex)
 	{
-		/** Pointer to the node associated with the vertex. */
-		INode* Node;
-
-		/** The description of the Node's vertex. */
-		FOutputDataVertexDescription Description;
-
-		FOutputDataVertex()
-		:	Node(nullptr)
-		{
-		}
-
-		FOutputDataVertex(INode* InNode, const FOutputDataVertexDescription& InDescription)
-		:	Node(InNode)
-		,	Description(InDescription)
-		{
-		}
-	};
+		return InVertex.VertexName;
+	}
 
-	/** FDataEdge
-	 *
-	 * An edge describes a connection between two Nodes. 
-	 */
-	struct METASOUNDGRAPHCORE_API FDataEdge
+	/** Create an FDataVertexKey from a FOutputDataVertex. */
+	FORCEINLINE FDataVertexKey MakeDataVertexKey(const FOutputDataVertex& InVertex)
 	{
-		/** Vertex producing the parameter. */
-		FOutputDataVertex From;
-
-		/** Vertex consuming the parameter. */
-		FInputDataVertex To;
-
-		FDataEdge()
-		{
-		}
-
-		FDataEdge(const FOutputDataVertex& InFrom, const FInputDataVertex& InTo)
-		:	From(InFrom)
-		,	To(InTo)
-		{
-		}
-	};
+		return InVertex.VertexName;
+	}
 
 	/** INodeBase
 	 * 
@@ -153,11 +127,11 @@ namespace Metasound
 			/** Return the type name of this node. */
 			virtual const FName& GetClassName() const = 0;
 
-			/** Return an array of input parameter descriptions for this node. */
-			virtual const TArray<FInputDataVertexDescription>& GetInputs() const = 0;
+			/** Return a collection of input parameter descriptions for this node. */
+			virtual const FInputDataVertexCollection& GetInputDataVertices() const = 0;
 
-			/** Return an array of output parameter descriptions for this node. */
-			virtual const TArray<FOutputDataVertexDescription>& GetOutputs() const = 0;
+			/** Return a collection of output parameter descriptions for this node. */
+			virtual const FOutputDataVertexCollection& GetOutputDataVertices() const = 0;
 	};
 
 	// Forward declare
@@ -172,6 +146,7 @@ namespace Metasound
 		public:
 			virtual ~INode() {}
 
+			// TODO: consider making this return a sharedref to allow for more flexible management of factories. 
 			/** Return a reference to the default operator factory. */
 			virtual IOperatorFactory& GetDefaultOperatorFactory() = 0;
 
@@ -182,6 +157,97 @@ namespace Metasound
 			*/
 	};
 
+
+	/** FOutputDataSource describes the source of data which is produced within
+	 * a graph and exposed external to the graph. 
+	 */
+	struct FOutputDataSource
+	{
+		/** Node containing the output data vertex. */
+		INode* Node = nullptr;
+
+		/** Output data vertex. */
+		FOutputDataVertex Vertex;
+
+		FOutputDataSource()
+		{
+		}
+
+		FOutputDataSource(INode& InNode, const FOutputDataVertex& InVertex)
+		:	Node(&InNode)
+		,	Vertex(InVertex)
+		{
+		}
+	};
+
+	/** FInputDataSource describes the destination of data which produced 
+	 * external to the graph and read internal to the graph.
+	 */
+	struct FInputDataDestination
+	{
+		/** Node containing the input data vertex. */
+		INode* Node = nullptr;
+
+		/** Input data vertex of edge. */
+		FInputDataVertex Vertex;
+
+		FInputDataDestination()
+		{
+		}
+
+		FInputDataDestination(INode& InNode, const FInputDataVertex& InVertex)
+		:	Node(&InNode)
+		,	Vertex(InVertex)
+		{
+		}
+	};
+
+	/** Key type for an FOutputDataSource or FInputDataDestination. */
+	typedef TTuple<const INode*, FString> FNodeDataVertexKey;
+
+	/** FOutputDataSourceCollection contains multiple FOutputDataSources mapped 
+	 * by FNodeDataVertexKeys.
+	 */
+	typedef TMap<FNodeDataVertexKey, FOutputDataSource> FOutputDataSourceCollection;
+
+	/** FInputDataDestinationCollection contains multiple 
+	 * FInputDataDestinations mapped by FNodeDataVertexKeys.
+	 */
+	typedef TMap<FNodeDataVertexKey, FInputDataDestination> FInputDataDestinationCollection;
+
+	/** Make a FNodeDataVertexKey from an FOutputDataSource. */
+	FORCEINLINE FNodeDataVertexKey MakeSourceDataVertexKey(const FOutputDataSource& InSource)
+	{
+		return FNodeDataVertexKey(InSource.Node, InSource.Vertex.VertexName);
+	}
+
+	FORCEINLINE FNodeDataVertexKey MakeDestinationDataVertexKey(const FInputDataDestination& InDestination)
+	{
+		return FNodeDataVertexKey(InDestination.Node, InDestination.Vertex.VertexName);
+	}
+
+	/** FDataEdge
+	 *
+	 * An edge describes a connection between two Nodes. 
+	 */
+	struct FDataEdge
+	{
+		FOutputDataSource From;
+
+		FInputDataDestination To;
+
+		FDataEdge()
+		{
+		}
+
+		FDataEdge(const FOutputDataSource& InFrom, const FInputDataDestination& InTo)
+		:	From(InFrom)
+		,	To(InTo)
+		{
+		}
+	};
+
+
 	/** Interface for graph of nodes. */
 	class IGraph : public INodeBase
 	{
@@ -192,10 +258,10 @@ namespace Metasound
 			virtual const TArray<FDataEdge>& GetDataEdges() const = 0;
 
 			/** Get vertices which contain input parameters. */
-			virtual const TArray<FInputDataVertex>& GetInputDataVertices() const = 0;
+			virtual const FInputDataDestinationCollection& GetInputDataDestinations() const = 0;
 
 			/** Get vertices which contain output parameters. */
-			virtual const TArray<FOutputDataVertex>& GetOutputDataVertices() const = 0;
+			virtual const FOutputDataSourceCollection& GetOutputDataSources() const = 0;
 	};
 
 
@@ -206,5 +272,4 @@ namespace Metasound
 			virtual ~INodeFactory() {}
 			virtual TUniquePtr<INode> CreateNode(const FString& InJsonString) = 0;
 	};
-
 }
diff --git a/Engine/Plugins/Runtime/Metasound/Source/MetasoundGraphCore/Public/MetasoundOperatorBuilder.h b/Engine/Plugins/Runtime/Metasound/Source/MetasoundGraphCore/Public/MetasoundOperatorBuilder.h
--- a/Engine/Plugins/Runtime/Metasound/Source/MetasoundGraphCore/Public/MetasoundOperatorBuilder.h
+++ b/Engine/Plugins/Runtime/Metasound/Source/MetasoundGraphCore/Public/MetasoundOperatorBuilder.h
@@ -9,7 +9,6 @@
 
 namespace Metasound
 {
-	// TODO: Builder needs a type registry to create inputs. Or can have specific input/output nodes and avoid registry.
 	class METASOUNDGRAPHCORE_API FOperatorBuilder : public IOperatorBuilder
 	{
 		public:
diff --git a/Engine/Plugins/Runtime/Metasound/Source/MetasoundGraphCore/Public/MetasoundOutputNode.h b/Engine/Plugins/Runtime/Metasound/Source/MetasoundGraphCore/Public/MetasoundOutputNode.h
--- a/Engine/Plugins/Runtime/Metasound/Source/MetasoundGraphCore/Public/MetasoundOutputNode.h
+++ b/Engine/Plugins/Runtime/Metasound/Source/MetasoundGraphCore/Public/MetasoundOutputNode.h
@@ -61,7 +61,7 @@ namespace Metasound
 							return TUniquePtr<IOperator>(nullptr);
 						}
 
-						return MakeUnique<FOutputOperator>(InNode.GetDescription(), InInputDataReferences.GetDataReadReference<DataType>(DataReferenceName));
+						return MakeUnique<FOutputOperator>(DataReferenceName, InInputDataReferences.GetDataReadReference<DataType>(DataReferenceName));
 					}
 
 				private:
@@ -71,28 +71,25 @@ namespace Metasound
 		public:
 
 
-			template<typename = typename TEnableIf< TIsConstructible<DataType>::Value >::Type >
-			TOutputNode(const FString& InDescription)
-			:	Description(InDescription)
-			,	Factory(InDescription)
+			TOutputNode(const FString& InNodeDescription, const FString& InVertexName)
+			:	NodeDescription(InNodeDescription)
+			,	VertexName(InVertexName)
+			,	Factory(InVertexName)
 			{
-				Outputs.Add(MakeOutputDataVertexDescription<DataType>(Description, FText::GetEmpty()));
-				Inputs.Add(MakeInputDataVertexDescription<DataType>(Description, FText::GetEmpty()));
-			}
+				FInputDataVertex InputVertex = MakeInputDataVertex<DataType>(VertexName, FText::GetEmpty());
+				FDataVertexKey InputKey = MakeDataVertexKey(InputVertex);
 
-			template<typename... ArgTypes, typename = typename TEnableIf< TIsConstructible<DataType, ArgTypes...>::Value >::Type >
-			TOutputNode(const FString& InDescription, ArgTypes&&... Args)
-			:	Description(InDescription)
-			,	Factory(InDescription)
-			,	Data(Forward<ArgTypes>(Args)...)
-			{
-				Outputs.Add(MakeOutputDataVertexDescription<DataType>(Description, FText::GetEmpty()));
-				Inputs.Add(MakeInputDataVertexDescription<DataType>(Description, FText::GetEmpty()));
+				Inputs.Add(InputKey, InputVertex);
+
+				FOutputDataVertex OutputVertex = MakeOutputDataVertex<DataType>(VertexName, FText::GetEmpty());
+				FDataVertexKey OutputKey = MakeDataVertexKey(OutputVertex);
+
+				Outputs.Add(OutputKey, OutputVertex);
 			}
 
 			virtual const FString& GetDescription() const override
 			{
-				return Description;
+				return NodeDescription;
 			}
 
 			virtual const FName& GetClassName() const override
@@ -103,12 +100,12 @@ namespace Metasound
 				return ClassName;
 			}
 
-			virtual const TArray<FInputDataVertexDescription>& GetInputs() const override
+			virtual const FInputDataVertexCollection& GetInputDataVertices() const override
 			{
 				return Inputs;
 			}
 
-			virtual const TArray<FOutputDataVertexDescription>& GetOutputs() const override
+			virtual const FOutputDataVertexCollection& GetOutputDataVertices() const override
 			{
 				return Outputs;
 			}
@@ -119,14 +116,14 @@ namespace Metasound
 			}
 
 		private:
-			FString Description;
+			FString NodeDescription;
+			FString VertexName;
 
 			FOutputOperatorFactory Factory;
 
-			DataType Data;
 
-			TArray<FInputDataVertexDescription> Inputs;
-			TArray<FOutputDataVertexDescription> Outputs;
+			FInputDataVertexCollection Inputs;
+			FOutputDataVertexCollection Outputs;
 	};
 }
 
