add convenience function to data ref collection to allow a "Get or Create" method.

diff --git a/Engine/Plugins/Runtime/Metasound/Source/MetasoundGraphCore/Public/MetasoundDataReferenceCollection.h b/Engine/Plugins/Runtime/Metasound/Source/MetasoundGraphCore/Public/MetasoundDataReferenceCollection.h
--- a/Engine/Plugins/Runtime/Metasound/Source/MetasoundGraphCore/Public/MetasoundDataReferenceCollection.h
+++ b/Engine/Plugins/Runtime/Metasound/Source/MetasoundGraphCore/Public/MetasoundDataReferenceCollection.h
@@ -192,6 +192,29 @@ namespace Metasound
 				return FDataRefType(*static_cast<const FDataRefType*>(DataRefPtr));
 			}
 
+			/** Returns a readable data ref from the collection or construct one
+			 * if one is not there.
+			 *
+			 * @param InName - Name of readable data.
+			 * @param ConstructorArgs - Arguments to pass to constructor of TDataReadReference<DataType>
+			 *
+			 * @return A readable data reference.
+			 */
+			template<typename DataType, typename... ConstructorArgTypes>
+			TDataReadReference<DataType> GetDataReadReferenceOrConstruct(const FString& InName, ConstructorArgTypes&&... ConstructorArgs) const
+			{
+				typedef TDataReadReference<DataType> FDataRefType;
+
+				if (ContainsDataReadReference<DataType>(InName))
+				{
+					return GetDataReadReference<DataType>(InName);
+				}
+				else
+				{
+					return FDataRefType(Forward<ConstructorArgTypes>(ConstructorArgs)...);
+				}
+			}
+
 			/** Returns a writable data ref from the collection.
 			 *
 			 * @param InName - Name of readable data.
@@ -212,6 +235,28 @@ namespace Metasound
 				return FDataRefType(*static_cast<const FDataRefType*>(DataRefPtr));
 			}
 
+			/** Returns a writable data ref from the collection or construct
+			 * one if one is not there.
+			 *
+			 * @param InName - Name of readable data.
+			 * @param ConstructorArgs - Arguments to pass to constructor of TDataWriteReference<DataType>
+			 *
+			 * @return A readable data reference.
+			 */
+			template<typename DataType, typename... ConstructorArgTypes>
+			TDataWriteReference<DataType> GetDataWriteReferenceOrConstruct(const FString& InName, ConstructorArgTypes&&... ConstructorArgs) const
+			{
+				typedef TDataWriteReference<DataType> FDataRefType;
+
+				if (ContainsDataWriteReference<DataType>(InName))
+				{
+					return GetDataWriteReference<DataType>(InName);
+				}
+				else
+				{
+					return FDataRefType(Forward<ConstructorArgTypes>(ConstructorArgs)...);
+				}
+			}
 
 		private:
 
