Initial implementation of Metasound::Frontend. Includes all classes and abstractions for dealing with the metasound document tree and graph editing.

Overall super happy with the way this turned out, since it lets us support exporting metasound graphs, creating local subgraphs in a metasound graph, and other really cool things.

High Level:
1) a Metasound with any sorts of inputs and outputs can be created using a UMetasoundNodeAsset, which contains a struct called the FMetasoundDocument.
2) since it would be very easy to create a malformed FMetasoundDocument by just editing the struct directly, all editing of the FMetasoundDocument is done through a set of APIs: FGraphHandle, FNodeHandle, FInputHandle, and FOutputHandle
2) You can access the main graph for a UMetasoundNodeAsset with UMetasoundNodeAsset::GetRootGraphHandle(), which returns a Metasound::Frontend::FGraphHandle.
3) The FGraphHandle lets you do all sorts of things, including:
   i) Adding and removing nodes from the graph (AddNewNode/RemoveNode)
   ii) adding and removing inputs/outputs from a graph (AddNewInput/RemoveInput, AddNewOutput/RemoveOutput)
   iii) Get a handle to manipulate individual nodes in the graph (GetAllNodes/GetOutputNodes/GetInputNodes/GetOutputNodeWithName/GetInputNodeWithName)
   iv) Get any metadata associated with that graph (GetGraphMetadata)
   v) Create a subgraph that lives in the same document (CreateEmptySubgraphNode)
   vi) Export a graph to a new Metasound asset or JSON
   vii) compile that graph to an executable, playable, Metasound::IOperator (currently not all the way implemented in this CL but I'm getting to that tomorrow)
4) The FNodeHandle lets you manipulate any indivdual node on any given graph, primarily:
    i) Getting input pins (GetAllInputs/GetAllOutputs/GetInputWithName/GetOutputWithName)
    ii) get an FGraphHandle for the subgraph for this node if there is one (GetContainedGraph)
5) FInputHandle/FOutputHandle lets you manipulate individual input/output pins, primarily:
   i) Seeing if two pins can be connected (CanConnectTo)
   ii) Connecting two pins (Connect/ConnectWithConverterNode, the latter of which isn't implemented yet)

Mid Level (the metasound document model):
1) FMetasoundDocument is an aggregate and totally serializable to binary or text using the uproperty serializer.
2) A FMetasoundDocument owns a root class, which is it's primary graph, and a list of dependencies, which could be c++ implemented nodes, metasound graphs living in other assets, or subgraphs living in the current metasound document. the Dependencies list owns a list of any dependencies for the root class, as well as any nested dependencies (for example, if the root class has a subgraph that requires a C++ implemented sine osc, the subgraph and the sine osc will both be in the Dependencies array.
3) Each FMetasoundClassDescription contains that metasounds metadata, as well as a list of inputs and outputs for that metasound, and optionally a fully implemented graph.
4) If the metasound class description has a fully implemented graph, that graph is a list of FMetasoundNodeDescriptions, which themselves describe which nodes they are connected to.
5) Each FMetasoundClassDescription contains a list of uint32s corresponding to a dependency in the document's Dependency list.

Low level (how individual elements of the FMetasoundDocument are accessed and edited)
1) We can't hand out weak pointers to the FMetasoundDocument itself unfortunately, but We are able to enforce a strict visitor pattern using the class FMetasound::Frontend:: FDescriptionAccessPoint. Currently, uobjects like UMetasoundNodeAsset can hand out weak pointers to the FDescriptionAccessPoint, and the FDescriptionAccessPoint can return raw pointers to specific elements in the FMetasoundDocument using FDescriptionAccessPoint::GetElementAtPath. Since individual elements in the document can get moved around for a number of reasons (for example, we added or removed a dependency to the document, or a node to a specific subgraph), we enforce that specific elements get accessed via a path that doesn't cache specific indices or references to specific nested structs.
2) Paths to specific elements are saved as FDescPaths, which is basically an array of enums/uint32s/strings. I added some syntactical sugar using bracket overloads to more easily build out paths to specific elements. At first this looks whacky, but after a few uses it ends up being pretty easy to write paths using chained bracket operators. All paths start at a document type, and autocomplete does most of the hard work using a set of enums.

For example, to access a specific node on a specific subgraph (which is the most complicated use case):
FDescPath()[Path::EFromDocument::ToDepenencies][TEXT("MyCoolSubgraph")][Path::EFromClass::ToGraph][Path::EFromGraph::ToNodes][12 /* or whatever node id I want */];
3) Finally, Metasound::Frontend:::TDescriptionPtr wraps a weak ptr to a FDescriptionAccessPoint and a FDescPath and ends up getting used like a weak pointer. Under the hood though, it's enforcing a visitor pattern.
4) All of the big branching statements for following a FDescPath to a specific element in the metasound document is hidden in FDescriptionAccessPoint::GoToNext, which uses TVariants to cache the current step and go to the next step.
5) Most of the complicated parts of implementing FGraphHandle is dynamically adding and removing dependencies to documents and classes behind the scenes when we add/remove new nodes.
6) I also added support for basic undo/redo callbacks, but haven't explicitly added em yet to any specific handles.

diff --git a/Engine/Plugins/Runtime/Metasound/Metasound.uplugin b/Engine/Plugins/Runtime/Metasound/Metasound.uplugin
--- a/Engine/Plugins/Runtime/Metasound/Metasound.uplugin
+++ b/Engine/Plugins/Runtime/Metasound/Metasound.uplugin
@@ -25,6 +25,11 @@
 			"Name" : "MetasoundGenerator",
 			"Type" : "Runtime",
 			"LoadingPhase" : "Default"
+		},
+		{
+			"Name" : "MetasoundFrontend",
+			"Type" : "Runtime",
+			"LoadingPhase" : "Default"
 		}
 	]
 }
diff --git a/Engine/Plugins/Runtime/Metasound/Source/MetasoundFrontend/MetasoundFrontend.Build.cs b/Engine/Plugins/Runtime/Metasound/Source/MetasoundFrontend/MetasoundFrontend.Build.cs
new file mode 100644
--- /dev/null
+++ b/Engine/Plugins/Runtime/Metasound/Source/MetasoundFrontend/MetasoundFrontend.Build.cs
@@ -0,0 +1,29 @@
+// Copyright Epic Games, Inc. All Rights Reserved.
+
+using UnrealBuildTool;
+using System.IO;
+
+namespace UnrealBuildTool.Rules
+{
+	public class MetasoundFrontend : ModuleRules
+	{
+        public MetasoundFrontend(ReadOnlyTargetRules Target) : base(Target)
+		{
+            PublicDependencyModuleNames.AddRange(
+				new string[] {
+                    "Core",
+					"SignalProcessing"
+                }
+            );
+
+			PrivateDependencyModuleNames.AddRange(
+				new string[]
+				{
+					"MetasoundGraphCore",
+					"CoreUObject",
+					"Serialization"
+				}
+			);
+		}
+	}
+}
diff --git a/Engine/Plugins/Runtime/Metasound/Source/MetasoundFrontend/Private/MetasoundFrontend.cpp b/Engine/Plugins/Runtime/Metasound/Source/MetasoundFrontend/Private/MetasoundFrontend.cpp
new file mode 100644
--- /dev/null
+++ b/Engine/Plugins/Runtime/Metasound/Source/MetasoundFrontend/Private/MetasoundFrontend.cpp
@@ -0,0 +1,1285 @@
+// Copyright Epic Games, Inc. All Rights Reserved.
+
+#include "MetasoundFrontend.h"
+#include "Backends/JsonStructSerializerBackend.h"
+#include "StructSerializer.h"
+#include "MetasoundNodeAsset.h"
+#include "MetasoundAudioFormats.h"
+
+static int32 MetasoundUndoRollLimitCvar = 128;
+FAutoConsoleVariableRef CVarMetasoundUndoRollLimit(
+	TEXT("au.Metasound.Frontend.UndoRollLimit"),
+	MetasoundUndoRollLimitCvar,
+	TEXT("Sets the maximum size of our undo buffer for graph editing in the Metasound Frontend.\n")
+	TEXT("n: Number of undoable actions we buffer."),
+	ECVF_Default);
+
+namespace Metasound
+{
+	namespace Frontend
+	{
+		TArray<FNodeClassInfo> GetAllAvailableNodeClasses()
+		{
+			// UEAU-471
+			ensureAlwaysMsgf(false, TEXT("Implement me!"));
+			return TArray<FNodeClassInfo>();
+		}
+
+		TArray<FNodeClassInfo> GetAllNodeClassesInNamespace(const FString& InNamespace)
+		{
+			// UEAU-471
+			ensureAlwaysMsgf(false, TEXT("Implement me!"));
+			return TArray<FNodeClassInfo>();
+		}
+
+		TArray<FNodeClassInfo> GetAllNodesWhoseNameContains(const FString& InSubstring)
+		{
+			// UEAU-471
+			ensureAlwaysMsgf(false, TEXT("Implement me!"));
+			return TArray<FNodeClassInfo>();
+		}
+
+		TArray<FNodeClassInfo> GetAllNodesWithAnOutputOfType(const FName& InType)
+		{
+			// UEAU-471
+			ensureAlwaysMsgf(false, TEXT("Implement me!"));
+			return TArray<FNodeClassInfo>();
+		}
+
+		TArray<FNodeClassInfo> GetAllNodesWithAnInputOfType(const FName& InType)
+		{
+			// UEAU-471
+			ensureAlwaysMsgf(false, TEXT("Implement me!"));
+			return TArray<FNodeClassInfo>();
+		}
+
+		// gets all metadata (name, description, author, what to say if it's missing) for a given node.
+		FMetasoundClassMetadata GetMetadataForNode(const FNodeClassInfo& InInfo)
+		{
+			// UEAU-471
+			ensureAlwaysMsgf(false, TEXT("Implement me!"));
+			return FMetasoundClassMetadata();
+		}
+
+		FMetasoundClassDescription GetClassDescriptionForNode(const FNodeClassInfo& InInfo)
+		{
+			// UEAU-471
+			ensureAlwaysMsgf(false, TEXT("Implement me!"));
+			return FMetasoundClassDescription();
+		}
+
+		FInputHandle::FInputHandle(FHandleInitParams::EPrivateToken PrivateToken, const FHandleInitParams& InParams, const FString& InputName)
+			: ITransactable(MetasoundUndoRollLimitCvar, InParams.InOwningAsset)
+			, NodePtr(InParams.InAccessPoint, InParams.InPath)
+			, NodeClass(InParams.InAccessPoint, Path::GetDependencyPath(InParams.InClassName))
+			, InputPtr(InParams.InAccessPoint, NodeClass.GetPath()[Path::EFromClass::ToInputs][*InputName])
+		{
+			if (InParams.InAccessPoint.IsValid())
+			{
+				// Test both pointers to the graph and it's owning class description.
+				ensureAlwaysMsgf(NodePtr.IsValid() && NodeClass.IsValid() && InputPtr.IsValid(), TEXT("Tried to build FGraphHandle with Invalid Path: %s"), *Path::GetPrintableString(InParams.InPath));
+			}
+		}
+
+		FInputHandle FInputHandle::InvalidHandle()
+		{
+			FDescPath NullPath = FDescPath();
+			FString NullString = FString();
+			FHandleInitParams InitParams = { nullptr, NullPath , NullString, nullptr };
+			return FInputHandle(FHandleInitParams::PrivateToken, InitParams, NullString);
+		}
+
+		bool FInputHandle::IsValid() const
+		{
+			return NodePtr.IsValid() && NodeClass.IsValid() && InputPtr.IsValid();
+		}
+
+		bool FInputHandle::IsConnected() const
+		{
+			if (!IsValid())
+			{
+				return false;
+			}
+
+			const FMetasoundNodeConnectionDescription* Connection = GetConnectionDescription();
+			if (Connection)
+			{
+				return Connection->NodeID != FMetasoundNodeConnectionDescription::DisconnectedNodeID;
+			}
+			else
+			{
+				return false;
+			}
+		}
+
+		FName FInputHandle::GetInputType() const
+		{
+			if (!IsValid())
+			{
+				return FName();
+			}
+
+			return InputPtr->TypeName;
+		}
+
+		FString FInputHandle::GetInputName() const
+		{
+			if (!IsValid())
+			{
+				return FString();
+			}
+
+			return InputPtr->Name;
+		}
+
+		FString FInputHandle::GetInputTooltip() const
+		{
+			if (!IsValid())
+			{
+				return FString();
+			}
+
+			return InputPtr->ToolTip;
+		}
+
+		FOutputHandle FInputHandle::GetCurrentlyConnectedOutput() const
+		{
+			if (!IsValid())
+			{
+				return FOutputHandle::InvalidHandle();
+			}
+
+			if (!IsConnected())
+			{
+				return FOutputHandle::InvalidHandle();
+			}
+
+			const FMetasoundNodeConnectionDescription* Connection = GetConnectionDescription();
+			check(Connection);
+
+			FString OutputName = Connection->OutputNameOrLiteral;
+			uint32 OutputNodeID = Connection->NodeID;
+
+			// All node connections are in the same graph, so we just need to go up one level to the Nodes array
+			// and look up the node by it's unique ID.
+			FDescPath OutputNodePath = (NodePtr.GetPath() << 1)[OutputNodeID];
+			FHandleInitParams InitParams = { NodePtr.GetAccessPoint(), OutputNodePath, NodeClass->Metadata.NodeName, OwningAsset };
+
+			return FOutputHandle(FHandleInitParams::PrivateToken, InitParams,  OutputName);
+		}
+
+		FConnectability FInputHandle::CanConnectTo(const FOutputHandle& InHandle) const
+		{
+			FConnectability OutConnectability;
+			OutConnectability.Connectable = FConnectability::EConnectable::No;
+
+			if (InHandle.GetOutputType() == GetInputType())
+			{
+				OutConnectability.Connectable = FConnectability::EConnectable::Yes;
+				return OutConnectability;
+			}
+
+			//@todo: scan for possible converter nodes here. (UEAU-473)
+			return OutConnectability;
+		}
+
+		bool FInputHandle::Connect(FOutputHandle& InHandle)
+		{
+			if (!IsValid() || !InHandle.IsValid())
+			{
+				return false;
+			}
+
+			if (!ensureAlwaysMsgf(InHandle.GetOutputType() == GetInputType(), TEXT("Tried to connect incompatible types!")))
+			{
+				return false;
+			}
+
+			uint32 OutputNodeID = InHandle.GetOwningNodeID();
+			FString OutputName = InHandle.GetOutputName();
+
+			FMetasoundNodeConnectionDescription* Connection = GetConnectionDescription();
+
+			if (!Connection)
+			{
+				TArray<FMetasoundNodeConnectionDescription>& Connections = NodePtr->InputConnections;
+				FMetasoundNodeConnectionDescription NewConnection;
+				NewConnection.InputName = GetInputName();
+				Connection = &Connections.Add_GetRef(NewConnection);
+			}
+
+			Connection->NodeID = OutputNodeID;
+			Connection->OutputNameOrLiteral = OutputName;
+
+			return true;
+		}
+
+		bool FInputHandle::ConnectWithConverterNode(FOutputHandle& InHandle, FString& InNodeClassName)
+		{
+			// (UEAU-473)
+			ensureAlwaysMsgf(false, TEXT("Implement me!"));
+			return false;
+		}
+
+		const FMetasoundNodeConnectionDescription* FInputHandle::GetConnectionDescription() const
+		{
+			if (!IsValid())
+			{
+				return nullptr;
+			}
+
+			const TArray<FMetasoundNodeConnectionDescription>& Connections = (*NodePtr).InputConnections;
+
+			for (const FMetasoundNodeConnectionDescription& Connection : Connections)
+			{
+				if (Connection.InputName == GetInputName())
+				{
+					return &Connection;
+				}
+			}
+
+			return nullptr;
+		}
+
+		FMetasoundNodeConnectionDescription* FInputHandle::GetConnectionDescription()
+		{
+			if (!IsValid())
+			{
+				return nullptr;
+			}
+
+			TArray<FMetasoundNodeConnectionDescription>& Connections = (*NodePtr).InputConnections;
+
+			for (FMetasoundNodeConnectionDescription& Connection : Connections)
+			{
+				if (Connection.InputName == GetInputName())
+				{
+					return &Connection;
+				}
+			}
+
+			return nullptr;
+		}
+
+		FOutputHandle::FOutputHandle(FHandleInitParams::EPrivateToken InToken /* unused */, const FHandleInitParams& InParams, const FString& InOutputName)
+			: ITransactable(MetasoundUndoRollLimitCvar, InParams.InOwningAsset)
+			, NodePtr(InParams.InAccessPoint, InParams.InPath)
+			, NodeClass(InParams.InAccessPoint, Path::GetDependencyPath(InParams.InClassName))
+			, OutputPtr(InParams.InAccessPoint, NodeClass.GetPath()[Path::EFromClass::ToOutputs][*InOutputName])
+		{
+			if (InParams.InAccessPoint.IsValid())
+			{
+				// Test both pointers to the graph and it's owning class description.
+				ensureAlwaysMsgf(NodePtr.IsValid() && NodeClass.IsValid() && OutputPtr.IsValid(), TEXT("Tried to build FGraphHandle with Invalid Path: %s"), *Path::GetPrintableString(InParams.InPath));
+			}
+		}
+
+		FOutputHandle FOutputHandle::InvalidHandle()
+		{
+			FString NullString;
+			FDescPath NullPath;
+			FHandleInitParams InitParams = { nullptr, NullPath, NullString, nullptr };
+			return FOutputHandle(FHandleInitParams::PrivateToken, InitParams, NullString);
+		}
+
+		bool FOutputHandle::IsValid() const
+		{
+			return NodePtr.IsValid() && NodeClass.IsValid() && OutputPtr.IsValid();
+		}
+
+		FName FOutputHandle::GetOutputType() const
+		{
+			if (!IsValid())
+			{
+				return FName();
+			}
+
+			return OutputPtr->TypeName;
+		}
+
+		FString FOutputHandle::GetOutputName() const
+		{
+			if (!IsValid())
+			{
+				return FString();
+			}
+
+			return OutputPtr->Name;
+		}
+
+		FString FOutputHandle::GetOutputTooltip() const
+		{
+			if (!IsValid())
+			{
+				return FString();
+			}
+
+			return OutputPtr->ToolTip;
+		}
+
+		uint32 FOutputHandle::GetOwningNodeID() const
+		{
+			if (!IsValid())
+			{
+				return INDEX_NONE;
+			}
+
+			return NodePtr->UniqueID;
+		}
+
+		FConnectability FOutputHandle::CanConnectTo(const FInputHandle& InHandle) const
+		{
+			return InHandle.CanConnectTo(*this);
+		}
+
+		bool FOutputHandle::Connect(FInputHandle& InHandle)
+		{
+			if (!IsValid() || !InHandle.IsValid())
+			{
+				return false;
+			}
+
+			return InHandle.Connect(*this);
+		}
+
+		bool FOutputHandle::ConnectWithConverterNode(FInputHandle& InHandle, FString& InNodeClassName)
+		{
+			return InHandle.ConnectWithConverterNode(*this, InNodeClassName);
+		}
+
+		FNodeHandle::FNodeHandle(FHandleInitParams::EPrivateToken PrivateToken, const FHandleInitParams& InParams)
+			: ITransactable(MetasoundUndoRollLimitCvar, InParams.InOwningAsset)
+			, NodePtr(InParams.InAccessPoint, InParams.InPath)
+			, NodeClass(InParams.InAccessPoint, Path::GetDependencyPath(InParams.InClassName))
+		{
+			if (InParams.InAccessPoint.IsValid())
+			{
+				// Test both pointers to the graph and it's owning class description.
+				ensureAlwaysMsgf(NodePtr.IsValid() && NodeClass.IsValid(), TEXT("Tried to build FGraphHandle with Invalid Path: %s"), *Path::GetPrintableString(InParams.InPath));
+			}
+		}
+
+		FNodeHandle FNodeHandle::InvalidHandle()
+		{
+			FDescPath InvalidPath;
+			FString InvalidClassName;
+			FHandleInitParams InitParams = { nullptr, InvalidPath, InvalidClassName, nullptr };
+			return FNodeHandle(FHandleInitParams::PrivateToken, InitParams);
+		}
+
+		bool FNodeHandle::IsValid() const
+		{
+			return NodePtr.IsValid() && NodeClass.IsValid();
+		}
+
+		TArray<FInputHandle> FNodeHandle::GetAllInputs()
+		{
+			TArray<FInputHandle> OutArray;
+
+			if (!IsValid())
+			{
+				return OutArray;
+			}
+
+			TArray<FMetasoundInputDescription>& InputDescriptions = NodeClass->Inputs;
+
+			for (FMetasoundInputDescription& InputDescription : InputDescriptions)
+			{
+				FDescPath NodePath = NodePtr.GetPath();
+				FString ClassName = GetNodeClassName();
+				FHandleInitParams InitParams = { NodePtr.GetAccessPoint(), NodePath, ClassName , OwningAsset };
+				OutArray.Emplace(FHandleInitParams::PrivateToken, InitParams, InputDescription.Name);
+			}
+
+			return OutArray;
+		}
+
+		TArray<FOutputHandle> FNodeHandle::GetAllOutputs()
+		{
+			TArray<FOutputHandle> OutArray;
+
+			if (!IsValid())
+			{
+				return OutArray;
+			}
+
+			TArray<FMetasoundOutputDescription>& OutputDescriptions = NodeClass->Outputs;
+
+			for (FMetasoundOutputDescription& OutputDescription : OutputDescriptions)
+			{
+				FDescPath NodePath = NodePtr.GetPath();
+				FString NodeClassName = GetNodeClassName();
+				FHandleInitParams InitParams = { NodePtr.GetAccessPoint(), NodePath, NodeClassName, OwningAsset };
+				OutArray.Emplace(FHandleInitParams::PrivateToken, InitParams, OutputDescription.Name);
+			}
+
+			return OutArray;
+		}
+
+		FInputHandle FNodeHandle::GetInputWithName(const FString& InName)
+		{
+			if (!IsValid())
+			{
+				return FInputHandle::InvalidHandle();
+			}
+
+			TArray<FMetasoundInputDescription>& InputDescriptions = NodeClass->Inputs;
+
+			for (FMetasoundInputDescription& InputDescription : InputDescriptions)
+			{
+				if (InputDescription.Name == InName)
+				{
+					FString ClassName = GetNodeClassName();
+					FHandleInitParams InitParams = { NodePtr.GetAccessPoint(), NodePtr.GetPath(), ClassName, OwningAsset };
+					return FInputHandle(FHandleInitParams::PrivateToken, InitParams, InputDescription.Name);
+				}
+			}
+
+			ensureAlwaysMsgf(false, TEXT("Couldn't find an input with this name on this node!"));
+			return FInputHandle::InvalidHandle();
+		}
+
+		FOutputHandle FNodeHandle::GetOutputWithName(const FString& InName)
+		{
+			if (!IsValid())
+			{
+				return FOutputHandle::InvalidHandle();
+			}
+
+			TArray<FMetasoundOutputDescription>& OutputDescriptions = NodeClass->Outputs;
+
+			for (FMetasoundOutputDescription& OutputDescription : OutputDescriptions)
+			{
+				if (OutputDescription.Name == InName)
+				{
+					FString NodeClassName = GetNodeClassName();
+					FHandleInitParams InitParams = { NodePtr.GetAccessPoint(), NodePtr.GetPath(), NodeClassName , OwningAsset };
+					return FOutputHandle(FHandleInitParams::PrivateToken, InitParams, OutputDescription.Name);
+				}
+			}
+
+			ensureAlwaysMsgf(false, TEXT("Couldn't find an output with this name on this node!"));
+			return FOutputHandle::InvalidHandle();
+		}
+
+		EMetasoundClassType FNodeHandle::GetNodeType()
+		{
+			if (!IsValid())
+			{
+				return EMetasoundClassType::Invalid;
+			}
+
+			return NodeClass->Metadata.NodeType;
+		}
+
+		FString FNodeHandle::GetNodeClassName()
+		{
+			if (!IsValid())
+			{
+				return FString();
+			}
+
+			return NodeClass->Metadata.NodeName;
+		}
+
+		void FNodeHandle::GetContainedGraph(FGraphHandle& OutGraph)
+		{
+			if (!IsValid())
+			{
+				OutGraph = FGraphHandle::InvalidHandle();
+			}
+
+			if (!ensureAlwaysMsgf(GetNodeType() == EMetasoundClassType::MetasoundGraph, TEXT("Tried to get the Metasound Graph for a node that was not a Metasound graph.")))
+			{
+				OutGraph = FGraphHandle::InvalidHandle();
+			}
+
+			FDescPath ContainedGraphPath = NodeClass.GetPath()[Path::EFromClass::ToGraph];
+			FHandleInitParams InitParams = { NodeClass.GetAccessPoint(), ContainedGraphPath, NodeClass->Metadata.NodeName, OwningAsset };
+			// Todo: link this up to look for externally implemented graphs as well.
+			OutGraph = FGraphHandle(FHandleInitParams::PrivateToken, InitParams);
+		}
+
+		uint32 FNodeHandle::GetNodeID()
+		{
+			if (!IsValid())
+			{
+				return INDEX_NONE;
+			}
+
+			return NodePtr->UniqueID;
+		}
+
+		uint32 FNodeHandle::GetNodeID(const FDescPath& InNodePath)
+		{
+			if (!ensureAlwaysMsgf(InNodePath.Path.Num() != 0, TEXT("Tried to get a node ID from an empty path.")))
+			{
+				return INDEX_NONE;
+			}
+
+			const Path::FElement& LastElementInPath = InNodePath.Path.Last();
+			
+			if (!ensureAlwaysMsgf(LastElementInPath.CurrentDescType == Path::EDescType::Node, TEXT("Tried to get the node ID for a path that was not set up for a node.")))
+			{
+				return INDEX_NONE;
+			}
+
+			return LastElementInPath.LookupID;
+		}
+
+		FGraphHandle::FGraphHandle(FHandleInitParams::EPrivateToken PrivateToken, const FHandleInitParams& InParams)
+			: ITransactable(MetasoundUndoRollLimitCvar, InParams.InOwningAsset)
+			, GraphPtr(InParams.InAccessPoint, InParams.InPath)
+			, GraphsClassDeclaration(InParams.InAccessPoint, Path::GetOwningClassDescription(InParams.InPath))
+			, OwningDocument(InParams.InAccessPoint, FDescPath())
+		{
+			if (InParams.InAccessPoint.IsValid())
+			{
+				// Test both pointers to the graph and it's owning class description.
+				ensureAlwaysMsgf(GraphPtr.IsValid() && GraphsClassDeclaration.IsValid(), TEXT("Tried to build FGraphHandle with Invalid Path: %s"), *Path::GetPrintableString(InParams.InPath));
+			}
+		}
+
+		uint32 FGraphHandle::FindNewUniqueNodeId()
+		{
+			// Assumption here is that we will never need more than ten thousand nodes,
+			// and four digits are easy enough to read/remember when looking at metasound graph documents.
+			static const uint32 NodeIDMax = 9999;
+
+			if (!IsValid())
+			{
+				return INDEX_NONE;
+			}
+
+			TArray<FMetasoundNodeDescription>& Nodes = GraphPtr->Nodes;
+
+			if (!ensureAlwaysMsgf(((uint32)Nodes.Num()) < NodeIDMax, TEXT("Too many nodes to guarantee a unique node ID. Increase the value of NodeIDMax.")))
+			{
+				return INDEX_NONE;
+			}
+
+			while (uint32 RandomID = FMath::RandRange(1, NodeIDMax))
+			{
+				// Scan through the nodes in this graph to see if they match this ID.
+				// If it does, generate a new random ID.
+				bool bFoundMatch = false;
+				for (FMetasoundNodeDescription& Node : Nodes)
+				{
+					if (Node.UniqueID == RandomID)
+					{
+						bFoundMatch = true;
+						break;
+					}
+				}
+
+				if (!bFoundMatch)
+				{
+					return RandomID;
+				}
+			}
+
+			return INDEX_NONE;
+		}
+
+		uint32 FGraphHandle::FindNewUniqueDependencyId()
+		{
+			// Assumption here is that we will never need more than ten thousand dependencies,
+			// and four digits are easy enough to read/remember when looking at metasound graph documents.
+			static const uint32 DependencyIDMax = 9999;
+
+			if (!IsValid())
+			{
+				return INDEX_NONE;
+			}
+
+			TArray<FMetasoundClassDescription>& Dependencies = OwningDocument->Dependencies;
+
+			if (!ensureAlwaysMsgf(Dependencies.Num() < DependencyIDMax, TEXT("Too many nodes to guarantee a unique node ID. Increase the value of NodeIDMax.")))
+			{
+				return INDEX_NONE;
+			}
+
+			while (uint32 RandomID = FMath::RandRange(1, DependencyIDMax))
+			{
+				// Scan through the nodes in this graph to see if they match this ID.
+				// If it does, generate a new random ID.
+				bool bFoundMatch = false;
+				for (FMetasoundClassDescription& Dependency : Dependencies)
+				{
+					if (Dependency.UniqueID == RandomID)
+					{
+						bFoundMatch = true;
+						break;
+					}
+				}
+
+				if (!bFoundMatch)
+				{
+					return RandomID;
+				}
+			}
+
+			return INDEX_NONE;
+		}
+
+		FGraphHandle FGraphHandle::InvalidHandle()
+		{
+			FDescPath InvalidPath;
+			FString InvalidClassName;
+			FHandleInitParams InitParams = { nullptr, InvalidPath, InvalidClassName, nullptr };
+			return FGraphHandle(FHandleInitParams::Token, InitParams);
+		}
+
+		bool FGraphHandle::IsValid()
+		{
+			return GraphPtr.IsValid() && GraphsClassDeclaration.IsValid();
+		}
+
+		TArray<FNodeHandle> FGraphHandle::GetAllNodes()
+		{
+			TArray<FNodeHandle> OutArray;
+
+			if (!IsValid())
+			{
+				return OutArray;
+			}
+			
+			const TArray<FMetasoundNodeDescription>& NodeDescriptions = GraphPtr->Nodes;
+			for (const FMetasoundNodeDescription& NodeDescription : NodeDescriptions)
+			{
+				FDescPath NodePath = GraphPtr.GetPath()[Path::EFromGraph::ToNodes][NodeDescription.UniqueID];
+
+				FHandleInitParams InitParams = { GraphPtr.GetAccessPoint(), NodePath, NodeDescription.Name, OwningAsset };
+				OutArray.Emplace(FHandleInitParams::PrivateToken, InitParams);
+			}
+
+			return OutArray;
+		}
+
+		TArray<FNodeHandle> FGraphHandle::GetOutputNodes()
+		{
+			TArray<FNodeHandle> OutArray;
+
+			if (!IsValid())
+			{
+				return OutArray;
+			}
+
+			const TArray<FMetasoundNodeDescription>& NodeDescriptions = GraphPtr->Nodes;
+			for (const FMetasoundNodeDescription& NodeDescription : NodeDescriptions)
+			{
+				if (NodeDescription.ObjectTypeOfNode == EMetasoundClassType::Output)
+				{
+					FDescPath NodePath = GraphPtr.GetPath()[Path::EFromGraph::ToNodes][NodeDescription.UniqueID];
+					FHandleInitParams InitParams = { GraphPtr.GetAccessPoint(), NodePath, NodeDescription.Name, OwningAsset };
+					OutArray.Emplace(FHandleInitParams::PrivateToken, InitParams);
+				}
+			}
+
+			return OutArray;
+		}
+
+		TArray<FNodeHandle> FGraphHandle::GetInputNodes()
+		{
+			TArray<FNodeHandle> OutArray;
+
+			if (!IsValid())
+			{
+				return OutArray;
+			}
+
+			const TArray<FMetasoundNodeDescription>& NodeDescriptions = GraphPtr->Nodes;
+			for (const FMetasoundNodeDescription& NodeDescription : NodeDescriptions)
+			{
+				if (NodeDescription.ObjectTypeOfNode == EMetasoundClassType::Input)
+				{
+					FDescPath NodePath = GraphPtr.GetPath()[Path::EFromGraph::ToNodes][NodeDescription.UniqueID];
+					FHandleInitParams InitParams = { GraphPtr.GetAccessPoint(), NodePath, NodeDescription.Name, OwningAsset };
+					OutArray.Emplace(FHandleInitParams::PrivateToken, InitParams);
+				}
+			}
+
+			return OutArray;
+		}
+
+		FNodeHandle FGraphHandle::GetOutputNodeWithName(const FString& InName)
+		{
+			if (!IsValid())
+			{
+				return FNodeHandle::InvalidHandle();
+			}
+
+			const TArray<FMetasoundNodeDescription>& NodeDescriptions = GraphPtr->Nodes;
+			for (const FMetasoundNodeDescription& NodeDescription : NodeDescriptions)
+			{
+				if (NodeDescription.ObjectTypeOfNode == EMetasoundClassType::Output && NodeDescription.Name == InName)
+				{
+					FDescPath NodePath = GraphPtr.GetPath()[Path::EFromGraph::ToNodes][NodeDescription.UniqueID];
+
+					FHandleInitParams InitParams = { GraphPtr.GetAccessPoint(), NodePath, InName, OwningAsset };
+					return FNodeHandle(FHandleInitParams::PrivateToken, InitParams);
+				}
+			}
+
+			ensureAlwaysMsgf(false, TEXT("Tried to get output node %s, but it didn't exist"), *InName);
+			return FNodeHandle::InvalidHandle();
+		}
+
+		FNodeHandle FGraphHandle::GetInputNodeWithName(const FString& InName)
+		{
+			if (!IsValid())
+			{
+				return FNodeHandle::InvalidHandle();
+			}
+
+			const TArray<FMetasoundNodeDescription>& NodeDescriptions = GraphPtr->Nodes;
+			for (const FMetasoundNodeDescription& NodeDescription : NodeDescriptions)
+			{
+				if (NodeDescription.ObjectTypeOfNode == EMetasoundClassType::Input && NodeDescription.Name == InName)
+				{
+					FDescPath NodePath = GraphPtr.GetPath()[Path::EFromGraph::ToNodes][NodeDescription.UniqueID];
+
+					// todo: add special enum for input/output nodes
+					FHandleInitParams InitParams = { GraphPtr.GetAccessPoint(), NodePath, InName, OwningAsset };
+					return FNodeHandle(FHandleInitParams::PrivateToken, InitParams);
+				}
+			}
+
+			ensureAlwaysMsgf(false, TEXT("Tried to get output node %s, but it didn't exist"), *InName);
+			return FNodeHandle::InvalidHandle();
+		}
+
+		FNodeHandle FGraphHandle::AddNewInput(const FMetasoundInputDescription& InDescription)
+		{
+			if (!IsValid())
+			{
+				return FNodeHandle::InvalidHandle();
+			}
+			
+			// @todo: verify that InDescription.TypeName is a valid Metasound type.
+
+			const uint32 NewUniqueId = FindNewUniqueNodeId();
+			
+			if (!ensureAlwaysMsgf(NewUniqueId != INDEX_NONE, TEXT("FindNewUniqueNodeId failed!")))
+			{
+				return FNodeHandle::InvalidHandle();
+			}
+
+			
+			TArray<FMetasoundInputDescription>& Inputs = GraphsClassDeclaration->Inputs;
+
+			// Sanity check that this input has a unique name.
+			for (const FMetasoundInputDescription& Input : Inputs)
+			{
+				if (!ensureAlwaysMsgf(Input.Name != InDescription.Name, TEXT("Tried to add a new input with a name that already exists!")))
+				{
+					return FNodeHandle::InvalidHandle();
+				}
+			}
+
+			// Add the input to this node's class description.
+			Inputs.Add(InDescription);
+
+			FMetasoundNodeDescription NewNodeDescription;
+			NewNodeDescription.Name = InDescription.Name;
+			NewNodeDescription.UniqueID = NewUniqueId;
+			NewNodeDescription.ObjectTypeOfNode = EMetasoundClassType::Input;
+
+			GraphPtr->Nodes.Add(NewNodeDescription);
+
+			FDescPath NodePath = GraphPtr.GetPath()[Path::EFromGraph::ToNodes][NewNodeDescription.UniqueID];
+			FHandleInitParams InitParams = { GraphPtr.GetAccessPoint(), NodePath, NewNodeDescription.Name, OwningAsset };
+			// todo: add special enum for input and output nodes
+			return FNodeHandle(FHandleInitParams::PrivateToken, InitParams);
+		}
+
+		bool FGraphHandle::RemoveInput(const FString& InputName)
+		{
+			if (!IsValid())
+			{
+				return false;
+			}
+
+			TArray<FMetasoundInputDescription>& Inputs = GraphsClassDeclaration->Inputs;
+			int32 IndexOfInputToRemove = -1;
+
+			for (int32 InputIndex = 0; InputIndex < Inputs.Num(); InputIndex++)
+			{
+				if (Inputs[InputIndex].Name == InputName)
+				{
+					IndexOfInputToRemove = InputIndex;
+					break;
+				}
+			}
+
+			if (!ensureAlwaysMsgf(IndexOfInputToRemove >= 0, TEXT("Tried to remove an Input that didn't exist: %s"), *InputName))
+			{
+				return false;
+			}
+
+			// find the corresponding node handle to delete.
+			FNodeHandle InputNode = GetInputNodeWithName(InputName);
+
+			// If we found the input declared in the class description but couldn't find it in the graph,
+			// something has gone terribly wrong. Remove the input from the description, but still ensure.
+			if (!ensureAlwaysMsgf(InputNode.IsValid(), TEXT(R"(Couldn't find an input node with name %s, even though we found the input listed as a dependency.
+				This indicates the underlying FMetasoundClassDescription is corrupted.
+				Removing the Input in the class dependency to resolve...)"), *InputName))
+			{
+				Inputs.RemoveAt(IndexOfInputToRemove);
+				return true;
+			}
+
+			// Finally, remove the node, and remove the input.
+			if (!ensureAlwaysMsgf(RemoveNode(InputNode, true), TEXT("Call to RemoveNode failed.")))
+			{
+				return false;
+			}
+
+			Inputs.RemoveAt(IndexOfInputToRemove);
+			return true;
+		}
+
+		FNodeHandle FGraphHandle::AddNewOutput(const FMetasoundOutputDescription& OutDescription)
+		{
+			if (!IsValid())
+			{
+				return FNodeHandle::InvalidHandle();
+			}
+
+			// @todo: verify that InDescription.TypeName is a valid Metasound type.
+
+			const uint32 NewUniqueId = FindNewUniqueNodeId();
+
+			if (!ensureAlwaysMsgf(NewUniqueId != INDEX_NONE, TEXT("FindNewUniqueNodeId failed")))
+			{
+				return FNodeHandle::InvalidHandle();
+			}
+
+
+			TArray<FMetasoundOutputDescription>& Outputs = GraphsClassDeclaration->Outputs;
+
+			// Sanity check that this input has a unique name.
+			for (const FMetasoundOutputDescription& Output : Outputs)
+			{
+				if (!ensureAlwaysMsgf(Output.Name != OutDescription.Name, TEXT("Tried to add a new output with a name that already exists!")))
+				{
+					return FNodeHandle::InvalidHandle();
+				}
+			}
+
+			// Add the output to this node's class description.
+			Outputs.Add(OutDescription);
+
+			// Add a node for this output to the graph description.
+			FMetasoundNodeDescription NewNodeDescription;
+			NewNodeDescription.Name = OutDescription.Name;
+			NewNodeDescription.UniqueID = NewUniqueId;
+			NewNodeDescription.ObjectTypeOfNode = EMetasoundClassType::Output;
+
+			GraphPtr->Nodes.Add(NewNodeDescription);
+
+			FDescPath NodePath = GraphPtr.GetPath()[Path::EFromGraph::ToNodes][NewNodeDescription.UniqueID];
+
+			// todo: add special enum or class for input/output nodes
+			FHandleInitParams InitParams = { GraphPtr.GetAccessPoint(), NodePath, NewNodeDescription.Name, OwningAsset };
+			return FNodeHandle(FHandleInitParams::PrivateToken, InitParams);
+		}
+
+		bool FGraphHandle::RemoveOutput(const FString& OutputName)
+		{
+			if (!IsValid())
+			{
+				return false;
+			}
+
+			TArray<FMetasoundOutputDescription>& Outputs = GraphsClassDeclaration->Outputs;
+			int32 IndexOfOutputToRemove = -1;
+
+			for (int32 OutputIndex = 0; OutputIndex < Outputs.Num(); OutputIndex++)
+			{
+				if (Outputs[OutputIndex].Name == OutputName)
+				{
+					IndexOfOutputToRemove = OutputIndex;
+					break;
+				}
+			}
+
+			if (!ensureAlwaysMsgf(IndexOfOutputToRemove >= 0, TEXT("Tried to remove an Output that didn't exist: %s"), *OutputName))
+			{
+				return false;
+			}
+
+			// find the corresponding node handle to delete.
+			FNodeHandle OutputNode = GetOutputNodeWithName(OutputName);
+
+			// If we found the input declared in the class description but couldn't find it in the graph,
+			// something has gone terribly wrong. Remove the input from the description, but still ensure.
+			if (!ensureAlwaysMsgf(OutputNode.IsValid(), TEXT(R"(Couldn't find an input node with name %s, even though we found the input listed as a dependency.
+				This indicates the underlying FMetasoundClassDescription is corrupted.
+				Removing the Input in the class dependency to resolve...)"), *OutputName))
+			{
+				Outputs.RemoveAt(IndexOfOutputToRemove);
+				return true;
+			}
+
+			// Finally, remove the node, and remove the input.
+			if (!ensureAlwaysMsgf(RemoveNode(OutputNode, true), TEXT("Call to RemoveNode failed.")))
+			{
+				return false;
+			}
+
+			Outputs.RemoveAt(IndexOfOutputToRemove);
+			return true;
+		}
+
+		FNodeHandle FGraphHandle::AddNewNode(const FNodeClassInfo& InNodeClass)
+		{
+			if (!IsValid())
+			{
+				return FNodeHandle::InvalidHandle();
+			}
+
+			// First, scan our dependency list to see if this node already exists there, and if not, get it.
+			TArray<FMetasoundClassDescription>& Dependencies = OwningDocument->Dependencies;
+
+			bool bFoundMatchingDependencyInDocument = false;
+
+			for (FMetasoundClassDescription& Dependency : Dependencies)
+			{
+				if (Dependency.Metadata.NodeName == InNodeClass.NodeName && Dependency.Metadata.NodeType == InNodeClass.NodeType)
+				{
+					bFoundMatchingDependencyInDocument = true;
+
+					// If this dependency was in the document's dependency list, check to see if we need to add it to this class' dependencies.
+					bool bFoundDependencyInLocalClass = false;
+					for (uint32 DependencyID : GraphsClassDeclaration->DependencyIDs)
+					{
+						if (DependencyID == Dependency.UniqueID)
+						{
+							bFoundDependencyInLocalClass = true;
+							break;
+						}
+					}
+
+					if (!bFoundDependencyInLocalClass)
+					{
+						// This dependency is already referenced somewhere in the document, but not for this graph's class. Add it.
+						GraphsClassDeclaration->DependencyIDs.Add(Dependency.UniqueID);
+						UE_LOG(LogTemp, Verbose, TEXT("Adding %s as a dependency for Metasound graph %s in Document %s"), *InNodeClass.NodeName, *GetGraphMetadata().NodeName, *OwningDocument->RootClass.Metadata.NodeName);
+					}
+
+					break;
+				}
+			}
+
+			
+
+			// If we haven't added a node of this class to the graph yet, add it to the dependencies for this class.
+			if (!bFoundMatchingDependencyInDocument)
+			{
+				FMetasoundClassDescription NewDependencyClassDescription = GetClassDescriptionForNode(InNodeClass);
+				NewDependencyClassDescription.UniqueID = FindNewUniqueDependencyId();
+				Dependencies.Add(NewDependencyClassDescription);
+				GraphsClassDeclaration->DependencyIDs.Add(NewDependencyClassDescription.UniqueID);
+
+				UE_LOG(LogTemp, Verbose, TEXT("Adding %s is used in graph %s, adding as a new dependency for Metasound Document %s"), *InNodeClass.NodeName, *GetGraphMetadata().NodeName, *OwningDocument->RootClass.Metadata.NodeName);
+			}
+
+			// Add a new node instance for this class.
+
+			const uint32 NewUniqueId = FindNewUniqueNodeId();
+			if (!ensureAlwaysMsgf(NewUniqueId != INDEX_NONE, TEXT("Call to FindNewUniqueNodeId failed!")))
+			{
+				return FNodeHandle::InvalidHandle();
+			}
+
+			// Add a node for this output to the graph description.
+			FMetasoundNodeDescription NewNodeDescription;
+			NewNodeDescription.Name = InNodeClass.NodeName;
+			NewNodeDescription.UniqueID = NewUniqueId;
+			NewNodeDescription.ObjectTypeOfNode = InNodeClass.NodeType;
+
+			GraphPtr->Nodes.Add(NewNodeDescription);
+
+			FDescPath NodePath = GraphPtr.GetPath()[Path::EFromGraph::ToNodes][NewNodeDescription.UniqueID];
+			FHandleInitParams InitParams = { GraphPtr.GetAccessPoint(), NodePath, NewNodeDescription.Name, OwningAsset };
+			return FNodeHandle(FHandleInitParams::PrivateToken, InitParams);
+		}
+
+		bool FGraphHandle::RemoveNode(FNodeHandle& InNode, bool bEvenIfInputOrOutputNode /*= false*/)
+		{
+			if (!IsValid())
+			{
+				return false;
+			}
+
+			// First, find the node in our nodes list,
+			// while also checking to see if this is the only node of this class left in this graph.
+			TArray<FMetasoundNodeDescription>& Nodes = GraphPtr->Nodes;
+			int32 IndexOfNodeToRemove = -1;
+			int32 NodesOfClass = 0;
+
+			FString NodeClassName = InNode.GetNodeClassName();
+
+			for (int32 NodeIndex = 0; NodeIndex < Nodes.Num(); NodeIndex++)
+			{
+				if (Nodes[NodeIndex].Name == NodeClassName)
+				{
+					NodesOfClass++;
+				}
+				
+				if (Nodes[NodeIndex].UniqueID == InNode.GetNodeID())
+				{
+					IndexOfNodeToRemove = NodeIndex;
+				}
+
+				// If we've found the matching node,
+				// and have found that there is more than one node of this class,
+				// we have found all the info we need.
+				if (IndexOfNodeToRemove > 0 && NodesOfClass > 1)
+				{
+					break;
+				}
+			}
+
+			if (!ensureAlwaysMsgf(IndexOfNodeToRemove >= 0, TEXT(R"(Couldn't find node corresponding to handle (%s ID: %u).
+				Are you sure this FNodeHandle was generated from this FGraphHandle?)"),
+				*InNode.GetNodeClassName(),
+				InNode.GetNodeType()))
+			{
+				return false;
+			}
+
+			// This should never hit based on the logic above.
+			if (!ensureAlwaysMsgf(NodesOfClass > 0, TEXT("Found node with matching ID (%u) but mismatched class (%s). Likely means that the underlying class description was corrupted."), 
+				*InNode.GetNodeClassName(),
+				InNode.GetNodeType()))
+			{
+				return false;
+			}
+
+			// If this node was the only node of this class remaining in the graph,
+			// Remove its ID as a dependency for the graph.
+			if (NodesOfClass < 2)
+			{
+				TArray<uint32>& DependencyIDs = GraphsClassDeclaration->DependencyIDs;
+				int32 IndexOfDependencyToRemove = -1;
+
+				uint32 UniqueIDForThisDependency = INDEX_NONE;
+				TArray<FMetasoundClassDescription>& DependencyClasses = OwningDocument->Dependencies;
+				
+				// scan the owning document's depenency classes for a dependency with this name.
+				int32 IndexOfDependencyInDocument = -1;
+				for (int32 Index = 0; Index < DependencyClasses.Num(); Index++)
+				{
+					if (DependencyClasses[Index].Metadata.NodeName == NodeClassName)
+					{
+						IndexOfDependencyInDocument = Index;
+						UniqueIDForThisDependency = DependencyClasses[Index].UniqueID;
+						break;
+					}
+				}
+
+				if (!ensureAlwaysMsgf(IndexOfDependencyToRemove >= 0, 
+					TEXT("Couldn't find node class %s in the list of dependencies for this document, but found it in the nodes list. This likely means that the underlying class description is corrupted.)"), *NodeClassName))
+				{
+					return false;
+				}
+
+				for (int32 DependencyIndex = 0; DependencyIndex < DependencyIDs.Num(); DependencyIndex++)
+				{
+					const bool bIsDependencyForNodeToRemove = DependencyIDs[DependencyIndex] == UniqueIDForThisDependency;
+
+					if (bIsDependencyForNodeToRemove)
+					{
+						IndexOfDependencyToRemove = DependencyIndex;
+						break;
+					}
+				}
+
+				if (ensureAlwaysMsgf(IndexOfDependencyToRemove > 0, TEXT(R"(Couldn't find node class %s in the list of dependencies for this graph, but found it in the nodes list.
+					This likely means that the underlying class description is corrupted.)"), *NodeClassName))
+				{
+					DependencyIDs.RemoveAt(IndexOfDependencyToRemove);
+				}
+
+				// Finally, check to see if there are any classes remaining in this document that depend on this class, and remove it from our dependencies list.
+				bool bFoundUsageOfDependencyInClass = false;
+
+				for (FMetasoundClassDescription& Dependency : DependencyClasses)
+				{
+					if (Dependency.DependencyIDs.Contains(UniqueIDForThisDependency))
+					{
+						bFoundUsageOfDependencyInClass = true;
+					}
+				}
+
+				// Also scan the root graph for this document, which lives outside of the Dependencies list.
+				FMetasoundClassDescription& RootClass = OwningDocument->RootClass;
+				if (RootClass.DependencyIDs.Contains(UniqueIDForThisDependency))
+				{
+					bFoundUsageOfDependencyInClass = true;
+				}
+
+				if (!bFoundUsageOfDependencyInClass)
+				{
+					// we can safely delete this dependency from the document.
+					DependencyClasses.RemoveAt(IndexOfDependencyInDocument);
+				}
+			}
+
+			// Finally, remove the node from the nodes list.
+			Nodes.RemoveAt(IndexOfNodeToRemove);
+			return true;
+		}
+
+		FMetasoundClassMetadata FGraphHandle::GetGraphMetadata()
+		{
+			if (!IsValid())
+			{
+				return FMetasoundClassMetadata();
+			}
+
+			return GraphsClassDeclaration->Metadata;
+		}
+
+		UMetasound* FGraphHandle::SaveToNewMetasoundAsset(const FString& InPath)
+		{
+			if (!IsValid())
+			{
+				return nullptr;
+			}
+
+			UMetasound* NewMetasoundNode = NewObject<UMetasound>();
+			NewMetasoundNode->SetMetasoundDocument(OwningDocument.GetChecked());
+
+			ensureAlwaysMsgf(false, TEXT("Implement the actual saving part!"));
+			ensureAlwaysMsgf(false, TEXT("Decide whether to delete this graph from the asset that owns it!"));
+
+			return NewMetasoundNode;
+		}
+
+
+		bool FGraphHandle::ExportToJSONAsset(const FString& InAbsolutePath)
+		{
+			if (!IsValid())
+			{
+				return false;
+			}
+
+			if (TUniquePtr<FArchive> FileWriter = TUniquePtr<FArchive>(IFileManager::Get().CreateFileWriter(*InAbsolutePath)))
+			{
+				FJsonStructSerializerBackend Backend(*FileWriter, EStructSerializerBackendFlags::Default);
+				FStructSerializer::Serialize<FMetasoundClassDescription>(GraphsClassDeclaration.GetChecked(), Backend);
+				FileWriter->Close();
+
+				return true;
+			}
+			else
+			{
+				ensureAlwaysMsgf(false, TEXT("Failed to create a filewriter with the given path."));
+				return false;
+			}
+		}
+
+		bool FGraphHandle::InflateNodeDirectlyIntoGraph(FNodeHandle& InNode)
+		{
+			ensureAlwaysMsgf(false, TEXT("Implement me!"));
+			// nontrivial but required anyways for graph inflation (UEAU-475)
+
+			//step 0: get the node's FMetasoundClassDescription
+			//step 1: check if the node is itself a metasound
+			//step 2: get the FMetasoundGraphDescription for the node from the Dependencies list.
+			//step 3: create new unique IDs for each node in the subgraph.
+			//step 4: add nodes from subgraph to the current graph.
+			//step 5: rebuild connections for new nodes in current graph based on the new IDs.
+			//step 6: delete Input nodes and Output nodes from the subgraph, and rebuild connections from this graph directly to the nodes in the subgraph.
+			return false;
+		}
+
+		TTuple<FGraphHandle, FNodeHandle> FGraphHandle::CreateEmptySubgraphNode(const FMetasoundClassMetadata& InInfo)
+		{
+			auto BuildInvalidTupleHandle = [&]() -> TTuple<FGraphHandle, FNodeHandle>
+			{
+				return TTuple<FGraphHandle, FNodeHandle>(FGraphHandle::InvalidHandle(), FNodeHandle::InvalidHandle());
+			};
+
+			if (!IsValid())
+			{
+				return BuildInvalidTupleHandle();
+			}
+
+			// Sanity check that the given name isn't already in our graph's dependency list.
+			TArray<FMetasoundClassDescription>& Dependencies = OwningDocument->Dependencies;
+			for (FMetasoundClassDescription& Dependency : Dependencies)
+			{
+				if (!ensureAlwaysMsgf(Dependency.Metadata.NodeName != InInfo.NodeName, TEXT("Tried to create a new subgraph with name %s but there was already a dependency named that in the graph."), *InInfo.NodeName))
+				{
+					return BuildInvalidTupleHandle();
+				}
+			}
+
+			// Create a new class in this graph's dependencies list:
+			uint32 NewUniqueIDForGraph = FindNewUniqueDependencyId();
+			if (!ensureAlwaysMsgf(NewUniqueIDForGraph != INDEX_NONE, TEXT("Call to FindNewUniqueNodeId failed!")))
+			{
+				return BuildInvalidTupleHandle();
+			}
+
+			FMetasoundClassDescription& NewGraphClass = Dependencies.Add_GetRef(FMetasoundClassDescription());
+			NewGraphClass.Metadata = InInfo;
+			NewGraphClass.Metadata.NodeType = EMetasoundClassType::MetasoundGraph;
+			NewGraphClass.UniqueID = NewUniqueIDForGraph;
+
+			// Add the new subgraph's ID as a dependency for the current graph:
+			GraphsClassDeclaration->DependencyIDs.Add(NewUniqueIDForGraph);
+
+			// Generate a new FGraphHandle for this subgraph:
+			FDescPath PathForNewGraph = FDescPath()[Path::EFromDocument::ToDependencies][*InInfo.NodeName][Path::EFromClass::ToGraph];
+			FHandleInitParams InitParams = { GraphsClassDeclaration.GetAccessPoint(), PathForNewGraph, InInfo.NodeName, OwningAsset };
+			FGraphHandle SubgraphHandle = FGraphHandle(FHandleInitParams::PrivateToken, InitParams);
+
+			// Create the node for this subgraph in the current graph:
+			const uint32 NewUniqueId = FindNewUniqueNodeId();
+			if (!ensureAlwaysMsgf(NewUniqueId != INDEX_NONE, TEXT("Call to FindNewUniqueNodeId failed!")))
+			{
+				return BuildInvalidTupleHandle();
+			}
+
+			// Add a node for this output to the graph description.
+			FMetasoundNodeDescription NewNodeDescription;
+			NewNodeDescription.Name = InInfo.NodeName;
+			NewNodeDescription.UniqueID = NewUniqueId;
+			NewNodeDescription.ObjectTypeOfNode = InInfo.NodeType;
+
+			GraphPtr->Nodes.Add(NewNodeDescription);
+
+			FDescPath NodePath = GraphPtr.GetPath()[Path::EFromGraph::ToNodes][NewNodeDescription.UniqueID];
+			FHandleInitParams NodeInitParams = { GraphPtr.GetAccessPoint(), NodePath, InInfo.NodeName, OwningAsset };
+			FNodeHandle SubgraphNode = FNodeHandle(FHandleInitParams::PrivateToken, NodeInitParams);
+
+			return TTuple<FGraphHandle, FNodeHandle>(SubgraphHandle, SubgraphNode);
+		}
+
+		TUniquePtr<Metasound::IOperator> FGraphHandle::BuildOperator()
+		{
+			// @todo: bring over operator builder from previous frontend draft with this new implementation.
+			ensureAlwaysMsgf(false, TEXT("Implement Me!"));
+			return nullptr;
+		}
+
+		FGraphHandle GetGraphHandleForClass(const FMetasoundClassDescription& InClass)
+		{
+			ensureAlwaysMsgf(false, TEXT("Implement Me!"));
+
+			// to implement this, we'll need to
+			// step 1. add tags for metasound asset UObject types that make their node name/inputs/outputs asset registry searchable 
+			// step 2. search the asset registry for the assets.
+			// step 3. Consider a runtime implementation for this using soft object paths.
+			return FGraphHandle::InvalidHandle();
+		}
+	}
+}
+
+IMPLEMENT_MODULE(FDefaultModuleImpl, MetasoundFrontend);
diff --git a/Engine/Plugins/Runtime/Metasound/Source/MetasoundFrontend/Private/MetasoundFrontendBaseClasses.cpp b/Engine/Plugins/Runtime/Metasound/Source/MetasoundFrontend/Private/MetasoundFrontendBaseClasses.cpp
new file mode 100644
--- /dev/null
+++ b/Engine/Plugins/Runtime/Metasound/Source/MetasoundFrontend/Private/MetasoundFrontendBaseClasses.cpp
@@ -0,0 +1,791 @@
+// Copyright Epic Games, Inc. All Rights Reserved.
+
+#include "MetasoundFrontendBaseClasses.h"
+
+namespace Metasound
+{
+	namespace Frontend
+	{
+
+		ITransactable::ITransactable(uint32 InUndoLimit, TWeakObjectPtr<UObject> InOwningAsset /* = nullptr */)
+			: UndoLimit(InUndoLimit)
+			, OwningAsset(InOwningAsset)
+		{
+		}
+
+		bool ITransactable::Undo()
+		{
+			if (UndoTransactableStack.Num() == 0)
+			{
+				return false;
+			}
+
+			if (TSharedPtr<ITransactable> Owner = OwningTransactable.Pin())
+			{
+				Owner->DiscardUndoFromOwnedTransactable(AsShared());
+			}
+
+			TSharedPtr<ITransactable> Transactable = UndoTransactableStack.Pop().Pin();
+			TSharedPtr<ITransactable> PreviousTransactable = AsShared();
+
+			RedoTransactableStack.Push(Transactable);
+
+			while(Transactable)
+			{
+				if (Transactable == PreviousTransactable)
+				{
+					bool bUndoSucceeded = Transactable->PerformLocalUndo();
+					if (bUndoSucceeded && OwningAsset.IsValid())
+					{
+						OwningAsset->MarkPackageDirty();
+					}
+
+					return bUndoSucceeded;
+				}
+
+				PreviousTransactable = Transactable;
+				Transactable = Transactable->UndoTransactableStack.Pop().Pin();
+				PreviousTransactable->RedoTransactableStack.Push(Transactable);
+			}
+
+			return false;
+		}
+
+		bool ITransactable::Redo()
+		{
+			if (RedoTransactableStack.Num() == 0)
+			{
+				return false;
+			}
+
+			if (TSharedPtr<ITransactable> Owner = OwningTransactable.Pin())
+			{
+				Owner->DiscardRedoFromOwnedTransactable(AsShared());
+			}
+
+			TSharedPtr<ITransactable> Transactable = RedoTransactableStack.Pop().Pin();
+			TSharedPtr<ITransactable> PreviousTransactable = AsShared();
+
+			UndoTransactableStack.Push(Transactable);
+
+			while (Transactable)
+			{
+				if (Transactable == PreviousTransactable)
+				{
+					const bool bRedoSucceeded = Transactable->PerformLocalRedo();
+					if (bRedoSucceeded && OwningAsset.IsValid())
+					{
+						OwningAsset->MarkPackageDirty();
+					}
+
+					return bRedoSucceeded;
+				}
+
+				PreviousTransactable = Transactable;
+				Transactable = Transactable->RedoTransactableStack.Pop().Pin();
+				PreviousTransactable->UndoTransactableStack.Push(Transactable);
+			}
+
+			return false;
+		}
+
+		void ITransactable::CommitTransaction(FReversibleTransaction&& InTransactionDescription)
+		{
+			if (((uint32)LocalUndoTransactionStack.Num()) >= UndoLimit)
+			{
+				// Discard the oldest undo action.
+				LocalUndoTransactionStack.RemoveAt(0);
+			}
+
+			LocalUndoTransactionStack.Push(MoveTemp(InTransactionDescription));
+			RedoTransactableStack.Reset();
+
+			TWeakPtr<ITransactable> WeakThisPtr = AsShared();
+			UndoTransactableStack.Push(WeakThisPtr);
+
+			if (TSharedPtr<ITransactable> Owner = OwningTransactable.Pin())
+			{
+				Owner->PushUndoFromOwnedTransactable(WeakThisPtr);
+			}
+
+			if (OwningAsset.IsValid())
+			{
+				OwningAsset->MarkPackageDirty();
+			}
+		}
+
+		bool ITransactable::RegisterOwningTransactable(ITransactable& InOwningTransactable)
+		{
+			// Sanity check that this will not create a cycle.
+			while (TSharedPtr<ITransactable> Owner = InOwningTransactable.OwningTransactable.Pin())
+			{
+				if (Owner.Get() == this)
+				{
+					return false;
+				}
+			}
+
+			OwningTransactable = InOwningTransactable.AsShared();
+			return true;
+		}
+
+
+		bool ITransactable::PerformLocalUndo()
+		{
+			if (LocalUndoTransactionStack.Num() == 0)
+			{
+				return false;
+			}
+
+			FReversibleTransaction Transaction = LocalUndoTransactionStack.Pop();
+			bool bResult = Transaction.UndoTransaction();
+			if (bResult)
+			{
+				LocalRedoTransactionStack.Push(MoveTemp(Transaction));
+			}
+
+			return bResult;
+		}
+
+		bool ITransactable::PerformLocalRedo()
+		{
+			if (LocalRedoTransactionStack.Num() == 0)
+			{
+				return false;
+			}
+
+			FReversibleTransaction Transaction = LocalRedoTransactionStack.Pop();
+			bool bResult = Transaction.RedoTransaction();
+			if (bResult)
+			{
+				LocalUndoTransactionStack.Push(MoveTemp(Transaction));
+			}
+
+			return bResult;
+		}
+
+		bool ITransactable::DiscardUndoFromOwnedTransactable(TWeakPtr<ITransactable> InOwnedTransactable)
+		{
+			// NOTE: This relies on TArray.Push adding an element to the end of the array.
+			// Start at the end of the array and work backwards until we find the most recent undo operation for this transactable.
+			for (int32 Index = UndoTransactableStack.Num() - 1; Index >= 0; Index--)
+			{
+				if (UndoTransactableStack[Index] == InOwnedTransactable)
+				{
+					// Discard, but preserve the stack order.
+					UndoTransactableStack.RemoveAt(Index);
+					return true;
+				}
+			}
+
+			return false;
+		}
+
+		bool ITransactable::DiscardRedoFromOwnedTransactable(TWeakPtr<ITransactable> InOwnedTransactable)
+		{
+			// NOTE: This relies on TArray.Push adding an element to the end of the array.
+			// Start at the end of the array and work backwards until we find the most recent undo operation for this transactable.
+			for (int32 Index = RedoTransactableStack.Num() - 1; Index >= 0; Index--)
+			{
+				if (RedoTransactableStack[Index] == InOwnedTransactable)
+				{
+					// Discard, but preserve the stack order.
+					RedoTransactableStack.RemoveAt(Index);
+					return true;
+				}
+			}
+
+			return false;
+		}
+
+		bool ITransactable::PushUndoFromOwnedTransactable(TWeakPtr<ITransactable> InOwnedTransactable)
+		{
+			UndoTransactableStack.Push(MoveTemp(InOwnedTransactable));
+			RedoTransactableStack.Reset();
+
+			return true;
+		}
+
+
+
+		bool ITransactable::PushRedoFromOwnedTransactable(TWeakPtr<ITransactable> InOwnedTransactable)
+		{
+			RedoTransactableStack.Push(MoveTemp(InOwnedTransactable));
+			DiscardUndoFromOwnedTransactable(InOwnedTransactable);
+
+			return true;
+		}
+
+		FDescriptionAccessPoint::FDescriptionAccessPoint(FMetasoundDocument& InRootDocument)
+			: RootDocument(InRootDocument)
+		{
+		}
+
+		FMetasoundDocument& FDescriptionAccessPoint::GetRoot()
+		{
+			return RootDocument;
+		}
+
+		FMetasoundClassDescription* FDescriptionAccessPoint::GetClassFromPath(const FDescPath& InPathFromRoot)
+		{
+			FDescPath CurrentPath = InPathFromRoot;
+
+			FMetasoundDescriptionPtr Ptr;
+			Ptr.Set<FMetasoundDocument*>(&RootDocument);
+			FDescriptionUnwindStep CurrentStep = { Ptr, Path::EDescType::Document };
+
+			while (CurrentPath.Num() != 0 && CurrentStep.Type != Path::EDescType::Invalid)
+			{
+				CurrentStep = GoToNext(CurrentPath, CurrentStep);
+			}
+
+			if (ensureAlwaysMsgf(CurrentPath.Num() == 0, TEXT("Couldn't resolve part of the path.")))
+			{
+				return CurrentStep.DescriptionStructPtr.Get<FMetasoundClassDescription*>();
+			}
+			else
+			{
+				return nullptr;
+			}
+		}
+
+		FMetasoundNodeDescription* FDescriptionAccessPoint::GetNodeFromPath(const FDescPath& InPathFromRoot)
+		{
+			FDescPath CurrentPath = InPathFromRoot;
+			
+			FMetasoundDescriptionPtr Ptr;
+			Ptr.Set<FMetasoundDocument*>(&RootDocument);
+			FDescriptionUnwindStep CurrentStep = { Ptr, Path::EDescType::Document };
+
+			while (CurrentPath.Num() != 0 && CurrentStep.Type != Path::EDescType::Invalid)
+			{
+				CurrentStep = GoToNext(CurrentPath, CurrentStep);
+			}
+
+			if (ensureAlwaysMsgf(CurrentPath.Num() == 0, TEXT("Couldn't resolve part of the path.")))
+			{
+				return CurrentStep.DescriptionStructPtr.Get<FMetasoundNodeDescription*>();
+			}
+			else
+			{
+				return nullptr;
+			}
+		}
+
+		FMetasoundGraphDescription* FDescriptionAccessPoint::GetGraphFromPath(const FDescPath& InPathFromRoot)
+		{
+			FDescPath CurrentPath = InPathFromRoot;
+			
+			FMetasoundDescriptionPtr Ptr;
+			Ptr.Set<FMetasoundDocument*>(&RootDocument);
+			FDescriptionUnwindStep CurrentStep = { Ptr, Path::EDescType::Document };
+
+			while (CurrentPath.Num() != 0 && CurrentStep.Type != Path::EDescType::Invalid)
+			{
+				CurrentStep = GoToNext(CurrentPath, CurrentStep);
+			}
+
+			if (ensureAlwaysMsgf(CurrentPath.Num() == 0, TEXT("Couldn't resolve part of the path.")))
+			{
+				return CurrentStep.DescriptionStructPtr.Get<FMetasoundGraphDescription*>();
+			}
+			else
+			{
+				return nullptr;
+			}
+		}
+
+		FMetasoundInputDescription* FDescriptionAccessPoint::GetInputFromPath(const FDescPath& InPathFromRoot)
+		{
+			FDescPath CurrentPath = InPathFromRoot;
+			
+			FMetasoundDescriptionPtr Ptr;
+			Ptr.Set<FMetasoundDocument*>(&RootDocument);
+			FDescriptionUnwindStep CurrentStep = { Ptr, Path::EDescType::Document };
+
+			while (CurrentPath.Num() != 0 && CurrentStep.Type != Path::EDescType::Invalid)
+			{
+				CurrentStep = GoToNext(CurrentPath, CurrentStep);
+			}
+
+			if (ensureAlwaysMsgf(CurrentPath.Num() == 0, TEXT("Couldn't resolve part of the path.")))
+			{
+				return CurrentStep.DescriptionStructPtr.Get<FMetasoundInputDescription*>();
+			}
+			else
+			{
+				return nullptr;
+			}
+		}
+
+		FMetasoundOutputDescription* FDescriptionAccessPoint::GetOutputFromPath(const FDescPath& InPathFromRoot)
+		{
+			FDescPath CurrentPath = InPathFromRoot;
+			
+			FMetasoundDescriptionPtr Ptr;
+			Ptr.Set<FMetasoundDocument*>(&RootDocument);
+			FDescriptionUnwindStep CurrentStep = { Ptr, Path::EDescType::Document };
+
+			while (CurrentPath.Num() != 0 && CurrentStep.Type != Path::EDescType::Invalid)
+			{
+				CurrentStep = GoToNext(CurrentPath, CurrentStep);
+			}
+
+			if (ensureAlwaysMsgf(CurrentPath.Num() == 0, TEXT("Couldn't resolve part of the path.")))
+			{
+				return CurrentStep.DescriptionStructPtr.Get<FMetasoundOutputDescription*>();
+			}
+			else
+			{
+				return nullptr;
+			}
+		}
+
+		FMetasoundClassMetadata* FDescriptionAccessPoint::GetMetadataFromPath(const FDescPath& InPathFromRoot)
+		{
+			FDescPath CurrentPath = InPathFromRoot;
+			
+			FMetasoundDescriptionPtr Ptr;
+			Ptr.Set<FMetasoundDocument*>(&RootDocument);
+			FDescriptionUnwindStep CurrentStep = { Ptr, Path::EDescType::Document };
+
+			while (CurrentPath.Num() != 0 && CurrentStep.Type != Path::EDescType::Invalid)
+			{
+				CurrentStep = GoToNext(CurrentPath, CurrentStep);
+			}
+
+			if (ensureAlwaysMsgf(CurrentPath.Num() == 0, TEXT("Couldn't resolve part of the path.")))
+			{
+				return CurrentStep.DescriptionStructPtr.Get<FMetasoundClassMetadata*>();
+			}
+			else
+			{
+				return nullptr;
+			}
+		}
+
+		Metasound::Frontend::FClassDependencyIDs* FDescriptionAccessPoint::GetClassDependencyIDsFromPath(const FDescPath& InPathFromRoot)
+		{
+			FDescPath CurrentPath = InPathFromRoot;
+			
+			FMetasoundDescriptionPtr Ptr;
+			Ptr.Set<FMetasoundDocument*>(&RootDocument);
+			FDescriptionUnwindStep CurrentStep = { Ptr, Path::EDescType::Document };
+
+			while (CurrentPath.Num() != 0 && CurrentStep.Type != Path::EDescType::Invalid)
+			{
+				CurrentStep = GoToNext(CurrentPath, CurrentStep);
+			}
+
+			if (ensureAlwaysMsgf(CurrentPath.Num() == 0, TEXT("Couldn't resolve part of the path.")))
+			{
+				return CurrentStep.DescriptionStructPtr.Get<FClassDependencyIDs*>();
+			}
+			else
+			{
+				return nullptr;
+			}
+		}
+
+		FDescriptionAccessPoint::FDescriptionUnwindStep FDescriptionAccessPoint::GoToNext(FDescPath& InPath, FDescriptionUnwindStep InElement)
+		{
+			auto InvalidUnwindStep = []() -> FDescriptionUnwindStep
+			{
+				FDescriptionUnwindStep InvalidStep;
+				InvalidStep.DescriptionStructPtr.Set<void*>(nullptr);
+				InvalidStep.Type = Path::EDescType::Invalid;
+				return InvalidStep;
+			};
+
+			if (!ensureMsgf(InPath.Path.Num() != 0, TEXT("Attempted to unwind an empty path.")))
+			{
+				return InvalidUnwindStep();
+			}
+
+			Path::FElement NextStep = InPath.Path[0];
+
+			InPath.Path.RemoveAt(0);
+
+			switch (InElement.Type)
+			{
+				case Path::EDescType::Document:
+				{
+					FMetasoundDocument* Document = InElement.DescriptionStructPtr.Get<FMetasoundDocument*>();
+					switch (NextStep.CurrentDescType)
+					{
+						case Path::EDescType::Document:
+						{
+							FDescriptionUnwindStep UnwindStep;
+							UnwindStep.DescriptionStructPtr.Set<FMetasoundDocument*>(Document);
+							UnwindStep.Type = Path::EDescType::Document;
+							return UnwindStep;
+						}
+						case Path::EDescType::Class:
+						{
+							FDescriptionUnwindStep UnwindStep;
+							UnwindStep.DescriptionStructPtr.Set<FMetasoundClassDescription*>(&(Document->RootClass));
+							UnwindStep.Type = Path::EDescType::Class;
+							return UnwindStep;
+						}
+						case Path::EDescType::DocDependencies:
+						{
+							// The next element in a path after Dependencies will always be the name of the dependency.
+							Path::FElement DependencyElement = InPath.Path[0];
+							InPath.Path.RemoveAt(0);
+							if (!ensureAlwaysMsgf(DependencyElement.CurrentDescType == Path::EDescType::Class, TEXT("Invalid path set up.")))
+							{
+								return InvalidUnwindStep();
+							}
+
+							FString& DependencyName = DependencyElement.LookupName;
+
+							uint32 DependencyID = DependencyElement.LookupID;
+							
+							if (!ensureAlwaysMsgf(DependencyID != INDEX_NONE || DependencyName.Len() > 0, TEXT("Path to a dependency did not include a valid ID or dependency name.")))
+							{
+								return InvalidUnwindStep();
+							}
+
+							if (DependencyID == FMetasoundClassDescription::RootClassID)
+							{
+								FDescriptionUnwindStep UnwindStep;
+								UnwindStep.DescriptionStructPtr.Set<FMetasoundClassDescription*>(&(Document->RootClass));
+								UnwindStep.Type = Path::EDescType::Class;
+								return UnwindStep;
+							}
+
+							TArray<FMetasoundClassDescription>& DependenciesList = Document->Dependencies;
+
+							// Dependencies can be looked up by ID or by name.
+							if (DependencyID != INDEX_NONE)
+							{
+								// Scan the dependencies list for the matching lookup ID.
+								for (FMetasoundClassDescription& Dependency : DependenciesList)
+								{
+									if (Dependency.UniqueID == DependencyID)
+									{
+										FDescriptionUnwindStep UnwindStep;
+										UnwindStep.DescriptionStructPtr.Set<FMetasoundClassDescription*>(&Dependency);
+										UnwindStep.Type = Path::EDescType::Class;
+										return UnwindStep;
+									}
+								}
+							}
+							else
+							{
+								// fall back to scanning the dependencies list for the matching lookup name.
+								for (FMetasoundClassDescription& Dependency : DependenciesList)
+								{
+									if (Dependency.Metadata.NodeName == DependencyName)
+									{
+										FDescriptionUnwindStep UnwindStep;
+										UnwindStep.DescriptionStructPtr.Set<FMetasoundClassDescription*>(&Dependency);
+										UnwindStep.Type = Path::EDescType::Class;
+										return UnwindStep;
+									}
+								}
+							}
+							
+
+							// If we reached the end of the Dependencies list and didn't find a match, ensure.
+							ensureAlwaysMsgf(false, TEXT("Couldn't find dependency %s in path."), *DependencyName);
+							return InvalidUnwindStep();
+						}
+					}
+					break;
+				}
+				case Path::EDescType::Class:
+				{
+					FMetasoundClassDescription* ClassDescription = InElement.DescriptionStructPtr.Get<FMetasoundClassDescription*>();
+					switch (NextStep.CurrentDescType)
+					{
+						case Path::EDescType::Graph:
+						{
+							FDescriptionUnwindStep UnwindStep;
+							UnwindStep.DescriptionStructPtr.Set<FMetasoundGraphDescription*>(&(ClassDescription->Graph));
+							UnwindStep.Type = Path::EDescType::Graph;
+							return UnwindStep;
+						}
+						case Path::EDescType::ClassDependencies:
+						{
+							FDescriptionUnwindStep UnwindStep;
+							UnwindStep.DescriptionStructPtr.Set<FClassDependencyIDs*>(&(ClassDescription->DependencyIDs));
+							UnwindStep.Type = Path::EDescType::ClassDependencies;
+							return UnwindStep;
+						}
+						case Path::EDescType::Inputs:
+						{
+							// The next element after an Inputs element should always be the name of an input.
+							Path::FElement InputElement = InPath.Path[0];
+							InPath.Path.RemoveAt(0);
+							if (!ensureAlwaysMsgf(InputElement.CurrentDescType == Path::EDescType::Input, TEXT("Invalid path set up.")))
+							{
+								return InvalidUnwindStep();
+							}
+
+							FString& InputName = InputElement.LookupName;
+
+							// Scan the inputs list for the lookup name.
+							TArray<FMetasoundInputDescription>& InputsList = ClassDescription->Inputs;
+							for (FMetasoundInputDescription& Input : InputsList)
+							{
+								if (Input.Name == InputName)
+								{
+									FDescriptionUnwindStep UnwindStep;
+									UnwindStep.DescriptionStructPtr.Set<FMetasoundInputDescription*>(&Input);
+									UnwindStep.Type = Path::EDescType::Input;
+									return UnwindStep;
+								}
+							}
+
+							// If we reached the end of the Inputs list and didn't find a match, ensure.
+							ensureAlwaysMsgf(false, TEXT("Couldn't find input %s in path."), *InputName);
+							return InvalidUnwindStep();
+						}
+						case Path::EDescType::Outputs:
+						{
+							// The next element after an Outputs element should always be the name of an output.
+							Path::FElement OutputElement = InPath.Path[0];
+							InPath.Path.RemoveAt(0);
+							if (!ensureAlwaysMsgf(OutputElement.CurrentDescType == Path::EDescType::Output, TEXT("Invalid path set up.")))
+							{
+								return InvalidUnwindStep();
+							}
+						
+							FString& OutputName = OutputElement.LookupName;
+
+							// Scan the outputs list for the lookup name.
+							TArray<FMetasoundOutputDescription>& OutputsList = ClassDescription->Outputs;
+							for (FMetasoundOutputDescription& Output : OutputsList)
+							{
+								if (Output.Name == NextStep.LookupName)
+								{
+									FDescriptionUnwindStep UnwindStep;
+									UnwindStep.DescriptionStructPtr.Set<FMetasoundOutputDescription*>(&Output);
+									UnwindStep.Type = Path::EDescType::Output;
+									return UnwindStep;
+								}
+							}
+
+							// If we reached the end of the Inputs list and didn't find a match, ensure.
+							ensureAlwaysMsgf(false, TEXT("Couldn't find output %s in path."), *NextStep.LookupName);
+							return InvalidUnwindStep();
+						}
+						case Path::EDescType::Metadata:
+						{
+							FDescriptionUnwindStep UnwindStep;
+							UnwindStep.DescriptionStructPtr.Set<FMetasoundClassMetadata*>(&(ClassDescription->Metadata));
+							UnwindStep.Type = Path::EDescType::Metadata;
+
+							return UnwindStep;
+						}
+						default:
+						{
+							ensureAlwaysMsgf(false, TEXT("Invalid path- Tried to path directly from a Class Description to a type that wasn't a direct memember of the Class"));
+							return InvalidUnwindStep();
+						}
+					}
+
+					break;
+				}
+				case Path::EDescType::Graph:
+				{
+					if (!ensureAlwaysMsgf(NextStep.CurrentDescType == Path::EDescType::Nodes, TEXT("Invalid path. the Graph description only contains the Nodes list.")))
+					{
+						return InvalidUnwindStep();
+					}
+
+					FMetasoundGraphDescription* GraphDescription = InElement.DescriptionStructPtr.Get<FMetasoundGraphDescription*>();
+				
+					if (!ensureAlwaysMsgf(InPath.Path.Num() != 0, TEXT("Incomplete path! path stopped at Nodes list without specifying a node ID.")))
+					{
+						return InvalidUnwindStep();
+					}
+
+					Path::FElement NodeElement = InPath.Path[0];
+					InPath.Path.RemoveAt(0);
+
+					if (!ensureAlwaysMsgf(NodeElement.CurrentDescType == Path::EDescType::Node, TEXT("Invalid path! a Nodes element must always be followed by a Node ID.")))
+					{
+						return InvalidUnwindStep();
+					}
+				
+					uint32 NodeID = NodeElement.LookupID;
+
+					TArray<FMetasoundNodeDescription>& NodeList = GraphDescription->Nodes;
+					for (FMetasoundNodeDescription& Node : NodeList)
+					{
+						if (Node.UniqueID == NodeID)
+						{
+							FDescriptionUnwindStep UnwindStep;
+							UnwindStep.DescriptionStructPtr.Set<FMetasoundNodeDescription*>(&Node);
+							UnwindStep.Type = Path::EDescType::Node;
+							return UnwindStep;
+						}
+					}
+					break;
+				}
+				case Path::EDescType::Inputs:
+				{
+					ensureAlwaysMsgf(false, TEXT("Invalid path. Inputs will always be pathed directly after a Class."));
+					return InvalidUnwindStep();
+				}
+				case Path::EDescType::DocDependencies:
+				{
+					ensureAlwaysMsgf(false, TEXT("Invalid path. Document dependencies should always follow after a Document and should always list a specific dependency by ID or name."));
+					return InvalidUnwindStep();
+				}
+				case Path::EDescType::Input:
+				case Path::EDescType::Output:
+				case Path::EDescType::Node:
+				case Path::EDescType::Metadata:
+				default:
+				{
+					ensureAlwaysMsgf(false, TEXT("Invalid path. Input, Output, Node, and Metadata don't have any child elements."));
+					return InvalidUnwindStep();
+				}
+			}
+
+			// if we ever hit this, we likely missed a return on one of these branches.
+			checkNoEntry();
+			return InvalidUnwindStep();
+		}
+
+		FDescPath Path::GetPathToClassForNode(FDescPath InPathForNode, FString& InNodeName)
+		{
+			return FDescPath()[EFromDocument::ToDependencies][*InNodeName];
+		}
+
+
+		FDescPath Path::GetOwningClassDescription(FDescPath InPathForGraph)
+		{
+			// Backtrack from the end of the path until we find a Class element.
+			for (int32 Level = InPathForGraph.Path.Num() - 1; Level >= 0; Level--)
+			{
+				if (InPathForGraph.Path[Level].CurrentDescType == Path::EDescType::Class)
+				{
+					return InPathForGraph;
+				}
+				else
+				{
+					InPathForGraph.Path.RemoveAt(InPathForGraph.Path.Num() - 1);
+				}
+			}
+
+			return InPathForGraph;
+		}
+
+		Metasound::Frontend::FDescPath Path::GetClassDescription(const FString& InClassName)
+		{
+			return FDescPath()[Path::EFromDocument::ToDependencies][*InClassName];
+		}
+
+		FDescPath Path::GetDependencyPath(const FString& InDependencyName)
+		{
+			return FDescPath()[Path::EFromDocument::ToDependencies][*InDependencyName];
+		}
+
+		FDescPath Path::GetDependencyPath(uint32 InDependencyID)
+		{
+			return FDescPath()[Path::EFromDocument::ToDependencies][InDependencyID];
+		}
+
+		FString Path::GetPrintableString(FDescPath InPath)
+		{
+			FString OutString = FString(TEXT("//"));
+			for (FElement& PathElement : InPath.Path)
+			{
+				switch (PathElement.CurrentDescType)
+				{
+					case EDescType::Document:
+					{
+						OutString += TEXT("Document/");
+						break;
+					}
+					case EDescType::Class:
+					{
+						OutString += TEXT("Class/");
+						break;
+					}
+					case EDescType::DocDependencies:
+					{
+						OutString += TEXT("Dependencies(");
+
+						if (PathElement.LookupID != INDEX_NONE)
+						{
+							OutString.AppendInt(PathElement.LookupID);
+						}
+						else
+						{
+							OutString += PathElement.LookupName;
+						}
+						
+						OutString += TEXT(")/");
+						break;
+					}
+					case EDescType::ClassDependencies:
+					{
+						OutString += TEXT("Dependencies(");
+						OutString.AppendInt(PathElement.LookupID);
+						OutString += TEXT(")/");
+						break;
+					}
+					case EDescType::Graph:
+					{
+						OutString += TEXT("Graph/");
+						break;
+					}
+					case EDescType::Inputs:
+					{
+						OutString += TEXT("Inputs/");
+						break;
+					}
+					case EDescType::Input:
+					{
+						OutString += TEXT("Input(");
+						OutString += PathElement.LookupName;
+						OutString += TEXT(")/");
+						break;
+					}
+					case EDescType::Metadata:
+					{
+						OutString += TEXT("Metadata/");
+						break;
+					}
+					case EDescType::Nodes:
+					{
+						OutString += TEXT("Nodes/");
+						break;
+					}
+					case EDescType::Node:
+					{
+						OutString += TEXT("Node(");
+						OutString.AppendInt(PathElement.LookupID);
+						OutString += TEXT(")/");
+						break;
+					}
+					case EDescType::Outputs:
+					{
+						OutString += TEXT("Outputs/");
+						break;
+					}
+					case EDescType::Output:
+					{
+						OutString += TEXT("Output(");
+						OutString += PathElement.LookupName;
+						OutString += TEXT(")/");
+						break;
+					}
+					default:
+					{
+						OutString += TEXT("Unknown/");
+						break;
+					}
+				}
+			}
+
+			return OutString;
+		}
+
+	}
+}
diff --git a/Engine/Plugins/Runtime/Metasound/Source/MetasoundFrontend/Private/MetasoundNodeAsset.cpp b/Engine/Plugins/Runtime/Metasound/Source/MetasoundFrontend/Private/MetasoundNodeAsset.cpp
new file mode 100644
--- /dev/null
+++ b/Engine/Plugins/Runtime/Metasound/Source/MetasoundFrontend/Private/MetasoundNodeAsset.cpp
@@ -0,0 +1,30 @@
+// Copyright Epic Games, Inc. All Rights Reserved.
+
+#include "MetasoundNodeAsset.h"
+#include "CoreMinimal.h"
+
+#include "Backends/JsonStructSerializerBackend.h"
+#include "StructSerializer.h"
+
+UMetasound::UMetasound(const FObjectInitializer& ObjectInitializer)
+	: Super(ObjectInitializer)
+	, AccessPoint(new Metasound::Frontend::FDescriptionAccessPoint(RootMetasoundDocument))
+{
+}
+
+bool UMetasound::ExportToJSON(const FString& InAbsolutePath)
+{
+	if (TUniquePtr<FArchive> FileWriter = TUniquePtr<FArchive>(IFileManager::Get().CreateFileWriter(*InAbsolutePath)))
+	{
+		FJsonStructSerializerBackend Backend(*FileWriter, EStructSerializerBackendFlags::Default);
+		FStructSerializer::Serialize<FMetasoundDocument>(RootMetasoundDocument, Backend);
+		FileWriter->Close();
+
+		return true;
+	}
+	else
+	{
+		ensureAlwaysMsgf(false, TEXT("Failed to create a filewriter with the given path."));
+		return false;
+	}
+}
diff --git a/Engine/Plugins/Runtime/Metasound/Source/MetasoundFrontend/Public/MetasoundFrontend.h b/Engine/Plugins/Runtime/Metasound/Source/MetasoundFrontend/Public/MetasoundFrontend.h
new file mode 100644
--- /dev/null
+++ b/Engine/Plugins/Runtime/Metasound/Source/MetasoundFrontend/Public/MetasoundFrontend.h
@@ -0,0 +1,322 @@
+// Copyright Epic Games, Inc. All Rights Reserved.
+
+#pragma once
+
+#include "CoreMinimal.h"
+#include "MetasoundGraph.h"
+#include "MetasoundFrontendDataLayout.h"
+#include "MetasoundFrontendBaseClasses.h"
+
+
+class UMetasound;
+
+namespace Metasound
+{	
+	/**
+	 * These classes are used for editing Metasounds.
+	 * In general, the workflow for editing a Metasound graph will be:
+	 * 1) Load or create a metasound asset.
+	 * 2) Call UMetasound::GetGraphHandle() to get a handle to the graph for that Metasound asset.
+	 * 
+	 * All Metasound documents are saved as a FMetasoundClassDescription, which itself can own FMetasoundClassDescriptions in a tree-like hierarchy.
+	 * Typically the workflow for creating a Metasound subgraph will be
+	 * 1) Get a Metasound::Frontend::FGraphHandle (typically through the two steps described in the first paragraph)
+	 * 2) Build a FMetasoundClassMetadata struct with whatever name/author/description you want to give the subgraph,
+	 * 3) call Metasound::Frontend::FGraphHandle::CreateEmptySubgraphNode, which will return that subgraph as it's own FGraphHandle,
+	 *    as well as the FNodeHandle for that subgraph in the current graph.
+	 *
+	 * If you want to bounce that subgraph to it's own saved Metasound Asset, call FGraphHandle::SaveToNewMetasoundAsset on that subgraph's FGraphHandle.
+	 *
+	 * General note- these apis are NOT thread safe. 
+	 * Make sure that any FGraphHandle, FNodeHandle, FInputHandle and FOutputHandle that access similar data are called on the same thread.
+	 */
+	namespace Frontend
+	{
+
+		// Struct with the basics of a node class' information,
+		// used to look up that node from our node browser functions,
+		// and also used in FGraphHandle::AddNewNode.
+		struct FNodeClassInfo
+		{
+			FString NodeName;
+			EMetasoundClassType NodeType;
+
+		};
+
+		// Get all available nodes of any type.
+		TArray<FNodeClassInfo> GetAllAvailableNodeClasses();
+
+		// Get all possible nodes whose name begins with a specific namespace.
+		TArray<FNodeClassInfo> GetAllNodeClassesInNamespace(const FString& InNamespace);
+
+		// Similar to GetAllNodeClassesInNamespace, except searches all nodes for a match with a substring.
+		TArray<FNodeClassInfo> GetAllNodesWhoseNameContains(const FString& InSubstring);
+
+		// Searches for any node types that can output a specific type.
+		TArray<FNodeClassInfo> GetAllNodesWithAnOutputOfType(const FName& InType);
+
+		TArray<FNodeClassInfo> GetAllNodesWithAnInputOfType(const FName& InType);
+
+		// gets all metadata (name, description, author, what to say if it's missing) for a given node.
+		FMetasoundClassMetadata GetMetadataForNode(const FNodeClassInfo& InInfo);
+
+
+		// Generates a new FMetasoundClassDescription for a given node class. Only used by classes that manipulate Metasound Description data directly.
+		FMetasoundClassDescription GetClassDescriptionForNode(const FNodeClassInfo& InInfo);
+
+		// Struct that indicates whether an input and an output can be connected,
+		// and whether an intermediate node is necessary to connect the two.
+		struct FConnectability
+		{
+			enum class EConnectable
+			{
+				Yes,
+				No,
+				YesWithConverterNode
+			};
+
+			EConnectable Connectable;
+
+			// If Connectable is EConnectable::YesWithConverterNode,
+			// this will be a populated list of nodes we can use 
+			// to convert between the input and output.
+			TArray<FNodeClassInfo> PossibleConverterNodeClasses;
+		};
+
+		class FInputHandle;
+
+		struct FHandleInitParams
+		{
+			TWeakPtr<FDescriptionAccessPoint> InAccessPoint;
+
+			// Path in the document to the element we're getting a handle to.
+			// for the FGraphHandle, this is a path to the graph within the graph's Class description.
+			// for the FNodeHandle, FInputHandle and FOutputHandle, this is a path to the individual node within a Graph description.
+			const FDescPath& InPath;
+
+			// Class name for the graph or node to get a handle for.
+			const FString& InClassName;
+
+			// The asset that owns the FMetasoundDocument this handle is associated with.
+			TWeakObjectPtr<UObject> InOwningAsset;
+
+		private:
+			// This is used to grant specific classes the ability to create handles.
+			enum EPrivateToken { Token };
+			static const EPrivateToken PrivateToken = EPrivateToken::Token;
+
+			friend class FOutputHandle;
+			friend class FInputHandle;
+			friend class FNodeHandle;
+			friend class FGraphHandle;
+			friend class UMetasound;
+		};
+
+		class FOutputHandle : protected ITransactable
+		{
+		public:
+			FOutputHandle() = delete;
+
+			explicit FOutputHandle(FHandleInitParams::EPrivateToken InToken, const FHandleInitParams& InParams, const FString& InOutputName);
+
+			~FOutputHandle() = default;
+
+			static FOutputHandle InvalidHandle();
+
+			bool IsValid() const;
+
+			// @todo: with the current spec, we'd have to scan all nodes in the graph for this output to solve this.
+			// Is this worth it, or should we just require folks to use FInputHandle::IsConnected()?
+			//bool IsConnected() const;
+
+			FName GetOutputType() const;
+			FString GetOutputName() const;
+			FString GetOutputTooltip() const;
+			uint32 GetOwningNodeID() const;
+
+			// @todo: with the current spec, we'd have to scan all nodes in the graph for this output to solve this.
+			// Is this worth it, or should we just require folks to use FInputHandle::GetCurrentlyConnectedOutput()?
+			//FInputHandle GetCurrentlyConnectedInput();
+
+			FConnectability CanConnectTo(const FInputHandle& InHandle) const;
+			bool Connect(FInputHandle& InHandle);
+			bool ConnectWithConverterNode(FInputHandle& InHandle, FString& InNodeClassName);
+
+		private:
+			TDescriptionPtr<FMetasoundNodeDescription> NodePtr;
+			TDescriptionPtr<FMetasoundClassDescription> NodeClass;
+			TDescriptionPtr<FMetasoundOutputDescription> OutputPtr;
+		};
+
+		class FInputHandle : protected ITransactable
+		{
+		public:
+			FInputHandle() = delete;
+			~FInputHandle() = default;
+
+			// Sole constructor for FInputHandle. Can only be used by friend classes for FHandleInitParams.
+			explicit FInputHandle(FHandleInitParams::EPrivateToken PrivateToken, const FHandleInitParams& InParams, const FString& InputName);
+
+			static FInputHandle InvalidHandle();
+
+			bool IsValid() const;
+			bool IsConnected() const;
+
+			FName GetInputType() const;
+			FString GetInputName() const;
+			FString GetInputTooltip() const;
+
+			FOutputHandle GetCurrentlyConnectedOutput() const;
+
+			FConnectability CanConnectTo(const FOutputHandle& InHandle) const;
+			bool Connect(FOutputHandle& InHandle);
+			bool ConnectWithConverterNode(FOutputHandle& InHandle, FString& InNodeClassName);
+
+		private:
+			TDescriptionPtr<FMetasoundNodeDescription> NodePtr;
+			TDescriptionPtr<FMetasoundClassDescription> NodeClass;
+			TDescriptionPtr<FMetasoundInputDescription> InputPtr;
+
+			FString InputName;
+
+			const FMetasoundNodeConnectionDescription* GetConnectionDescription() const;
+			FMetasoundNodeConnectionDescription* GetConnectionDescription();
+		};
+
+		class FGraphHandle;
+
+		// Opaque handle to a single node on a graph.
+		// Can retrieve metadata about that node,
+		// get FNodeHandles for other nodes it is connected to, and validate and modify connections.
+		class FNodeHandle : protected ITransactable
+		{
+		public:
+			FNodeHandle() = delete;
+			~FNodeHandle() = default;
+
+			// Sole constructor for FNodeHandle. Can only be used by friends of FHandleInitParams.
+			explicit FNodeHandle(FHandleInitParams::EPrivateToken PrivateToken, const FHandleInitParams& InParams);
+
+			static FNodeHandle InvalidHandle();
+
+			bool IsValid() const;
+
+			TArray<FInputHandle> GetAllInputs();
+			TArray<FOutputHandle> GetAllOutputs();
+
+			FInputHandle GetInputWithName(const FString& InName);
+			FOutputHandle GetOutputWithName(const FString& InName);
+
+			EMetasoundClassType GetNodeType();
+			FString GetNodeClassName();
+
+			// If this node is itself a Metasound,
+			// use this to return the contained graph for the metasound.
+			// Otherwise it will return an invalid FGraphHandle.
+			void GetContainedGraph(FGraphHandle& OutGraph);
+
+			uint32 GetNodeID();
+			static uint32 GetNodeID(const FDescPath& InNodePath);
+
+		private:
+
+			TDescriptionPtr<FMetasoundNodeDescription> NodePtr;
+			TDescriptionPtr<FMetasoundClassDescription> NodeClass;
+
+			uint32 NodeID;
+		};
+
+		class FGraphHandle : protected ITransactable
+		{
+		public:
+			FGraphHandle() = delete;
+			~FGraphHandle() = default;
+
+			// Sole constructor for FGraphHandle. Can only be used by friends of FHandleInitParams.
+			FGraphHandle(FHandleInitParams::EPrivateToken PrivateToken, const FHandleInitParams& InParams);
+
+			static FGraphHandle InvalidHandle();
+
+			bool IsValid();
+
+			TArray<FNodeHandle> GetAllNodes();
+			TArray<FNodeHandle> GetOutputNodes();
+			TArray<FNodeHandle> GetInputNodes();
+
+			FNodeHandle GetOutputNodeWithName(const FString& InName);
+			FNodeHandle GetInputNodeWithName(const FString& InName);
+
+			FNodeHandle AddNewInput(const FMetasoundInputDescription& InDescription);
+			bool RemoveInput(const FString& InputName);
+
+			FNodeHandle AddNewOutput(const FMetasoundOutputDescription& OutDescription);
+			bool RemoveOutput(const FString& OutputName);
+
+			FNodeHandle AddNewNode(const FNodeClassInfo& InNodeClass);
+
+			// Remove the node corresponding to this node handle.
+			// On success, invalidates the received node handle.
+			bool RemoveNode(FNodeHandle& InNode, bool bEvenIfInputOrOutputNode = false);
+
+			// Returns the metadata for the current graph, including the name, description and author.
+			FMetasoundClassMetadata GetGraphMetadata();
+
+			// Creates a new UMetasound object with the current graph.
+			// @param InPath, if in editor, the content directory path to save this metasound to.
+			// @returns a new MetasoundNodeAsset object on success, nullptr on failure.
+			// @todo decide if this means that we remove the graph in our current asset.
+			//       This seems dangerous since we can fail to save the new asset and scratch the old one.
+			UMetasound* SaveToNewMetasoundAsset(const FString& InPath);
+
+			// Exports this graph to a JSON at the given path.
+			// @returns true on success.
+			bool ExportToJSONAsset(const FString& InAbsolutePath);
+
+			// If the FNodeHandle given is itself a Metasound graph,
+			// and the FNodeHandle is a direct member of this FGraphHandle,
+			// If successful, this will invalidate the FNodeHandle and paste the graph for this node directly
+			// into the graph.
+			// If not successful, InNode will not be affected.
+			// @returns true on success, false on failure.
+			bool InflateNodeDirectlyIntoGraph(FNodeHandle& InNode);
+
+			// When called, creates a new node in this graph that can be edited directly with the FGraphHandle
+			// that gets returned.
+			// @param InInfo the name, author, etc that should be used for this graph.
+			// @returns a tuple of the FGraphHandle for this subgraph (which can be independently manipulated),
+			//          as well as the FNodeHandle for the subgraph's node in this graph.
+			//          Both of these will be invalid on failure.
+			TTuple<FGraphHandle, FNodeHandle> CreateEmptySubgraphNode(const FMetasoundClassMetadata& InInfo);
+
+			// This invokes the Metasound Builder to synchronously compile a 
+			// metasound operator, which can then be used for playback.
+			// @returns nullptr on failure.
+			// @todo: forward errors from the Metasound Builder.
+			// @todo: make an API for doing this async.
+			TUniquePtr<IOperator> BuildOperator();
+
+		private:
+
+			// The graph struct itself.
+			TDescriptionPtr<FMetasoundGraphDescription> GraphPtr;
+
+			// The class description for the graph corresponding to this.
+			TDescriptionPtr<FMetasoundClassDescription> GraphsClassDeclaration;
+
+			// the outermost document for all dependencies for the document this graph lives in.
+			TDescriptionPtr<FMetasoundDocument> OwningDocument;
+
+			// Scans all existing node ids to guarantee a new unique ID.
+			uint32 FindNewUniqueNodeId();
+
+			// Scans all existing dependency IDs in the root dependency list to guarantee a new unique ID when adding a dependency.
+			uint32 FindNewUniqueDependencyId();
+		};
+
+		// If there's a Metasound node referenced by a class in it's dependencies list,
+		// but the node' graph isn't implemented in that dependency list,
+		// this can be used to attempt to find which asset that node is implemented in,
+		// and return a graph handle for that node.
+		FGraphHandle GetGraphHandleForClass(const FMetasoundClassDescription& InClass);
+	}	
+}
diff --git a/Engine/Plugins/Runtime/Metasound/Source/MetasoundFrontend/Public/MetasoundFrontendBaseClasses.h b/Engine/Plugins/Runtime/Metasound/Source/MetasoundFrontend/Public/MetasoundFrontendBaseClasses.h
new file mode 100644
--- /dev/null
+++ b/Engine/Plugins/Runtime/Metasound/Source/MetasoundFrontend/Public/MetasoundFrontendBaseClasses.h
@@ -0,0 +1,658 @@
+// Copyright Epic Games, Inc. All Rights Reserved.
+
+#pragma once
+
+#include "CoreMinimal.h"
+#include "Misc/TVariant.h"
+#include "MetasoundFrontendDataLayout.h"
+
+
+namespace Metasound
+{	
+	namespace Frontend
+	{
+		// Utility base class for classes that want to support naive multi-level undo/redo.
+		class ITransactable : public TSharedFromThis<ITransactable>
+		{
+		public:
+
+			ITransactable() = delete;
+			ITransactable(uint32 InUndoLimit, TWeakObjectPtr<UObject> InOwningAsset = nullptr);
+
+			/*
+			 * Undo a single operation for this or any of it's owned transactables.
+			 * @returns true on success.
+			 */
+			bool Undo();
+
+			/*
+			 * Redo a single operation for this or any of it's owned transactables.
+			 * @returns true on success.
+			 */
+			bool Redo();
+
+		protected:
+			struct FReversibleTransaction
+			{
+				// Function that will be executed when Undo is called.
+				// This should return true on success.
+				TFunction<bool()> UndoTransaction;
+
+				// Function that will be executed when Redo is called.
+				// This should return true on success.
+				TFunction<bool()> RedoTransaction;
+
+				FReversibleTransaction(TFunction<bool()>&& InUndo, TFunction<bool()>&& InRedo)
+					: UndoTransaction(MoveTemp(InUndo))
+					, RedoTransaction(MoveTemp(InRedo))
+				{
+				}
+			};
+
+			// Commits an undoable/redoable transaction for this object.
+			void CommitTransaction(FReversibleTransaction&& InTransactionDescription);
+			
+			// Use this function to denote an ITransactable that should consider this ITransactable part of it's undo roll.
+			// For example, if an undoable transaction is committed to a FNodeHandle,
+			// and you want FGraphHandle::Undo to be able to undo that transaction,
+			// you would call FInputHandle::RegisterOwningTransactable(OwningGraphHandle). 
+			// 
+			// @returns true on success, false if registering this as an owning transaction would create a cycle.
+			bool RegisterOwningTransactable(ITransactable& InOwningTransactable);
+
+		private:
+
+			// Execute the topmost transaction on the undo stack,
+			// push it to redo stack,
+			// and notify the OnwingTransactable.
+			bool PerformLocalUndo();
+
+			// Execute the topmost transaction on the redo stack,
+			// push it to the undo stack,
+			// and notify the OwningTransactable.
+			bool PerformLocalRedo();
+
+			// Called from a owned transactable to denote that it is popping an undo or redo action.
+			// When this occurs, we discard the topmost call to the owned transactable from the UndoTransactableStack.
+			bool DiscardUndoFromOwnedTransactable(TWeakPtr<ITransactable> InOwnedTransactable);
+			bool DiscardRedoFromOwnedTransactable(TWeakPtr<ITransactable> InOwnedTransactable);
+
+			bool PushUndoFromOwnedTransactable(TWeakPtr<ITransactable> InOwnedTransactable);
+			bool PushRedoFromOwnedTransactable(TWeakPtr<ITransactable> InOwnedTransactable);
+
+			// TODO: Make these ring buffers to avoid constantly moving elements around.
+			TArray<TWeakPtr<ITransactable>> UndoTransactableStack;
+			TArray<TWeakPtr<ITransactable>> RedoTransactableStack;
+			TArray<FReversibleTransaction> LocalUndoTransactionStack;
+			TArray<FReversibleTransaction> LocalRedoTransactionStack;
+
+			// When CommitTransaction is called,
+			// this ITransactable is pushed to the ITransactable::TransactableStack of this ITransactable object.
+			TWeakPtr<ITransactable> OwningTransactable;
+
+		protected:
+			uint32 UndoLimit;
+
+			// If this transactable is operating on a UAsset, this can be filled to mark it's corresponding package as dirty.
+			TWeakObjectPtr<UObject> OwningAsset;
+		};
+
+		namespace Path
+		{
+			// This series of enums is used to denote 
+			// the path to a specific element in the class description.
+
+			// This enum is every possible Description struct we have.
+			enum class EDescType : uint8
+			{
+				Document, // FMetasoundDocument
+				Class, // FMetasoundClassDescription
+				Metadata, // FMetasoundClassMetadata
+				DocDependencies, // TArray<FMetasoundClassDescription>
+				ClassDependencies,
+				Graph, // 
+				Nodes,
+				Node,
+				Inputs,
+				Input,
+				Outputs,
+				Output,
+				Invalid
+			};
+
+
+			// This series of enums is used to delineate which description types are direct children or other descripiton types.
+			enum class EFromDocument : uint32
+			{
+				ToRootClass,
+				ToDependencies
+			};
+
+			
+			enum class EFromClass : uint32
+			{
+				ToMetadata,
+				ToInputs,
+				ToOutputs,
+				ToDependencies,
+				ToGraph,
+				ToInvalid
+			};
+
+			enum class EFromGraph : uint32
+			{
+				ToNodes,
+				Invalid
+			};
+
+			struct FElement
+			{
+				EDescType CurrentDescType;
+				
+				// This is only used if EDescType is Inputs, Outputs, or Dependencies.
+				FString LookupName;
+
+				// This is only used if EDescType == Node.
+				uint32 LookupID;
+
+				FElement()
+					: CurrentDescType(EDescType::Invalid)
+					, LookupID(INDEX_NONE)
+				{}
+			};
+		}
+
+		// This struct represents a selector for a specific element in a FMetasoundClassDescription.
+		// Typically these will be used with FDescriptionAccessPoint::GetElementAtPath.
+		struct FDescPath
+		{
+			TArray<Path::FElement> Path;
+
+			FDescPath()
+			{
+				Path::FElement RootElement;
+				RootElement.CurrentDescType = Path::EDescType::Document;
+				Path.Add(RootElement);
+			}
+
+			int32 Num() const
+			{
+				return Path.Num();
+			}
+
+			bool IsValid() const
+			{
+				return Path.Num() > 0 &&
+					Path[0].CurrentDescType == Path::EDescType::Document;
+			}
+
+			/** 
+			 * Bracket operator used for setting up a path into a metasound class description.
+			 * These can be chained, for example:
+			 * FDescPath PathToNode = RootClass[Path::EFromClass::ToGraph][Path::EFromGraph::ToNodes][12];
+			*/
+			FDescPath operator[](Path::EFromDocument InElement) const
+			{
+				if (!ensureAlwaysMsgf(Path.Num() > 0, TEXT("Tried to append to a path that had no root.")))
+				{
+					return FDescPath();
+				}
+
+				// Copy the old path, and based on the last element in that path, append a new child element.
+				FDescPath NewPath = *this;
+
+				const Path::FElement& LastElementInPath = Path.Last();
+
+				if(!ensureAlwaysMsgf(LastElementInPath.CurrentDescType == Path::EDescType::Document, TEXT("Tried to build an invalid path.")))
+				{
+					return FDescPath();
+				}
+
+				Path::FElement NewElement = Path::FElement();
+
+				switch (InElement)
+				{
+					case Path::EFromDocument::ToRootClass:
+					{
+						NewElement.CurrentDescType = Path::EDescType::Class;
+						break;
+					}
+					case Path::EFromDocument::ToDependencies:
+					{
+						NewElement.CurrentDescType = Path::EDescType::DocDependencies;
+						break;
+					}
+				}
+
+				NewPath.Path.Add(NewElement);
+				return NewPath;
+			}
+
+			FDescPath operator[](Path::EFromClass InElement) const
+			{
+				if (!ensureAlwaysMsgf(Path.Num() > 0, TEXT("Tried to append to a path that had no root.")))
+				{
+					return FDescPath();
+				}
+
+				// Copy the old path, and based on the last element in that path, append a new child element.
+				FDescPath NewPath = *this;
+
+				const Path::FElement& LastElementInPath = Path.Last();
+
+				if (!ensureAlwaysMsgf(LastElementInPath.CurrentDescType == Path::EDescType::Class, TEXT("Tried to build an invalid path.")))
+				{
+					return FDescPath();
+				}
+
+				Path::FElement NewElement = Path::FElement();
+
+				switch (InElement)
+				{
+					case Path::EFromClass::ToDependencies:
+					{
+						NewElement.CurrentDescType = Path::EDescType::ClassDependencies;
+						break;
+					}
+					case Path::EFromClass::ToGraph:
+					{
+						NewElement.CurrentDescType = Path::EDescType::Graph;
+						break;
+					}
+					case Path::EFromClass::ToInputs:
+					{
+						NewElement.CurrentDescType = Path::EDescType::Inputs;
+						break;
+					}
+					case Path::EFromClass::ToOutputs:
+					{
+						NewElement.CurrentDescType = Path::EDescType::Outputs;
+						break;
+					}
+					case Path::EFromClass::ToInvalid:
+					default:
+					{
+						ensureAlwaysMsgf(false, TEXT("Invalid path set up. Returning an empty path."));
+						return FDescPath();
+						break;
+					}
+				}
+
+				NewPath.Path.Add(NewElement);
+				return NewPath;
+			}
+
+			FDescPath operator[](Path::EFromGraph InElement) const
+			{
+				if (!ensureAlwaysMsgf(Path.Num() > 0, TEXT("Tried to append to a path that had no root.")))
+				{
+					return FDescPath();
+				}
+
+				// Copy the old path, and based on the last element in that path, append a new child element.
+				FDescPath NewPath = *this;
+
+				const Path::FElement& LastElementInPath = Path.Last();
+
+				if (!ensureAlwaysMsgf(LastElementInPath.CurrentDescType == Path::EDescType::Graph, TEXT("Tried to build an invalid path.")))
+				{
+					return FDescPath();
+				}
+
+				Path::FElement NewElement = Path::FElement();
+
+				switch (InElement)
+				{
+					case Path::EFromGraph::ToNodes:
+					{
+						NewElement.CurrentDescType = Path::EDescType::Nodes;
+						break;
+					}
+					case Path::EFromGraph::Invalid:
+					default:
+					{
+						ensureAlwaysMsgf(false, TEXT("Invalid path set up. Returning an empty path."));
+						return FDescPath();
+						break;
+					}
+				}
+
+				NewPath.Path.Add(NewElement);
+				return NewPath;
+			}
+
+			FDescPath operator[](uint32 InElement) const
+			{
+				if (!ensureAlwaysMsgf(Path.Num() > 0, TEXT("Tried to append to a path that had no root.")))
+				{
+					return FDescPath();
+				}
+
+				// Copy the old path, and based on the last element in that path, append a new child element.
+				FDescPath NewPath = *this;
+
+				const Path::FElement& LastElementInPath = Path.Last();
+
+				if (!ensureAlwaysMsgf(LastElementInPath.CurrentDescType == Path::EDescType::Nodes || LastElementInPath.CurrentDescType == Path::EDescType::DocDependencies, TEXT("Tried to build an invalid path.")))
+				{
+					return FDescPath();
+				}
+
+				Path::FElement NewElement = Path::FElement();
+
+				if (LastElementInPath.CurrentDescType == Path::EDescType::Nodes)
+				{
+					NewElement.CurrentDescType = Path::EDescType::Node;
+					NewElement.LookupID = InElement;
+				}
+				else if (LastElementInPath.CurrentDescType == Path::EDescType::ClassDependencies)
+				{
+					NewElement.CurrentDescType = Path::EDescType::DocDependencies;
+					NewElement.LookupID = InElement;
+				}
+
+				NewPath.Path.Add(NewElement);
+				return NewPath;
+			}
+
+			FDescPath operator [](const TCHAR* InElementName) const
+			{
+				if (!ensureMsgf(Path.Num() > 0, TEXT("Tried to append to a path that had no root.")))
+				{
+					return FDescPath();
+				}
+
+				// Copy the old path, and based on the last element in that path, append a new child element.
+				FDescPath NewPath = *this;
+
+				const Path::FElement& LastElementInPath = Path.Last();
+				Path::FElement NewElement = Path::FElement();
+
+				switch (LastElementInPath.CurrentDescType)
+				{
+					case Path::EDescType::DocDependencies:
+					{
+						NewElement.CurrentDescType = Path::EDescType::Class;
+						NewElement.LookupName = FString(InElementName);
+						break;
+					}
+					case Path::EDescType::Inputs:
+					{
+						NewElement.CurrentDescType = Path::EDescType::Input;
+						NewElement.LookupName = FString(InElementName);
+						break;
+					}
+					case Path::EDescType::Outputs:
+					{
+						NewElement.CurrentDescType = Path::EDescType::Output;
+						NewElement.LookupName = FString(InElementName);
+						break;
+					}
+					/** Class descriptions and graph descriptions are navigated via enums. */
+					case Path::EDescType::Class:
+					case Path::EDescType::Graph:
+					{
+						ensureAlwaysMsgf(false, TEXT("Invalid path set up. Classes and Graphs are accessed by using the DescriptionPath::EClass and DesciptionPath::EGraph enums, respectively."));
+						return FDescPath();
+						break;
+					}
+					/** Dependencies for individual classes are navigated to via a uint32. */
+					case Path::EDescType::ClassDependencies:
+					{
+						ensureAlwaysMsgf(false, TEXT("Invalid path set up. Dependencies for individual classes are accessed by using a uint32."));
+						return FDescPath();
+						break;;
+					}
+					/** Node list descriptions are navigated via a uint32. */
+					case Path::EDescType::Nodes:
+					{
+						ensureAlwaysMsgf(false, TEXT("Invalid path set up. Node lists are accessed by using a uint32."));
+						return FDescPath();
+						break;
+					}
+					/** Individiual inputs, outpus and nodes, as well as metadata, do not have any child elements. */
+					case Path::EDescType::Input:
+					case Path::EDescType::Output:
+					case Path::EDescType::Node:
+					case Path::EDescType::Metadata:
+					default:
+					{
+						ensureAlwaysMsgf(false, TEXT("Invalid path. Tried to add more pathing off of a description type that has no children."));
+						return FDescPath();
+						break;
+					}
+				}
+
+				NewPath.Path.Add(NewElement);
+				return NewPath;
+			}
+
+			// This can be used to go up multiple levels in a path. Effectively removes the last UnwindCount elements from the path.
+			FDescPath& operator<< (uint32 UnwindCount)
+			{
+				if (!ensureAlwaysMsgf(((uint32)Path.Num()) > UnwindCount, TEXT("Tried to unwind past the root of an FDescPath!")))
+				{
+					return *this;
+				}
+
+				while (UnwindCount-- > 0)
+				{
+					Path.RemoveAt(Path.Num() - 1);
+				}
+
+				return *this;
+			}
+		};
+
+		// Series of utility functions for description paths.
+		namespace Path
+		{
+			// When given a path for a node, finds that node's declaration in the Dependencies list.
+			FDescPath GetPathToClassForNode(FDescPath InPathForNode, FString& InNodeName);
+			
+			// Returns a path for whatever class description encapsulates the given graph.
+			FDescPath GetOwningClassDescription(FDescPath InPathForGraph);
+			
+			FDescPath GetClassDescription(const FString& InClassName);
+
+			FDescPath GetDependencyPath(const FString& InDependencyName);
+
+			FDescPath GetDependencyPath(uint32 InDependencyID);
+
+			// Generates a human readable string for a path into a Metasound description. Used for debugging.
+			FString GetPrintableString(FDescPath InPath);
+		}
+
+		typedef TArray<uint32> FClassDependencyIDs;
+
+		/**
+		 * Utility class that lives alongside a FMetasoundClassDescription to allow it's individual elements to be safely read and edited.
+		 * Note: Could be moved to a private header.
+		 *       It's also only necessary because we can't wrap a FMetasoundClassDescription in a shared ptr and also use it as a UPROPERTY.
+		 */
+		class FDescriptionAccessPoint
+		{
+		public:
+			FDescriptionAccessPoint() = delete;
+			FDescriptionAccessPoint(FMetasoundDocument& InRootDocument);
+
+			// Returns the root class description associated with this FDescriptionAccessPoint.
+			FMetasoundDocument& GetRoot();
+
+			// Get a descriptor element at a specific path from the root of a metasound.
+			// @see FDescPath for more details.
+			// @returns a pointer to a descriptor struct, or an invalid ptr if the path is invalid.
+			template <typename MetasoundDescriptionType>
+			MetasoundDescriptionType* GetElementAtPath(const FDescPath& InPathFromRoot)
+			{
+				if (TIsSame<MetasoundDescriptionType, FMetasoundClassDescription>::Value)
+				{
+					return (MetasoundDescriptionType*) GetClassFromPath(InPathFromRoot);
+				}
+				else if (TIsSame<MetasoundDescriptionType, FMetasoundNodeDescription>::Value)
+				{
+					return (MetasoundDescriptionType*) GetNodeFromPath(InPathFromRoot);
+				}
+				else if (TIsSame<MetasoundDescriptionType, FMetasoundGraphDescription>::Value)
+				{
+					return (MetasoundDescriptionType*) GetGraphFromPath(InPathFromRoot);
+				}
+				else if (TIsSame<MetasoundDescriptionType, FMetasoundInputDescription>::Value)
+				{
+					return (MetasoundDescriptionType*) GetInputFromPath(InPathFromRoot);
+				}
+				else if (TIsSame<MetasoundDescriptionType, FMetasoundOutputDescription>::Value)
+				{
+					return (MetasoundDescriptionType*) GetOutputFromPath(InPathFromRoot);
+				}
+				else if (TIsSame<MetasoundDescriptionType, FMetasoundClassMetadata>::Value)
+				{
+					return (MetasoundDescriptionType*) GetMetadataFromPath(InPathFromRoot);
+				}
+				else if (TIsSame<MetasoundDescriptionType, FClassDependencyIDs>::Value)
+				{
+					return (MetasoundDescriptionType*) GetClassDependencyIDsFromPath(InPathFromRoot);
+				}
+				else
+				{
+					ensureAlwaysMsgf(false, TEXT("Tried to call GetElementAtPath with an invalid type."));
+					return nullptr;
+				}
+			}
+
+		private:
+			// The root class description for all description structs
+			// accessed via this class.
+			FMetasoundDocument& RootDocument;
+
+			// These functions result in some duplicate code, but allow us to hide
+			// the implementation for GetElementAtPath.
+			FMetasoundClassDescription* GetClassFromPath(const FDescPath& InPathFromRoot);
+			FMetasoundNodeDescription* GetNodeFromPath(const FDescPath& InPathFromRoot);
+			FMetasoundGraphDescription* GetGraphFromPath(const FDescPath& InPathFromRoot);
+			FMetasoundInputDescription* GetInputFromPath(const FDescPath& InPathFromRoot);
+			FMetasoundOutputDescription* GetOutputFromPath(const FDescPath& InPathFromRoot);
+			FMetasoundClassMetadata* GetMetadataFromPath(const FDescPath& InPathFromRoot);
+			FClassDependencyIDs* GetClassDependencyIDsFromPath(const FDescPath& InPathFromRoot);
+
+			using FMetasoundDescriptionPtr = TVariant<
+				FMetasoundDocument*,
+				FMetasoundClassDescription*,
+				FMetasoundNodeDescription*,
+				FMetasoundGraphDescription*,
+				FMetasoundInputDescription*,
+				FMetasoundOutputDescription*,
+				FMetasoundClassMetadata*,
+				FClassDependencyIDs*,
+				void*>
+				;
+
+			struct FDescriptionUnwindStep
+			{
+				FMetasoundDescriptionPtr DescriptionStructPtr;
+				Path::EDescType Type;
+			};
+
+			// This function pops off elements of InPath to go to the next level down in a Metasound description.
+			FDescriptionUnwindStep GoToNext(FDescPath& InPath, FDescriptionUnwindStep InElement);
+		};
+
+		// Utility class used by the frontend handle classes to access the underlying Description data they are modifying.
+		// Similar to a TSoftObjectPtr: 
+		// this class caches a FDescPath to a specific element in a Metasound Class Description, and a weak pointer to a FDescriptionAccessPoint
+		// to use it with.
+		template <typename MetasoundDescriptionType>
+		class TDescriptionPtr
+		{
+			static constexpr bool bIsValidMetasoundDescriptionType =
+				TIsSame< MetasoundDescriptionType, FMetasoundDocument>::Value ||
+				TIsSame< MetasoundDescriptionType, FMetasoundGraphDescription>::Value ||
+				TIsSame< MetasoundDescriptionType, FMetasoundNodeDescription>::Value ||
+				TIsSame< MetasoundDescriptionType, FMetasoundInputDescription>::Value ||
+				TIsSame< MetasoundDescriptionType, FMetasoundOutputDescription>::Value ||
+				TIsSame< MetasoundDescriptionType, FMetasoundClassMetadata>::Value ||
+				TIsSame< MetasoundDescriptionType, FMetasoundClassDescription>::Value;
+
+			static_assert(bIsValidMetasoundDescriptionType, R"(Tried to use a Metasound::Frontend::FDescriptionAccessor with an invalid type. Supported types are: 
+			FMetasoundDocument			
+			FMetasoundGraphDescription, 
+			FMetasoundNodeDescription,
+			FMetasoundInputDescription,
+			FMetasoundOutputDescription,
+			FMetasoundClassMetadata,
+			FMetasoundClassDescription)");
+
+			TWeakPtr<FDescriptionAccessPoint> AccessPoint;
+			FDescPath PathFromRoot;
+
+		public:
+			TDescriptionPtr() = delete;
+
+			TDescriptionPtr(TWeakPtr<FDescriptionAccessPoint> InAccessPoint, const FDescPath& InPathFromRoot)
+				: AccessPoint(InAccessPoint)
+				, PathFromRoot(InPathFromRoot)
+			{
+			}
+
+			// @returns false if the description this accessor was referencing has been destroyed,
+			// or if this was created with bad arguments.
+			bool IsValid() const
+			{
+				return AccessPoint.IsValid() && PathFromRoot.IsValid();
+			}
+
+			FDescPath GetPath() const
+			{
+				return PathFromRoot;
+			}
+
+			TWeakPtr<FDescriptionAccessPoint> GetAccessPoint() const
+			{
+				return AccessPoint;
+			}
+
+			// returns a pointer or reference to the description data struct.
+			// Note that this is NOT thread safe- this should be called on the same thread as whomever owns the lifecycle of the underlying MetasoundDescriptionType.
+			MetasoundDescriptionType* Get() const
+			{
+				if (TSharedPtr<FDescriptionAccessPoint> PinnedAccessPoint = AccessPoint.Pin())
+				{
+					if (TIsSame<MetasoundDescriptionType, FMetasoundDocument>::Value)
+					{
+						return (MetasoundDescriptionType*) &(PinnedAccessPoint->GetRoot());
+					}
+					else
+					{
+						return PinnedAccessPoint->GetElementAtPath<MetasoundDescriptionType>(PathFromRoot);
+					}
+				}
+				else
+				{
+					return nullptr;
+				}
+			}
+
+			MetasoundDescriptionType& GetChecked() const
+			{
+				MetasoundDescriptionType* ElementPtr = Get();
+				check(ElementPtr);
+				return *ElementPtr;
+			}
+
+			MetasoundDescriptionType& operator*() const
+			{
+				return GetChecked();
+			}
+
+			MetasoundDescriptionType* operator->() const
+			{
+				return Get();
+			}
+
+			FORCEINLINE explicit operator bool() const
+			{
+				return IsValid();
+			}
+		};
+	}
+}
diff --git a/Engine/Plugins/Runtime/Metasound/Source/MetasoundFrontend/Public/MetasoundFrontendDataLayout.h b/Engine/Plugins/Runtime/Metasound/Source/MetasoundFrontend/Public/MetasoundFrontendDataLayout.h
new file mode 100644
--- /dev/null
+++ b/Engine/Plugins/Runtime/Metasound/Source/MetasoundFrontend/Public/MetasoundFrontendDataLayout.h
@@ -0,0 +1,226 @@
+// Copyright Epic Games, Inc. All Rights Reserved.
+
+#pragma once
+
+#include "CoreMinimal.h"
+#include "MetasoundGraph.h"
+#include "MetasoundFrontendDataLayout.generated.h"
+
+UENUM()
+enum class EMetasoundClassType : uint8
+{
+	// Objects with the type Primitive are always available as part of the Metasound specification.
+	// The Metasound Builder can optimize usage of these types of nodes.
+	Primitive,
+	// Objects with the type External are implemented in compiled code.
+	External,
+	// Objects with this type are Metasound graphs unto themselves,
+	// they can be used like any other node.
+	MetasoundGraph,
+	// Input into the owning metasound. Only used in FMetasoundNodeDescription.
+	Input,
+	//Output from the owning metasound. Only used in FMetasoundNodeDescription.
+	Output,
+	Invalid UMETA(Hidden)
+};
+
+USTRUCT()
+struct FMetasoundClassMetadata
+{
+	GENERATED_BODY()
+
+	// the title for this metasound object. Will be the way the metasound is described in the content browser.
+	UPROPERTY()
+	FString NodeName;
+
+	// This will always be set to EMetasoundObjectType::Metasound.
+	UPROPERTY()
+	EMetasoundClassType NodeType;
+
+	// Optional longform description of what this Metasound does.
+	// Can be displayed in a tooltip.
+	UPROPERTY()
+	FString MetasoundDescription;
+
+	// If another Metasound depends on this metasoundand doesn't package it, 
+	// this prompt will show if we can't find this Metasound. 
+	// It should describe how to get this metasound.
+	UPROPERTY()
+	FString PromptIfMissing;
+
+	// Original author of this Metasound Object.
+	UPROPERTY()
+	FString AuthorName;
+};
+
+USTRUCT()
+struct FMetasoundInputDescription
+{
+	GENERATED_BODY()
+
+	// The descriptive name of this input.
+	UPROPERTY()
+	FString Name;
+
+	// The type of this input. This type matches with a type declared via the DECLARE_METASOUND_DATA_REFERENCE_TYPES macro.
+	UPROPERTY()
+	FName TypeName;
+
+	// Optional description text about this input.
+	UPROPERTY()
+	FString ToolTip;
+};
+
+USTRUCT()
+struct FMetasoundOutputDescription
+{
+	GENERATED_BODY()
+
+	// The descriptive name of this output.
+	UPROPERTY()
+	FString Name;
+
+	// The type of this output. This type matches with a type declared via the DECLARE_METASOUND_DATA_REFERENCE_TYPES macro.
+	UPROPERTY()
+	FName TypeName;
+
+	// Optional description text about this output.
+	UPROPERTY()
+	FString ToolTip;
+};
+
+// This is used to describe a single incoming connection or literal for a given input on a FMetasoundNodeDescription.
+USTRUCT()
+struct FMetasoundNodeConnectionDescription
+{
+	GENERATED_BODY()
+
+	static const uint32 DisconnectedNodeID = INDEX_NONE;
+
+	// Name of the input this connection is for. Should match an input name in the Inputs array of this node's FMetasoundClassDescription.
+	UPROPERTY()
+	FString InputName;
+
+	// Unique ID for the node this connection is from.
+	UPROPERTY()
+	uint32 NodeID;
+
+	// FMetasoundOutputDescription::Name of the output this connection is from if NodeID is valid,
+	// or a literal value if NodeID is FMetasoundNodeConnectionDescription::DisconnectedNodeID.
+	UPROPERTY()
+	FString OutputNameOrLiteral;
+};
+
+// description for an individual node within the FMetasoundGraphDescription.
+// corresponds to an instance of a FMetasoundObjectDescription listed in the Dependencies array.
+USTRUCT()
+struct FMetasoundNodeDescription
+{
+	GENERATED_BODY()
+
+	// Unique integer given to this node. Only has to be unique within the Nodes array.
+	UPROPERTY()
+	uint32 UniqueID;
+
+	// This will either match a FMetasoundObjectDescription::Metadata::Name in the Dependencies list of the owning FMetasoundObjectDescription,
+	// Or it will match a FMetasoundInputDescription::Name or FMetasoundOutputDescription::Name, depending on the value of ObjectTypeOfNode.
+	// TODO: Now that each dependency class has a unique ID, we can consider using that rather than the full name.
+	UPROPERTY()
+	FString Name;
+
+	// This describes what type of node this object represents.
+	UPROPERTY()
+	EMetasoundClassType ObjectTypeOfNode;
+
+	//TODO: See if we need to explicitly list the typename of the input/output if ObjectTypeOfNode is Input or Output.
+
+	// list of connections or literal values for each connected input pin on this node.
+	UPROPERTY()
+	TArray<FMetasoundNodeConnectionDescription> InputConnections;
+};
+
+USTRUCT()
+struct FMetasoundGraphDescription
+{
+	GENERATED_BODY()
+
+		UPROPERTY()
+		TArray<FMetasoundNodeDescription> Nodes;
+};
+
+// Full saved out description of a Metasound object,
+// which is anything that could be used as a node in a graph.
+// It is not called "FMetasoundNodeDescription" to avoid confusion with the individual node objects that compose
+// an asset's Metasound graph.
+USTRUCT()
+struct FMetasoundClassDescription
+{
+	GENERATED_BODY()
+
+	static const uint32 RootClassID = 0;
+
+	// Unique ID for this FMetasoundClassDescription in FMetasoundDocument::Dependencies.
+	// this will be set to FMetasoundClassDescription::RootClassID for FMetasoundDocument::RootClass.
+	UPROPERTY()
+	uint32 UniqueID;
+
+	UPROPERTY()
+	FMetasoundClassMetadata Metadata;
+
+	UPROPERTY()
+	TArray<FMetasoundInputDescription> Inputs;
+
+	UPROPERTY()
+	TArray<FMetasoundOutputDescription> Outputs;
+
+	// If this object is itself of the Metasound type, here we list the unique IDs of other objects it depends on to be built.
+	// These unique IDs are then found in FMetasoundDocument::Dependencies.
+	UPROPERTY()
+	TArray<uint32> DependencyIDs;
+
+	// If this object is itself of the Metasound type, We can fully describe that metasound's graph here.
+	// If this object is a Metasound and this Graph is empty,
+	// we will attempt to find the implemented version of this Metasound in the asset registry using Metasound::Frontend::FindGraphHandleForClass.
+	UPROPERTY()
+	FMetasoundGraphDescription Graph;
+};
+
+// This is used to describe the required inputs and outputs for a metasound, and is used to make sure we can use a metasound graph for specific applications.
+// For example, a UMetasoundSource needs to generate audio, so its RequiredOutputs will contain "MainAudioOutput"
+USTRUCT()
+struct FMetasoundArchetype
+{
+	GENERATED_BODY()
+
+	// Name of the archetype we're using.
+	UPROPERTY()
+	FName ArchetypeName;
+
+	// whatever inputs are required to support this archetype.
+	UPROPERTY()
+	TArray<FMetasoundInputDescription> RequiredInputs;
+
+	// whatever outputs are required to support this archetype.
+	UPROPERTY()
+	TArray<FMetasoundOutputDescription> RequiredOutputs;
+};
+
+// Outermost description of a single metasound.
+USTRUCT()
+struct FMetasoundDocument
+{
+	GENERATED_BODY()
+
+	// description of what kind of metasound this is going to be (i.e. source effect, playable source, generic node, etc).
+	UPROPERTY()
+	FMetasoundArchetype Archetype;
+
+	// The highest level Metasound Class in this document.
+	// Typically of type MetasoundGraph, contains graph of the node classes listed in Dependencies.
+	UPROPERTY()
+	FMetasoundClassDescription RootClass;
+
+	// this is the list of every dependency required by the RootClass, as well as nested dependencies.
+	UPROPERTY()
+	TArray<FMetasoundClassDescription> Dependencies;
+};
\ No newline at end of file
diff --git a/Engine/Plugins/Runtime/Metasound/Source/MetasoundFrontend/Public/MetasoundNodeAsset.h b/Engine/Plugins/Runtime/Metasound/Source/MetasoundFrontend/Public/MetasoundNodeAsset.h
new file mode 100644
--- /dev/null
+++ b/Engine/Plugins/Runtime/Metasound/Source/MetasoundFrontend/Public/MetasoundNodeAsset.h
@@ -0,0 +1,100 @@
+// Copyright Epic Games, Inc. All Rights Reserved.
+
+#pragma once
+
+#include "CoreMinimal.h"
+#include "MetasoundFrontend.h"
+#include "MetasoundNodeAsset.generated.h"
+
+
+/**
+ * This asset type is used for Metasound assets that can only be used as nodes in other Metasound graphs.
+ * Because of this, they can have any inputs or outputs they need.
+ */
+UCLASS(hidecategories = object, BlueprintType)
+class UMetasound : public UObject
+{
+	GENERATED_BODY()
+
+protected:
+	UPROPERTY()
+	FMetasoundDocument RootMetasoundDocument;
+
+	TSharedPtr<Metasound::Frontend::FDescriptionAccessPoint> AccessPoint;
+
+public:
+	UMetasound(const FObjectInitializer& ObjectInitializer);
+
+	FMetasoundClassMetadata GetMetadata()
+	{
+		return RootMetasoundDocument.RootClass.Metadata;
+	}
+
+	// This can be used to update the metadata (name, author, etc) for this metasound.
+	// @param InMetadata may be updated with any corrections we do to the input metadata.
+	void SetMetadata(FMetasoundClassMetadata& InMetadata)
+	{
+		InMetadata.NodeType = EMetasoundClassType::MetasoundGraph;
+		RootMetasoundDocument.RootClass.Metadata = InMetadata;
+		MarkPackageDirty();
+	}
+
+	// delete this asset's current metasound document,
+	// and replace it with InClassDescription.
+	void SetMetasoundDocument(const FMetasoundDocument& InDocument)
+	{
+		RootMetasoundDocument = InDocument;
+	}
+
+	// returns a weak pointer that can be used to build a TDescriptionPtr
+	// for direct editing of the FMetasoundClassDescription tree.
+	// For advance use only, and requires knowledge of Metasound::Frontend::FDescPath syntax.
+	// For most use cases, use GetGraphHandle() instead.
+	TWeakPtr<Metasound::Frontend::FDescriptionAccessPoint> GetGraphAccessPoint()
+	{
+		return TWeakPtr<Metasound::Frontend::FDescriptionAccessPoint>(AccessPoint);
+	}
+
+	// Get the handle for the root metasound graph of this asset.
+	Metasound::Frontend::FGraphHandle GetRootGraphHandle()
+	{
+		using FDescPath = Metasound::Frontend::FDescPath;
+		using FGraphHandle = Metasound::Frontend::FGraphHandle;
+		using FHandleInitParams = Metasound::Frontend::FHandleInitParams;
+		using EFromDocument = Metasound::Frontend::Path::EFromDocument;
+		using EFromClass = Metasound::Frontend::Path::EFromClass;
+
+		FDescPath PathToGraph = FDescPath()[EFromDocument::ToRootClass][EFromClass::ToGraph];
+		FHandleInitParams InitParams = { GetGraphAccessPoint(), PathToGraph, RootMetasoundDocument.RootClass.Metadata.NodeName, MakeWeakObjectPtr(this) };
+		return FGraphHandle(FHandleInitParams::PrivateToken, InitParams);
+	}
+
+	TArray<Metasound::Frontend::FGraphHandle> GetAllSubgraphHandles()
+	{
+		using FDescPath = Metasound::Frontend::FDescPath;
+		using FGraphHandle = Metasound::Frontend::FGraphHandle;
+		using FHandleInitParams = Metasound::Frontend::FHandleInitParams;
+		using EFromDocument = Metasound::Frontend::Path::EFromDocument;
+		using EFromClass = Metasound::Frontend::Path::EFromClass;
+
+		TArray<FGraphHandle> OutArray;
+
+		FDescPath RootPathForDependencyGraphs = FDescPath()[EFromDocument::ToDependencies];
+
+		TArray<FMetasoundClassDescription>& DependenciesList = RootMetasoundDocument.Dependencies;
+		for (FMetasoundClassDescription& Dependency : DependenciesList)
+		{
+			const bool bIsSubgraph = Dependency.Metadata.NodeType == EMetasoundClassType::MetasoundGraph && Dependency.Graph.Nodes.Num() > 0;
+			if (bIsSubgraph)
+			{
+				FDescPath SubgraphPath = RootPathForDependencyGraphs[Dependency.UniqueID][EFromClass::ToGraph];
+				FHandleInitParams InitParams = { GetGraphAccessPoint(), SubgraphPath, Dependency.Metadata.NodeName, MakeWeakObjectPtr(this) };
+				OutArray.Emplace(FHandleInitParams::PrivateToken, InitParams);
+			}
+		}
+
+		return OutArray;
+	}
+
+	bool ExportToJSON(const FString& InAbsolutePath);
+};
\ No newline at end of file
