Fix for missing includes for non-unity compiles.

diff --git a/Engine/Plugins/Runtime/Metasound/Source/MetasoundFrontend/Private/MetasoundFrontend.cpp b/Engine/Plugins/Runtime/Metasound/Source/MetasoundFrontend/Private/MetasoundFrontend.cpp
--- a/Engine/Plugins/Runtime/Metasound/Source/MetasoundFrontend/Private/MetasoundFrontend.cpp
+++ b/Engine/Plugins/Runtime/Metasound/Source/MetasoundFrontend/Private/MetasoundFrontend.cpp
@@ -1,11 +1,17 @@
 // Copyright Epic Games, Inc. All Rights Reserved.
 
 #include "MetasoundFrontend.h"
+
 #include "Backends/JsonStructSerializerBackend.h"
 #include "StructSerializer.h"
+#include "HAL/FileManager.h"
+#include "HAL/IConsoleManager.h"
+#include "Modules/ModuleManager.h"
+
 #include "MetasoundNodeAsset.h"
 #include "MetasoundAudioFormats.h"
 
+
 static int32 MetasoundUndoRollLimitCvar = 128;
 FAutoConsoleVariableRef CVarMetasoundUndoRollLimit(
 	TEXT("au.Metasound.Frontend.UndoRollLimit"),
diff --git a/Engine/Plugins/Runtime/Metasound/Source/MetasoundFrontend/Private/MetasoundNodeAsset.cpp b/Engine/Plugins/Runtime/Metasound/Source/MetasoundFrontend/Private/MetasoundNodeAsset.cpp
--- a/Engine/Plugins/Runtime/Metasound/Source/MetasoundFrontend/Private/MetasoundNodeAsset.cpp
+++ b/Engine/Plugins/Runtime/Metasound/Source/MetasoundFrontend/Private/MetasoundNodeAsset.cpp
@@ -5,6 +5,7 @@
 
 #include "Backends/JsonStructSerializerBackend.h"
 #include "StructSerializer.h"
+#include "HAL/FileManager.h"
 
 UMetasound::UMetasound(const FObjectInitializer& ObjectInitializer)
 	: Super(ObjectInitializer)
