Clang compile fix: friend class UMetasound is in the global namespace, so FHandleInitParams needed to establish friendship with ::UMetasound rather than UMetasound.

diff --git a/Engine/Plugins/Runtime/Metasound/Source/MetasoundFrontend/Public/MetasoundFrontend.h b/Engine/Plugins/Runtime/Metasound/Source/MetasoundFrontend/Public/MetasoundFrontend.h
--- a/Engine/Plugins/Runtime/Metasound/Source/MetasoundFrontend/Public/MetasoundFrontend.h
+++ b/Engine/Plugins/Runtime/Metasound/Source/MetasoundFrontend/Public/MetasoundFrontend.h
@@ -109,7 +109,7 @@ namespace Metasound
 			friend class FInputHandle;
 			friend class FNodeHandle;
 			friend class FGraphHandle;
-			friend class UMetasound;
+			friend class ::UMetasound;
 		};
 
 		class FOutputHandle : protected ITransactable
