Add MetasoundEngine Module

diff --git a/Engine/Plugins/Runtime/Metasound/Metasound.uplugin b/Engine/Plugins/Runtime/Metasound/Metasound.uplugin
--- a/Engine/Plugins/Runtime/Metasound/Metasound.uplugin
+++ b/Engine/Plugins/Runtime/Metasound/Metasound.uplugin
@@ -30,6 +30,11 @@
 			"Name" : "MetasoundFrontend",
 			"Type" : "Runtime",
 			"LoadingPhase" : "Default"
-		}
+		},
+		{
+			"Name" : "MetasoundEngine",
+			"Type" : "Runtime",
+			"LoadingPhase" : "Default"
+		},
 	]
 }
diff --git a/Engine/Plugins/Runtime/Metasound/Source/MetasoundEngine/MetasoundEngine.Build.cs b/Engine/Plugins/Runtime/Metasound/Source/MetasoundEngine/MetasoundEngine.Build.cs
new file mode 100644
--- /dev/null
+++ b/Engine/Plugins/Runtime/Metasound/Source/MetasoundEngine/MetasoundEngine.Build.cs
@@ -0,0 +1,32 @@
+// Copyright Epic Games, Inc. All Rights Reserved.
+
+using UnrealBuildTool;
+using System.IO;
+
+namespace UnrealBuildTool.Rules
+{
+	public class MetasoundEngine : ModuleRules
+	{
+		public MetasoundEngine(ReadOnlyTargetRules Target) : base(Target)
+		{
+			PrivateDependencyModuleNames.AddRange
+			(
+				new string[]
+				{
+					"Core",
+					"MetasoundFrontend",
+					"MetasoundGraphCore"
+				}
+			);
+
+			PublicDependencyModuleNames.AddRange
+			(
+				new string[]
+				{
+					"CoreUObject",
+					"Engine"
+				}
+			);
+		}
+	}
+}
diff --git a/Engine/Plugins/Runtime/Metasound/Source/MetasoundEngine/Private/MetasoundEngineModule.cpp b/Engine/Plugins/Runtime/Metasound/Source/MetasoundEngine/Private/MetasoundEngineModule.cpp
new file mode 100644
--- /dev/null
+++ b/Engine/Plugins/Runtime/Metasound/Source/MetasoundEngine/Private/MetasoundEngineModule.cpp
@@ -0,0 +1,16 @@
+// Copyright Epic Games, Inc. All Rights Reserved.
+
+#include "MetasoundEngineModule.h"
+
+
+DEFINE_LOG_CATEGORY(LogMetasoundEngine);
+
+class FMetasoundEngineModule : public IMetasoundEngineModule
+{
+	virtual void StartupModule() override
+	{
+		UE_LOG(LogMetasoundEngine, Log, TEXT("Metasound Engine Initialized"));
+	}
+};
+
+IMPLEMENT_MODULE(FMetasoundEngineModule, MetasoundEngine);
diff --git a/Engine/Plugins/Runtime/Metasound/Source/MetasoundEngine/Private/MetasoundEngineModule.h b/Engine/Plugins/Runtime/Metasound/Source/MetasoundEngine/Private/MetasoundEngineModule.h
new file mode 100644
--- /dev/null
+++ b/Engine/Plugins/Runtime/Metasound/Source/MetasoundEngine/Private/MetasoundEngineModule.h
@@ -0,0 +1,12 @@
+// Copyright Epic Games, Inc. All Rights Reserved.
+#pragma once
+
+#include "CoreMinimal.h"
+#include "Logging/LogMacros.h"
+#include "Modules/ModuleInterface.h"
+
+DECLARE_LOG_CATEGORY_EXTERN(LogMetasoundEngine, Log, All);
+
+class METASOUNDENGINE_API IMetasoundEngineModule : public IModuleInterface
+{
+};
