IWYU Fix

diff --git a/Engine/Plugins/Runtime/Metasound/Source/MetasoundEngine/Private/MetasoundEngineModule.cpp b/Engine/Plugins/Runtime/Metasound/Source/MetasoundEngine/Private/MetasoundEngineModule.cpp
--- a/Engine/Plugins/Runtime/Metasound/Source/MetasoundEngine/Private/MetasoundEngineModule.cpp
+++ b/Engine/Plugins/Runtime/Metasound/Source/MetasoundEngine/Private/MetasoundEngineModule.cpp
@@ -1,10 +1,11 @@
 // Copyright Epic Games, Inc. All Rights Reserved.
-
 #include "MetasoundEngineModule.h"
 
+#include "Modules/ModuleManager.h"
 
 DEFINE_LOG_CATEGORY(LogMetasoundEngine);
 
+
 class FMetasoundEngineModule : public IMetasoundEngineModule
 {
 	virtual void StartupModule() override
