Move UMetasound to MetasoundEngine (step toward editor enablement as EdGraph reference will be required on class).

diff --git a/Engine/Plugins/Runtime/Metasound/Metasound.uplugin b/Engine/Plugins/Runtime/Metasound/Metasound.uplugin
--- a/Engine/Plugins/Runtime/Metasound/Metasound.uplugin
+++ b/Engine/Plugins/Runtime/Metasound/Metasound.uplugin
@@ -36,5 +36,10 @@
 			"Type" : "Runtime",
 			"LoadingPhase" : "Default"
 		},
+		{
+			"Name": "MetasoundEditor",
+			"Type": "Editor",
+			"LoadingPhase": "PostEngineInit"
+		},
 	]
 }
diff --git a/Engine/Plugins/Runtime/Metasound/Source/MetasoundEngine/MetasoundEngine.Build.cs b/Engine/Plugins/Runtime/Metasound/Source/MetasoundEngine/MetasoundEngine.Build.cs
--- a/Engine/Plugins/Runtime/Metasound/Source/MetasoundEngine/MetasoundEngine.Build.cs
+++ b/Engine/Plugins/Runtime/Metasound/Source/MetasoundEngine/MetasoundEngine.Build.cs
@@ -14,7 +14,6 @@ namespace UnrealBuildTool.Rules
 				new string[]
 				{
 					"Core",
-					"MetasoundFrontend",
 					"MetasoundGraphCore"
 				}
 			);
@@ -24,7 +23,9 @@ namespace UnrealBuildTool.Rules
 				new string[]
 				{
 					"CoreUObject",
-					"Engine"
+					"Engine",
+					"MetasoundFrontend",
+					"Serialization"
 				}
 			);
 		}
diff --git a/Engine/Plugins/Runtime/Metasound/Source/MetasoundEngine/Private/Metasound.cpp b/Engine/Plugins/Runtime/Metasound/Source/MetasoundEngine/Private/Metasound.cpp
new file mode 100644
--- /dev/null
+++ b/Engine/Plugins/Runtime/Metasound/Source/MetasoundEngine/Private/Metasound.cpp
@@ -0,0 +1,99 @@
+// Copyright Epic Games, Inc. All Rights Reserved.
+#include "Metasound.h"
+
+#include "Backends/JsonStructSerializerBackend.h"
+#include "CoreMinimal.h"
+#include "HAL/FileManager.h"
+#include "StructSerializer.h"
+
+UMetasound::UMetasound(const FObjectInitializer& ObjectInitializer)
+	: Super(ObjectInitializer)
+	, AccessPoint(MakeShared<Metasound::Frontend::FDescriptionAccessPoint>(RootMetasoundDocument))
+{
+}
+
+FMetasoundClassMetadata UMetasound::GetMetadata()
+{
+	return RootMetasoundDocument.RootClass.Metadata;
+}
+
+// This can be used to update the metadata (name, author, etc) for this metasound.
+// @param InMetadata may be updated with any corrections we do to the input metadata.
+void UMetasound::SetMetadata(FMetasoundClassMetadata& InMetadata)
+{
+	InMetadata.NodeType = EMetasoundClassType::MetasoundGraph;
+	RootMetasoundDocument.RootClass.Metadata = InMetadata;
+	MarkPackageDirty();
+}
+
+// delete this asset's current metasound document,
+// and replace it with InClassDescription.
+void UMetasound::SetMetasoundDocument(const FMetasoundDocument& InDocument)
+{
+	RootMetasoundDocument = InDocument;
+}
+
+// returns a weak pointer that can be used to build a TDescriptionPtr
+// for direct editing of the FMetasoundClassDescription tree.
+// For advance use only, and requires knowledge of Metasound::Frontend::FDescPath syntax.
+// For most use cases, use GetGraphHandle() instead.
+TWeakPtr<Metasound::Frontend::FDescriptionAccessPoint> UMetasound::GetGraphAccessPoint()
+{
+	return TWeakPtr<Metasound::Frontend::FDescriptionAccessPoint>(AccessPoint);
+}
+
+// Get the handle for the root metasound graph of this asset.
+Metasound::Frontend::FGraphHandle UMetasound::GetRootGraphHandle()
+{
+	using FDescPath = Metasound::Frontend::FDescPath;
+	using FGraphHandle = Metasound::Frontend::FGraphHandle;
+	using FHandleInitParams = Metasound::Frontend::FHandleInitParams;
+	using EFromDocument = Metasound::Frontend::Path::EFromDocument;
+	using EFromClass = Metasound::Frontend::Path::EFromClass;
+
+	FDescPath PathToGraph = FDescPath()[EFromDocument::ToRootClass][EFromClass::ToGraph];
+	FHandleInitParams InitParams = { GetGraphAccessPoint(), PathToGraph, RootMetasoundDocument.RootClass.Metadata.NodeName, MakeWeakObjectPtr(this) };
+	return FGraphHandle(FHandleInitParams::PrivateToken, InitParams);
+}
+
+TArray<Metasound::Frontend::FGraphHandle> UMetasound::GetAllSubgraphHandles()
+{
+	using FDescPath = Metasound::Frontend::FDescPath;
+	using FGraphHandle = Metasound::Frontend::FGraphHandle;
+	using FHandleInitParams = Metasound::Frontend::FHandleInitParams;
+	using EFromDocument = Metasound::Frontend::Path::EFromDocument;
+	using EFromClass = Metasound::Frontend::Path::EFromClass;
+
+	TArray<FGraphHandle> OutArray;
+
+	FDescPath RootPathForDependencyGraphs = FDescPath()[EFromDocument::ToDependencies];
+
+	TArray<FMetasoundClassDescription>& DependenciesList = RootMetasoundDocument.Dependencies;
+	for (FMetasoundClassDescription& Dependency : DependenciesList)
+	{
+		const bool bIsSubgraph = Dependency.Metadata.NodeType == EMetasoundClassType::MetasoundGraph && Dependency.Graph.Nodes.Num() > 0;
+		if (bIsSubgraph)
+		{
+			FDescPath SubgraphPath = RootPathForDependencyGraphs[Dependency.UniqueID][EFromClass::ToGraph];
+			FHandleInitParams InitParams = { GetGraphAccessPoint(), SubgraphPath, Dependency.Metadata.NodeName, MakeWeakObjectPtr(this) };
+			OutArray.Emplace(FHandleInitParams::PrivateToken, InitParams);
+		}
+	}
+
+	return OutArray;
+}
+
+bool UMetasound::ExportToJSON(const FString& InAbsolutePath)
+{
+	if (TUniquePtr<FArchive> FileWriter = TUniquePtr<FArchive>(IFileManager::Get().CreateFileWriter(*InAbsolutePath)))
+	{
+		FJsonStructSerializerBackend Backend(*FileWriter, EStructSerializerBackendFlags::Default);
+		FStructSerializer::Serialize<FMetasoundDocument>(RootMetasoundDocument, Backend);
+		FileWriter->Close();
+
+		return true;
+	}
+
+	ensureAlwaysMsgf(false, TEXT("Failed to create a filewriter with the given path."));
+	return false;
+}
diff --git a/Engine/Plugins/Runtime/Metasound/Source/MetasoundEngine/Private/MetasoundEngineModule.cpp b/Engine/Plugins/Runtime/Metasound/Source/MetasoundEngine/Private/MetasoundEngineModule.cpp
--- a/Engine/Plugins/Runtime/Metasound/Source/MetasoundEngine/Private/MetasoundEngineModule.cpp
+++ b/Engine/Plugins/Runtime/Metasound/Source/MetasoundEngine/Private/MetasoundEngineModule.cpp
@@ -1,13 +1,31 @@
 // Copyright Epic Games, Inc. All Rights Reserved.
 #include "MetasoundEngineModule.h"
 
+#include "Metasound.h"
 #include "Modules/ModuleManager.h"
 
-DEFINE_LOG_CATEGORY(LogMetasoundEngine);
 
+DEFINE_LOG_CATEGORY(LogMetasoundEngine);
 
 class FMetasoundEngineModule : public IMetasoundEngineModule
 {
+	virtual UMetasound* DeserializeMetasound(const FString& InPath) override
+	{
+		FMetasoundDocument MetasoundDoc;
+		ensureAlwaysMsgf(false, TEXT("Implement the actual loading part!"));
+
+		UMetasound* NewMetasoundNode = NewObject<UMetasound>();
+		NewMetasoundNode->SetMetasoundDocument(MetasoundDoc);
+
+		return NewMetasoundNode;
+	}
+
+	virtual void SerializeMetasound(const UMetasound& InMetasound, const FString& InPath) override
+	{
+		ensureAlwaysMsgf(false, TEXT("Implement the actual saving part!"));
+		ensureAlwaysMsgf(false, TEXT("Decide whether to delete this graph from the asset that owns it!"));
+	}
+
 	virtual void StartupModule() override
 	{
 		UE_LOG(LogMetasoundEngine, Log, TEXT("Metasound Engine Initialized"));
diff --git a/Engine/Plugins/Runtime/Metasound/Source/MetasoundEngine/Private/MetasoundEngineModule.h b/Engine/Plugins/Runtime/Metasound/Source/MetasoundEngine/Private/MetasoundEngineModule.h
--- a/Engine/Plugins/Runtime/Metasound/Source/MetasoundEngine/Private/MetasoundEngineModule.h
+++ b/Engine/Plugins/Runtime/Metasound/Source/MetasoundEngine/Private/MetasoundEngineModule.h
@@ -5,8 +5,23 @@
 #include "Logging/LogMacros.h"
 #include "Modules/ModuleInterface.h"
 
+// Forward Declarations
+class UMetasound;
+
 DECLARE_LOG_CATEGORY_EXTERN(LogMetasoundEngine, Log, All);
 
-class METASOUNDENGINE_API IMetasoundEngineModule : public IModuleInterface
+
+class IMetasoundEngineModule : public IModuleInterface
 {
+	// Deserializes Metasound from the graph at the provided path.
+	// @param InPath Content directory path to load metasound from.
+	// @returns New Metasound object on success, nullptr on failure.
+	// @todo decide if this means that we remove the graph in our current asset.
+	//       This seems dangerous since we can fail to save the new asset and scratch the old one.
+	virtual UMetasound* DeserializeMetasound(const FString& InPath) = 0;
+
+	// Serializes Metasound to the provided path.
+	// @param InMetasound Metasound to serialize.
+	// @param InPath Path to serialize to.
+	virtual void SerializeMetasound(const UMetasound& InMetasound, const FString& InPath) = 0;
 };
diff --git a/Engine/Plugins/Runtime/Metasound/Source/MetasoundEngine/Public/Metasound.h b/Engine/Plugins/Runtime/Metasound/Source/MetasoundEngine/Public/Metasound.h
new file mode 100644
--- /dev/null
+++ b/Engine/Plugins/Runtime/Metasound/Source/MetasoundEngine/Public/Metasound.h
@@ -0,0 +1,52 @@
+// Copyright Epic Games, Inc. All Rights Reserved.
+
+#pragma once
+
+#include "CoreMinimal.h"
+#include "MetasoundFrontend.h"
+#include "MetasoundFrontendDataLayout.h"
+
+#include "Metasound.generated.h"
+
+
+/**
+ * This asset type is used for Metasound assets that can only be used as nodes in other Metasound graphs.
+ * Because of this, they can have any inputs or outputs they need.
+ */
+UCLASS(hidecategories = object, BlueprintType, MinimalAPI)
+class UMetasound : public UObject
+{
+	GENERATED_BODY()
+
+protected:
+	UPROPERTY()
+	FMetasoundDocument RootMetasoundDocument;
+
+	TSharedPtr<Metasound::Frontend::FDescriptionAccessPoint> AccessPoint;
+
+public:
+	UMetasound(const FObjectInitializer& ObjectInitializer);
+
+	FMetasoundClassMetadata GetMetadata();
+
+	// This can be used to update the metadata (name, author, etc) for this metasound.
+	// @param InMetadata may be updated with any corrections we do to the input metadata.
+	void SetMetadata(FMetasoundClassMetadata& InMetadata);
+
+	// delete this asset's current metasound document,
+	// and replace it with InClassDescription.
+	void SetMetasoundDocument(const FMetasoundDocument& InDocument);
+
+	// returns a weak pointer that can be used to build a TDescriptionPtr
+	// for direct editing of the FMetasoundClassDescription tree.
+	// For advance use only, and requires knowledge of Metasound::Frontend::FDescPath syntax.
+	// For most use cases, use GetGraphHandle() instead.
+	TWeakPtr<Metasound::Frontend::FDescriptionAccessPoint> GetGraphAccessPoint();
+
+	// Get the handle for the root metasound graph of this asset.
+	Metasound::Frontend::FGraphHandle GetRootGraphHandle();
+
+	TArray<Metasound::Frontend::FGraphHandle> GetAllSubgraphHandles();
+
+	bool ExportToJSON(const FString& InAbsolutePath);
+};
\ No newline at end of file
diff --git a/Engine/Plugins/Runtime/Metasound/Source/MetasoundFrontend/Private/MetasoundFrontend.cpp b/Engine/Plugins/Runtime/Metasound/Source/MetasoundFrontend/Private/MetasoundFrontend.cpp
--- a/Engine/Plugins/Runtime/Metasound/Source/MetasoundFrontend/Private/MetasoundFrontend.cpp
+++ b/Engine/Plugins/Runtime/Metasound/Source/MetasoundFrontend/Private/MetasoundFrontend.cpp
@@ -8,7 +8,6 @@
 #include "HAL/IConsoleManager.h"
 #include "Modules/ModuleManager.h"
 
-#include "MetasoundNodeAsset.h"
 #include "MetasoundAudioFormats.h"
 
 
@@ -1150,23 +1149,6 @@ namespace Metasound
 			return GraphsClassDeclaration->Metadata;
 		}
 
-		UMetasound* FGraphHandle::SaveToNewMetasoundAsset(const FString& InPath)
-		{
-			if (!IsValid())
-			{
-				return nullptr;
-			}
-
-			UMetasound* NewMetasoundNode = NewObject<UMetasound>();
-			NewMetasoundNode->SetMetasoundDocument(OwningDocument.GetChecked());
-
-			ensureAlwaysMsgf(false, TEXT("Implement the actual saving part!"));
-			ensureAlwaysMsgf(false, TEXT("Decide whether to delete this graph from the asset that owns it!"));
-
-			return NewMetasoundNode;
-		}
-
-
 		bool FGraphHandle::ExportToJSONAsset(const FString& InAbsolutePath)
 		{
 			if (!IsValid())
diff --git a/Engine/Plugins/Runtime/Metasound/Source/MetasoundFrontend/Private/MetasoundNodeAsset.cpp b/Engine/Plugins/Runtime/Metasound/Source/MetasoundFrontend/Private/MetasoundNodeAsset.cpp
deleted file mode 100644
diff --git a/Engine/Plugins/Runtime/Metasound/Source/MetasoundFrontend/Public/MetasoundFrontend.h b/Engine/Plugins/Runtime/Metasound/Source/MetasoundFrontend/Public/MetasoundFrontend.h
--- a/Engine/Plugins/Runtime/Metasound/Source/MetasoundFrontend/Public/MetasoundFrontend.h
+++ b/Engine/Plugins/Runtime/Metasound/Source/MetasoundFrontend/Public/MetasoundFrontend.h
@@ -226,7 +226,7 @@ namespace Metasound
 			uint32 NodeID;
 		};
 
-		class FGraphHandle : protected ITransactable
+		class METASOUNDFRONTEND_API FGraphHandle : protected ITransactable
 		{
 		public:
 			FGraphHandle() = delete;
@@ -261,13 +261,6 @@ namespace Metasound
 			// Returns the metadata for the current graph, including the name, description and author.
 			FMetasoundClassMetadata GetGraphMetadata();
 
-			// Creates a new UMetasound object with the current graph.
-			// @param InPath, if in editor, the content directory path to save this metasound to.
-			// @returns a new MetasoundNodeAsset object on success, nullptr on failure.
-			// @todo decide if this means that we remove the graph in our current asset.
-			//       This seems dangerous since we can fail to save the new asset and scratch the old one.
-			UMetasound* SaveToNewMetasoundAsset(const FString& InPath);
-
 			// Exports this graph to a JSON at the given path.
 			// @returns true on success.
 			bool ExportToJSONAsset(const FString& InAbsolutePath);
diff --git a/Engine/Plugins/Runtime/Metasound/Source/MetasoundFrontend/Public/MetasoundFrontendBaseClasses.h b/Engine/Plugins/Runtime/Metasound/Source/MetasoundFrontend/Public/MetasoundFrontendBaseClasses.h
--- a/Engine/Plugins/Runtime/Metasound/Source/MetasoundFrontend/Public/MetasoundFrontendBaseClasses.h
+++ b/Engine/Plugins/Runtime/Metasound/Source/MetasoundFrontend/Public/MetasoundFrontendBaseClasses.h
@@ -12,7 +12,7 @@ namespace Metasound
 	namespace Frontend
 	{
 		// Utility base class for classes that want to support naive multi-level undo/redo.
-		class ITransactable : public TSharedFromThis<ITransactable>
+		class METASOUNDFRONTEND_API ITransactable : public TSharedFromThis<ITransactable>
 		{
 		public:
 
@@ -463,14 +463,14 @@ namespace Metasound
 			FString GetPrintableString(FDescPath InPath);
 		}
 
-		typedef TArray<uint32> FClassDependencyIDs;
+		using FClassDependencyIDs = TArray<uint32>;
 
 		/**
 		 * Utility class that lives alongside a FMetasoundClassDescription to allow it's individual elements to be safely read and edited.
 		 * Note: Could be moved to a private header.
 		 *       It's also only necessary because we can't wrap a FMetasoundClassDescription in a shared ptr and also use it as a UPROPERTY.
 		 */
-		class FDescriptionAccessPoint
+		class METASOUNDFRONTEND_API FDescriptionAccessPoint
 		{
 		public:
 			FDescriptionAccessPoint() = delete;
diff --git a/Engine/Plugins/Runtime/Metasound/Source/MetasoundFrontend/Public/MetasoundNodeAsset.h b/Engine/Plugins/Runtime/Metasound/Source/MetasoundFrontend/Public/MetasoundNodeAsset.h
deleted file mode 100644
