Backout file submitted mistakenly from WIP CL
#rnx

diff --git a/Engine/Plugins/Runtime/Metasound/Metasound.uplugin b/Engine/Plugins/Runtime/Metasound/Metasound.uplugin
--- a/Engine/Plugins/Runtime/Metasound/Metasound.uplugin
+++ b/Engine/Plugins/Runtime/Metasound/Metasound.uplugin
@@ -35,11 +35,6 @@
 			"Name" : "MetasoundEngine",
 			"Type" : "Runtime",
 			"LoadingPhase" : "Default"
-		},
-		{
-			"Name": "MetasoundEditor",
-			"Type": "Editor",
-			"LoadingPhase": "PostEngineInit"
-		},
+		}
 	]
 }
