Boilerplate for empty Metasound Graph Implementation & editor module support

diff --git a/Engine/Plugins/Runtime/Metasound/Metasound.uplugin b/Engine/Plugins/Runtime/Metasound/Metasound.uplugin
--- a/Engine/Plugins/Runtime/Metasound/Metasound.uplugin
+++ b/Engine/Plugins/Runtime/Metasound/Metasound.uplugin
@@ -35,6 +35,11 @@
 			"Name" : "MetasoundEngine",
 			"Type" : "Runtime",
 			"LoadingPhase" : "Default"
-		}
+		},
+		{
+			"Name": "MetasoundEditor",
+			"Type": "Editor",
+			"LoadingPhase": "PostEngineInit"
+		},
 	]
 }
diff --git a/Engine/Plugins/Runtime/Metasound/Source/MetasoundEditor/MetasoundEditor.Build.cs b/Engine/Plugins/Runtime/Metasound/Source/MetasoundEditor/MetasoundEditor.Build.cs
new file mode 100644
--- /dev/null
+++ b/Engine/Plugins/Runtime/Metasound/Source/MetasoundEditor/MetasoundEditor.Build.cs
@@ -0,0 +1,51 @@
+// Copyright Epic Games, Inc. All Rights Reserved.
+
+using UnrealBuildTool;
+using System.IO;
+
+namespace UnrealBuildTool.Rules
+{
+	public class MetasoundEditor : ModuleRules
+	{
+		public MetasoundEditor(ReadOnlyTargetRules Target) : base(Target)
+		{
+			OptimizeCode = CodeOptimization.Never;
+
+			PrivateDependencyModuleNames.AddRange
+			(
+				new string[]
+				{
+					"AudioExtensions",
+					"AudioMixer",
+					"MetasoundEngine",
+					"MetasoundFrontend",
+					"MetasoundGraphCore",
+					"ToolMenus",
+				}
+			);
+
+			PublicDependencyModuleNames.AddRange
+			(
+				new string[]
+				{
+					"ApplicationCore",
+					"AudioMixer",
+					"ClassViewer",
+					"ContentBrowser",
+					"Core",
+					"CoreUObject",
+					"DetailCustomizations",
+					"EditorStyle",
+					"Engine",
+					"GraphEditor",
+					"InputCore",
+					"PropertyEditor",
+					"RenderCore",
+					"Slate",
+					"SlateCore",
+					"UnrealEd",
+				}
+			);
+		}
+	}
+}
diff --git a/Engine/Plugins/Runtime/Metasound/Source/MetasoundEditor/Private/MetasoundAssetTypeActions.cpp b/Engine/Plugins/Runtime/Metasound/Source/MetasoundEditor/Private/MetasoundAssetTypeActions.cpp
new file mode 100644
--- /dev/null
+++ b/Engine/Plugins/Runtime/Metasound/Source/MetasoundEditor/Private/MetasoundAssetTypeActions.cpp
@@ -0,0 +1,42 @@
+// Copyright Epic Games, Inc. All Rights Reserved.
+
+#include "MetasoundAssetTypeActions.h"
+
+#include "Metasound.h"
+#include "MetasoundEditor.h"
+
+
+#define LOCTEXT_NAMESPACE "AssetTypeActions"
+
+namespace MetasoundEditorUtils
+{
+	static const TArray<FText> AssetTypeActionSubMenu
+	{
+		LOCTEXT("AssetMetasoundSubMenu", "Metasound")
+	};
+} // namespace MetasoundEditorUtils
+
+UClass* FAssetTypeActions_Metasound::GetSupportedClass() const
+{
+	return UMetasound::StaticClass();
+}
+
+const TArray<FText>& FAssetTypeActions_Metasound::GetSubMenus() const
+{
+	return MetasoundEditorUtils::AssetTypeActionSubMenu;
+}
+
+void FAssetTypeActions_Metasound::OpenAssetEditor(const TArray<UObject*>& InObjects, TSharedPtr<IToolkitHost> ToolkitHost)
+{
+	const EToolkitMode::Type Mode = ToolkitHost.IsValid() ? EToolkitMode::WorldCentric : EToolkitMode::Standalone;
+
+	for (UObject* Object : InObjects)
+	{
+		if (UMetasound* Metasound = Cast<UMetasound>(Object))
+		{
+			TSharedRef<FMetasoundEditor> NewEditor = MakeShared<FMetasoundEditor>();
+			NewEditor->InitMetasoundEditor(Mode, ToolkitHost, Metasound);
+		}
+	}
+}
+#undef LOCTEXT_NAMESPACE
diff --git a/Engine/Plugins/Runtime/Metasound/Source/MetasoundEditor/Private/MetasoundAssetTypeActions.h b/Engine/Plugins/Runtime/Metasound/Source/MetasoundEditor/Private/MetasoundAssetTypeActions.h
new file mode 100644
--- /dev/null
+++ b/Engine/Plugins/Runtime/Metasound/Source/MetasoundEditor/Private/MetasoundAssetTypeActions.h
@@ -0,0 +1,23 @@
+// Copyright Epic Games, Inc. All Rights Reserved.
+#pragma once
+
+#include "CoreMinimal.h"
+#include "AssetTypeActions_Base.h"
+
+
+// Forward Declarations
+class IToolkitHost;
+
+
+class FAssetTypeActions_Metasound : public FAssetTypeActions_Base
+{
+public:
+	// IAssetTypeActions Implementation
+	virtual FText GetName() const override { return NSLOCTEXT("AssetTypeActions", "AssetTypeActions_Metasound", "Metasound"); }
+	virtual FColor GetTypeColor() const override { return FColor(220, 160, 0); }
+	virtual UClass* GetSupportedClass() const override;
+	virtual uint32 GetCategories() override { return EAssetTypeCategories::Sounds; }
+	virtual void OpenAssetEditor(const TArray<UObject*>& InObjects, TSharedPtr<IToolkitHost> ToolkitHost) override;
+
+	virtual const TArray<FText>& GetSubMenus() const override;
+};
diff --git a/Engine/Plugins/Runtime/Metasound/Source/MetasoundEditor/Private/MetasoundEditor.cpp b/Engine/Plugins/Runtime/Metasound/Source/MetasoundEditor/Private/MetasoundEditor.cpp
new file mode 100644
--- /dev/null
+++ b/Engine/Plugins/Runtime/Metasound/Source/MetasoundEditor/Private/MetasoundEditor.cpp
@@ -0,0 +1,900 @@
+// Copyright Epic Games, Inc. All Rights Reserved.
+
+#include "MetasoundEditor.h"
+
+#include "EdGraph/EdGraphNode.h"
+#include "EdGraphUtilities.h"
+#include "Editor.h"
+#include "EditorStyleSet.h"
+#include "Framework/Commands/GenericCommands.h"
+#include "Framework/MultiBox/MultiBoxBuilder.h"
+#include "GraphEditor.h"
+#include "GraphEditorActions.h"
+#include "HAL/PlatformApplicationMisc.h"
+#include "IDetailsView.h"
+#include "Kismet2/BlueprintEditorUtils.h"
+#include "Modules/ModuleManager.h"
+#include "MetasoundEditorCommands.h"
+#include "MetasoundEditorGraph.h"
+#include "MetasoundEditorGraphNode.h"
+#include "MetasoundEditorGraphSchema.h"
+#include "PropertyEditorModule.h"
+#include "ScopedTransaction.h"
+#include "SMetasoundPalette.h"
+#include "SNodePanel.h"
+#include "Widgets/Docking/SDockTab.h"
+#include "Components/AudioComponent.h"
+
+#define LOCTEXT_NAMESPACE "MetasoundEditor"
+
+const FName FMetasoundEditor::GraphCanvasTabId(TEXT("MetasoundEditor_GraphCanvas"));
+const FName FMetasoundEditor::PropertiesTabId(TEXT("MetasoundEditor_Properties"));
+const FName FMetasoundEditor::PaletteTabId(TEXT("MetasoundEditor_Palette"));
+
+FMetasoundEditor::FMetasoundEditor()
+	: Metasound(nullptr)
+{
+}
+
+void FMetasoundEditor::RegisterTabSpawners(const TSharedRef<FTabManager>& InTabManager)
+{
+	WorkspaceMenuCategory = InTabManager->AddLocalWorkspaceMenuCategory(LOCTEXT("WorkspaceMenu_MetasoundEditor", "Metasound Editor"));
+	auto WorkspaceMenuCategoryRef = WorkspaceMenuCategory.ToSharedRef();
+
+	FAssetEditorToolkit::RegisterTabSpawners(InTabManager);
+
+	InTabManager->RegisterTabSpawner(GraphCanvasTabId, FOnSpawnTab::CreateSP(this, &FMetasoundEditor::SpawnTab_GraphCanvas))
+		.SetDisplayName(LOCTEXT("GraphCanvasTab", "Viewport"))
+		.SetGroup(WorkspaceMenuCategoryRef)
+		.SetIcon(FSlateIcon(FEditorStyle::GetStyleSetName(), "GraphEditor.EventGraph_16x"));
+
+	InTabManager->RegisterTabSpawner(PropertiesTabId, FOnSpawnTab::CreateSP(this, &FMetasoundEditor::SpawnTab_Properties))
+		.SetDisplayName(LOCTEXT("DetailsTab", "Details"))
+		.SetGroup(WorkspaceMenuCategoryRef)
+		.SetIcon(FSlateIcon(FEditorStyle::GetStyleSetName(), "LevelEditor.Tabs.Details"));
+
+	InTabManager->RegisterTabSpawner(PaletteTabId, FOnSpawnTab::CreateSP(this, &FMetasoundEditor::SpawnTab_Palette))
+		.SetDisplayName(LOCTEXT("PaletteTab", "Palette"))
+		.SetGroup(WorkspaceMenuCategoryRef)
+		.SetIcon(FSlateIcon(FEditorStyle::GetStyleSetName(), "Kismet.Tabs.Palette"));
+}
+
+void FMetasoundEditor::UnregisterTabSpawners(const TSharedRef<FTabManager>& InTabManager)
+{
+	FAssetEditorToolkit::UnregisterTabSpawners(InTabManager);
+
+	InTabManager->UnregisterTabSpawner(GraphCanvasTabId);
+	InTabManager->UnregisterTabSpawner(PropertiesTabId);
+	InTabManager->UnregisterTabSpawner(PaletteTabId);
+}
+
+FMetasoundEditor::~FMetasoundEditor()
+{
+	// Stop any playing sounds when the editor closes
+	UAudioComponent* Component = GEditor->GetPreviewAudioComponent();
+	if (Component && Component->IsPlaying())
+	{
+		Stop();
+	}
+
+	check(GEditor);
+	GEditor->UnregisterForUndo(this);
+}
+
+void FMetasoundEditor::InitMetasoundEditor(const EToolkitMode::Type Mode, const TSharedPtr<IToolkitHost>& InitToolkitHost, UObject* ObjectToEdit)
+{
+	Metasound = CastChecked<UMetasound>(ObjectToEdit);
+
+	// Support undo/redo
+	Metasound->SetFlags(RF_Transactional);
+	
+	if (!Metasound->GetGraph())
+	{
+		UMetasoundEditorGraph* Graph = NewObject<UMetasoundEditorGraph>();
+		Graph->Metasound = Metasound;
+		Graph->Schema = UMetasoundEditorGraphSchema::StaticClass();
+		Metasound->SetGraph(Graph);
+	}
+
+	GEditor->RegisterForUndo(this);
+
+	FGraphEditorCommands::Register();
+	FMetasoundEditorCommands::Register();
+
+	BindGraphCommands();
+
+	CreateInternalWidgets();
+
+	const TSharedRef<FTabManager::FLayout> StandaloneDefaultLayout = FTabManager::NewLayout("Standalone_MetasoundEditor_Layout_v1")
+	->AddArea
+	(
+		FTabManager::NewPrimaryArea()
+		->SetOrientation(Orient_Vertical)
+		->Split
+		(
+			FTabManager::NewStack()
+			->SetSizeCoefficient(0.1f)
+			->SetHideTabWell(true)
+			->AddTab(GetToolbarTabId(), ETabState::OpenedTab)
+		)
+		->Split(FTabManager::NewSplitter()
+			->SetOrientation(Orient_Horizontal)
+			->SetSizeCoefficient(0.9f)
+			->Split
+			(
+				FTabManager::NewStack()
+				->SetSizeCoefficient(0.225f)
+				->SetHideTabWell(true)
+				->AddTab(PropertiesTabId, ETabState::OpenedTab)
+			)
+			->Split
+			(
+				FTabManager::NewStack()
+				->SetSizeCoefficient(0.65f)
+				->SetHideTabWell(true)
+				->AddTab(GraphCanvasTabId, ETabState::OpenedTab)
+			)
+			->Split
+			(
+				FTabManager::NewStack()
+				->SetSizeCoefficient(0.125f)
+				->SetHideTabWell(true)
+				->AddTab(PaletteTabId, ETabState::OpenedTab)
+			)
+		)
+	);
+
+	const bool bCreateDefaultStandaloneMenu = true;
+	const bool bCreateDefaultToolbar = true;
+	FAssetEditorToolkit::InitAssetEditor(Mode, InitToolkitHost, TEXT("MetasoundEditorApp"), StandaloneDefaultLayout, bCreateDefaultStandaloneMenu, bCreateDefaultToolbar, ObjectToEdit, false);
+
+	ExtendToolbar();
+	RegenerateMenusAndToolbars();
+}
+
+UMetasound* FMetasoundEditor::GetMetasound() const
+{
+	return Metasound;
+}
+
+void FMetasoundEditor::SetSelection(TArray<UObject*> SelectedObjects)
+{
+	if (MetasoundProperties.IsValid())
+	{
+		MetasoundProperties->SetObjects(SelectedObjects);
+	}
+}
+
+bool FMetasoundEditor::GetBoundsForSelectedNodes(FSlateRect& Rect, float Padding)
+{
+	return MetasoundGraphEditor->GetBoundsForSelectedNodes(Rect, Padding);
+}
+
+int32 FMetasoundEditor::GetNumberOfSelectedNodes() const
+{
+	return MetasoundGraphEditor->GetSelectedNodes().Num();
+}
+
+FName FMetasoundEditor::GetToolkitFName() const
+{
+	return FName("MetasoundEditor");
+}
+
+FText FMetasoundEditor::GetBaseToolkitName() const
+{
+	return LOCTEXT("AppLabel", "Metasound Editor");
+}
+
+FString FMetasoundEditor::GetWorldCentricTabPrefix() const
+{
+	return LOCTEXT("WorldCentricTabPrefix", "Metasound ").ToString();
+}
+
+FLinearColor FMetasoundEditor::GetWorldCentricTabColorScale() const
+{
+	return FLinearColor(0.3f, 0.2f, 0.5f, 0.5f);
+}
+
+TSharedRef<SDockTab> FMetasoundEditor::SpawnTab_GraphCanvas(const FSpawnTabArgs& Args)
+{
+	check(Args.GetTabId() == GraphCanvasTabId);
+
+	TSharedRef<SDockTab> SpawnedTab = SNew(SDockTab)
+		.Label(LOCTEXT("MetasoundGraphCanvasTitle", "Viewport"));
+
+	if (MetasoundGraphEditor.IsValid())
+	{
+		SpawnedTab->SetContent(MetasoundGraphEditor.ToSharedRef());
+	}
+
+	return SpawnedTab;
+}
+
+TSharedRef<SDockTab> FMetasoundEditor::SpawnTab_Properties(const FSpawnTabArgs& Args)
+{
+	check(Args.GetTabId() == PropertiesTabId);
+
+	return SNew(SDockTab)
+		.Icon(FEditorStyle::GetBrush("LevelEditor.Tabs.Details"))
+		.Label(LOCTEXT("MetasoundDetailsTitle", "Details"))
+		[
+			MetasoundProperties.ToSharedRef()
+		];
+}
+
+TSharedRef<SDockTab> FMetasoundEditor::SpawnTab_Palette(const FSpawnTabArgs& Args)
+{
+	check(Args.GetTabId() == PaletteTabId);
+
+	return SNew(SDockTab)
+		.Icon(FEditorStyle::GetBrush("Kismet.Tabs.Palette"))
+		.Label(LOCTEXT("MetasoundPaletteTitle", "Palette"))
+		[
+			Palette.ToSharedRef()
+		];
+}
+
+void FMetasoundEditor::AddReferencedObjects(FReferenceCollector& Collector)
+{
+	Collector.AddReferencedObject(Metasound);
+}
+
+void FMetasoundEditor::PostUndo(bool bSuccess)
+{
+	if (MetasoundGraphEditor.IsValid())
+	{
+		MetasoundGraphEditor->ClearSelectionSet();
+		MetasoundGraphEditor->NotifyGraphChanged();
+		FSlateApplication::Get().DismissAllMenus();
+	}
+
+}
+
+void FMetasoundEditor::NotifyPostChange(const FPropertyChangedEvent& PropertyChangedEvent, FProperty* PropertyThatChanged)
+{
+	if (MetasoundGraphEditor.IsValid() && PropertyChangedEvent.ChangeType != EPropertyChangeType::Interactive)
+	{
+		MetasoundGraphEditor->NotifyGraphChanged();
+	}
+}
+
+void FMetasoundEditor::CreateInternalWidgets()
+{
+	MetasoundGraphEditor = CreateGraphEditorWidget();
+
+	FDetailsViewArgs Args;
+	Args.bHideSelectionTip = true;
+	Args.NotifyHook = this;
+
+	FPropertyEditorModule& PropertyModule = FModuleManager::LoadModuleChecked<FPropertyEditorModule>("PropertyEditor");
+	MetasoundProperties = PropertyModule.CreateDetailView(Args);
+	MetasoundProperties->SetObject(Metasound);
+
+	Palette = SNew(SMetasoundPalette);
+}
+
+void FMetasoundEditor::ExtendToolbar()
+{
+	TSharedPtr<FExtender> ToolbarExtender = MakeShared<FExtender>();
+	ToolbarExtender->AddToolBarExtension
+	(
+		"Asset",
+		EExtensionHook::After,
+		GetToolkitCommands(),
+		FToolBarExtensionDelegate::CreateLambda([](FToolBarBuilder& ToolbarBuilder)
+		{
+			ToolbarBuilder.BeginSection("Audition");
+			{
+				ToolbarBuilder.AddToolBarButton(FMetasoundEditorCommands::Get().Play);
+				ToolbarBuilder.AddToolBarButton(FMetasoundEditorCommands::Get().Stop);
+			}
+			ToolbarBuilder.EndSection();
+		})
+	);
+
+	AddToolbarExtender(ToolbarExtender);
+}
+
+void FMetasoundEditor::BindGraphCommands()
+{
+	const FMetasoundEditorCommands& Commands = FMetasoundEditorCommands::Get();
+
+	ToolkitCommands->MapAction(
+		Commands.Play,
+		FExecuteAction::CreateSP(this, &FMetasoundEditor::Play));
+
+	ToolkitCommands->MapAction(
+		Commands.Stop,
+		FExecuteAction::CreateSP(this, &FMetasoundEditor::Stop));
+
+	ToolkitCommands->MapAction(
+		Commands.TogglePlayback,
+		FExecuteAction::CreateSP(this, &FMetasoundEditor::TogglePlayback));
+
+	ToolkitCommands->MapAction(
+		FGenericCommands::Get().Undo,
+		FExecuteAction::CreateSP(this, &FMetasoundEditor::UndoGraphAction));
+
+	ToolkitCommands->MapAction(
+		FGenericCommands::Get().Redo,
+		FExecuteAction::CreateSP(this, &FMetasoundEditor::RedoGraphAction));
+}
+
+void FMetasoundEditor::Play()
+{
+// TODO: Implement play
+
+// 	GEditor->PlayPreviewSound(Metasound);
+
+// 	MetasoundGraphEditor->RegisterActiveTimer(0.0f, 
+// 		FWidgetActiveTimerDelegate::CreateLambda([](double InCurrentTime, float InDeltaTime)
+// 		{
+// 			UAudioComponent* PreviewComp = GEditor->GetPreviewAudioComponent();
+// 			if (PreviewComp && PreviewComp->IsPlaying())
+// 			{
+// 				return EActiveTimerReturnType::Continue;
+// 			}
+// 			else
+// 			{
+// 				return EActiveTimerReturnType::Stop;
+// 			}
+// 		})
+// 	);
+}
+
+void FMetasoundEditor::PlayNode()
+{
+	// already checked that only one node is selected
+	const FGraphPanelSelectionSet SelectedNodes = GetSelectedNodes();
+
+	for (FGraphPanelSelectionSet::TConstIterator NodeIt(SelectedNodes); NodeIt; ++NodeIt)
+	{
+		PlaySingleNode(CastChecked<UEdGraphNode>(*NodeIt));
+	}
+}
+
+bool FMetasoundEditor::CanPlayNode() const
+{
+	return false;
+}
+
+void FMetasoundEditor::Stop()
+{
+	check(GEditor);
+	GEditor->ResetPreviewAudioComponent();
+}
+
+void FMetasoundEditor::TogglePlayback()
+{
+	check(GEditor);
+
+	UAudioComponent* Component = GEditor->GetPreviewAudioComponent();
+	if (Component && Component->IsPlaying())
+	{
+		Stop();
+	}
+	else
+	{
+		Play();
+	}
+}
+
+void FMetasoundEditor::PlaySingleNode(UEdGraphNode* Node)
+{
+	// TODO: Implement? Will we support single node playback?
+}
+
+void FMetasoundEditor::SyncInBrowser()
+{
+	TArray<UObject*> ObjectsToSync;
+
+	const FGraphPanelSelectionSet SelectedNodes = GetSelectedNodes();
+	for (FGraphPanelSelectionSet::TConstIterator NodeIt(SelectedNodes); NodeIt; ++NodeIt)
+	{
+		// TODO: Implement sync to referenced Metasound if selected node is a reference to another metasound
+	}
+
+	if (!ObjectsToSync.Num())
+	{
+		ObjectsToSync.Add(Metasound);
+	}
+
+	check(GEditor);
+	GEditor->SyncBrowserToObjects(ObjectsToSync);
+}
+
+void FMetasoundEditor::AddInput()
+{
+}
+
+bool FMetasoundEditor::CanAddInput() const
+{
+	return GetSelectedNodes().Num() == 1;
+}
+
+void FMetasoundEditor::DeleteInput()
+{
+}
+
+bool FMetasoundEditor::CanDeleteInput() const
+{
+	return true;
+}
+
+void FMetasoundEditor::OnCreateComment()
+{
+	Metasound->MarkPackageDirty();
+}
+
+TSharedRef<SGraphEditor> FMetasoundEditor::CreateGraphEditorWidget()
+{
+	if (!GraphEditorCommands.IsValid())
+	{
+		GraphEditorCommands = MakeShared<FUICommandList>();
+
+		GraphEditorCommands->MapAction(FMetasoundEditorCommands::Get().BrowserSync,
+			FExecuteAction::CreateSP(this, &FMetasoundEditor::SyncInBrowser));
+
+		GraphEditorCommands->MapAction(FMetasoundEditorCommands::Get().AddInput,
+			FExecuteAction::CreateSP(this, &FMetasoundEditor::AddInput),
+			FCanExecuteAction::CreateSP(this, &FMetasoundEditor::CanAddInput));
+
+		GraphEditorCommands->MapAction(FMetasoundEditorCommands::Get().DeleteInput,
+			FExecuteAction::CreateSP(this, &FMetasoundEditor::DeleteInput),
+			FCanExecuteAction::CreateSP(this, &FMetasoundEditor::CanDeleteInput));
+
+		// Graph Editor Commands
+		GraphEditorCommands->MapAction(FGraphEditorCommands::Get().CreateComment,
+			FExecuteAction::CreateSP(this, &FMetasoundEditor::OnCreateComment));
+
+		// Editing commands
+		GraphEditorCommands->MapAction(FGenericCommands::Get().SelectAll,
+			FExecuteAction::CreateSP(this, &FMetasoundEditor::SelectAllNodes),
+			FCanExecuteAction::CreateSP(this, &FMetasoundEditor::CanSelectAllNodes));
+
+		GraphEditorCommands->MapAction(FGenericCommands::Get().Delete,
+			FExecuteAction::CreateSP(this, &FMetasoundEditor::DeleteSelectedNodes),
+			FCanExecuteAction::CreateSP(this, &FMetasoundEditor::CanDeleteNodes));
+
+		GraphEditorCommands->MapAction(FGenericCommands::Get().Copy,
+			FExecuteAction::CreateSP(this, &FMetasoundEditor::CopySelectedNodes),
+			FCanExecuteAction::CreateSP(this, &FMetasoundEditor::CanCopyNodes));
+
+		GraphEditorCommands->MapAction(FGenericCommands::Get().Cut,
+			FExecuteAction::CreateSP(this, &FMetasoundEditor::CutSelectedNodes),
+			FCanExecuteAction::CreateSP(this, &FMetasoundEditor::CanCutNodes));
+
+		GraphEditorCommands->MapAction(FGenericCommands::Get().Paste,
+			FExecuteAction::CreateSP(this, &FMetasoundEditor::PasteNodes),
+			FCanExecuteAction::CreateSP(this, &FMetasoundEditor::CanPasteNodes));
+
+		GraphEditorCommands->MapAction(FGenericCommands::Get().Duplicate,
+			FExecuteAction::CreateSP(this, &FMetasoundEditor::DuplicateNodes),
+			FCanExecuteAction::CreateSP(this, &FMetasoundEditor::CanDuplicateNodes));
+
+		// Alignment Commands
+		GraphEditorCommands->MapAction(FGraphEditorCommands::Get().AlignNodesTop,
+			FExecuteAction::CreateSP(this, &FMetasoundEditor::OnAlignTop));
+
+		GraphEditorCommands->MapAction(FGraphEditorCommands::Get().AlignNodesMiddle,
+			FExecuteAction::CreateSP(this, &FMetasoundEditor::OnAlignMiddle));
+
+		GraphEditorCommands->MapAction(FGraphEditorCommands::Get().AlignNodesBottom,
+			FExecuteAction::CreateSP(this, &FMetasoundEditor::OnAlignBottom));
+
+		GraphEditorCommands->MapAction(FGraphEditorCommands::Get().AlignNodesLeft,
+			FExecuteAction::CreateSP(this, &FMetasoundEditor::OnAlignLeft));
+
+		GraphEditorCommands->MapAction(FGraphEditorCommands::Get().AlignNodesCenter,
+			FExecuteAction::CreateSP(this, &FMetasoundEditor::OnAlignCenter));
+
+		GraphEditorCommands->MapAction(FGraphEditorCommands::Get().AlignNodesRight,
+			FExecuteAction::CreateSP(this, &FMetasoundEditor::OnAlignRight));
+
+		GraphEditorCommands->MapAction(FGraphEditorCommands::Get().StraightenConnections,
+			FExecuteAction::CreateSP(this, &FMetasoundEditor::OnStraightenConnections));
+
+		// Distribution Commands
+		GraphEditorCommands->MapAction(FGraphEditorCommands::Get().DistributeNodesHorizontally,
+			FExecuteAction::CreateSP(this, &FMetasoundEditor::OnDistributeNodesH));
+
+		GraphEditorCommands->MapAction(FGraphEditorCommands::Get().DistributeNodesVertically,
+			FExecuteAction::CreateSP(this, &FMetasoundEditor::OnDistributeNodesV));
+	}
+
+	FGraphAppearanceInfo AppearanceInfo;
+	AppearanceInfo.CornerText = LOCTEXT("AppearanceCornerText_Metasound", "Metasound");
+
+	SGraphEditor::FGraphEditorEvents InEvents;
+	InEvents.OnSelectionChanged = SGraphEditor::FOnSelectionChanged::CreateSP(this, &FMetasoundEditor::OnSelectedNodesChanged);
+	InEvents.OnTextCommitted = FOnNodeTextCommitted::CreateSP(this, &FMetasoundEditor::OnNodeTitleCommitted);
+	InEvents.OnNodeDoubleClicked = FSingleNodeEvent::CreateSP(this, &FMetasoundEditor::PlaySingleNode);
+
+	return SNew(SGraphEditor)
+		.AdditionalCommands(GraphEditorCommands)
+		.IsEditable(true)
+		.Appearance(AppearanceInfo)
+		.GraphToEdit(Metasound->GetGraph())
+		.GraphEvents(InEvents)
+		.AutoExpandActionMenu(true)
+		.ShowGraphStateOverlay(false);
+}
+
+FGraphPanelSelectionSet FMetasoundEditor::GetSelectedNodes() const
+{
+	FGraphPanelSelectionSet CurrentSelection;
+	if (MetasoundGraphEditor.IsValid())
+	{
+		CurrentSelection = MetasoundGraphEditor->GetSelectedNodes();
+	}
+	return CurrentSelection;
+}
+
+void FMetasoundEditor::OnSelectedNodesChanged(const TSet<UObject*>& NewSelection)
+{
+	TArray<UObject*> Selection;
+
+	if(NewSelection.Num())
+	{
+		for(TSet<UObject*>::TConstIterator SetIt(NewSelection); SetIt; ++SetIt)
+		{
+			if (Cast<UMetasoundEditorGraphNode>(*SetIt))
+			{
+				Selection.Add(GetMetasound());
+			}
+			else
+			{
+				Selection.Add(*SetIt);
+			}
+		}
+	}
+	else
+	{
+		Selection.Add(GetMetasound());
+	}
+
+	SetSelection(Selection);
+}
+
+void FMetasoundEditor::OnNodeTitleCommitted(const FText& NewText, ETextCommit::Type CommitInfo, UEdGraphNode* NodeBeingChanged)
+{
+	if (NodeBeingChanged)
+	{
+		const FScopedTransaction Transaction(LOCTEXT("RenameNode", "Rename Node"));
+		NodeBeingChanged->Modify();
+		NodeBeingChanged->OnRenameNode(NewText.ToString());
+	}
+}
+
+void FMetasoundEditor::SelectAllNodes()
+{
+	MetasoundGraphEditor->SelectAllNodes();
+}
+
+bool FMetasoundEditor::CanSelectAllNodes() const
+{
+	return true;
+}
+
+void FMetasoundEditor::DeleteSelectedNodes()
+{
+	const FScopedTransaction Transaction(NSLOCTEXT("UnrealEd", "MetasoundEditorDeleteSelectedNode", "Delete Selected Metasound Node"));
+
+	MetasoundGraphEditor->GetCurrentGraph()->Modify();
+
+	const FGraphPanelSelectionSet SelectedNodes = GetSelectedNodes();
+
+	MetasoundGraphEditor->ClearSelectionSet();
+
+	for (FGraphPanelSelectionSet::TConstIterator NodeIt(SelectedNodes); NodeIt; ++NodeIt)
+	{
+		UEdGraphNode* Node = CastChecked<UEdGraphNode>(*NodeIt);
+		if (Node->CanUserDeleteNode())
+		{
+			if (UMetasoundEditorGraphNode* MetasoundGraphNode = Cast<UMetasoundEditorGraphNode>(Node))
+			{
+				// TODO: Implement delete selected here
+			}
+		}
+	}
+}
+
+bool FMetasoundEditor::CanDeleteNodes() const
+{
+	const FGraphPanelSelectionSet SelectedNodes = GetSelectedNodes();
+
+	if (SelectedNodes.Num() == 1)
+	{
+		for (FGraphPanelSelectionSet::TConstIterator NodeIt(SelectedNodes); NodeIt; ++NodeIt)
+		{
+			if (Cast<UMetasoundEditorGraphNode>(*NodeIt))
+			{
+				// Return false if only root node is selected, as it can't be deleted
+				return false;
+			}
+		}
+	}
+
+	return SelectedNodes.Num() > 0;
+}
+
+void FMetasoundEditor::DeleteSelectedDuplicatableNodes()
+{
+	// Cache off the old selection
+	const FGraphPanelSelectionSet OldSelectedNodes = GetSelectedNodes();
+
+	// Clear the selection and only select the nodes that can be duplicated
+	FGraphPanelSelectionSet RemainingNodes;
+	MetasoundGraphEditor->ClearSelectionSet();
+
+	for (FGraphPanelSelectionSet::TConstIterator SelectedIter(OldSelectedNodes); SelectedIter; ++SelectedIter)
+	{
+		UEdGraphNode* Node = Cast<UEdGraphNode>(*SelectedIter);
+		if ((Node != nullptr) && Node->CanDuplicateNode())
+		{
+			MetasoundGraphEditor->SetNodeSelection(Node, true);
+		}
+		else
+		{
+			RemainingNodes.Add(Node);
+		}
+	}
+
+	// Delete the duplicatable nodes
+	DeleteSelectedNodes();
+
+	// Reselect whatever's left from the original selection after the deletion
+	MetasoundGraphEditor->ClearSelectionSet();
+
+	for (FGraphPanelSelectionSet::TConstIterator SelectedIter(RemainingNodes); SelectedIter; ++SelectedIter)
+	{
+		if (UEdGraphNode* Node = Cast<UEdGraphNode>(*SelectedIter))
+		{
+			MetasoundGraphEditor->SetNodeSelection(Node, true);
+		}
+	}
+}
+
+void FMetasoundEditor::CutSelectedNodes()
+{
+	CopySelectedNodes();
+	// Cut should only delete nodes that can be duplicated
+	DeleteSelectedDuplicatableNodes();
+}
+
+bool FMetasoundEditor::CanCutNodes() const
+{
+	return CanCopyNodes() && CanDeleteNodes();
+}
+
+void FMetasoundEditor::CopySelectedNodes()
+{
+	// Export the selected nodes and place the text on the clipboard
+	const FGraphPanelSelectionSet SelectedNodes = GetSelectedNodes();
+
+	FString ExportedText;
+
+	for (FGraphPanelSelectionSet::TConstIterator SelectedIter(SelectedNodes); SelectedIter; ++SelectedIter)
+	{
+		if(UMetasoundEditorGraphNode* Node = Cast<UMetasoundEditorGraphNode>(*SelectedIter))
+		{
+			Node->PrepareForCopying();
+		}
+	}
+
+	FEdGraphUtilities::ExportNodesToText(SelectedNodes, /*out*/ ExportedText);
+	FPlatformApplicationMisc::ClipboardCopy(*ExportedText);
+
+	// Make sure Metasound remains the owner of the copied nodes
+	for (FGraphPanelSelectionSet::TConstIterator SelectedIter(SelectedNodes); SelectedIter; ++SelectedIter)
+	{
+		if (UMetasoundEditorGraphNode* Node = Cast<UMetasoundEditorGraphNode>(*SelectedIter))
+		{
+			Node->PostCopyNode();
+		}
+	}
+}
+
+bool FMetasoundEditor::CanCopyNodes() const
+{
+	// If any of the nodes can be duplicated then we should allow copying
+	const FGraphPanelSelectionSet SelectedNodes = GetSelectedNodes();
+	for (FGraphPanelSelectionSet::TConstIterator SelectedIter(SelectedNodes); SelectedIter; ++SelectedIter)
+	{
+		UEdGraphNode* Node = Cast<UEdGraphNode>(*SelectedIter);
+		if ((Node != nullptr) && Node->CanDuplicateNode())
+		{
+			return true;
+		}
+	}
+	return false;
+}
+
+void FMetasoundEditor::PasteNodes()
+{
+	PasteNodesAtLocation(MetasoundGraphEditor->GetPasteLocation());
+}
+
+void FMetasoundEditor::PasteNodesAtLocation(const FVector2D& Location)
+{
+	UEdGraph* Graph = Metasound->GetGraph();
+	if (!Graph)
+	{
+		return;
+	}
+
+	// Undo/Redo support
+	const FScopedTransaction Transaction(NSLOCTEXT("UnrealEd", "MetasoundEditorPaste", "Paste Metasound Node"));
+	Graph->Modify();
+	Metasound->Modify();
+
+	// Clear the selection set (newly pasted stuff will be selected)
+	MetasoundGraphEditor->ClearSelectionSet();
+
+	// Grab the text to paste from the clipboard.
+	FString TextToImport;
+	FPlatformApplicationMisc::ClipboardPaste(TextToImport);
+
+	// Import the nodes
+	TSet<UEdGraphNode*> PastedNodes;
+	FEdGraphUtilities::ImportNodesFromText(Graph, TextToImport, /*out*/ PastedNodes);
+
+	//Average position of nodes so we can move them while still maintaining relative distances to each other
+	FVector2D AvgNodePosition(0.0f,0.0f);
+
+	for (TSet<UEdGraphNode*>::TIterator It(PastedNodes); It; ++It)
+	{
+		UEdGraphNode* Node = *It;
+		AvgNodePosition.X += Node->NodePosX;
+		AvgNodePosition.Y += Node->NodePosY;
+	}
+
+	if (PastedNodes.Num() > 0)
+	{
+		float InvNumNodes = 1.0f/float(PastedNodes.Num());
+		AvgNodePosition.X *= InvNumNodes;
+		AvgNodePosition.Y *= InvNumNodes;
+	}
+
+	for (TSet<UEdGraphNode*>::TIterator It(PastedNodes); It; ++It)
+	{
+		UEdGraphNode* Node = *It;
+
+		if (UMetasoundEditorGraphNode* SoundGraphNode = Cast<UMetasoundEditorGraphNode>(Node))
+		{
+			// TODO: Add newly referenced nodes to MS reference list
+		}
+
+		// Select the newly pasted stuff
+		MetasoundGraphEditor->SetNodeSelection(Node, true);
+
+		Node->NodePosX = (Node->NodePosX - AvgNodePosition.X) + Location.X ;
+		Node->NodePosY = (Node->NodePosY - AvgNodePosition.Y) + Location.Y ;
+
+		Node->SnapToGrid(SNodePanel::GetSnapGridSize());
+
+		// Give new node a different Guid from the old one
+		Node->CreateNewGuid();
+	}
+
+	// TODO: Compile metasound here???
+
+	// Update UI
+	MetasoundGraphEditor->NotifyGraphChanged();
+
+	Metasound->PostEditChange();
+	Metasound->MarkPackageDirty();
+}
+
+bool FMetasoundEditor::CanPasteNodes() const
+{
+	FString ClipboardContent;
+	FPlatformApplicationMisc::ClipboardPaste(ClipboardContent);
+
+	// TODO: Implement
+// 	const bool bCanPasteNodes = FEdGraphUtilities::CanImportNodesFromText(Metasound->MetasoundGraph, ClipboardContent);
+// 	return bCanPasteNodes;
+
+	return false;
+}
+
+void FMetasoundEditor::DuplicateNodes()
+{
+	// Copy and paste current selection
+	CopySelectedNodes();
+	PasteNodes();
+}
+
+bool FMetasoundEditor::CanDuplicateNodes() const
+{
+	return CanCopyNodes();
+}
+
+void FMetasoundEditor::UndoGraphAction()
+{
+	check(GEditor);
+	GEditor->UndoTransaction();
+}
+
+void FMetasoundEditor::RedoGraphAction()
+{
+	// Clear selection, to avoid holding refs to nodes that go away
+	MetasoundGraphEditor->ClearSelectionSet();
+
+	check(GEditor);
+	GEditor->RedoTransaction();
+}
+
+void FMetasoundEditor::OnAlignTop()
+{
+	if (MetasoundGraphEditor.IsValid())
+	{
+		MetasoundGraphEditor->OnAlignTop();
+	}
+}
+
+void FMetasoundEditor::OnAlignMiddle()
+{
+	if (MetasoundGraphEditor.IsValid())
+	{
+		MetasoundGraphEditor->OnAlignMiddle();
+	}
+}
+
+void FMetasoundEditor::OnAlignBottom()
+{
+	if (MetasoundGraphEditor.IsValid())
+	{
+		MetasoundGraphEditor->OnAlignBottom();
+	}
+}
+
+void FMetasoundEditor::OnAlignLeft()
+{
+	if (MetasoundGraphEditor.IsValid())
+	{
+		MetasoundGraphEditor->OnAlignLeft();
+	}
+}
+
+void FMetasoundEditor::OnAlignCenter()
+{
+	if (MetasoundGraphEditor.IsValid())
+	{
+		MetasoundGraphEditor->OnAlignCenter();
+	}
+}
+
+void FMetasoundEditor::OnAlignRight()
+{
+	if (MetasoundGraphEditor.IsValid())
+	{
+		MetasoundGraphEditor->OnAlignRight();
+	}
+}
+
+void FMetasoundEditor::OnStraightenConnections()
+{
+	if (MetasoundGraphEditor.IsValid())
+	{
+		MetasoundGraphEditor->OnStraightenConnections();
+	}
+}
+
+void FMetasoundEditor::OnDistributeNodesH()
+{
+	if (MetasoundGraphEditor.IsValid())
+	{
+		MetasoundGraphEditor->OnDistributeNodesH();
+	}
+}
+
+void FMetasoundEditor::OnDistributeNodesV()
+{
+	if (MetasoundGraphEditor.IsValid())
+	{
+		MetasoundGraphEditor->OnDistributeNodesV();
+	}
+}
+
+#undef LOCTEXT_NAMESPACE
diff --git a/Engine/Plugins/Runtime/Metasound/Source/MetasoundEditor/Private/MetasoundEditor.h b/Engine/Plugins/Runtime/Metasound/Source/MetasoundEditor/Private/MetasoundEditor.h
new file mode 100644
--- /dev/null
+++ b/Engine/Plugins/Runtime/Metasound/Source/MetasoundEditor/Private/MetasoundEditor.h
@@ -0,0 +1,214 @@
+// Copyright Epic Games, Inc. All Rights Reserved.
+
+#pragma once
+
+#include "CoreMinimal.h"
+#include "EditorUndoClient.h"
+#include "IMetasoundEditor.h"
+#include "Misc/NotifyHook.h"
+#include "Toolkits/IToolkitHost.h"
+#include "UObject/GCObject.h"
+
+
+// Forward Declarations
+class FSlateRect;
+class FTabManager;
+class IDetailsView;
+class IToolkitHost;
+class SDockableTab;
+class SGraphEditor;
+class SMetasoundPalette;
+class UEdGraphNode;
+class UMetasound;
+
+struct FPropertyChangedEvent;
+
+
+class FMetasoundEditor : public IMetasoundEditor, public FGCObject, public FNotifyHook, public FEditorUndoClient
+{
+public:
+	FMetasoundEditor();
+
+	virtual void RegisterTabSpawners(const TSharedRef<FTabManager>& TabManager) override;
+	virtual void UnregisterTabSpawners(const TSharedRef<FTabManager>& TabManager) override;
+
+	/** Destructor */
+	virtual ~FMetasoundEditor();
+
+	/** Edits the specified Metasound object */
+	void InitMetasoundEditor(const EToolkitMode::Type Mode, const TSharedPtr<IToolkitHost >& InitToolkitHost, UObject* ObjectToEdit);
+
+	/** IMetasoundEditor interface */
+	virtual UMetasound* GetMetasound() const override;
+	virtual void SetSelection(TArray<UObject*> SelectedObjects) override;
+	virtual bool GetBoundsForSelectedNodes(FSlateRect& Rect, float Padding) override;
+	virtual int32 GetNumberOfSelectedNodes() const override;
+	virtual TSet<UObject*> GetSelectedNodes() const override;
+
+	/** IToolkit interface */
+	virtual FName GetToolkitFName() const override;
+	virtual FText GetBaseToolkitName() const override;
+	virtual FString GetWorldCentricTabPrefix() const override;
+	virtual FLinearColor GetWorldCentricTabColorScale() const override;
+
+	virtual FString GetDocumentationLink() const override
+	{
+		return FString(TEXT("Engine/Audio/Metasounds/Editor"));
+	}
+
+	/** FGCObject interface */
+	virtual void AddReferencedObjects(FReferenceCollector& Collector) override;
+
+	//~ Begin FEditorUndoClient Interface
+	virtual void PostUndo(bool bSuccess) override;
+	virtual void PostRedo(bool bSuccess) override { PostUndo(bSuccess); }
+	// End of FEditorUndoClient
+
+	/** Whether pasting the currently selected nodes is permissible */
+	bool CanPasteNodes() const;
+
+	/** Paste the contents of the clipboard at the provided location */
+	void PasteNodesAtLocation(const FVector2D& Location);
+
+private:
+	TSharedRef<SDockTab> SpawnTab_GraphCanvas(const FSpawnTabArgs& Args);
+	TSharedRef<SDockTab> SpawnTab_Properties(const FSpawnTabArgs& Args);
+	TSharedRef<SDockTab> SpawnTab_Palette(const FSpawnTabArgs& Args);
+
+protected:
+	/** Called when the selection changes in the GraphEditor */
+	void OnSelectedNodesChanged(const TSet<UObject*>& NewSelection);
+
+	/**
+	 * Called when a node's title is committed for a rename
+	 *
+	 * @param	NewText				New title text
+	 * @param	CommitInfo			How text was committed
+	 * @param	NodeBeingChanged	The node being changed
+	 */
+	void OnNodeTitleCommitted(const FText& NewText, ETextCommit::Type CommitInfo, UEdGraphNode* NodeBeingChanged);
+
+	/** Select every node in the graph */
+	void SelectAllNodes();
+	/** Whether we can select every node */
+	bool CanSelectAllNodes() const;
+
+	/** Delete the currently selected nodes */
+	void DeleteSelectedNodes();
+
+	/** Whether we are able to delete the currently selected nodes */
+	bool CanDeleteNodes() const;
+
+	/** Delete only the currently selected nodes that can be duplicated */
+	void DeleteSelectedDuplicatableNodes();
+
+	/** Cut the currently selected nodes */
+	void CutSelectedNodes();
+
+	/** Whether we are able to cut the currently selected nodes */
+	bool CanCutNodes() const;
+
+	/** Copy the currently selected nodes */
+	void CopySelectedNodes();
+
+	/** Whether copying the currently selected nodes is permissible */
+	bool CanCopyNodes() const;
+
+	/** Paste the contents of the clipboard */
+	void PasteNodes();
+
+	/** Duplicate the currently selected nodes */
+	void DuplicateNodes();
+
+	/** Whether we are able to duplicate the currently selected nodes */
+	bool CanDuplicateNodes() const;
+
+	/** Called to undo the last action */
+	void UndoGraphAction();
+
+	/** Called to redo the last undone action */
+	void RedoGraphAction();
+
+	void OnAlignTop();
+	void OnAlignMiddle();
+	void OnAlignBottom();
+	void OnAlignLeft();
+	void OnAlignCenter();
+	void OnAlignRight();
+
+	void OnStraightenConnections();
+
+	void OnDistributeNodesH();
+	void OnDistributeNodesV();
+
+private:
+	/** FNotifyHook interface */
+	virtual void NotifyPostChange(const FPropertyChangedEvent& PropertyChangedEvent, FProperty* PropertyThatChanged) override;
+
+	/** Creates all internal widgets for the tabs to point at */
+	void CreateInternalWidgets();
+
+	/** Builds the toolbar widget for the Metasound editor */
+	void ExtendToolbar();
+
+	/** Binds new graph commands to delegates */
+	void BindGraphCommands();
+
+	/** Toolbar command methods */
+	void Play();
+	void PlayNode();
+	
+
+	/** Whether we can play the current selection of nodes */
+	bool CanPlayNode() const;
+	void Stop();
+
+	/** Either play the Metasound or stop currently playing sound */
+	void TogglePlayback();
+	
+	/** Plays a single specified node */
+	void PlaySingleNode(UEdGraphNode* Node);
+
+	/** Sync the content browser to the current selection of nodes */
+	void SyncInBrowser();
+
+	/** Add an input to the currently selected node */
+	void AddInput();
+	/** Whether we can add an input to the currently selected node */
+	bool CanAddInput() const;
+
+	/** Delete an input from the currently selected node */
+	void DeleteInput();
+	/** Whether we can delete an input from the currently selected node */
+	bool CanDeleteInput() const;
+
+	/* Create comment node on graph */
+	void OnCreateComment();
+
+	/** Create new graph editor widget */
+	TSharedRef<SGraphEditor> CreateGraphEditorWidget();
+
+private:
+	/** The Metasound asset being inspected */
+	UMetasound* Metasound;
+
+	/** List of open tool panels; used to ensure only one exists at any one time */
+	TMap<FName, TWeakPtr<SDockableTab>> SpawnedToolPanels;
+
+	/** New Graph Editor */
+	TSharedPtr<SGraphEditor> MetasoundGraphEditor;
+
+	/** Properties tab */
+	TSharedPtr<IDetailsView> MetasoundProperties;
+
+	/** Palette of Sound Node types */
+	TSharedPtr<SMetasoundPalette> Palette;
+
+	/** Command list for this editor */
+	TSharedPtr<FUICommandList> GraphEditorCommands;
+
+	/**	The tab ids for all the tabs used */
+	static const FName GraphCanvasTabId;
+	static const FName PropertiesTabId;
+	static const FName PaletteTabId;
+};
diff --git a/Engine/Plugins/Runtime/Metasound/Source/MetasoundEditor/Private/MetasoundEditorCommands.cpp b/Engine/Plugins/Runtime/Metasound/Source/MetasoundEditor/Private/MetasoundEditorCommands.cpp
new file mode 100644
--- /dev/null
+++ b/Engine/Plugins/Runtime/Metasound/Source/MetasoundEditor/Private/MetasoundEditorCommands.cpp
@@ -0,0 +1,16 @@
+// Copyright Epic Games, Inc. All Rights Reserved.
+#include "MetasoundEditorCommands.h"
+
+#define LOCTEXT_NAMESPACE "MetasoundEditorCommands"
+
+void FMetasoundEditorCommands::RegisterCommands()
+{
+	UI_COMMAND(Play, "Play", "Plays (or restarts) the Metasound", EUserInterfaceActionType::Button, FInputChord());
+	UI_COMMAND(Stop, "Stop", "Stops Metasound (If currently playing)", EUserInterfaceActionType::Button, FInputChord());
+	UI_COMMAND(TogglePlayback, "Toggle Playback", "Plays the Metasound or stops the currently playing Metasound", EUserInterfaceActionType::Button, FInputChord(EKeys::SpaceBar));
+
+	UI_COMMAND(BrowserSync, "Sync to Browser", "Selects the SoundWave in the content browser", EUserInterfaceActionType::Button, FInputChord());
+	UI_COMMAND(AddInput, "Add Input", "Adds an input to the node", EUserInterfaceActionType::Button, FInputChord());
+	UI_COMMAND(DeleteInput, "Delete Input", "Removes an input from the node", EUserInterfaceActionType::Button, FInputChord());
+}
+#undef LOCTEXT_NAMESPACE
diff --git a/Engine/Plugins/Runtime/Metasound/Source/MetasoundEditor/Private/MetasoundEditorGraphConnectionDrawingPolicy.cpp b/Engine/Plugins/Runtime/Metasound/Source/MetasoundEditor/Private/MetasoundEditorGraphConnectionDrawingPolicy.cpp
new file mode 100644
--- /dev/null
+++ b/Engine/Plugins/Runtime/Metasound/Source/MetasoundEditor/Private/MetasoundEditorGraphConnectionDrawingPolicy.cpp
@@ -0,0 +1,83 @@
+// Copyright Epic Games, Inc. All Rights Reserved.
+
+#include "MetasoundEditorGraphConnectionDrawingPolicy.h"
+
+#include "Misc/App.h"
+#include "MetasoundEditorGraph.h"
+#include "MetasoundEditorGraphSchema.h"
+#include "Components/AudioComponent.h"
+
+#include "Editor.h"
+
+FConnectionDrawingPolicy* FMetasoundGraphConnectionDrawingPolicyFactory::CreateConnectionPolicy(const class UEdGraphSchema* Schema, int32 InBackLayerID, int32 InFrontLayerID, float ZoomFactor, const class FSlateRect& InClippingRect, class FSlateWindowElementList& InDrawElements, class UEdGraph* InGraphObj) const
+{
+	if (Schema->IsA(UMetasoundEditorGraphSchema::StaticClass()))
+	{
+		return new FMetasoundGraphConnectionDrawingPolicy(InBackLayerID, InFrontLayerID, ZoomFactor, InClippingRect, InDrawElements, InGraphObj);
+	}
+	return nullptr;
+}
+
+FMetasoundGraphConnectionDrawingPolicy::FMetasoundGraphConnectionDrawingPolicy(int32 InBackLayerID, int32 InFrontLayerID, float ZoomFactor, const FSlateRect& InClippingRect, FSlateWindowElementList& InDrawElements, UEdGraph* InGraphObj)
+	: FConnectionDrawingPolicy(InBackLayerID, InFrontLayerID, ZoomFactor, InClippingRect, InDrawElements)
+	, GraphObj(InGraphObj)
+{
+	// Cache off the editor options
+	ActiveColor = Settings->TraceAttackColor;
+	InactiveColor = Settings->TraceReleaseColor;
+
+	ActiveWireThickness = Settings->TraceAttackWireThickness;
+	InactiveWireThickness = Settings->TraceReleaseWireThickness;
+
+	// Don't want to draw ending arrowheads
+	ArrowImage = nullptr;
+	ArrowRadius = FVector2D::ZeroVector;
+}
+
+void FMetasoundGraphConnectionDrawingPolicy::Draw(TMap<TSharedRef<SWidget>, FArrangedWidget>& InPinGeometries, FArrangedChildren& ArrangedNodes)
+{
+	// Draw everything
+	FConnectionDrawingPolicy::Draw(InPinGeometries, ArrangedNodes);
+}
+
+// Give specific editor modes a chance to highlight this connection or darken non-interesting connections
+void FMetasoundGraphConnectionDrawingPolicy::DetermineWiringStyle(UEdGraphPin* OutputPin, UEdGraphPin* InputPin, /*inout*/ FConnectionParams& Params)
+{
+	Params.AssociatedPin1 = OutputPin;
+	Params.AssociatedPin2 = InputPin;
+
+	// Get the schema and grab the default color from it
+	check(OutputPin);
+	check(GraphObj);
+	const UEdGraphSchema* Schema = GraphObj->GetSchema();
+
+	Params.WireColor = Schema->GetPinTypeColor(OutputPin->PinType);
+
+	if (!InputPin)
+	{
+		return;
+	}
+	
+	bool bExecuted = false;
+
+	// Run thru the predecessors, and on
+	if (FExecPairingMap* PredecessorMap = PredecessorNodes.Find(InputPin->GetOwningNode()))
+	{
+		if (FTimePair* Times = PredecessorMap->Find(OutputPin->GetOwningNode()))
+		{
+			bExecuted = true;
+
+			Params.WireThickness = ActiveWireThickness;
+			Params.WireColor = ActiveColor;
+
+			Params.bDrawBubbles = true;
+		}
+	}
+
+	if (!bExecuted)
+	{
+		// It's not followed, fade it and keep it thin
+		Params.WireColor = InactiveColor;
+		Params.WireThickness = InactiveWireThickness;
+	}
+}
diff --git a/Engine/Plugins/Runtime/Metasound/Source/MetasoundEditor/Private/MetasoundEditorGraphConnectionDrawingPolicy.h b/Engine/Plugins/Runtime/Metasound/Source/MetasoundEditor/Private/MetasoundEditorGraphConnectionDrawingPolicy.h
new file mode 100644
--- /dev/null
+++ b/Engine/Plugins/Runtime/Metasound/Source/MetasoundEditor/Private/MetasoundEditorGraphConnectionDrawingPolicy.h
@@ -0,0 +1,65 @@
+// Copyright Epic Games, Inc. All Rights Reserved.
+
+#pragma once
+
+#include "CoreMinimal.h"
+#include "Layout/ArrangedWidget.h"
+#include "Widgets/SWidget.h"
+#include "EdGraphUtilities.h"
+#include "ConnectionDrawingPolicy.h"
+
+class FSlateRect;
+class FSlateWindowElementList;
+class UEdGraph;
+class UEdGraphSchema;
+
+struct FMetasoundGraphConnectionDrawingPolicyFactory : public FGraphPanelPinConnectionFactory
+{
+public:
+	virtual ~FMetasoundGraphConnectionDrawingPolicyFactory() = default;
+
+	// FGraphPanelPinConnectionFactory
+	virtual class FConnectionDrawingPolicy* CreateConnectionPolicy(const UEdGraphSchema* Schema, int32 InBackLayerID, int32 InFrontLayerID, float ZoomFactor, const FSlateRect& InClippingRect, FSlateWindowElementList& InDrawElements, class UEdGraph* InGraphObj) const override;
+	// ~FGraphPanelPinConnectionFactory
+};
+
+
+// This class draws the connections for an UEdGraph using a Metasound schema
+class FMetasoundGraphConnectionDrawingPolicy : public FConnectionDrawingPolicy
+{
+protected:
+	// Times for one execution pair within the current graph
+	struct FTimePair
+	{
+		double PredExecTime;
+		double ThisExecTime;
+
+		FTimePair()
+			: PredExecTime(0.0)
+			, ThisExecTime(0.0)
+		{
+		}
+	};
+
+	// Map of pairings
+	using FExecPairingMap = TMap<UEdGraphNode*, FTimePair>;
+
+	// Map of nodes that preceded before a given node in the execution sequence (one entry for each pairing)
+	TMap<UEdGraphNode*, FExecPairingMap> PredecessorNodes;
+
+	UEdGraph* GraphObj;
+
+	FLinearColor ActiveColor;
+	FLinearColor InactiveColor;
+
+	float ActiveWireThickness;
+	float InactiveWireThickness;
+
+public:
+	FMetasoundGraphConnectionDrawingPolicy(int32 InBackLayerID, int32 InFrontLayerID, float ZoomFactor, const FSlateRect& InClippingRect, FSlateWindowElementList& InDrawElements, UEdGraph* InGraphObj);
+
+	// FConnectionDrawingPolicy interface
+	virtual void DetermineWiringStyle(UEdGraphPin* OutputPin, UEdGraphPin* InputPin, /*inout*/ FConnectionParams& Params) override;
+	virtual void Draw(TMap<TSharedRef<SWidget>, FArrangedWidget>& PinGeometries, FArrangedChildren& ArrangedNodes) override;
+	// End of FConnectionDrawingPolicy interface
+};
diff --git a/Engine/Plugins/Runtime/Metasound/Source/MetasoundEditor/Private/MetasoundEditorGraphNode.cpp b/Engine/Plugins/Runtime/Metasound/Source/MetasoundEditor/Private/MetasoundEditorGraphNode.cpp
new file mode 100644
--- /dev/null
+++ b/Engine/Plugins/Runtime/Metasound/Source/MetasoundEditor/Private/MetasoundEditorGraphNode.cpp
@@ -0,0 +1,218 @@
+// Copyright Epic Games, Inc. All Rights Reserved.
+
+#include "MetasoundEditorGraphNode.h"
+
+#include "Editor/EditorEngine.h"
+#include "Engine/Font.h"
+#include "Framework/Commands/GenericCommands.h"
+#include "GraphEditorActions.h"
+#include "Metasound.h"
+#include "MetasoundEditorGraph.h"
+#include "MetasoundEditorCommands.h"
+#include "ScopedTransaction.h"
+#include "ToolMenus.h"
+
+#define LOCTEXT_NAMESPACE "MetasoundEditorGraphNode"
+
+
+UMetasoundEditorGraphNode::UMetasoundEditorGraphNode(const FObjectInitializer& ObjectInitializer)
+	: Super(ObjectInitializer)
+{
+}
+
+void UMetasoundEditorGraphNode::PostLoad()
+{
+	Super::PostLoad();
+
+	for (int32 Index = 0; Index < Pins.Num(); ++Index)
+	{
+		UEdGraphPin* Pin = Pins[Index];
+		if (Pin->PinName.IsNone())
+		{
+			// Makes sure pin has a name for lookup purposes but user will never see it
+			if (Pin->Direction == EGPD_Input)
+			{
+				Pin->PinName = CreateUniquePinName(TEXT("Input"));
+			}
+			else
+			{
+				Pin->PinName = CreateUniquePinName(TEXT("Output"));
+			}
+			Pin->PinFriendlyName = FText::FromString(TEXT(" "));
+		}
+	}
+}
+
+void UMetasoundEditorGraphNode::CreateInputPin()
+{
+	FString PinName; // get from UMetasound
+	UEdGraphPin* NewPin = CreatePin(EGPD_Input, TEXT("MetasoundEditorGraphNode"), *PinName);
+	if (NewPin->PinName.IsNone())
+	{
+		// Makes sure pin has a name for lookup purposes but user will never see it
+		NewPin->PinName = CreateUniquePinName(TEXT("Input"));
+		NewPin->PinFriendlyName = FText::FromString(TEXT(" "));
+	}
+}
+
+void UMetasoundEditorGraphNode::AddInputPin()
+{
+	const FScopedTransaction Transaction( NSLOCTEXT("UnrealEd", "MetasoundEditorAddInput", "Add Metasound Input") );
+	Modify();
+	CreateInputPin();
+
+// 	UMetasound* Metasound = CastChecked<UMetasoundEditorGraph>(GetGraph())->Metasound;
+	// TODO: Add metasound input pin
+
+	// Refresh the current graph, so the pins can be updated
+	GetGraph()->NotifyGraphChanged();
+}
+
+void UMetasoundEditorGraphNode::RemoveInputPin(UEdGraphPin* InGraphPin)
+{
+	const FScopedTransaction Transaction( NSLOCTEXT("UnrealEd", "SoundCueEditorDeleteInput", "Delete Metasound Input") );
+	Modify();
+
+	TArray<UEdGraphPin*> InputPins;
+	GetInputPins(InputPins);
+
+	for (int32 InputIndex = 0; InputIndex < InputPins.Num(); InputIndex++)
+	{
+		if (InGraphPin == InputPins[InputIndex])
+		{
+			InGraphPin->MarkPendingKill();
+			Pins.Remove(InGraphPin);
+
+			// Remove pin from UMetasound
+			break;
+		}
+	}
+
+	// Refresh the current graph, so the pins can be updated
+	GetGraph()->NotifyGraphChanged();
+}
+
+int32 UMetasoundEditorGraphNode::EstimateNodeWidth() const
+{
+	const int32 EstimatedCharWidth = 6;
+	FString NodeTitle = GetNodeTitle(ENodeTitleType::FullTitle).ToString();
+	UFont* Font = GetDefault<UEditorEngine>()->EditorFont;
+	int32 Result = NodeTitle.Len()*EstimatedCharWidth;
+
+	if (Font)
+	{
+		Result = Font->GetStringSize(*NodeTitle);
+	}
+
+	return Result;
+}
+
+bool UMetasoundEditorGraphNode::CanAddInputPin() const
+{
+	// Get from UMetasound
+	return false;
+}
+
+FText UMetasoundEditorGraphNode::GetNodeTitle(ENodeTitleType::Type TitleType) const
+{
+	// Get from UMetasound
+	return FText::GetEmpty();
+}
+
+void UMetasoundEditorGraphNode::PrepareForCopying()
+{
+}
+
+void UMetasoundEditorGraphNode::PostCopyNode()
+{
+}
+
+void UMetasoundEditorGraphNode::PostEditImport()
+{
+}
+
+void UMetasoundEditorGraphNode::PostDuplicate(bool bDuplicateForPIE)
+{
+	Super::PostDuplicate(bDuplicateForPIE);
+
+	if (!bDuplicateForPIE)
+	{
+		CreateNewGuid();
+	}
+}
+
+void UMetasoundEditorGraphNode::CreateInputPins()
+{
+	// Get from UMetasound... archetype?
+}
+
+void UMetasoundEditorGraphNode::GetNodeContextMenuActions(UToolMenu* Menu, UGraphNodeContextMenuContext* Context) const
+{
+	if (Context->Pin)
+	{
+		// If on an input that can be deleted, show option
+		if (Context->Pin->Direction == EGPD_Input /*&& SoundNode->ChildNodes.Num() > SoundNode->GetMinChildNodes()*/)
+		{
+			FToolMenuSection& Section = Menu->AddSection("MetasoundEditorGraphDeleteInput");
+			Section.AddMenuEntry(FMetasoundEditorCommands::Get().DeleteInput);
+		}
+	}
+	else if (Context->Node)
+	{
+		{
+			FToolMenuSection& Section = Menu->AddSection("MetasoundEditorGraphNodeAlignment");
+			Section.AddSubMenu("Alignment", LOCTEXT("AlignmentHeader", "Alignment"), FText(), FNewToolMenuDelegate::CreateLambda([](UToolMenu* SubMenu)
+			{
+				{
+					FToolMenuSection& SubMenuSection = SubMenu->AddSection("EdGraphSchemaAlignment", LOCTEXT("AlignHeader", "Align"));
+					SubMenuSection.AddMenuEntry(FGraphEditorCommands::Get().AlignNodesTop);
+					SubMenuSection.AddMenuEntry(FGraphEditorCommands::Get().AlignNodesMiddle);
+					SubMenuSection.AddMenuEntry(FGraphEditorCommands::Get().AlignNodesBottom);
+					SubMenuSection.AddMenuEntry(FGraphEditorCommands::Get().AlignNodesLeft);
+					SubMenuSection.AddMenuEntry(FGraphEditorCommands::Get().AlignNodesCenter);
+					SubMenuSection.AddMenuEntry(FGraphEditorCommands::Get().AlignNodesRight);
+					SubMenuSection.AddMenuEntry(FGraphEditorCommands::Get().StraightenConnections);
+				}
+
+				{
+					FToolMenuSection& SubMenuSection = SubMenu->AddSection("EdGraphSchemaDistribution", LOCTEXT("DistributionHeader", "Distribution"));
+					SubMenuSection.AddMenuEntry(FGraphEditorCommands::Get().DistributeNodesHorizontally);
+					SubMenuSection.AddMenuEntry(FGraphEditorCommands::Get().DistributeNodesVertically);
+				}
+			}));
+		}
+
+		{
+			FToolMenuSection& Section = Menu->AddSection("MetasoundEditorGraphNodeEdit");
+			Section.AddMenuEntry(FGenericCommands::Get().Delete);
+			Section.AddMenuEntry(FGenericCommands::Get().Cut);
+			Section.AddMenuEntry(FGenericCommands::Get().Copy);
+			Section.AddMenuEntry(FGenericCommands::Get().Duplicate);
+		}
+
+		{
+			FToolMenuSection& Section = Menu->AddSection("MetasoundEditorGraphNodeAddPlaySync");
+			if (CanAddInputPin())
+			{
+				Section.AddMenuEntry(FMetasoundEditorCommands::Get().AddInput);
+			}
+		}
+	}
+}
+
+FText UMetasoundEditorGraphNode::GetTooltipText() const
+{
+	FText Tooltip = FText::GetEmpty();
+
+	// TODO: Retrieve from UMetasound
+
+	return Tooltip;
+}
+
+FString UMetasoundEditorGraphNode::GetDocumentationExcerptName() const
+{
+	// Default the node to searching for an excerpt named for the C++ node class name, including the U prefix.
+	// This is done so that the excerpt name in the doc file can be found by find-in-files when searching for the full class name.
+	return FString::Printf(TEXT("%s%s"), UMetasound::StaticClass()->GetPrefixCPP(), *UMetasound::StaticClass()->GetName());
+}
+#undef LOCTEXT_NAMESPACE
diff --git a/Engine/Plugins/Runtime/Metasound/Source/MetasoundEditor/Private/MetasoundEditorGraphNodeFactory.h b/Engine/Plugins/Runtime/Metasound/Source/MetasoundEditor/Private/MetasoundEditorGraphNodeFactory.h
new file mode 100644
--- /dev/null
+++ b/Engine/Plugins/Runtime/Metasound/Source/MetasoundEditor/Private/MetasoundEditorGraphNodeFactory.h
@@ -0,0 +1,25 @@
+// Copyright Epic Games, Inc. All Rights Reserved.
+
+#pragma once
+
+#include "CoreMinimal.h"
+#include "EdGraphUtilities.h"
+#include "GraphNode.h"
+#include "SGraphNode.h"
+#include "SGraphNodeSoundResult.h"
+#include "SGraphNodeSoundBase.h"
+#include "Widgets/DeclarativeSyntaxSupport.h"
+
+
+class FMetasoundGraphNodeFactory : public FGraphPanelNodeFactory
+{
+	virtual TSharedPtr<SGraphNode> CreateNode(UEdGraphNode* InNode) const override
+	{
+		if (UMetasoundEditorGraphNode* BaseSoundNode = Cast<UMetasoundEditorGraphNode>(InNode))
+		{
+			return SNew(SGraphNodeSoundResult, RootSoundNode);
+		}
+
+		return nullptr;
+	}
+};
diff --git a/Engine/Plugins/Runtime/Metasound/Source/MetasoundEditor/Private/MetasoundEditorGraphNode_Base.cpp b/Engine/Plugins/Runtime/Metasound/Source/MetasoundEditor/Private/MetasoundEditorGraphNode_Base.cpp
new file mode 100644
--- /dev/null
+++ b/Engine/Plugins/Runtime/Metasound/Source/MetasoundEditor/Private/MetasoundEditorGraphNode_Base.cpp
@@ -0,0 +1,230 @@
+// Copyright Epic Games, Inc. All Rights Reserved.
+
+#include "MetasoundEditorGraphNode_Base.h"
+
+#include "EdGraph/EdGraphSchema.h"
+#include "MetasoundEditorGraphSchema.h"
+
+
+UMetasoundEditorGraphNode_Base::UMetasoundEditorGraphNode_Base(const FObjectInitializer& ObjectInitializer)
+	: Super(ObjectInitializer)
+{
+}
+
+UEdGraphPin* UMetasoundEditorGraphNode_Base::GetOutputPin()
+{
+	UEdGraphPin* OutputPin = nullptr;
+
+	for (int32 PinIndex = 0; PinIndex < Pins.Num(); PinIndex++)
+	{
+		if (Pins[PinIndex]->Direction == EGPD_Output)
+		{
+			// should only ever be one output pin
+			check(OutputPin == nullptr);
+			OutputPin = Pins[PinIndex];
+		}
+	}
+
+	return OutputPin;
+}
+
+void UMetasoundEditorGraphNode_Base::GetInputPins(TArray<class UEdGraphPin*>& OutInputPins)
+{
+	OutInputPins.Empty();
+
+	for (int32 PinIndex = 0; PinIndex < Pins.Num(); PinIndex++)
+	{
+		if (Pins[PinIndex]->Direction == EGPD_Input)
+		{
+			OutInputPins.Add(Pins[PinIndex]);
+		}
+	}
+}
+
+UEdGraphPin* UMetasoundEditorGraphNode_Base::GetInputPin(int32 InputIndex)
+{
+	check(InputIndex >= 0 && InputIndex < GetInputCount());
+
+	for (int32 PinIndex = 0, FoundInputs = 0; PinIndex < Pins.Num(); PinIndex++)
+	{
+		if (Pins[PinIndex]->Direction == EGPD_Input)
+		{
+			if (InputIndex == FoundInputs)
+			{
+				return Pins[PinIndex];
+			}
+			else
+			{
+				FoundInputs++;
+			}
+		}
+	}
+
+	return nullptr;
+}
+
+int32 UMetasoundEditorGraphNode_Base::GetInputCount() const
+{
+	int32 InputCount = 0;
+
+	for (int32 PinIndex = 0, FoundInputs = 0; PinIndex < Pins.Num(); PinIndex++)
+	{
+		if (Pins[PinIndex]->Direction == EGPD_Input)
+		{
+			InputCount++;
+		}
+	}
+
+	return InputCount;
+}
+
+void UMetasoundEditorGraphNode_Base::InsertNewNode(UEdGraphPin* FromPin, UEdGraphPin* NewLinkPin, TSet<UEdGraphNode*>& OutNodeList)
+{
+	const UMetasoundEditorGraphSchema* Schema = CastChecked<UMetasoundEditorGraphSchema>(GetSchema());
+
+	// The pin we are creating from already has a connection that needs to be broken. We want to "insert" the new node in between, so that the output of the new node is hooked up too
+	UEdGraphPin* OldLinkedPin = FromPin->LinkedTo[0];
+	check(OldLinkedPin);
+
+	FromPin->BreakAllPinLinks();
+
+	// Hook up the old linked pin to the first valid output pin on the new node
+	for (int32 OutpinPinIdx=0; OutpinPinIdx<Pins.Num(); OutpinPinIdx++)
+	{
+		UEdGraphPin* OutputExecPin = Pins[OutpinPinIdx];
+		check(OutputExecPin);
+		if (ECanCreateConnectionResponse::CONNECT_RESPONSE_MAKE == Schema->CanCreateConnection(OldLinkedPin, OutputExecPin).Response)
+		{
+			if (Schema->TryCreateConnection(OldLinkedPin, OutputExecPin))
+			{
+				OutNodeList.Add(OldLinkedPin->GetOwningNode());
+				OutNodeList.Add(this);
+			}
+			break;
+		}
+	}
+
+	if (Schema->TryCreateConnection(FromPin, NewLinkPin))
+	{
+		OutNodeList.Add(FromPin->GetOwningNode());
+		OutNodeList.Add(this);
+	}
+}
+
+void UMetasoundEditorGraphNode_Base::AllocateDefaultPins()
+{
+	check(Pins.Num() == 0);
+
+	CreateInputPins();
+
+	if (!IsRootNode())
+	{
+		// Todo --> conform to UMetasound outputs
+		CreatePin(EGPD_Output, TEXT("Metasound"), TEXT("Output"));
+	}
+}
+
+void UMetasoundEditorGraphNode_Base::ReconstructNode()
+{
+	// Break any links to 'orphan' pins
+	for (int32 PinIndex = 0; PinIndex < Pins.Num(); ++PinIndex)
+	{
+		UEdGraphPin* Pin = Pins[PinIndex];
+		TArray<class UEdGraphPin*>& LinkedToRef = Pin->LinkedTo;
+		for (int32 LinkIdx=0; LinkIdx < LinkedToRef.Num(); LinkIdx++)
+		{
+			UEdGraphPin* OtherPin = LinkedToRef[LinkIdx];
+			// If we are linked to a pin that its owner doesn't know about, break that link
+			if (!OtherPin->GetOwningNode()->Pins.Contains(OtherPin))
+			{
+				Pin->LinkedTo.Remove(OtherPin);
+			}
+		}
+	}
+
+	// Store the old Input and Output pins
+	TArray<UEdGraphPin*> OldInputPins;
+	GetInputPins(OldInputPins);
+	UEdGraphPin* OldOutputPin = GetOutputPin();
+
+	// Move the existing pins to a saved array
+	TArray<UEdGraphPin*> OldPins(Pins);
+	Pins.Reset();
+
+	// Recreate the new pins
+	AllocateDefaultPins();
+
+	// Get new Input and Output pins
+	TArray<UEdGraphPin*> NewInputPins;
+	GetInputPins(NewInputPins);
+	UEdGraphPin* NewOutputPin = GetOutputPin();
+
+	for (int32 PinIndex = 0; PinIndex < OldInputPins.Num(); PinIndex++)
+	{
+		if (PinIndex < NewInputPins.Num())
+		{
+			NewInputPins[PinIndex]->MovePersistentDataFromOldPin(*OldInputPins[PinIndex]);
+		}
+	}
+
+	if (OldOutputPin)
+	{
+		NewOutputPin->MovePersistentDataFromOldPin(*OldOutputPin);
+	}
+
+	// Throw away the original pins
+	for (UEdGraphPin* OldPin : OldPins)
+	{
+		OldPin->Modify();
+		UEdGraphNode::DestroyPin(OldPin);
+	}
+}
+
+void UMetasoundEditorGraphNode_Base::AutowireNewNode(UEdGraphPin* FromPin)
+{
+	if (FromPin)
+	{
+		const UMetasoundEditorGraphSchema* Schema = CastChecked<UMetasoundEditorGraphSchema>(GetSchema());
+
+		TSet<UEdGraphNode*> NodeList;
+
+		// auto-connect from dragged pin to first compatible pin on the new node
+		for (int32 i=0; i<Pins.Num(); i++)
+		{
+			UEdGraphPin* Pin = Pins[i];
+			check(Pin);
+			FPinConnectionResponse Response = Schema->CanCreateConnection(FromPin, Pin);
+			if (ECanCreateConnectionResponse::CONNECT_RESPONSE_MAKE == Response.Response)
+			{
+				if (Schema->TryCreateConnection(FromPin, Pin))
+				{
+					NodeList.Add(FromPin->GetOwningNode());
+					NodeList.Add(this);
+				}
+				break;
+			}
+			else if(ECanCreateConnectionResponse::CONNECT_RESPONSE_BREAK_OTHERS_A == Response.Response)
+			{
+				InsertNewNode(FromPin, Pin, NodeList);
+				break;
+			}
+		}
+
+		// Send all nodes that received a new pin connection a notification
+		for (auto It = NodeList.CreateConstIterator(); It; ++It)
+		{
+			UEdGraphNode* Node = (*It);
+			Node->NodeConnectionListChanged();
+		}
+	}
+}
+
+bool UMetasoundEditorGraphNode_Base::CanCreateUnderSpecifiedSchema(const UEdGraphSchema* Schema) const
+{
+	return Schema->IsA(UMetasoundEditorGraphSchema::StaticClass());
+}
+
+FString UMetasoundEditorGraphNode_Base::GetDocumentationLink() const
+{
+	return TEXT("Shared/GraphNodes/Metasound");
+}
diff --git a/Engine/Plugins/Runtime/Metasound/Source/MetasoundEditor/Private/MetasoundEditorGraphNode_Root.cpp b/Engine/Plugins/Runtime/Metasound/Source/MetasoundEditor/Private/MetasoundEditorGraphNode_Root.cpp
new file mode 100644
--- /dev/null
+++ b/Engine/Plugins/Runtime/Metasound/Source/MetasoundEditor/Private/MetasoundEditorGraphNode_Root.cpp
@@ -0,0 +1,42 @@
+// Copyright Epic Games, Inc. All Rights Reserved.
+
+#include "MetasoundEditorGraphNode_Root.h"
+
+#include "ToolMenus.h"
+#include "GraphEditorSettings.h"
+#include "MetasoundEditorCommands.h"
+
+
+#define LOCTEXT_NAMESPACE "MetasoundGraphNode_Root"
+
+UMetasoundEditorGraphNode_Root::UMetasoundEditorGraphNode_Root(const FObjectInitializer& ObjectInitializer)
+	: Super(ObjectInitializer)
+{
+}
+
+FLinearColor UMetasoundEditorGraphNode_Root::GetNodeTitleColor() const
+{
+	return GetDefault<UGraphEditorSettings>()->ResultNodeTitleColor;
+}
+
+FText UMetasoundEditorGraphNode_Root::GetNodeTitle(ENodeTitleType::Type TitleType) const
+{
+	return LOCTEXT("RootTitle", "Output");
+}
+
+FText UMetasoundEditorGraphNode_Root::GetTooltipText() const
+{
+	return LOCTEXT("RootToolTip", "Wire the final Sound Node into this node");
+}
+
+void UMetasoundEditorGraphNode_Root::CreateInputPins()
+{
+	// TODO: Implement input pin creation from Metasound doc data
+	CreatePin(EGPD_Input, TEXT("MetasoundEditorGraphNode"), TEXT("Root"), NAME_None);
+}
+
+void UMetasoundEditorGraphNode_Root::GetNodeContextMenuActions(UToolMenu* Menu, UGraphNodeContextMenuContext* Context) const
+{
+	// TODO: Implement menu actions
+}
+#undef LOCTEXT_NAMESPACE
diff --git a/Engine/Plugins/Runtime/Metasound/Source/MetasoundEditor/Private/MetasoundEditorGraphSchema.cpp b/Engine/Plugins/Runtime/Metasound/Source/MetasoundEditor/Private/MetasoundEditorGraphSchema.cpp
new file mode 100644
--- /dev/null
+++ b/Engine/Plugins/Runtime/Metasound/Source/MetasoundEditor/Private/MetasoundEditorGraphSchema.cpp
@@ -0,0 +1,338 @@
+// Copyright Epic Games, Inc. All Rights Reserved.
+
+#include "MetasoundEditorGraphSchema.h"
+
+#include "EdGraph/EdGraph.h"
+#include "EdGraphNode_Comment.h"
+#include "Editor.h"
+#include "Engine/Selection.h"
+#include "Framework/MultiBox/MultiBoxBuilder.h"
+#include "GraphEditor.h"
+#include "GraphEditorActions.h"
+#include "Layout/SlateRect.h"
+#include "Metasound.h"
+#include "MetasoundEditor.h"
+#include "MetasoundEditorGraph.h"
+#include "MetasoundEditorGraphNode.h"
+#include "MetasoundEditorGraphNode_Root.h"
+#include "MetasoundEditorUtilities.h"
+#include "ScopedTransaction.h"
+#include "ToolMenus.h"
+#include "UObject/UObjectHash.h"
+#include "UObject/UObjectIterator.h"
+#include "Toolkits/ToolkitManager.h"
+
+#define LOCTEXT_NAMESPACE "MetasoundSchema"
+
+namespace MetasoundEditorUtils
+{
+	TSharedPtr<FMetasoundEditor> GetMetasoundEditor(const UMetasound* Metasound)
+	{
+		if (Metasound)
+		{
+			TSharedPtr<IToolkit> FoundAssetEditor = FToolkitManager::Get().FindEditorForAsset(CastChecked<const UObject>(Metasound));
+			return StaticCastSharedPtr<FMetasoundEditor, IToolkit>(FoundAssetEditor);
+		}
+
+		return nullptr;
+	}
+
+	TSharedPtr<FMetasoundEditor> GetMetasoundEditor(const UEdGraph& EdGraph)
+	{
+		const UMetasoundEditorGraph* MetasoundGraph = CastChecked<const UMetasoundEditorGraph>(&EdGraph);
+		return GetMetasoundEditor(MetasoundGraph->Metasound);
+	}
+} // namespace MetasoundEditorUtils
+
+UEdGraphNode* FMetasoundGraphSchemaAction_NewNode::PerformAction(UEdGraph* ParentGraph, UEdGraphPin* FromPin, const FVector2D Location, bool bSelectNewNode/* = true*/)
+{
+	UMetasoundEditorGraph* MetasoundGraph = CastChecked<UMetasoundEditorGraph>(ParentGraph);
+	UMetasound* Metasound = MetasoundGraph->Metasound;
+	const FScopedTransaction Transaction(LOCTEXT("MetasoundEditorNewSoundNode", "Metasound Editor: New Metasound Node"));
+	ParentGraph->Modify();
+
+
+	Metasound->Modify();
+
+	UMetasoundEditorGraphNode_Root* GraphNode = nullptr;
+// TODO: Implement
+// 	UMetasoundEditorGraphNode_Root* GraphNode = MetasoundGraph->CreateBlankNode<MetasoundNodeClass>();
+
+	Metasound->PostEditChange();
+	Metasound->MarkPackageDirty();
+
+	return GraphNode;
+}
+
+void FMetasoundGraphSchemaAction_NewNode::ConnectToSelectedNodes(UMetasound* NewNode, UEdGraph* ParentGraph) const
+{
+	// TODO: Implement
+}
+
+UEdGraphNode* FMetasoundGraphSchemaAction_NewFromSelected::PerformAction(UEdGraph* ParentGraph, UEdGraphPin* FromPin, const FVector2D Location, bool bSelectNewNode/* = true*/)
+{
+	// TODO: Implement
+	return nullptr;
+}
+
+UEdGraphNode* FMetasoundGraphSchemaAction_NewComment::PerformAction(UEdGraph* ParentGraph, UEdGraphPin* FromPin, const FVector2D Location, bool bSelectNewNode/* = true*/)
+{
+	UEdGraphNode_Comment* CommentTemplate = NewObject<UEdGraphNode_Comment>();
+
+	FSlateRect Bounds;
+	FVector2D SpawnLocation = Location;
+	TSharedPtr<FMetasoundEditor> MetasoundEditor = MetasoundEditorUtils::GetMetasoundEditor(*ParentGraph);
+
+	if (MetasoundEditor.IsValid() && MetasoundEditor->GetBoundsForSelectedNodes(Bounds, 50.0f))
+	{
+		CommentTemplate->SetBounds(Bounds);
+		SpawnLocation.X = CommentTemplate->NodePosX;
+		SpawnLocation.Y = CommentTemplate->NodePosY;
+	}
+
+	return FEdGraphSchemaAction_NewNode::SpawnNodeFromTemplate<UEdGraphNode_Comment>(ParentGraph, CommentTemplate, SpawnLocation);
+}
+
+UEdGraphNode* FMetasoundGraphSchemaAction_Paste::PerformAction(UEdGraph* ParentGraph, UEdGraphPin* FromPin, const FVector2D Location, bool bSelectNewNode/* = true*/)
+{
+	TSharedPtr<FMetasoundEditor> MetasoundEditor = MetasoundEditorUtils::GetMetasoundEditor(*ParentGraph);
+	if (MetasoundEditor.IsValid())
+	{
+		MetasoundEditor->PasteNodesAtLocation(Location);
+	}
+
+	return nullptr;
+}
+
+UMetasoundEditorGraphSchema::UMetasoundEditorGraphSchema(const FObjectInitializer& ObjectInitializer)
+	: Super(ObjectInitializer)
+{
+}
+
+bool UMetasoundEditorGraphSchema::ConnectionCausesLoop(const UEdGraphPin* InputPin, const UEdGraphPin* OutputPin) const
+{
+	// TODO: Implement loop check
+
+	// Simple connection to root node
+	return false;
+}
+
+void UMetasoundEditorGraphSchema::GetPaletteActions(FGraphActionMenuBuilder& ActionMenuBuilder) const
+{
+	GetAllMetasoundActions(ActionMenuBuilder, /* bShowSelectedActions*/ false);
+	GetCommentAction(ActionMenuBuilder);
+}
+
+void UMetasoundEditorGraphSchema::GetGraphContextActions(FGraphContextMenuBuilder& ContextMenuBuilder) const
+{
+	GetAllMetasoundActions(ContextMenuBuilder, /* bShowSelectedActions*/ true);
+
+	GetCommentAction(ContextMenuBuilder, ContextMenuBuilder.CurrentGraph);
+
+	if (!ContextMenuBuilder.FromPin)
+	{
+		TSharedPtr<FMetasoundEditor> MetasoundEditor = MetasoundEditorUtils::GetMetasoundEditor(*ContextMenuBuilder.CurrentGraph);
+		if (MetasoundEditor.IsValid() && MetasoundEditor->CanPasteNodes())
+		{
+			TSharedPtr<FMetasoundGraphSchemaAction_Paste> NewAction(new FMetasoundGraphSchemaAction_Paste(FText::GetEmpty(), LOCTEXT("PasteHereAction", "Paste here"), FText::GetEmpty(), 0));
+			ContextMenuBuilder.AddAction(NewAction);
+		}
+	}
+}
+
+void UMetasoundEditorGraphSchema::GetContextMenuActions(class UToolMenu* Menu, class UGraphNodeContextMenuContext* Context) const
+{
+	if (Context->Pin)
+	{
+		FToolMenuSection& Section = Menu->AddSection("MetasoundGraphSchemaPinActions", LOCTEXT("PinActionsMenuHeader", "Pin Actions"));
+
+		// Only displays the 'Break Link' option if there is a link to break
+		if (Context->Pin->LinkedTo.Num() > 0)
+		{
+			Section.AddMenuEntry(FGraphEditorCommands::Get().BreakPinLinks);
+		}
+	}
+	else if (Context->Node)
+	{
+		const UMetasoundEditorGraphNode* SoundGraphNode = Cast<const UMetasoundEditorGraphNode>(Context->Node);
+		{
+			FToolMenuSection& Section = Menu->AddSection("MetasoundGraphSchemaNodeActions", LOCTEXT("NodeActionsMenuHeader", "Node Actions"));
+			Section.AddMenuEntry(FGraphEditorCommands::Get().BreakNodeLinks);
+		}
+	}
+
+	Super::GetContextMenuActions(Menu, Context);
+}
+
+void UMetasoundEditorGraphSchema::CreateDefaultNodesForGraph(UEdGraph& Graph) const
+{
+	const int32 RootNodeHeightOffset = -58;
+
+	// Create the result node
+	FGraphNodeCreator<UMetasoundEditorGraphNode> NodeCreator(Graph);
+	UMetasoundEditorGraphNode* ResultRootNode = NodeCreator.CreateNode();
+	ResultRootNode->NodePosY = RootNodeHeightOffset;
+	NodeCreator.Finalize();
+	SetNodeMetaData(ResultRootNode, FNodeMetadata::DefaultGraphNode);
+}
+
+const FPinConnectionResponse UMetasoundEditorGraphSchema::CanCreateConnection(const UEdGraphPin* PinA, const UEdGraphPin* PinB) const
+{
+	// Make sure the pins are not on the same node
+	if (PinA->GetOwningNode() == PinB->GetOwningNode())
+	{
+		return FPinConnectionResponse(CONNECT_RESPONSE_DISALLOW, LOCTEXT("ConnectionSameNode", "Both are on the same node"));
+	}
+
+	// Compare the directions
+	const UEdGraphPin* InputPin = nullptr;
+	const UEdGraphPin* OutputPin = nullptr;
+
+	if (!CategorizePinsByDirection(PinA, PinB, /*out*/ InputPin, /*out*/ OutputPin))
+	{
+		return FPinConnectionResponse(CONNECT_RESPONSE_DISALLOW, LOCTEXT("ConnectionIncompatible", "Directions are not compatible"));
+	}
+
+	if (ConnectionCausesLoop(InputPin, OutputPin))
+	{
+		return FPinConnectionResponse(CONNECT_RESPONSE_DISALLOW, LOCTEXT("ConnectionLoop", "Connection would cause loop"));
+	}
+
+	// Break existing connections on inputs only - multiple output connections are acceptable
+	if (InputPin->LinkedTo.Num() > 0)
+	{
+		ECanCreateConnectionResponse ReplyBreakOutputs;
+		if (InputPin == PinA)
+		{
+			ReplyBreakOutputs = CONNECT_RESPONSE_BREAK_OTHERS_A;
+		}
+		else
+		{
+			ReplyBreakOutputs = CONNECT_RESPONSE_BREAK_OTHERS_B;
+		}
+		return FPinConnectionResponse(ReplyBreakOutputs, LOCTEXT("ConnectionReplace", "Replace existing connections"));
+	}
+
+	return FPinConnectionResponse(CONNECT_RESPONSE_MAKE, TEXT(""));
+}
+
+bool UMetasoundEditorGraphSchema::TryCreateConnection(UEdGraphPin* PinA, UEdGraphPin* PinB) const
+{
+	bool bModified = UEdGraphSchema::TryCreateConnection(PinA, PinB);
+	if (bModified)
+	{
+		// TODO: Compile Metasound???
+	}
+
+	return bModified;
+}
+
+bool UMetasoundEditorGraphSchema::ShouldHidePinDefaultValue(UEdGraphPin* Pin) const
+{
+	// TODO: Determine if should be hidden from doc data
+	return true;
+}
+
+FLinearColor UMetasoundEditorGraphSchema::GetPinTypeColor(const FEdGraphPinType& PinType) const
+{
+	return FLinearColor::White;
+}
+
+void UMetasoundEditorGraphSchema::BreakNodeLinks(UEdGraphNode& TargetNode) const
+{
+	Super::BreakNodeLinks(TargetNode);
+
+	// TODO: Update MS document
+}
+
+void UMetasoundEditorGraphSchema::BreakPinLinks(UEdGraphPin& TargetPin, bool bSendsNodeNotifcation) const
+{
+	const FScopedTransaction Transaction(NSLOCTEXT("UnrealEd", "GraphEd_BreakPinLinks", "Break Pin Links"));
+
+	Super::BreakPinLinks(TargetPin, bSendsNodeNotifcation);
+
+	// TODO: Update MS document
+}
+
+void UMetasoundEditorGraphSchema::GetAssetsGraphHoverMessage(const TArray<FAssetData>& Assets, const UEdGraph* HoverGraph, FString& OutTooltipText, bool& OutOkIcon) const
+{
+	OutOkIcon = false;
+
+	for (int32 AssetIdx = 0; AssetIdx < Assets.Num(); ++AssetIdx)
+	{
+		// As soon as one of the items is a sound wave, say we can drag it on... we actually eat only the sound waves.
+		USoundWave* SoundWav = Cast<USoundWave>(Assets[AssetIdx].GetAsset());
+		if (SoundWav)
+		{
+			OutOkIcon = true;
+			break;
+		}
+	}
+}
+
+void UMetasoundEditorGraphSchema::DroppedAssetsOnGraph(const TArray<FAssetData>& Assets, const FVector2D& GraphPosition, UEdGraph* Graph) const
+{
+	// TODO: Implement
+}
+
+void UMetasoundEditorGraphSchema::DroppedAssetsOnNode(const TArray<FAssetData>& Assets, const FVector2D& GraphPosition, UEdGraphNode* Node) const
+{
+	if (!Node->IsA<UMetasoundEditorGraphNode>())
+	{
+		return;
+	}
+
+	UMetasoundEditorGraphNode* GraphNode = CastChecked<UMetasoundEditorGraphNode>(Node);
+	UMetasoundEditorGraph* MetasoundGraph = CastChecked<UMetasoundEditorGraph>(Node->GetGraph());
+
+	// TODO: Implement adding reference node
+// 	UMetasound* Metasound = MetasoundGraph->Metasound;
+
+	MetasoundGraph->NotifyGraphChanged();
+}
+
+void UMetasoundEditorGraphSchema::GetAllMetasoundActions(FGraphActionMenuBuilder& ActionMenuBuilder, bool bShowSelectedActions) const
+{
+	// TODO: Add actions for each type of metasound node
+}
+
+void UMetasoundEditorGraphSchema::GetCommentAction(FGraphActionMenuBuilder& ActionMenuBuilder, const UEdGraph* CurrentGraph) const
+{
+	if (!ActionMenuBuilder.FromPin && CurrentGraph)
+	{
+		TSharedPtr<FMetasoundEditor> MetasoundEditor = MetasoundEditorUtils::GetMetasoundEditor(*CurrentGraph);
+
+		if (MetasoundEditor.IsValid())
+		{
+			const int32 NumSelected = MetasoundEditor->GetNumberOfSelectedNodes();
+			const FText MenuDescription = NumSelected > 0 ? LOCTEXT("CreateCommentAction", "Create Comment from Selection") : LOCTEXT("AddCommentAction", "Add Comment...");
+			const FText ToolTip = LOCTEXT("CreateCommentToolTip", "Creates a comment.");
+
+			TSharedPtr<FMetasoundGraphSchemaAction_NewComment> NewAction(new FMetasoundGraphSchemaAction_NewComment(FText::GetEmpty(), MenuDescription, ToolTip, 0));
+			ActionMenuBuilder.AddAction(NewAction);
+		}
+	}
+}
+
+int32 UMetasoundEditorGraphSchema::GetNodeSelectionCount(const UEdGraph* Graph) const
+{
+	if (const UMetasound* Metasound = CastChecked<UMetasoundEditorGraph>(Graph)->Metasound)
+	{
+		TSharedPtr<IToolkit> FoundAssetEditor = FToolkitManager::Get().FindEditorForAsset(CastChecked<UObject>(Metasound));
+		if (FoundAssetEditor.IsValid())
+		{
+			const FMetasoundEditor* AssetEditor = static_cast<FMetasoundEditor*>(FoundAssetEditor.Get());
+			return AssetEditor->GetNumberOfSelectedNodes();
+		}
+	}
+
+	return 0;
+}
+
+TSharedPtr<FEdGraphSchemaAction> UMetasoundEditorGraphSchema::GetCreateCommentAction() const
+{
+	TSharedPtr<FMetasoundGraphSchemaAction_NewComment> Comment = MakeShared<FMetasoundGraphSchemaAction_NewComment>();
+	return StaticCastSharedPtr<FEdGraphSchemaAction, FMetasoundGraphSchemaAction_NewComment>(Comment);
+}
+
+#undef LOCTEXT_NAMESPACE
diff --git a/Engine/Plugins/Runtime/Metasound/Source/MetasoundEditor/Private/MetasoundEditorModule.cpp b/Engine/Plugins/Runtime/Metasound/Source/MetasoundEditor/Private/MetasoundEditorModule.cpp
new file mode 100644
--- /dev/null
+++ b/Engine/Plugins/Runtime/Metasound/Source/MetasoundEditor/Private/MetasoundEditorModule.cpp
@@ -0,0 +1,91 @@
+// Copyright Epic Games, Inc. All Rights Reserved.
+
+#include "MetasoundEditorModule.h"
+
+#include "AssetTypeActions_Base.h"
+#include "Brushes/SlateImageBrush.h"
+#include "CoreMinimal.h"
+#include "EditorStyleSet.h"
+#include "MetasoundAssetTypeActions.h"
+#include "Modules/ModuleInterface.h"
+#include "Modules/ModuleManager.h"
+#include "Styling/CoreStyle.h"
+#include "Styling/SlateStyle.h"
+#include "Styling/SlateStyleMacros.h"
+#include "Styling/SlateStyleRegistry.h"
+#include "Styling/SlateTypes.h"
+#include "Templates/SharedPointer.h"
+
+DEFINE_LOG_CATEGORY(LogMetasoundEditor);
+
+
+namespace MetasoundEditorUtils
+{
+	static const FName AssetToolName = TEXT("AssetTools");
+
+	template <typename T>
+	void AddAssetAction(IAssetTools& AssetTools, TArray<TSharedPtr<FAssetTypeActions_Base>>& AssetArray)
+	{
+		TSharedPtr<T> AssetAction = MakeShared<T>();
+		TSharedPtr<FAssetTypeActions_Base> AssetActionBase = StaticCastSharedPtr<FAssetTypeActions_Base>(AssetAction);
+		AssetTools.RegisterAssetTypeActions(AssetAction.ToSharedRef());
+		AssetArray.Add(AssetActionBase);
+	}
+} // namespace MetasoundEditorUtils
+
+class FMetasoundSlateStyle : public FSlateStyleSet
+{
+public:
+	FMetasoundSlateStyle()
+		: FSlateStyleSet("MetasoundStyle")
+	{
+		SetParentStyleName(FEditorStyle::GetStyleSetName());
+
+		SetContentRoot(FPaths::EngineContentDir() / TEXT("Editor/Slate"));
+		SetCoreContentRoot(FPaths::EngineContentDir() / TEXT("Slate"));
+
+		static const FVector2D Icon20x20(20.0f, 20.0f);
+		static const FVector2D Icon40x40(40.0f, 40.0f);
+
+		// Metasound Editor
+		{
+			Set("MetasoundEditor.Play", new FSlateImageBrush(RootToContentDir(TEXT("Icons/icon_SCueEd_PlayCue_40x.png")), Icon40x40));
+			Set("MetasoundEditor.Play.Small", new FSlateImageBrush(RootToContentDir(TEXT("Icons/icon_SCueEd_PlayCue_40x.png")), Icon20x20));
+			Set("MetasoundEditor.Stop", new FSlateImageBrush(RootToContentDir(TEXT("Icons/icon_SCueEd_Stop_40x.png")), Icon40x40));
+			Set("MetasoundEditor.Stop.Small", new FSlateImageBrush(RootToContentDir(TEXT("Icons/icon_SCueEd_Stop_40x.png")), Icon20x20));
+		}
+
+		FSlateStyleRegistry::RegisterSlateStyle(*this);
+	}
+};
+
+class FMetasoundEditorModule : public IMetasoundEditorModule
+{
+	virtual void StartupModule() override
+	{
+		// Register Metasound asset type actions
+		IAssetTools& AssetTools = FModuleManager::LoadModuleChecked<FAssetToolsModule>(MetasoundEditorUtils::AssetToolName).Get();
+		MetasoundEditorUtils::AddAssetAction<FAssetTypeActions_Metasound>(AssetTools, AssetActions);
+
+		StyleSet = MakeShared<FMetasoundSlateStyle>();
+	}
+
+	virtual void ShutdownModule() override
+	{
+		if (FModuleManager::Get().IsModuleLoaded(MetasoundEditorUtils::AssetToolName))
+		{
+			IAssetTools& AssetTools = FModuleManager::GetModuleChecked<FAssetToolsModule>(MetasoundEditorUtils::AssetToolName).Get();
+			for (TSharedPtr<FAssetTypeActions_Base>& AssetAction : AssetActions)
+			{
+				AssetTools.UnregisterAssetTypeActions(AssetAction.ToSharedRef());
+			}
+		}
+		AssetActions.Reset();
+	}
+
+	TArray<TSharedPtr<FAssetTypeActions_Base>> AssetActions;
+	TSharedPtr<FSlateStyleSet> StyleSet;
+};
+
+
+IMPLEMENT_MODULE(FMetasoundEditorModule, MetasoundEditor);
diff --git a/Engine/Plugins/Runtime/Metasound/Source/MetasoundEditor/Private/MetasoundEditorModule.h b/Engine/Plugins/Runtime/Metasound/Source/MetasoundEditor/Private/MetasoundEditorModule.h
new file mode 100644
--- /dev/null
+++ b/Engine/Plugins/Runtime/Metasound/Source/MetasoundEditor/Private/MetasoundEditorModule.h
@@ -0,0 +1,12 @@
+// Copyright Epic Games, Inc. All Rights Reserved.
+
+#pragma once
+
+#include "CoreMinimal.h"
+#include "Modules/ModuleInterface.h"
+
+DECLARE_LOG_CATEGORY_EXTERN(LogMetasoundEditor, Log, All);
+
+class METASOUNDEDITOR_API IMetasoundEditorModule : public IModuleInterface
+{
+};
diff --git a/Engine/Plugins/Runtime/Metasound/Source/MetasoundEditor/Private/MetasoundEditorUtilities.cpp b/Engine/Plugins/Runtime/Metasound/Source/MetasoundEditor/Private/MetasoundEditorUtilities.cpp
new file mode 100644
--- /dev/null
+++ b/Engine/Plugins/Runtime/Metasound/Source/MetasoundEditor/Private/MetasoundEditorUtilities.cpp
@@ -0,0 +1,3 @@
+// Copyright Epic Games, Inc. All Rights Reserved.
+
+#include "MetasoundEditorUtilities.h"
diff --git a/Engine/Plugins/Runtime/Metasound/Source/MetasoundEditor/Private/MetasoundFactory.cpp b/Engine/Plugins/Runtime/Metasound/Source/MetasoundEditor/Private/MetasoundFactory.cpp
new file mode 100644
--- /dev/null
+++ b/Engine/Plugins/Runtime/Metasound/Source/MetasoundEditor/Private/MetasoundFactory.cpp
@@ -0,0 +1,23 @@
+// Copyright Epic Games, Inc. All Rights Reserved.
+
+#include "MetasoundFactory.h"
+
+#include "Metasound.h"
+#include "MetasoundEditorGraphNode.h"
+
+
+UMetasoundFactory::UMetasoundFactory(const FObjectInitializer& ObjectInitializer)
+	: Super(ObjectInitializer)
+{
+	SupportedClass = UMetasound::StaticClass();
+
+	bCreateNew = true;
+	bEditorImport = false;
+	bEditAfterNew = true;
+}
+
+UObject* UMetasoundFactory::FactoryCreateNew(UClass* InClass, UObject* InParent, FName Name, EObjectFlags Flags, UObject* InContext, FFeedbackContext* InFeedbackContext)
+{
+	UMetasound* MetaSound = NewObject<UMetasound>(InParent, Name, Flags);
+	return MetaSound;
+}
diff --git a/Engine/Plugins/Runtime/Metasound/Source/MetasoundEditor/Private/SMetasoundPalette.cpp b/Engine/Plugins/Runtime/Metasound/Source/MetasoundEditor/Private/SMetasoundPalette.cpp
new file mode 100644
--- /dev/null
+++ b/Engine/Plugins/Runtime/Metasound/Source/MetasoundEditor/Private/SMetasoundPalette.cpp
@@ -0,0 +1,19 @@
+// Copyright Epic Games, Inc. All Rights Reserved.
+
+#include "SMetasoundPalette.h"
+#include "MetasoundEditorGraphSchema.h"
+
+void SMetasoundPalette::Construct(const FArguments& InArgs)
+{
+	SGraphPalette::Construct(SGraphPalette::FArguments().AutoExpandActionMenu(true));
+}
+
+void SMetasoundPalette::CollectAllActions(FGraphActionListBuilderBase& OutAllActions)
+{
+	if (const UMetasoundEditorGraphSchema* Schema = GetDefault<UMetasoundEditorGraphSchema>())
+	{
+		FGraphActionMenuBuilder ActionMenuBuilder;
+		Schema->GetPaletteActions(ActionMenuBuilder);
+		OutAllActions.Append(ActionMenuBuilder);
+	}
+}
diff --git a/Engine/Plugins/Runtime/Metasound/Source/MetasoundEditor/Private/SMetasoundPalette.h b/Engine/Plugins/Runtime/Metasound/Source/MetasoundEditor/Private/SMetasoundPalette.h
new file mode 100644
--- /dev/null
+++ b/Engine/Plugins/Runtime/Metasound/Source/MetasoundEditor/Private/SMetasoundPalette.h
@@ -0,0 +1,23 @@
+// Copyright Epic Games, Inc. All Rights Reserved.
+
+
+#pragma once
+
+#include "CoreMinimal.h"
+#include "Widgets/DeclarativeSyntaxSupport.h"
+#include "SGraphPalette.h"
+
+//////////////////////////////////////////////////////////////////////////
+
+class SMetasoundPalette : public SGraphPalette
+{
+public:
+	SLATE_BEGIN_ARGS( SMetasoundPalette ) {};
+	SLATE_END_ARGS()
+
+	void Construct(const FArguments& InArgs);
+
+protected:
+	/** Callback used to populate all actions list in SGraphActionMenu */
+	virtual void CollectAllActions(FGraphActionListBuilderBase& OutAllActions) override;
+};
diff --git a/Engine/Plugins/Runtime/Metasound/Source/MetasoundEditor/Public/IMetasoundEditor.h b/Engine/Plugins/Runtime/Metasound/Source/MetasoundEditor/Public/IMetasoundEditor.h
new file mode 100644
--- /dev/null
+++ b/Engine/Plugins/Runtime/Metasound/Source/MetasoundEditor/Public/IMetasoundEditor.h
@@ -0,0 +1,18 @@
+// Copyright Epic Games, Inc. All Rights Reserved.
+
+#pragma once
+
+#include "CoreMinimal.h"
+
+#include "Metasound.h"
+#include "UObject/Object.h"
+
+
+class IMetasoundEditor : public FAssetEditorToolkit
+{
+	virtual UMetasound* GetMetasound() const = 0;
+	virtual void SetSelection(TArray<UObject*> SelectedObjects) = 0;
+	virtual bool GetBoundsForSelectedNodes(FSlateRect& Rect, float Padding) = 0;
+	virtual int32 GetNumberOfSelectedNodes() const = 0;
+	virtual TSet<UObject*> GetSelectedNodes() const = 0;
+};
diff --git a/Engine/Plugins/Runtime/Metasound/Source/MetasoundEditor/Public/MetasoundEditorCommands.h b/Engine/Plugins/Runtime/Metasound/Source/MetasoundEditor/Public/MetasoundEditorCommands.h
new file mode 100644
--- /dev/null
+++ b/Engine/Plugins/Runtime/Metasound/Source/MetasoundEditor/Public/MetasoundEditorCommands.h
@@ -0,0 +1,41 @@
+// Copyright Epic Games, Inc. All Rights Reserved.
+
+#pragma once
+
+#include "CoreMinimal.h"
+#include "Framework/Commands/Commands.h"
+#include "EditorStyleSet.h"
+
+class FMetasoundEditorCommands : public TCommands<FMetasoundEditorCommands>
+{
+public:
+	/** Constructor */
+	FMetasoundEditorCommands() 
+		: TCommands<FMetasoundEditorCommands>("MetasoundEditor", NSLOCTEXT("Contexts", "MetasoundEditor", "Metasound Graph Editor"), NAME_None, "MetasoundStyle")
+	{
+	}
+	
+	/** Plays the Metasound */
+	TSharedPtr<FUICommandInfo> Play;
+	
+	/** Stops the currently playing Metasound */
+	TSharedPtr<FUICommandInfo> Stop;
+
+	/** Plays stops the currently playing Metasound */
+	TSharedPtr<FUICommandInfo> TogglePlayback;
+
+	/** Selects the SoundWave in the content browser */
+	TSharedPtr<FUICommandInfo> BrowserSync;
+
+	/** Breaks the node input/output link */
+	TSharedPtr<FUICommandInfo> BreakLink;
+
+	/** Adds an input to the node */
+	TSharedPtr<FUICommandInfo> AddInput;
+
+	/** Removes an input from the node */
+	TSharedPtr<FUICommandInfo> DeleteInput;
+
+	/** Initialize commands */
+	virtual void RegisterCommands() override;
+};
diff --git a/Engine/Plugins/Runtime/Metasound/Source/MetasoundEditor/Public/MetasoundEditorGraph.h b/Engine/Plugins/Runtime/Metasound/Source/MetasoundEditor/Public/MetasoundEditorGraph.h
new file mode 100644
--- /dev/null
+++ b/Engine/Plugins/Runtime/Metasound/Source/MetasoundEditor/Public/MetasoundEditorGraph.h
@@ -0,0 +1,20 @@
+// Copyright Epic Games, Inc. All Rights Reserved.
+
+#pragma once
+
+#include "CoreMinimal.h"
+#include "EdGraph/EdGraph.h"
+#include "Metasound.h"
+#include "UObject/ObjectMacros.h"
+
+#include "MetasoundEditorGraph.generated.h"
+
+UCLASS(MinimalAPI)
+class UMetasoundEditorGraph : public UEdGraph
+{
+	GENERATED_BODY()
+
+public:
+	UPROPERTY()
+	UMetasound* Metasound;
+};
diff --git a/Engine/Plugins/Runtime/Metasound/Source/MetasoundEditor/Public/MetasoundEditorGraphNode.h b/Engine/Plugins/Runtime/Metasound/Source/MetasoundEditor/Public/MetasoundEditorGraphNode.h
new file mode 100644
--- /dev/null
+++ b/Engine/Plugins/Runtime/Metasound/Source/MetasoundEditor/Public/MetasoundEditorGraphNode.h
@@ -0,0 +1,52 @@
+// Copyright Epic Games, Inc. All Rights Reserved.
+
+#pragma once
+
+#include "CoreMinimal.h"
+#include "UObject/ObjectMacros.h"
+#include "MetasoundEditorGraphNode_Base.h"
+
+#include "MetaSoundEditorGraphNode.generated.h"
+
+// Forward Declarations
+class UEdGraphPin;
+class USoundNode;
+
+
+UCLASS(MinimalAPI)
+class UMetasoundEditorGraphNode : public UMetasoundEditorGraphNode_Base
+{
+	GENERATED_UCLASS_BODY()
+
+	/** Fix up the node's owner after being copied */
+	METASOUNDEDITOR_API void PostCopyNode();
+	/** Create a new input pin for this node */
+	METASOUNDEDITOR_API void CreateInputPin();
+	/** Add an input pin to this node and recompile the MetaSoundEditor */
+	METASOUNDEDITOR_API void AddInputPin();
+	/** Remove a specific input pin from this node and recompile the MetaSoundEditor */
+	METASOUNDEDITOR_API void RemoveInputPin(UEdGraphPin* InGraphPin);
+	/** Estimate the width of this Node from the length of its title */
+	METASOUNDEDITOR_API int32 EstimateNodeWidth() const;
+	/** Checks whether an input can be added to this node */
+	METASOUNDEDITOR_API bool CanAddInputPin() const;
+
+
+	// UMetasoundEditorGraphNode_Base interface
+	virtual void CreateInputPins() override;
+	// End of UMetasoundEditorGraphNode_Base interface
+
+	// UEdGraphNode interface
+	virtual FText GetNodeTitle(ENodeTitleType::Type TitleType) const override;
+	virtual void PrepareForCopying() override;
+	virtual void GetNodeContextMenuActions(class UToolMenu* Menu, class UGraphNodeContextMenuContext* Context) const override;
+	virtual FText GetTooltipText() const override;
+	virtual FString GetDocumentationExcerptName() const override;
+	// End of UEdGraphNode interface
+
+	// UObject interface
+	virtual void PostLoad() override;
+	virtual void PostEditImport() override;
+	virtual void PostDuplicate(bool bDuplicateForPIE) override;
+	// End of UObject interface
+};
diff --git a/Engine/Plugins/Runtime/Metasound/Source/MetasoundEditor/Public/MetasoundEditorGraphNode_Base.h b/Engine/Plugins/Runtime/Metasound/Source/MetasoundEditor/Public/MetasoundEditorGraphNode_Base.h
new file mode 100644
--- /dev/null
+++ b/Engine/Plugins/Runtime/Metasound/Source/MetasoundEditor/Public/MetasoundEditorGraphNode_Base.h
@@ -0,0 +1,53 @@
+// Copyright Epic Games, Inc. All Rights Reserved.
+
+#pragma once
+
+#include "CoreMinimal.h"
+#include "EdGraph/EdGraphNode.h"
+#include "UObject/ObjectMacros.h"
+
+#include "MetasoundEditorGraphNode_Base.generated.h"
+
+class UEdGraphPin;
+class UEdGraphSchema;
+
+UCLASS(MinimalAPI)
+class UMetasoundEditorGraphNode_Base : public UEdGraphNode
+{
+	GENERATED_UCLASS_BODY()
+
+	/** Create input pins required of the node */
+	virtual void CreateInputPins() { };
+
+	/** Whether or not this is the root node */
+	virtual bool IsRootNode() const { return false; }
+
+	/** Get the Output Pin (should only ever be one) */
+	METASOUNDEDITOR_API class UEdGraphPin* GetOutputPin();
+
+	/** Get all of the input pins */
+	METASOUNDEDITOR_API void GetInputPins(TArray<UEdGraphPin*>& OutInputPins);
+
+	/** Get the input pin at the provided index */
+	METASOUNDEDITOR_API class UEdGraphPin* GetInputPin(int32 InputIndex);
+
+	/** Get the current input pin count */
+	METASOUNDEDITOR_API int32 GetInputCount() const;
+
+	/**
+	 * Handles inserting the node between the FromPin and what the FromPin was original connected to
+	 *
+	 * @param FromPin		The pin this node is being spawned from
+	 * @param NewLinkPin	The new pin the FromPin will connect to
+	 * @param OutNodeList	Any nodes that are modified will get added to this list for notification purposes
+	 */
+	void InsertNewNode(UEdGraphPin* FromPin, UEdGraphPin* NewLinkPin, TSet<UEdGraphNode*>& OutNodeList);
+
+	// UEdGraphNode interface.
+	virtual void AllocateDefaultPins() override;
+	virtual void ReconstructNode() override;
+	virtual void AutowireNewNode(UEdGraphPin* FromPin) override;
+	virtual bool CanCreateUnderSpecifiedSchema(const UEdGraphSchema* Schema) const override;
+	virtual FString GetDocumentationLink() const override;
+	// End of UEdGraphNode interface.
+};
diff --git a/Engine/Plugins/Runtime/Metasound/Source/MetasoundEditor/Public/MetasoundEditorGraphNode_Root.h b/Engine/Plugins/Runtime/Metasound/Source/MetasoundEditor/Public/MetasoundEditorGraphNode_Root.h
new file mode 100644
--- /dev/null
+++ b/Engine/Plugins/Runtime/Metasound/Source/MetasoundEditor/Public/MetasoundEditorGraphNode_Root.h
@@ -0,0 +1,32 @@
+// Copyright Epic Games, Inc. All Rights Reserved.
+
+#pragma once
+
+#include "CoreMinimal.h"
+#include "MetasoundEditorGraphNode_Base.h"
+#include "UObject/ObjectMacros.h"
+
+#include "MetasoundEditorGraphNode_Root.generated.h"
+
+// Forward Declarations
+class UGraphNodeContextMenuContext;
+
+UCLASS(MinimalAPI)
+class UMetasoundEditorGraphNode_Root : public UMetasoundEditorGraphNode_Base
+{
+	GENERATED_UCLASS_BODY()
+
+	// UEdGraphNode interface
+	virtual FLinearColor GetNodeTitleColor() const override;
+	virtual FText GetTooltipText() const override;
+	virtual FText GetNodeTitle(ENodeTitleType::Type TitleType) const override;
+	virtual bool CanUserDeleteNode() const override { return false; }
+	virtual bool CanDuplicateNode() const override { return false; }
+	virtual void GetNodeContextMenuActions(class UToolMenu* Menu, UGraphNodeContextMenuContext* Context) const override;
+	// End of UEdGraphNode interface
+
+	// UMetasoundEditorGraphNode_Base interface
+	virtual void CreateInputPins() override;
+	virtual bool IsRootNode() const override { return true; }
+	// End of UMetasoundEditorGraphNode_Base interface
+};
diff --git a/Engine/Plugins/Runtime/Metasound/Source/MetasoundEditor/Public/MetasoundEditorGraphSchema.h b/Engine/Plugins/Runtime/Metasound/Source/MetasoundEditor/Public/MetasoundEditorGraphSchema.h
new file mode 100644
--- /dev/null
+++ b/Engine/Plugins/Runtime/Metasound/Source/MetasoundEditor/Public/MetasoundEditorGraphSchema.h
@@ -0,0 +1,136 @@
+// Copyright Epic Games, Inc. All Rights Reserved.
+
+#pragma once
+
+#include "AssetData.h"
+#include "CoreMinimal.h"
+#include "EdGraph/EdGraphSchema.h"
+#include "UObject/ObjectMacros.h"
+
+#include "MetasoundEditorGraphSchema.generated.h"
+
+class UEdGraph;
+class UMetasound;
+class UEdGraphNode;
+
+/** Action to add a node to the graph */
+USTRUCT()
+struct METASOUNDEDITOR_API FMetasoundGraphSchemaAction_NewNode : public FEdGraphSchemaAction
+{
+	GENERATED_USTRUCT_BODY();
+
+	/** Class of node we want to create */
+	UPROPERTY()
+	UClass* MetasoundNodeClass;
+
+	FMetasoundGraphSchemaAction_NewNode() 
+		: FEdGraphSchemaAction()
+		, MetasoundNodeClass(nullptr)
+	{}
+
+	FMetasoundGraphSchemaAction_NewNode(FText InNodeCategory, FText InMenuDesc, FText InToolTip, const int32 InGrouping)
+		: FEdGraphSchemaAction(MoveTemp(InNodeCategory), MoveTemp(InMenuDesc), MoveTemp(InToolTip), InGrouping)
+		, MetasoundNodeClass(nullptr)
+	{}
+
+	//~ Begin FEdGraphSchemaAction Interface
+	virtual UEdGraphNode* PerformAction(UEdGraph* ParentGraph, UEdGraphPin* FromPin, const FVector2D Location, bool bSelectNewNode = true) override;
+	//~ End FEdGraphSchemaAction Interface
+
+private:
+	/** Connects new node to output of selected nodes */
+	void ConnectToSelectedNodes(UMetasound* NewNodeclass, UEdGraph* ParentGraph) const;
+};
+
+/** Action to add nodes to the graph based on selected objects*/
+USTRUCT()
+struct METASOUNDEDITOR_API FMetasoundGraphSchemaAction_NewFromSelected : public FMetasoundGraphSchemaAction_NewNode
+{
+	GENERATED_USTRUCT_BODY();
+
+	FMetasoundGraphSchemaAction_NewFromSelected() 
+		: FMetasoundGraphSchemaAction_NewNode()
+	{}
+
+	FMetasoundGraphSchemaAction_NewFromSelected(FText InNodeCategory, FText InMenuDesc, FText InToolTip, const int32 InGrouping)
+		: FMetasoundGraphSchemaAction_NewNode(MoveTemp(InNodeCategory), MoveTemp(InMenuDesc), MoveTemp(InToolTip), InGrouping) 
+	{}
+
+	//~ Begin FEdGraphSchemaAction Interface
+	virtual UEdGraphNode* PerformAction(class UEdGraph* ParentGraph, UEdGraphPin* FromPin, const FVector2D Location, bool bSelectNewNode = true) override;
+	//~ End FEdGraphSchemaAction Interface
+};
+
+/** Action to create new comment */
+USTRUCT()
+struct METASOUNDEDITOR_API FMetasoundGraphSchemaAction_NewComment : public FEdGraphSchemaAction
+{
+	GENERATED_USTRUCT_BODY();
+
+	FMetasoundGraphSchemaAction_NewComment() 
+		: FEdGraphSchemaAction()
+	{}
+
+	FMetasoundGraphSchemaAction_NewComment(FText InNodeCategory, FText InMenuDesc, FText InToolTip, const int32 InGrouping)
+		: FEdGraphSchemaAction(MoveTemp(InNodeCategory), MoveTemp(InMenuDesc), MoveTemp(InToolTip), InGrouping)
+	{}
+
+	//~ Begin FEdGraphSchemaAction Interface
+	virtual UEdGraphNode* PerformAction(class UEdGraph* ParentGraph, UEdGraphPin* FromPin, const FVector2D Location, bool bSelectNewNode = true) override;
+	//~ End FEdGraphSchemaAction Interface
+};
+
+/** Action to paste clipboard contents into the graph */
+USTRUCT()
+struct METASOUNDEDITOR_API FMetasoundGraphSchemaAction_Paste : public FEdGraphSchemaAction
+{
+	GENERATED_USTRUCT_BODY();
+
+	FMetasoundGraphSchemaAction_Paste() 
+		: FEdGraphSchemaAction()
+	{}
+
+	FMetasoundGraphSchemaAction_Paste(FText InNodeCategory, FText InMenuDesc, FText InToolTip, const int32 InGrouping)
+		: FEdGraphSchemaAction(MoveTemp(InNodeCategory), MoveTemp(InMenuDesc), MoveTemp(InToolTip), InGrouping)
+	{}
+
+	//~ Begin FEdGraphSchemaAction Interface
+	virtual UEdGraphNode* PerformAction(class UEdGraph* ParentGraph, UEdGraphPin* FromPin, const FVector2D Location, bool bSelectNewNode = true) override;
+	//~ End FEdGraphSchemaAction Interface
+};
+
+UCLASS(MinimalAPI)
+class UMetasoundEditorGraphSchema : public UEdGraphSchema
+{
+	GENERATED_UCLASS_BODY()
+
+	/** Check whether connecting these pins would cause a loop */
+	bool ConnectionCausesLoop(const UEdGraphPin* InputPin, const UEdGraphPin* OutputPin) const;
+
+	/** Helper method to add items valid to the palette list */
+	METASOUNDEDITOR_API void GetPaletteActions(FGraphActionMenuBuilder& ActionMenuBuilder) const;
+
+	//~ Begin EdGraphSchema Interface
+	virtual void GetGraphContextActions(FGraphContextMenuBuilder& ContextMenuBuilder) const override;
+	virtual void GetContextMenuActions(class UToolMenu* Menu, class UGraphNodeContextMenuContext* Context) const override;
+	virtual void CreateDefaultNodesForGraph(UEdGraph& Graph) const override;
+	virtual const FPinConnectionResponse CanCreateConnection(const UEdGraphPin* A, const UEdGraphPin* B) const override;
+	virtual bool TryCreateConnection(UEdGraphPin* A, UEdGraphPin* B) const override;
+	virtual bool ShouldHidePinDefaultValue(UEdGraphPin* Pin) const override;
+	virtual FLinearColor GetPinTypeColor(const FEdGraphPinType& PinType) const override;
+	virtual void BreakNodeLinks(UEdGraphNode& TargetNode) const override;
+	virtual void BreakPinLinks(UEdGraphPin& TargetPin, bool bSendsNodeNotifcation) const override;
+	virtual void GetAssetsGraphHoverMessage(const TArray<FAssetData>& Assets, const UEdGraph* HoverGraph, FString& OutTooltipText, bool& OutOkIcon) const override;
+	virtual void DroppedAssetsOnGraph(const TArray<struct FAssetData>& Assets, const FVector2D& GraphPosition, UEdGraph* Graph) const override;
+	virtual void DroppedAssetsOnNode(const TArray<FAssetData>& Assets, const FVector2D& GraphPosition, UEdGraphNode* Node) const override;
+	virtual int32 GetNodeSelectionCount(const UEdGraph* Graph) const override;
+	virtual TSharedPtr<FEdGraphSchemaAction> GetCreateCommentAction() const override;
+	//~ End EdGraphSchema Interface
+
+private:
+	/** Adds actions for creating every type of GraphNode */
+	void GetAllMetasoundActions(FGraphActionMenuBuilder& ActionMenuBuilder, bool bShowSelectedActions) const;
+
+	/** Adds action for creating a comment */
+	void GetCommentAction(FGraphActionMenuBuilder& ActionMenuBuilder, const UEdGraph* CurrentGraph = nullptr) const;
+};
diff --git a/Engine/Plugins/Runtime/Metasound/Source/MetasoundEditor/Public/MetasoundEditorUtilities.h b/Engine/Plugins/Runtime/Metasound/Source/MetasoundEditor/Public/MetasoundEditorUtilities.h
new file mode 100644
--- /dev/null
+++ b/Engine/Plugins/Runtime/Metasound/Source/MetasoundEditor/Public/MetasoundEditorUtilities.h
@@ -0,0 +1,15 @@
+// Copyright Epic Games, Inc. All Rights Reserved.
+
+#pragma once
+
+#include "CoreMinimal.h"
+
+class IMetasoundEditor;
+class UEdGraph;
+class UMetasound;
+
+
+class METASOUNDEDITOR_API FMetasoundEditorUtilities
+{
+
+};
diff --git a/Engine/Plugins/Runtime/Metasound/Source/MetasoundEditor/Public/MetasoundFactory.h b/Engine/Plugins/Runtime/Metasound/Source/MetasoundEditor/Public/MetasoundFactory.h
new file mode 100644
--- /dev/null
+++ b/Engine/Plugins/Runtime/Metasound/Source/MetasoundEditor/Public/MetasoundFactory.h
@@ -0,0 +1,19 @@
+// Copyright Epic Games, Inc. All Rights Reserved.
+
+#pragma once
+
+#include "CoreMinimal.h"
+#include "Factories/Factory.h"
+#include "UObject/ObjectMacros.h"
+
+#include "MetasoundFactory.generated.h"
+
+UCLASS(hidecategories=Object, MinimalAPI)
+class UMetasoundFactory : public UFactory
+{
+	GENERATED_UCLASS_BODY()
+
+	//~ Begin UFactory Interface
+	virtual UObject* FactoryCreateNew(UClass* InClass, UObject* InParent, FName Name, EObjectFlags Flags, UObject* InContext, FFeedbackContext* InFeedbackContext) override;
+	//~ Begin UFactory Interface
+};
diff --git a/Engine/Plugins/Runtime/Metasound/Source/MetasoundEngine/Private/Metasound.cpp b/Engine/Plugins/Runtime/Metasound/Source/MetasoundEngine/Private/Metasound.cpp
--- a/Engine/Plugins/Runtime/Metasound/Source/MetasoundEngine/Private/Metasound.cpp
+++ b/Engine/Plugins/Runtime/Metasound/Source/MetasoundEngine/Private/Metasound.cpp
@@ -12,6 +12,24 @@ UMetasound::UMetasound(const FObjectInitializer& ObjectInitializer)
 {
 }
 
+#if WITH_EDITORONLY_DATA
+UEdGraph* UMetasound::GetGraph()
+{
+	return Graph;
+}
+
+UEdGraph& UMetasound::GetGraphChecked()
+{
+	check(Graph);
+	return *Graph;
+}
+
+void UMetasound::SetGraph(UEdGraph* InGraph)
+{
+	Graph = InGraph;
+}
+#endif // WITH_EDITORONLY_DATA
+
 FMetasoundClassMetadata UMetasound::GetMetadata()
 {
 	return RootMetasoundDocument.RootClass.Metadata;
diff --git a/Engine/Plugins/Runtime/Metasound/Source/MetasoundEngine/Private/MetasoundEngineModule.h b/Engine/Plugins/Runtime/Metasound/Source/MetasoundEngine/Private/MetasoundEngineModule.h
--- a/Engine/Plugins/Runtime/Metasound/Source/MetasoundEngine/Private/MetasoundEngineModule.h
+++ b/Engine/Plugins/Runtime/Metasound/Source/MetasoundEngine/Private/MetasoundEngineModule.h
@@ -16,12 +16,12 @@ class IMetasoundEngineModule : public IModuleInterface
 	// Deserializes Metasound from the graph at the provided path.
 	// @param InPath Content directory path to load metasound from.
 	// @returns New Metasound object on success, nullptr on failure.
-	// @todo decide if this means that we remove the graph in our current asset.
-	//       This seems dangerous since we can fail to save the new asset and scratch the old one.
 	virtual UMetasound* DeserializeMetasound(const FString& InPath) = 0;
 
 	// Serializes Metasound to the provided path.
 	// @param InMetasound Metasound to serialize.
 	// @param InPath Path to serialize to.
+	// @todo Decide if this means that we remove the graph in our current asset.
+	//       This seems dangerous since we can fail to save the new asset and scratch the old one.
 	virtual void SerializeMetasound(const UMetasound& InMetasound, const FString& InPath) = 0;
 };
diff --git a/Engine/Plugins/Runtime/Metasound/Source/MetasoundEngine/Public/Metasound.h b/Engine/Plugins/Runtime/Metasound/Source/MetasoundEngine/Public/Metasound.h
--- a/Engine/Plugins/Runtime/Metasound/Source/MetasoundEngine/Public/Metasound.h
+++ b/Engine/Plugins/Runtime/Metasound/Source/MetasoundEngine/Public/Metasound.h
@@ -1,11 +1,14 @@
 // Copyright Epic Games, Inc. All Rights Reserved.
-
 #pragma once
 
 #include "CoreMinimal.h"
 #include "MetasoundFrontend.h"
 #include "MetasoundFrontendDataLayout.h"
 
+#if WITH_EDITORONLY_DATA
+#include "EdGraph/EdGraph.h"
+#endif // WITH_EDITORONLY_DATA
+
 #include "Metasound.generated.h"
 
 
@@ -13,8 +16,8 @@
  * This asset type is used for Metasound assets that can only be used as nodes in other Metasound graphs.
  * Because of this, they can have any inputs or outputs they need.
  */
-UCLASS(hidecategories = object, BlueprintType, MinimalAPI)
-class UMetasound : public UObject
+UCLASS(hidecategories = object, BlueprintType)
+class METASOUNDENGINE_API UMetasound : public UObject
 {
 	GENERATED_BODY()
 
@@ -22,19 +25,37 @@ protected:
 	UPROPERTY()
 	FMetasoundDocument RootMetasoundDocument;
 
+#if WITH_EDITORONLY_DATA
+	UPROPERTY()
+	UEdGraph* Graph;
+#endif // WITH_EDITORONLY_DATA
+
 	TSharedPtr<Metasound::Frontend::FDescriptionAccessPoint> AccessPoint;
 
 public:
 	UMetasound(const FObjectInitializer& ObjectInitializer);
 
+#if WITH_EDITORONLY_DATA
+
+	// Returns the graph associated with this Metasound. Graph is required to be referenced on
+	// Metasound UObject for editor serialization purposes.
+	// @return Editor graph associated with UMetasound.
+	UEdGraph* GetGraph();
+	UEdGraph& GetGraphChecked();
+
+	// Sets the graph associated with this Metasound. Graph is required to be referenced on
+	// Metasound UObject for editor serialization purposes.
+	// @param Editor graph associated with UMetasound.
+	void SetGraph(UEdGraph* InGraph);
+#endif // WITH_EDITORONLY_DATA
+
 	FMetasoundClassMetadata GetMetadata();
 
-	// This can be used to update the metadata (name, author, etc) for this metasound.
-	// @param InMetadata may be updated with any corrections we do to the input metadata.
+	// Updates the Metasound's metadata (name, author, etc).
+	// @param InMetadata Metadata containing corrections to the class metadata.
 	void SetMetadata(FMetasoundClassMetadata& InMetadata);
 
-	// delete this asset's current metasound document,
-	// and replace it with InClassDescription.
+	// Deletes Metasound's current metasound document, and replaces it with InClassDescription.
 	void SetMetasoundDocument(const FMetasoundDocument& InDocument);
 
 	// returns a weak pointer that can be used to build a TDescriptionPtr
