IWYU Fixes
#rnx

diff --git a/Engine/Plugins/Runtime/Metasound/Source/MetasoundEditor/Private/MetasoundEditor.cpp b/Engine/Plugins/Runtime/Metasound/Source/MetasoundEditor/Private/MetasoundEditor.cpp
--- a/Engine/Plugins/Runtime/Metasound/Source/MetasoundEditor/Private/MetasoundEditor.cpp
+++ b/Engine/Plugins/Runtime/Metasound/Source/MetasoundEditor/Private/MetasoundEditor.cpp
@@ -2,28 +2,31 @@
 
 #include "MetasoundEditor.h"
 
+#include "Components/AudioComponent.h"
 #include "EdGraph/EdGraphNode.h"
 #include "EdGraphUtilities.h"
 #include "Editor.h"
 #include "EditorStyleSet.h"
 #include "Framework/Commands/GenericCommands.h"
+#include "Framework/Docking/TabManager.h"
 #include "Framework/MultiBox/MultiBoxBuilder.h"
 #include "GraphEditor.h"
 #include "GraphEditorActions.h"
 #include "HAL/PlatformApplicationMisc.h"
 #include "IDetailsView.h"
 #include "Kismet2/BlueprintEditorUtils.h"
-#include "Modules/ModuleManager.h"
 #include "MetasoundEditorCommands.h"
 #include "MetasoundEditorGraph.h"
 #include "MetasoundEditorGraphNode.h"
 #include "MetasoundEditorGraphSchema.h"
+#include "MetasoundFormat.h"
+#include "Misc/Attribute.h"
+#include "Modules/ModuleManager.h"
 #include "PropertyEditorModule.h"
 #include "ScopedTransaction.h"
 #include "SMetasoundPalette.h"
 #include "SNodePanel.h"
 #include "Widgets/Docking/SDockTab.h"
-#include "Components/AudioComponent.h"
 
 #define LOCTEXT_NAMESPACE "MetasoundEditor"
 
diff --git a/Engine/Plugins/Runtime/Metasound/Source/MetasoundEditor/Private/MetasoundEditor.h b/Engine/Plugins/Runtime/Metasound/Source/MetasoundEditor/Private/MetasoundEditor.h
--- a/Engine/Plugins/Runtime/Metasound/Source/MetasoundEditor/Private/MetasoundEditor.h
+++ b/Engine/Plugins/Runtime/Metasound/Source/MetasoundEditor/Private/MetasoundEditor.h
@@ -4,8 +4,12 @@
 
 #include "CoreMinimal.h"
 #include "EditorUndoClient.h"
+#include "Framework/Commands/UICommandList.h"
+#include "GraphEditor.h"
+#include "IDetailsView.h"
 #include "IMetasoundEditor.h"
 #include "Misc/NotifyHook.h"
+#include "SMetasoundPalette.h"
 #include "Toolkits/IToolkitHost.h"
 #include "UObject/GCObject.h"
 
diff --git a/Engine/Plugins/Runtime/Metasound/Source/MetasoundEditor/Public/IMetasoundEditor.h b/Engine/Plugins/Runtime/Metasound/Source/MetasoundEditor/Public/IMetasoundEditor.h
--- a/Engine/Plugins/Runtime/Metasound/Source/MetasoundEditor/Public/IMetasoundEditor.h
+++ b/Engine/Plugins/Runtime/Metasound/Source/MetasoundEditor/Public/IMetasoundEditor.h
@@ -4,7 +4,9 @@
 
 #include "CoreMinimal.h"
 
+#include "Layout/SlateRect.h"
 #include "Metasound.h"
+#include "Toolkits/AssetEditorToolkit.h"
 #include "UObject/Object.h"
 
 
