Remove stub code that was causing IWYU error, Remove future CL include
#rnx

diff --git a/Engine/Plugins/Runtime/Metasound/Source/MetasoundEditor/Private/MetasoundEditor.cpp b/Engine/Plugins/Runtime/Metasound/Source/MetasoundEditor/Private/MetasoundEditor.cpp
--- a/Engine/Plugins/Runtime/Metasound/Source/MetasoundEditor/Private/MetasoundEditor.cpp
+++ b/Engine/Plugins/Runtime/Metasound/Source/MetasoundEditor/Private/MetasoundEditor.cpp
@@ -19,7 +19,6 @@
 #include "MetasoundEditorGraph.h"
 #include "MetasoundEditorGraphNode.h"
 #include "MetasoundEditorGraphSchema.h"
-#include "MetasoundFormat.h"
 #include "Misc/Attribute.h"
 #include "Modules/ModuleManager.h"
 #include "PropertyEditorModule.h"
diff --git a/Engine/Plugins/Runtime/Metasound/Source/MetasoundEditor/Private/MetasoundEditorGraphSchema.cpp b/Engine/Plugins/Runtime/Metasound/Source/MetasoundEditor/Private/MetasoundEditorGraphSchema.cpp
--- a/Engine/Plugins/Runtime/Metasound/Source/MetasoundEditor/Private/MetasoundEditorGraphSchema.cpp
+++ b/Engine/Plugins/Runtime/Metasound/Source/MetasoundEditor/Private/MetasoundEditorGraphSchema.cpp
@@ -254,20 +254,11 @@ void UMetasoundEditorGraphSchema::BreakPinLinks(UEdGraphPin& TargetPin, bool bSe
 	// TODO: Update MS document
 }
 
-void UMetasoundEditorGraphSchema::GetAssetsGraphHoverMessage(const TArray<FAssetData>& Assets, const UEdGraph* HoverGraph, FString& OutTooltipText, bool& OutOkIcon) const
+void UMetasoundEditorGraphSchema::GetAssetsGraphHoverMessage(const TArray<FAssetData>& Assets, const UEdGraph* HoverGraph, FString& OutTooltipText, bool& bOutOkIcon) const
 {
-	OutOkIcon = false;
+	bOutOkIcon = false;
 
-	for (int32 AssetIdx = 0; AssetIdx < Assets.Num(); ++AssetIdx)
-	{
-		// As soon as one of the items is a sound wave, say we can drag it on... we actually eat only the sound waves.
-		USoundWave* SoundWav = Cast<USoundWave>(Assets[AssetIdx].GetAsset());
-		if (SoundWav)
-		{
-			OutOkIcon = true;
-			break;
-		}
-	}
+	// TODO: Implement
 }
 
 void UMetasoundEditorGraphSchema::DroppedAssetsOnGraph(const TArray<FAssetData>& Assets, const FVector2D& GraphPosition, UEdGraph* Graph) const
