Linux fix: header casing
#rnx

diff --git a/Engine/Plugins/Runtime/Metasound/Source/MetasoundEditor/Public/MetasoundEditorGraphNode.h b/Engine/Plugins/Runtime/Metasound/Source/MetasoundEditor/Public/MetasoundEditorGraphNode.h
--- a/Engine/Plugins/Runtime/Metasound/Source/MetasoundEditor/Public/MetasoundEditorGraphNode.h
+++ b/Engine/Plugins/Runtime/Metasound/Source/MetasoundEditor/Public/MetasoundEditorGraphNode.h
@@ -6,7 +6,7 @@
 #include "UObject/ObjectMacros.h"
 #include "MetasoundEditorGraphNode_Base.h"
 
-#include "MetaSoundEditorGraphNode.generated.h"
+#include "MetasoundEditorGraphNode.generated.h"
 
 // Forward Declarations
 class UEdGraphPin;
