Update metasound generator interface. Make synth component example an explicit example.

diff --git a/Engine/Plugins/Runtime/Metasound/Source/MetasoundGenerator/Private/MetasoundGenerator.cpp b/Engine/Plugins/Runtime/Metasound/Source/MetasoundGenerator/Private/MetasoundGenerator.cpp
--- a/Engine/Plugins/Runtime/Metasound/Source/MetasoundGenerator/Private/MetasoundGenerator.cpp
+++ b/Engine/Plugins/Runtime/Metasound/Source/MetasoundGenerator/Private/MetasoundGenerator.cpp
@@ -16,86 +16,115 @@
 
 namespace Metasound
 {
-	FMetasoundGenerator::FMetasoundGenerator(FOperatorUniquePtr InOperator, const FAudioReadRef& InOperatorReadBuffer)
-	:	Operator(MoveTemp(InOperator))
-	,	OperatorReadBuffer(InOperatorReadBuffer)
-	,	OperatorReadNum(InOperatorReadBuffer->Num())
-	,	ExecuteOperator(nullptr)
+	FMetasoundGenerator::FMetasoundGenerator(FOperatorUniquePtr InGraphOperator, const FMultichannelAudioFormatReadRef& InGraphOutputAudioRef)
+	:	NumChannels(0)
+	,	NumFramesPerExecute(0)
+	,	NumSamplesPerExecute(0)
+	,	GraphOutputAudioRef(InGraphOutputAudioRef) // Copy audio ref here to avoid unneeded construction of unused graph output audio object.
 	{
-		if (Operator.IsValid())
-		{
-			ExecuteOperator = Operator->GetExecuteFunction();
-			FDataReferenceCollection Inputs = Operator->GetInputs();
-
-			if (Inputs.ContainsDataWriteReference<FFrequency>(TEXT("Frequency")))
-			{
-				FrequencyRef = Inputs.GetDataWriteReference<FFrequency>(TEXT("Frequency"));
-			}
-
-			if (Inputs.ContainsDataWriteReference<FFloatTime>(TEXT("BopPeriod")))
-			{
-				BopPeriodRef = Inputs.GetDataWriteReference<FFloatTime>(TEXT("BopPeriod"));
-			}
-		}
+		SetGraph(MoveTemp(InGraphOperator), InGraphOutputAudioRef);
 	}
 
 	FMetasoundGenerator::~FMetasoundGenerator()
 	{
 	}
 
-	void FMetasoundGenerator::SetFrequency(float InFrequency)
+	bool FMetasoundGenerator::UpdateGraphOperator(FOperatorUniquePtr InGraphOperator, const FMultichannelAudioFormatReadRef& InGraphOutputAudioRef)
 	{
-		FrequencyRef->SetHertz(InFrequency);
+		if (InGraphOutputAudioRef->GetNumChannels() == NumChannels)
+		{
+			SetGraph(MoveTemp(InGraphOperator), InGraphOutputAudioRef);
+
+			return true;
+		}
+
+		return false;
+	}
+
+	void FMetasoundGenerator::SetGraph(FOperatorUniquePtr InGraphOperator, const FMultichannelAudioFormatReadRef& InGraphOutputAudioRef)
+	{
+		InterleavedAudioBuffer.Reset();
+		NumFramesPerExecute = 0;
+		NumSamplesPerExecute = 0;
+
+		// The graph operator and graph audio output contain all the values needed
+		// by the sound generator.
+		RootExecuter.SetOperator(MoveTemp(InGraphOperator));
+		GraphOutputAudioRef = InGraphOutputAudioRef;
+
+		// Query the graph output to get the number of output audio channels.
+		NumChannels = GraphOutputAudioRef->GetNumChannels();
+
+		if (NumChannels > 0)
+		{
+			// All buffers have same number of frames, so only need to query
+			// first buffer to know number of frames.
+			NumFramesPerExecute = GraphOutputAudioRef->GetBuffers()[0]->Num();
+
+			NumSamplesPerExecute = NumFramesPerExecute * NumChannels;
+
+			if (NumSamplesPerExecute > 0)
+			{
+				// Preallocate interleaved buffer as it is necessary for any audio generation calls.
+				InterleavedAudioBuffer.AddUninitialized(NumSamplesPerExecute);
+			}
+		}
 	}
 
-	void FMetasoundGenerator::SetBopPeriod(float InPeriodInSeconds)
+	int32 FMetasoundGenerator::GetNumChannels() const
 	{
-		BopPeriodRef->SetSeconds(InPeriodInSeconds);
+		return GraphOutputAudioRef->GetNumChannels();
 	}
 
-	int32 FMetasoundGenerator::OnGenerateAudio(float* OutAudio, int32 NumSamples)
+	int32 FMetasoundGenerator::OnGenerateAudio(float* OutAudio, int32 NumSamplesRemaining)
 	{
-		if ((NumSamples <= 0) || (nullptr == ExecuteOperator) || (OperatorReadNum <= 0))
+		if (NumSamplesRemaining <= 0)
 		{
 			return 0;
 		}
 
-		int32 WriteCount = FillWithBuffer(OverflowBuffer, OutAudio, NumSamples);
+		// If we have any audio left in the internal overflow buffer from 
+		// previous calls, write that to the output before generating more audio.
+		int32 NumSamplesWritten = FillWithBuffer(OverflowBuffer, OutAudio, NumSamplesRemaining);
 
-		if (WriteCount > 0)
+		if (NumSamplesWritten > 0)
 		{
-			NumSamples -= WriteCount;
-			OverflowBuffer.RemoveAtSwap(0 /* Index */, WriteCount /* Count */, false /* bAllowShrinking */);
-		}
+			NumSamplesRemaining -= NumSamplesWritten;
+			OverflowBuffer.RemoveAtSwap(0 /* Index */, NumSamplesWritten /* Count */, false /* bAllowShrinking */);
 
+		}
 
-		while (NumSamples > 0)
+		while (NumSamplesRemaining > 0)
 		{
-			// TODO: figure out good wrapper for this.
-			ExecuteOperator(Operator.Get());
+			// Call metasound graph operator.
+			RootExecuter.Execute();
 
-			int32 ThisLoopWriteCount = FillWithBuffer(*OperatorReadBuffer, &OutAudio[WriteCount], NumSamples);
+			// Interleave audio because ISoundGenerator interface expects interleaved audio.
+			InterleaveGeneratedAudio();
 
-			NumSamples -= ThisLoopWriteCount;
-			WriteCount += ThisLoopWriteCount;
+			// Add audio generated during graph execution to the output buffer.
+			int32 ThisLoopNumSamplesWritten = FillWithBuffer(InterleavedAudioBuffer, &OutAudio[NumSamplesWritten], NumSamplesRemaining);
 
-			if (0 == NumSamples)
+			NumSamplesRemaining -= ThisLoopNumSamplesWritten;
+			NumSamplesWritten += ThisLoopNumSamplesWritten;
+
+			// If not all the samples were written, then we have to save the 
+			// additional samples to the overflow buffer.
+			if (ThisLoopNumSamplesWritten < InterleavedAudioBuffer.Num())
 			{
-				if (ThisLoopWriteCount < OperatorReadBuffer->Num())
-				{
-					int32 OverflowCount = OperatorReadBuffer->Num() - ThisLoopWriteCount;
+				int32 OverflowCount = InterleavedAudioBuffer.Num() - ThisLoopNumSamplesWritten;
 
-					OverflowBuffer.Reset();
-					OverflowBuffer.AddUninitialized(OverflowCount);
+				OverflowBuffer.Reset();
+				OverflowBuffer.AddUninitialized(OverflowCount);
 
-					FMemory::Memcpy(OverflowBuffer.GetData(), &(OperatorReadBuffer->GetData()[ThisLoopWriteCount]), OverflowCount);
-				}
+				FMemory::Memcpy(OverflowBuffer.GetData(), &InterleavedAudioBuffer.GetData()[ThisLoopNumSamplesWritten], OverflowCount);
 			}
 		}
 
-		return WriteCount;
+		return NumSamplesWritten;
 	}
 
+
 	int32 FMetasoundGenerator::FillWithBuffer(const Audio::AlignedFloatBuffer& InBuffer, float* OutAudio, int32 MaxNumOutputSamples)
 	{
 		int32 InNum = InBuffer.Num();
@@ -116,138 +145,37 @@ namespace Metasound
 
 		return 0;
 	}
-} /* End namespace Metasound */
-
-
-USynthComponentMetasoundGenerator::USynthComponentMetasoundGenerator(const FObjectInitializer& ObjInitializer)
-:	Super(ObjInitializer)
-{
-}
-
-USynthComponentMetasoundGenerator::~USynthComponentMetasoundGenerator()
-{
-}
-
-void USynthComponentMetasoundGenerator::SetBopPeriod(float InBopPeriod)
-{
-	using namespace Metasound;
 
-	BopPeriod = InBopPeriod;
-	if (Generator.IsValid())
+	void FMetasoundGenerator::InterleaveGeneratedAudio()
 	{
-		FMetasoundGenerator* MetasoundGenerator = static_cast<FMetasoundGenerator*>(Generator.Get());
-		MetasoundGenerator->SetBopPeriod(BopPeriod);
-	}
-}
+		const FMultichannelAudioFormat& DeinterleavedAudio = *GraphOutputAudioRef;
 
-void USynthComponentMetasoundGenerator::SetFrequency(float InFrequency)
-{
-	using namespace Metasound;
+		// Prepare output buffer
+		InterleavedAudioBuffer.Reset();
 
-	Frequency = InFrequency;
-	if (Generator.IsValid())
-	{
-		FMetasoundGenerator* MetasoundGenerator = static_cast<FMetasoundGenerator*>(Generator.Get());
-		MetasoundGenerator->SetFrequency(Frequency);
-	}
-}
+		if (NumSamplesPerExecute > 0)
+		{
+			InterleavedAudioBuffer.AddUninitialized(NumSamplesPerExecute);
+		}
 
-#if WITH_EDITOR
-void USynthComponentMetasoundGenerator::PostEditChangeProperty(struct FPropertyChangedEvent & PropertyChangedEvent)
-{
-	using namespace Metasound;
+		const TArrayView<const FAudioBufferReadRef>& Buffers = DeinterleavedAudio.GetBuffers();
 
-	Super::PostEditChangeProperty(PropertyChangedEvent);
+		// Iterate over channels
+		for (int32 ChannelIndex = 0; ChannelIndex < NumChannels; ChannelIndex++)
+		{
+			const FAudioBuffer& InputBuffer = *Buffers[ChannelIndex];
 
-	static const FName FrequencyFName = GET_MEMBER_NAME_CHECKED(USynthComponentMetasoundGenerator, Frequency);
-	static const FName BopPeriodFName = GET_MEMBER_NAME_CHECKED(USynthComponentMetasoundGenerator, BopPeriod);
+			const float* InputPtr = InputBuffer.GetData();
+			float* OutputPtr = &InterleavedAudioBuffer.GetData()[ChannelIndex];
 
-	if (!HasAnyFlags(RF_ClassDefaultObject))
-	{
-		if (Generator.IsValid())
-		{
-			if (FProperty* PropertyThatChanged = PropertyChangedEvent.Property)
+			// Assign values to output for single channel.
+			for (int32 FrameIndex = 0; FrameIndex < NumFramesPerExecute; FrameIndex++)
 			{
-				const FName& Name = PropertyThatChanged->GetFName();
-				if (Name == FrequencyFName)
-				{
-					FMetasoundGenerator* MetasoundGenerator = static_cast<FMetasoundGenerator*>(Generator.Get());
-					MetasoundGenerator->SetFrequency(Frequency);
-				}
-				else if (Name == BopPeriodFName)
-				{
-					FMetasoundGenerator* MetasoundGenerator = static_cast<FMetasoundGenerator*>(Generator.Get());
-					MetasoundGenerator->SetBopPeriod(BopPeriod);
-				}
+				*OutputPtr = InputPtr[FrameIndex];
+				OutputPtr += NumChannels;
 			}
 		}
 	}
-
-}
-#endif
-
-ISoundGeneratorPtr USynthComponentMetasoundGenerator::CreateSoundGenerator(int32 InSampleRate, int32 InNumChannels)
-{
-	using namespace Metasound;
-	using FDataEdge = Metasound::FDataEdge; // Need to make explicit to disambiguate with existing FDataEdge.
-	using FOperatorUniquePtr = TUniquePtr<Metasound::IOperator>;
-
-	TInputNode<FFloatTime> BopPeriodInputNode(TEXT("Input Bop Period Node"), TEXT("BopPeriod"), 1.f, ETimeResolution::Seconds);
-	TInputNode<FFrequency> FrequencyInputNode(TEXT("Input Frequency Node"), TEXT("Frequency"), 100.f, EFrequencyResolution::Hertz);
-
-	FPeriodicBopNode BopNode(TEXT("Bop"), 1.f);
-	FOscNode OscNode(TEXT("Osc"), 100.f);
-	FADSRNode ADSRNode(TEXT("ADSR"), 10.f /* AttackMs */, 10.f /* DecayMs */, 50.f /* SustainMs */, 10.f /* ReleaseMs */);
-	FAudioMultiplyNode MultiplyNode(TEXT("Multiply"));
-
-	TOutputNode<FAudioBuffer> OutputNode(TEXT("Output Audio Node"), TEXT("Audio"));
-
-	FGraph Graph(TEXT("Graph"));
-
-	// Add rate of bop.
-	Graph.AddDataEdge(BopPeriodInputNode, TEXT("BopPeriod"), BopNode, TEXT("Period"));
-	// Add frequency controls
-	Graph.AddDataEdge(FrequencyInputNode, TEXT("Frequency"), OscNode, TEXT("Frequency"));
-
-	// TODO: maybe nodes should have name accessors to inputs and outputs instead of an array?
-	Graph.AddDataEdge(BopNode, TEXT("Bop"), ADSRNode, TEXT("Bop"));
-
-	// Hookup multiply nodes
-	Graph.AddDataEdge(OscNode, TEXT("Audio"), MultiplyNode, TEXT("InputBuffer1"));
-	Graph.AddDataEdge(ADSRNode, TEXT("Envelope"), MultiplyNode, TEXT("InputBuffer2"));
-
-	// Route to output.
-	Graph.AddDataEdge(MultiplyNode, TEXT("Audio"), OutputNode, TEXT("Audio"));
-
-	Graph.AddInputDataDestination(BopPeriodInputNode, TEXT("BopPeriod"));
-	Graph.AddInputDataDestination(FrequencyInputNode, TEXT("Frequency"));
-	Graph.AddOutputDataSource(OutputNode, TEXT("Audio"));
-
-
-	// Build everything
-	
-	FOperatorSettings Settings(InSampleRate, 256);
-	FOperatorBuilder Builder(Settings);
-
-	TArray<IOperatorBuilder::FBuildErrorPtr> Errors;
-	FOperatorUniquePtr Operator = Builder.BuildGraphOperator(Graph, Errors);
-
-	if (!Operator.IsValid())
-	{
-		return ISoundGeneratorPtr(nullptr);
-	}
-
-	FDataReferenceCollection OutputCollection = Operator->GetOutputs();
-	FAudioBufferReadRef OutAudio = OutputCollection.GetDataReadReference<FAudioBuffer>(TEXT("Audio"));
-
-	Generator = MakeShared<FMetasoundGenerator, ESPMode::ThreadSafe>(MoveTemp(Operator), OutAudio);
-	FMetasoundGenerator* MetasoundGenerator = static_cast<FMetasoundGenerator*>(Generator.Get());
-	MetasoundGenerator->SetFrequency(Frequency);
-	MetasoundGenerator->SetBopPeriod(BopPeriod);
-
-	return Generator;
-}
-
-
+} 
 
 
diff --git a/Engine/Plugins/Runtime/Metasound/Source/MetasoundGenerator/Private/SynthComponentMetasoundExample.cpp b/Engine/Plugins/Runtime/Metasound/Source/MetasoundGenerator/Private/SynthComponentMetasoundExample.cpp
new file mode 100644
--- /dev/null
+++ b/Engine/Plugins/Runtime/Metasound/Source/MetasoundGenerator/Private/SynthComponentMetasoundExample.cpp
@@ -0,0 +1,207 @@
+// Copyright Epic Games, Inc. All Rights Reserved.
+
+#include "SynthComponentMetasoundExample.h"
+
+#include "MetasoundOperatorInterface.h"
+#include "MetasoundNodeInterface.h"
+#include "MetasoundAudioFormats.h"
+#include "MetasoundGenerator.h"
+
+#include "MetasoundOscNode.h"
+#include "MetasoundAudioMultiplyNode.h"
+#include "MetasoundADSRNode.h"
+#include "MetasoundPeriodicBopNode.h"
+#include "MetasoundOutputNode.h"
+#include "MetasoundInputNode.h"
+#include "MetasoundOperatorBuilder.h"
+#include "MetasoundGraph.h"
+#include "MetasoundFrequency.h"
+#include "MetasoundTime.h"
+
+struct FSynthComponentMetasoundExampleGraph
+{
+	TUniquePtr<Metasound::FGraph> Graph;
+	TArray<TUniquePtr<Metasound::INode>> Nodes;
+};
+
+static TUniquePtr<FSynthComponentMetasoundExampleGraph> CreateSynthComponentMetasoundExampleGraph()
+{
+	using namespace Metasound;
+	
+	TUniquePtr<TInputNode<FFloatTime>> BopPeriodInputNode = MakeUnique<TInputNode<FFloatTime>>(TEXT("Input Bop Period Node"), TEXT("BopPeriod"), 1.f, ETimeResolution::Seconds);
+
+	TUniquePtr<TInputNode<FFrequency>> FrequencyInputNode = MakeUnique<TInputNode<FFrequency>>(TEXT("Input Frequency Node"), TEXT("Frequency"), 100.f, EFrequencyResolution::Hertz);
+
+	TUniquePtr<FPeriodicBopNode> BopNode = MakeUnique<FPeriodicBopNode>(TEXT("Bop"), 1.f);
+
+	TUniquePtr<FOscNode> OscNode = MakeUnique<FOscNode>(TEXT("Osc"), 100.f);
+
+	TUniquePtr<FADSRNode> ADSRNode = MakeUnique<FADSRNode>(TEXT("ADSR"), 10.f /* AttackMs */, 10.f /* DecayMs */, 50.f /* SustainMs */, 10.f /* ReleaseMs */);
+
+	TUniquePtr<FAudioMultiplyNode> MultiplyNode = MakeUnique<FAudioMultiplyNode>(TEXT("Multiply"));
+
+	TUniquePtr<TOutputNode<FAudioBuffer>> OutputNode = MakeUnique<TOutputNode<FAudioBuffer>>(TEXT("Output Audio Node"), TEXT("Audio"));
+
+	TUniquePtr<FGraph> Graph = MakeUnique<FGraph>(TEXT("Graph"));
+
+	// Add rate of bop.
+	Graph->AddDataEdge(*BopPeriodInputNode, TEXT("BopPeriod"), *BopNode, TEXT("Period"));
+	// Add frequency controls
+	Graph->AddDataEdge(*FrequencyInputNode, TEXT("Frequency"), *OscNode, TEXT("Frequency"));
+
+	Graph->AddDataEdge(*BopNode, TEXT("Bop"), *ADSRNode, TEXT("Bop"));
+
+	// Hookup multiply nodes
+	Graph->AddDataEdge(*OscNode, TEXT("Audio"), *MultiplyNode, TEXT("InputBuffer1"));
+	Graph->AddDataEdge(*ADSRNode, TEXT("Envelope"), *MultiplyNode, TEXT("InputBuffer2"));
+
+	// Route to output.
+	Graph->AddDataEdge(*MultiplyNode, TEXT("Audio"), *OutputNode, TEXT("Audio"));
+
+	Graph->AddInputDataDestination(*BopPeriodInputNode, TEXT("BopPeriod"));
+	Graph->AddInputDataDestination(*FrequencyInputNode, TEXT("Frequency"));
+	Graph->AddOutputDataSource(*OutputNode, TEXT("Audio"));
+
+	// Add everything to this handy-dandy container.
+	TUniquePtr<FSynthComponentMetasoundExampleGraph> ExampleGraph = MakeUnique<FSynthComponentMetasoundExampleGraph>();
+
+	ExampleGraph->Graph = MoveTemp(Graph);
+	ExampleGraph->Nodes.Add(MoveTemp(BopPeriodInputNode));
+	ExampleGraph->Nodes.Add(MoveTemp(FrequencyInputNode));
+	ExampleGraph->Nodes.Add(MoveTemp(BopNode));
+	ExampleGraph->Nodes.Add(MoveTemp(OscNode));
+	ExampleGraph->Nodes.Add(MoveTemp(ADSRNode));
+	ExampleGraph->Nodes.Add(MoveTemp(MultiplyNode));
+	ExampleGraph->Nodes.Add(MoveTemp(OutputNode));
+
+	return MoveTemp(ExampleGraph);
+}
+
+USynthComponentMetasoundExample::USynthComponentMetasoundExample(const FObjectInitializer& ObjInitializer)
+:	Super(ObjInitializer)
+{
+	NumChannels = 1;
+}
+
+USynthComponentMetasoundExample::~USynthComponentMetasoundExample()
+{
+}
+
+void USynthComponentMetasoundExample::SetFloatParameter(const FString& InName, float InValue)
+{
+	using namespace Metasound;
+
+	FloatParameters.Add(InName, InValue);
+
+	if (Generator.IsValid())
+	{
+		FMetasoundGenerator* MetasoundGenerator = static_cast<FMetasoundGenerator*>(Generator.Get());
+
+		SynthCommand([=]() 
+			{
+				MetasoundGenerator->SetInputValue(InName, InValue);
+			}
+		);
+	}
+}
+
+
+#if WITH_EDITOR
+void USynthComponentMetasoundExample::PostEditChangeProperty(struct FPropertyChangedEvent & PropertyChangedEvent)
+{
+	using namespace Metasound;
+	using FFloatMap = TMap<FString, float>;
+	using FStringFloatPair = TMap<FString, float>::ElementType;
+
+	Super::PostEditChangeProperty(PropertyChangedEvent);
+
+	static const FName FloatParamFName = GET_MEMBER_NAME_CHECKED(USynthComponentMetasoundExample, FloatParameters);
+
+	if (!HasAnyFlags(RF_ClassDefaultObject))
+	{
+		if (Generator.IsValid())
+		{
+			if (FProperty* PropertyThatChanged = PropertyChangedEvent.Property)
+			{
+				const FName& Name = PropertyThatChanged->GetFName();
+				if (Name == FloatParamFName)
+				{
+					FMetasoundGenerator* MetasoundGenerator = static_cast<FMetasoundGenerator*>(Generator.Get());
+
+					SynthCommand([=]() 
+						{
+							for (const FStringFloatPair& Pair : FloatParameters)
+							{
+								MetasoundGenerator->SetInputValue(Pair.Get<0>(), Pair.Get<1>());
+							}
+						}
+					);
+				}
+			}
+		}
+	}
+
+}
+#endif
+
+void USynthComponentMetasoundExample::SetGraphOperator(FOperatorUniquePtr InGraphOperator, const FString& InOutputAudioName)
+{
+	using namespace Metasound;
+
+	if (InGraphOperator.IsValid())
+	{
+		FDataReferenceCollection Outputs = InGraphOperator->GetOutputs();
+		FMultichannelAudioFormatReadRef Audio = Outputs.GetDataReadReferenceOrConstruct<FMultichannelAudioFormat>(InOutputAudioName, NumChannels, 256);
+		NumChannels = Audio->GetNumChannels();
+
+		Generator = MakeShared<Metasound::FMetasoundGenerator, ESPMode::ThreadSafe>(MoveTemp(InGraphOperator), Audio);
+	}
+	else
+	{
+		NumChannels = 0;
+		Generator.Reset();
+	}
+}
+
+bool USynthComponentMetasoundExample::PushGraphOperator(FOperatorUniquePtr InGraphOperator, const FString& InOutputAudioName)
+{
+	using namespace Metasound;
+
+	if (!Generator.IsValid())
+	{
+		return false;
+	}
+
+	FDataReferenceCollection Outputs = InGraphOperator->GetOutputs();
+	FMultichannelAudioFormatReadRef Audio = Outputs.GetDataReadReferenceOrConstruct<FMultichannelAudioFormat>(InOutputAudioName, NumChannels, 256);
+
+	if (Generator->UpdateGraphOperator(MoveTemp(InGraphOperator), Audio))
+	{
+		return true;
+	}
+
+	return false;
+}
+
+ISoundGeneratorPtr USynthComponentMetasoundExample::CreateSoundGenerator(int32 InSampleRate, int32 InNumChannels)
+{
+	using namespace Metasound;
+
+	// Build everything
+
+	TUniquePtr<FSynthComponentMetasoundExampleGraph> ExampleGraph = CreateSynthComponentMetasoundExampleGraph();
+
+	FOperatorSettings Settings(InSampleRate, 256);
+	FOperatorBuilder Builder(Settings);
+
+	TArray<IOperatorBuilder::FBuildErrorPtr> Errors;
+	FOperatorUniquePtr Operator = Builder.BuildGraphOperator(*ExampleGraph->Graph, Errors);
+
+	SetGraphOperator(MoveTemp(Operator), TEXT("Audio"));
+
+	return Generator;
+}
+
+
+
+
diff --git a/Engine/Plugins/Runtime/Metasound/Source/MetasoundGenerator/Public/MetasoundGenerator.h b/Engine/Plugins/Runtime/Metasound/Source/MetasoundGenerator/Public/MetasoundGenerator.h
--- a/Engine/Plugins/Runtime/Metasound/Source/MetasoundGenerator/Public/MetasoundGenerator.h
+++ b/Engine/Plugins/Runtime/Metasound/Source/MetasoundGenerator/Public/MetasoundGenerator.h
@@ -3,83 +3,124 @@
 #pragma once
 
 #include "CoreMinimal.h"
-#include "Components/SynthComponent.h"
 #include "DSP/Dsp.h"
 #include "MetasoundOperatorInterface.h"
-#include "MetasoundDataReferenceTypes.h"
-#include "MetasoundTime.h"
-#include "MetasoundFrequency.h"
+#include "MetasoundExecutableOperator.h"
+#include "MetasoundAudioFormats.h"
+#include "MetasoundPrimitives.h"
 #include "Sound/SoundGenerator.h"
 
-#include "MetasoundGenerator.generated.h"
-
 namespace Metasound
 {
+	/** FMetasoundGenerator generates audio from a given metasound IOperator
+	 * which produces a multichannel audio output.
+	 */
 	class METASOUNDGENERATOR_API FMetasoundGenerator : public ISoundGenerator
 	{
 	public:
 		using FOperatorUniquePtr = TUniquePtr<Metasound::IOperator>;
-		using FAudioReadRef = Metasound::FAudioBufferReadRef;
+		using FMultichannelAudioFormatReadRef = Metasound::FMultichannelAudioFormatReadRef;
+
+		/** Create the generator with a graph operator and an output audio reference.
+		 *
+		 * @param InGraphOperator - Unique pointer to the IOperator which executes the entire graph.
+		 * @param InGraphOutputAudioRef - Read reference to the audio buffer filled by the InGraphOperator.
+		 */
+		FMetasoundGenerator(FOperatorUniquePtr InGraphOperator, const FMultichannelAudioFormatReadRef& InGraphOutputAudioRef);
 
-		FMetasoundGenerator(FOperatorUniquePtr InOperator, const FAudioReadRef& InOperatorReadBuffer);
 		virtual ~FMetasoundGenerator();
 
-		void SetFrequency(float InFrequency);
-		void SetBopPeriod(float InPeriodInSeconds);
+		/** Set the value of a graph's input data using the assignment operator.
+		 *
+		 * @param InName - The name of the graph's input data reference.
+		 * @param InData - The value to assign.
+		 */
+		template<typename DataType>
+		void SetInputValue(const FString& InName, DataType InData)
+		{
+			typedef TDataWriteReference< TDecay<DataType>::Type > FDataWriteRef;
+
+			const FDataReferenceCollection& InputCollection = RootExecuter.GetInputs();
+
+			// Check if an input data reference with the given name and type exist in the graph.
+			bool bContainsWriteReference = InputCollection.ContainsDataWriteReference< TDecay<DataType>::Type >(InName);
+			if (ensureMsgf(bContainsWriteReference, TEXT("Operator does not contain write reference name \"%s\" of type \"%s\""), *InName, *GetMetasoundDataTypeName<DataType>().ToString()))
+			{
+				FDataWriteRef WriteRef = InputCollection.GetDataWriteReference<DataType>(InName);
+
+				// call assignment operator of DataType.
+				*WriteRef = InData;
+			}
+		}
+
+		/** Apply a function to the graph's input data.
+		 *
+		 * @param InName - The name of the graph's input data reference.
+		 * @param InFunc - A function which takes the DataType as an input.
+		 */ 
+		template<typename DataType>
+		void ApplyToInputValue(const FString& InName, TUniqueFunction<void(DataType&)> InFunc)
+		{
+			// Get decayed type as InFunc could take a const qualified type.
+			typedef TDataWriteReference< TDecay<DataType>::Type > FDataWriteRef;
+
+			const FDataReferenceCollection& InputCollection = RootExecuter.GetInputs();
+
+			// Check if an input data reference with the given name and type exists in the graph.
+			bool bContainsWriteReference = InputCollection.ContainsDataWriteReference< TDecay<DataType>::Type >(InName);
+			if (ensureMsgf(bContainsWriteReference, TEXT("Operator does not contain write reference name \"%s\" of type \"%s\""), *InName, *GetMetasoundDataTypeName<DataType>().ToString()))
+			{
+				FDataWriteRef WriteRef = InputCollection.GetDataWriteReference<DataType>(InName);
+
+				// Apply function to DataType
+				InFunc(*WriteRef);
+			}
+		}
+
+		/** Update the current graph operator with a new graph operator. The number of channels
+		 * of InGraphOutputAudioRef must match the existing number of channels reported by
+		 * GetNumChannels() in order for this function to successfully replace the graph operator.
+		 *
+		 * @param InGraphOperator - Unique pointer to the IOperator which executes the entire graph.
+		 * @param InGraphOutputAudioRef - Read reference to the audio buffer filled by the InGraphOperator.
+		 *
+		 * @return Returns true if the graph operator was successfully swapped. False otherwise.
+		 */
+		bool UpdateGraphOperator(FOperatorUniquePtr InGraphOperator, const FMultichannelAudioFormatReadRef& InGraphOutputAudioRef);
+
+		/** Return the number of audio channels. */
+		int32 GetNumChannels() const;
 
 		//~ Begin FSoundGenerator
 		virtual int32 OnGenerateAudio(float* OutAudio, int32 NumSamples) override;
 		//~ End FSoundGenerator
 
 	private:
-		int32 FillWithBuffer(const Audio::AlignedFloatBuffer& InBuffer, float* OutAudio, int32 MaxNumOutputSamples); 
-
-		using FExecuteFunction = Metasound::IOperator::FExecuteFunction;
-
-		Audio::AlignedFloatBuffer OverflowBuffer;
-		FOperatorUniquePtr Operator;
-		FAudioReadRef OperatorReadBuffer;
-
-		FFrequencyWriteRef FrequencyRef;
-		FFloatTimeWriteRef BopPeriodRef;
-
-		int32 OperatorReadNum;
-		FExecuteFunction ExecuteOperator;
-	};
-}
-
-UCLASS(ClassGroup = Synth, meta = (BlueprintSpawnableComponent))
-class METASOUNDGENERATOR_API USynthComponentMetasoundGenerator : public USynthComponent
-{
-    GENERATED_BODY()
-
-public:
-
-    USynthComponentMetasoundGenerator(const FObjectInitializer& ObjInitializer);
 
-    virtual ~USynthComponentMetasoundGenerator();
+		// Internal set graph after checking compatibility.
+		void SetGraph(FOperatorUniquePtr InGraphOperator, const FMultichannelAudioFormatReadRef& InGraphOutputAudioRef);
 
-	UPROPERTY(EditAnywhere, BlueprintReadWrite, Category=Metasound, meta = (ClampMin="0.0001", ClampMax="10.0"))
-	float BopPeriod;
+		// Fill OutAudio with data in InBuffer, up to maximum number of samples.
+		// Returns the number of samples used.
+		int32 FillWithBuffer(const Audio::AlignedFloatBuffer& InBuffer, float* OutAudio, int32 MaxNumOutputSamples); 
 
-	UPROPERTY(EditAnywhere, BlueprintReadWrite, Category=Metasound, meta = (ClampMin="20", ClampMax="800"))
-	float Frequency;
+		// Metasound creates deinterleaved audio while sound generator requires interleaved audio.
+		void InterleaveGeneratedAudio();
+		
+		FExecuter RootExecuter;
 
-	UFUNCTION(BlueprintCallable, Category=Metasound)
-	void SetBopPeriod(float InBopPeriod);
+		int32 NumChannels;
+		int32 NumFramesPerExecute;
+		int32 NumSamplesPerExecute;
 
-	UFUNCTION(BlueprintCallable, Category = Metasound)
-	void SetFrequency(float InFrequency);
+		FMultichannelAudioFormatReadRef GraphOutputAudioRef;
 
-    virtual ISoundGeneratorPtr CreateSoundGenerator(int32 InSampleRate, int32 InNumChannels) override;
+		Audio::AlignedFloatBuffer InterleavedAudioBuffer;
 
-#if WITH_EDITOR
-	virtual void PostEditChangeProperty(struct FPropertyChangedEvent & PropertyChangedEvent) override;
-#endif
+		Audio::AlignedFloatBuffer OverflowBuffer;
+	};
+}
 
-private:
-	ISoundGeneratorPtr Generator;
-};
 
 
 
diff --git a/Engine/Plugins/Runtime/Metasound/Source/MetasoundGenerator/Public/SynthComponentMetasoundExample.h b/Engine/Plugins/Runtime/Metasound/Source/MetasoundGenerator/Public/SynthComponentMetasoundExample.h
new file mode 100644
--- /dev/null
+++ b/Engine/Plugins/Runtime/Metasound/Source/MetasoundGenerator/Public/SynthComponentMetasoundExample.h
@@ -0,0 +1,73 @@
+// Copyright Epic Games, Inc. All Rights Reserved.
+
+#pragma once
+
+#include "CoreMinimal.h"
+
+#include "Components/SynthComponent.h"
+
+#include "MetasoundOperatorInterface.h"
+#include "MetasoundNodeInterface.h"
+#include "MetasoundAudioFormats.h"
+#include "MetasoundGenerator.h"
+
+#include "SynthComponentMetasoundExample.generated.h"
+
+
+/** USynthComponentMetasoundExample is an class to show some of the interface 
+ * for MetasoundGraphCore. It's not very useful except as a code example.
+ */
+UCLASS(ClassGroup = Synth, meta = (BlueprintSpawnableComponent))
+class METASOUNDGENERATOR_API USynthComponentMetasoundExample : public USynthComponent
+{
+    GENERATED_BODY()
+
+public:
+	using FOperatorUniquePtr = TUniquePtr<Metasound::IOperator>;
+	using FMultichannelAudioFormatReadRef = Metasound::FMultichannelAudioFormatReadRef;
+
+    USynthComponentMetasoundExample(const FObjectInitializer& ObjInitializer);
+
+    virtual ~USynthComponentMetasoundExample();
+
+	/** Parameters which can be set on the input of the metasound graph during
+	 * runtime. 
+	 */
+	UPROPERTY(EditAnywhere, BlueprintReadWrite, Category=Metasound)
+	TMap<FString, float> FloatParameters;
+
+	/** Set a float parameter on a metasound. 
+	 *
+	 * @param InName - Name of parameter to set.
+	 * @param InValue - Value of parameter.
+	 */
+	UFUNCTION(BlueprintCallable, Category=Metasound)
+	void SetFloatParameter(const FString& InName, float InValue);
+
+	/** Sets the graph operator to use to generate audio. The given operator will be used
+	 * the next time CreateSoundGenerator(...) is called.
+	 *
+	 * @param InGraphOperator - Operator which executes entire MetasoundGraph.
+	 * @param InOutputAudioName - Name of the FMultichannelAudioReadRef data ref in the graph output.
+	 */
+	void SetGraphOperator(FOperatorUniquePtr InGraphOperator, const FString& InOutputAudioName);
+
+	/** Attempt to push a new graph operator to the existing metasound generator.
+	 *
+	 * @param InGraphOperator - Operator which executes entire MetasoundGraph.
+	 * @param InOutputAudioName - Name of the FMultichannelAudioReadRef data ref in the graph output.
+	 *
+	 * @return True on success, false on failure.
+	 */
+	bool PushGraphOperator(FOperatorUniquePtr InGraphOperator, const FString& InOutputAudioName);
+
+    virtual ISoundGeneratorPtr CreateSoundGenerator(int32 InSampleRate, int32 InNumChannels) override;
+
+#if WITH_EDITOR
+	virtual void PostEditChangeProperty(struct FPropertyChangedEvent & PropertyChangedEvent) override;
+#endif
+
+private:
+
+	TSharedPtr<Metasound::FMetasoundGenerator, ESPMode::ThreadSafe> Generator;
+};
diff --git a/Engine/Plugins/Runtime/Metasound/Source/MetasoundGraphCore/Private/MetasoundGraphOperator.cpp b/Engine/Plugins/Runtime/Metasound/Source/MetasoundGraphCore/Private/MetasoundGraphOperator.cpp
--- a/Engine/Plugins/Runtime/Metasound/Source/MetasoundGraphCore/Private/MetasoundGraphOperator.cpp
+++ b/Engine/Plugins/Runtime/Metasound/Source/MetasoundGraphCore/Private/MetasoundGraphOperator.cpp
@@ -53,7 +53,7 @@ namespace Metasound
 
 	void FGraphOperator::Execute()
 	{
-		FExecutableOperatorWrapper* StackPtr = OperatorStack.GetData();
+		FExecuter* StackPtr = OperatorStack.GetData();
 		const int32 Num = OperatorStack.Num();
 		for (int32 i = 0; i < Num; i++)
 		{
diff --git a/Engine/Plugins/Runtime/Metasound/Source/MetasoundGraphCore/Public/MetasoundAudioBuffer.h b/Engine/Plugins/Runtime/Metasound/Source/MetasoundGraphCore/Public/MetasoundAudioBuffer.h
--- a/Engine/Plugins/Runtime/Metasound/Source/MetasoundGraphCore/Public/MetasoundAudioBuffer.h
+++ b/Engine/Plugins/Runtime/Metasound/Source/MetasoundGraphCore/Public/MetasoundAudioBuffer.h
@@ -138,5 +138,5 @@ namespace Metasound
 #endif				
 	};
 
-	DECLARE_METASOUND_DATA_REFERENCE_TYPES(FAudioBuffer, "AudioBuffer", 0x1b02dac2 , FAudioBufferTypeInfo, FAudioBufferReadRef, FAudioBufferWriteRef);
+	DECLARE_METASOUND_DATA_REFERENCE_TYPES(FAudioBuffer, "Audio:Buffer", 0x1b02dac2 , FAudioBufferTypeInfo, FAudioBufferReadRef, FAudioBufferWriteRef);
 };
diff --git a/Engine/Plugins/Runtime/Metasound/Source/MetasoundGraphCore/Public/MetasoundAudioFormats.h b/Engine/Plugins/Runtime/Metasound/Source/MetasoundGraphCore/Public/MetasoundAudioFormats.h
--- a/Engine/Plugins/Runtime/Metasound/Source/MetasoundGraphCore/Public/MetasoundAudioFormats.h
+++ b/Engine/Plugins/Runtime/Metasound/Source/MetasoundGraphCore/Public/MetasoundAudioFormats.h
@@ -18,7 +18,7 @@ namespace Metasound
 	 * of the FUnformattedAudio. All audio buffers within a FUnformattedAudio object must contain the 
 	 * same number of audio frames.
 	 */
-	class FUnformattedAudio
+	class METASOUNDGRAPHCORE_API FUnformattedAudio
 	{
 		public:
 			/** FUnformattedAudio Constructor.
@@ -93,7 +93,7 @@ namespace Metasound
 	 * of the FMultichannelAudioFormat. All audio buffers within a FMultichannelAudioFormat object must contain the 
 	 * same number of audio frames.
 	 */
-	class FMultichannelAudioFormat
+	class METASOUNDGRAPHCORE_API FMultichannelAudioFormat
 	{
 		public:
 			/** FMultichannelAudioFormat Constructor.
@@ -311,6 +311,11 @@ namespace Metasound
 	 */
 	class FMonoAudioFormat : public TStaticChannelAudioFormat<1>
 	{
+			// This is used to grant specific classes the ability to create FMonoAudioFormat
+			// using a FAudioBufferReadRef
+			//enum EPrivateToken { Token };
+			//static const EPrivateToken PrivateToken = EPrivateToken::Token;
+
 		public:
 			using Super = TStaticChannelAudioFormat<1>;
 
@@ -328,6 +333,19 @@ namespace Metasound
 			{
 			}
 
+			// TODO: clean up this mess.
+			// Construct a FMonoAudioFormat a readable buffer reference.
+			// 
+			// This constructor should only be used when it can be assured 
+			// that the constructed object will not provide writable access 
+			// to the passed in audio buffers. 
+			/*
+			FMonoAudioFormat(const FAudioBufferReadRef& InAudio, EPrivateToken InToken)
+			:	Super({WriteCast(InAudio)})
+			{
+			}
+			*/
+
 			/** Return writable audio buffer reference of center channel. */
 			FAudioBufferWriteRef GetCenter() { return GetBuffer<0>(); }
 
@@ -335,25 +353,29 @@ namespace Metasound
 			FAudioBufferReadRef GetCenter() const { return GetBuffer<0>(); }
 
 		private:
+
+
 			// Friendship with the data reference class gives it access to the 
 			// protected constructor for the scenario where a 
 			// TDataReadReference<*> is constructed with FAudioBufferReadRefs. 
 			// The constructor cannot be public as it would provide writable
 			// access to the passed in audio buffers, even though the passed 
 			// buffers explicitly were read references.
-			friend class TDataReadReference<FMonoAudioFormat>;
-
-			// Construct a FMonoAudioFormat a readable buffer reference.
-			// 
-			// This constructor should only be used when it can be assured 
-			// that the constructed object will not provide writable access 
-			// to the passed in audio buffers. 
-			FMonoAudioFormat(const FAudioBufferReadRef& InAudio)
-			:	TStaticChannelAudioFormat<1>({WriteCast(InAudio)})
-			{
-			}
+			//friend class TDataReadReference<FMonoAudioFormat>;
 	};
 
+	// Template specialization for TDataReadReference<FMonoAudioFormat> constructor to use
+	// private token.
+	/*
+	template<>
+	template<>
+	TDataReadReference<FMonoAudioFormat>::TDataReadReference(const FAudioBufferReadRef& InReadBuffer)
+	:	FDataReference(InReadBuffer, FMonoAudioFormat::PrivateToken)
+	{
+	}
+	*/
+	
+
 	/** FStereoAudioFormat represents stereo audio containing two channels of 
 	 * audio.
 	 *
@@ -422,4 +444,5 @@ namespace Metasound
 
 	DECLARE_METASOUND_DATA_REFERENCE_TYPES(FStereoAudioFormat, "Audio:Stereo", 0xb55304e2 , FStereoAudioFormatTypeInfo, FStereoAudioFormatReadRef, FStereoAudioFormatWriteRef);
 
+
 }
diff --git a/Engine/Plugins/Runtime/Metasound/Source/MetasoundGraphCore/Public/MetasoundDataReference.h b/Engine/Plugins/Runtime/Metasound/Source/MetasoundGraphCore/Public/MetasoundDataReference.h
--- a/Engine/Plugins/Runtime/Metasound/Source/MetasoundGraphCore/Public/MetasoundDataReference.h
+++ b/Engine/Plugins/Runtime/Metasound/Source/MetasoundGraphCore/Public/MetasoundDataReference.h
@@ -78,6 +78,14 @@ namespace Metasound
 		// without specialization.
 		static_assert(TSpecializationHelper<DataType>::Value, "TDataReferenceTypeInfo must be specialized.  Use macro DECLARE_METASOUND_DATA_REFERENCE_TYPES");
 	};
+
+	template<typename DataType>
+	const FName GetMetasoundDataTypeName() 
+	{
+		static const FName TypeName = FName(TDataReferenceTypeInfo< TDecay<DataType>::Type >::TypeName);
+
+		return TypeName;
+	}
 	
 	/** Template class for a paramter reference. 
 	 *
@@ -136,7 +144,7 @@ namespace Metasound
 			/** Return the name of the underlying type. */
 			virtual const FName& GetDataTypeName() const override
 			{
-				static const FName TypeName = FName(FInfoType::TypeName);
+				static const FName TypeName = GetMetasoundDataTypeName<DataType>();
 				return TypeName;
 			}
 
@@ -258,6 +266,8 @@ namespace Metasound
 			}
 	};
 
+
+
 	/** Cast a TDataReadReference to a TDataWriteReference. 
 	 *
 	 * In general TDataReadReferences should not be converted into TDataWriteReferences unless the caller
diff --git a/Engine/Plugins/Runtime/Metasound/Source/MetasoundGraphCore/Public/MetasoundExecutableOperator.h b/Engine/Plugins/Runtime/Metasound/Source/MetasoundGraphCore/Public/MetasoundExecutableOperator.h
--- a/Engine/Plugins/Runtime/Metasound/Source/MetasoundGraphCore/Public/MetasoundExecutableOperator.h
+++ b/Engine/Plugins/Runtime/Metasound/Source/MetasoundGraphCore/Public/MetasoundExecutableOperator.h
@@ -3,6 +3,7 @@
 #pragma once
 
 #include "CoreMinimal.h"
+
 #include "MetasoundOperatorInterface.h"
 
 namespace Metasound
@@ -22,6 +23,7 @@ namespace Metasound
 	 * 	  }
 	 * 	};
 	 */
+
 	template<class DerivedOperatorType>
 	class TExecutableOperator : public IOperator
 	{
@@ -46,20 +48,42 @@ namespace Metasound
 			}
 	};
 
-	/** FExecutableOperatorWrapper 
+	class FNoOpOperator : public IOperator
+	{
+		public:
+			virtual ~FNoOpOperator() {}
+
+			virtual FExecuteFunction GetExecuteFunction() override
+			{
+				return nullptr;
+			}
+	};
+
+	/** FExecuter
 	 *
 	 * Wraps an IOperator and provides an Execute() member function.
 	 */
-	class FExecutableOperatorWrapper 
+	class FExecuter : public IOperator
 	{
 		public:
 			using FOperatorPtr = TUniquePtr<IOperator>;
 			using FExecuteFunction = IOperator::FExecuteFunction;
 
-			FExecutableOperatorWrapper(FOperatorPtr InOperator)
-			:	Operator(MoveTemp(InOperator))
-			,	ExecuteFunction(&FExecutableOperatorWrapper::NoOp)
+			FExecuter()
+			:	ExecuteFunction(&FExecuter::NoOp)
 			{
+			}
+
+			FExecuter(FOperatorPtr InOperator)
+			:	ExecuteFunction(&FExecuter::NoOp)
+			{
+				SetOperator(MoveTemp(InOperator));
+			}
+
+			void SetOperator(FOperatorPtr InOperator)
+			{
+				Operator = MoveTemp(InOperator);
+
 				if (Operator.IsValid())
 				{
 					FExecuteFunction Func = Operator->GetExecuteFunction();
@@ -69,10 +93,10 @@ namespace Metasound
 						ExecuteFunction = Func;
 					}
 				}
-			}
-
-			virtual ~FExecutableOperatorWrapper()
-			{
+				else
+				{
+					ExecuteFunction = &FExecuter::NoOp;
+				}
 			}
 
 			void Execute()
@@ -82,7 +106,41 @@ namespace Metasound
 
 			bool IsNoOp()
 			{
-				return (ExecuteFunction == &FExecutableOperatorWrapper::NoOp);
+				return (ExecuteFunction == &FExecuter::NoOp);
+			}
+
+			bool IsValid() const
+			{
+				return Operator.IsValid();
+			}
+
+			virtual const FDataReferenceCollection& GetInputs() const override
+			{
+				static FDataReferenceCollection EmptyCollection;
+
+				if (Operator.IsValid())
+				{
+					return Operator->GetInputs();
+				}
+
+				return EmptyCollection;
+			}
+
+			virtual const FDataReferenceCollection& GetOutputs() const override
+			{
+				static FDataReferenceCollection EmptyCollection;
+
+				if (Operator.IsValid())
+				{
+					return Operator->GetOutputs();
+				}
+
+				return EmptyCollection;
+			}
+
+			virtual FExecuteFunction GetExecuteFunction() override
+			{
+				return ExecuteFunction;
 			}
 
 		private:
diff --git a/Engine/Plugins/Runtime/Metasound/Source/MetasoundGraphCore/Public/MetasoundGraphOperator.h b/Engine/Plugins/Runtime/Metasound/Source/MetasoundGraphCore/Public/MetasoundGraphOperator.h
--- a/Engine/Plugins/Runtime/Metasound/Source/MetasoundGraphCore/Public/MetasoundGraphOperator.h
+++ b/Engine/Plugins/Runtime/Metasound/Source/MetasoundGraphCore/Public/MetasoundGraphOperator.h
@@ -42,7 +42,7 @@ namespace Metasound
 
 			static void ExecuteFunction(IOperator* InOperator);
 
-			TArray<FExecutableOperatorWrapper> OperatorStack;
+			TArray<FExecuter> OperatorStack;
 
 			FDataReferenceCollection Inputs;
 			FDataReferenceCollection Outputs;
diff --git a/Engine/Plugins/Runtime/Metasound/Source/MetasoundGraphCore/Public/MetasoundTime.h b/Engine/Plugins/Runtime/Metasound/Source/MetasoundGraphCore/Public/MetasoundTime.h
--- a/Engine/Plugins/Runtime/Metasound/Source/MetasoundGraphCore/Public/MetasoundTime.h
+++ b/Engine/Plugins/Runtime/Metasound/Source/MetasoundGraphCore/Public/MetasoundTime.h
@@ -35,6 +35,9 @@ namespace Metasound
 			static_assert(TIsFloatingPoint<TimeType>::Value, "TimeType must be floating point");
 
 		public:
+			using FQuantizationType = QuantizationType;
+			using FTimeTYpe = TimeType;
+
 			/** Default constructor */
 			TTime()
 			:	Value(0)
