Moving implementation off of UMetasound to avoid engine requirement
- Move deserialize/serialize utilities off of module and into own FileConversion namespace. Stub in landing place for export in Intermediates.

diff --git a/Engine/Plugins/Runtime/Metasound/Source/MetasoundEngine/Private/Metasound.cpp b/Engine/Plugins/Runtime/Metasound/Source/MetasoundEngine/Private/Metasound.cpp
--- a/Engine/Plugins/Runtime/Metasound/Source/MetasoundEngine/Private/Metasound.cpp
+++ b/Engine/Plugins/Runtime/Metasound/Source/MetasoundEngine/Private/Metasound.cpp
@@ -8,7 +8,7 @@
 
 UMetasound::UMetasound(const FObjectInitializer& ObjectInitializer)
 	: Super(ObjectInitializer)
-	, AccessPoint(MakeShared<Metasound::Frontend::FDescriptionAccessPoint>(RootMetasoundDocument))
+	, FMetasoundAssetBase()
 {
 }
 
@@ -18,29 +18,34 @@ UEdGraph* UMetasound::GetGraph()
 	return Graph;
 }
 
+const UEdGraph* UMetasound::GetGraph() const
+{
+	return Graph;
+}
+
 UEdGraph& UMetasound::GetGraphChecked()
 {
 	check(Graph);
 	return *Graph;
 }
 
-void UMetasound::SetGraph(UEdGraph* InGraph)
+const UEdGraph& UMetasound::GetGraphChecked() const
 {
-	Graph = InGraph;
+	check(Graph);
+	return *Graph;
 }
-#endif // WITH_EDITORONLY_DATA
 
-FMetasoundClassMetadata UMetasound::GetMetadata()
+void UMetasound::SetGraph(UEdGraph* InGraph)
 {
-	return RootMetasoundDocument.RootClass.Metadata;
+	Graph = InGraph;
 }
+#endif // WITH_EDITORONLY_DATA
 
 // This can be used to update the metadata (name, author, etc) for this metasound.
 // @param InMetadata may be updated with any corrections we do to the input metadata.
 void UMetasound::SetMetadata(FMetasoundClassMetadata& InMetadata)
 {
-	InMetadata.NodeType = EMetasoundClassType::MetasoundGraph;
-	RootMetasoundDocument.RootClass.Metadata = InMetadata;
+	FMetasoundAssetBase::SetMetadata(InMetadata);
 	MarkPackageDirty();
 }
 
@@ -50,68 +55,3 @@ void UMetasound::SetMetasoundDocument(const FMetasoundDocument& InDocument)
 {
 	RootMetasoundDocument = InDocument;
 }
-
-// returns a weak pointer that can be used to build a TDescriptionPtr
-// for direct editing of the FMetasoundClassDescription tree.
-// For advance use only, and requires knowledge of Metasound::Frontend::FDescPath syntax.
-// For most use cases, use GetGraphHandle() instead.
-TWeakPtr<Metasound::Frontend::FDescriptionAccessPoint> UMetasound::GetGraphAccessPoint()
-{
-	return TWeakPtr<Metasound::Frontend::FDescriptionAccessPoint>(AccessPoint);
-}
-
-// Get the handle for the root metasound graph of this asset.
-Metasound::Frontend::FGraphHandle UMetasound::GetRootGraphHandle()
-{
-	using FDescPath = Metasound::Frontend::FDescPath;
-	using FGraphHandle = Metasound::Frontend::FGraphHandle;
-	using FHandleInitParams = Metasound::Frontend::FHandleInitParams;
-	using EFromDocument = Metasound::Frontend::Path::EFromDocument;
-	using EFromClass = Metasound::Frontend::Path::EFromClass;
-
-	FDescPath PathToGraph = FDescPath()[EFromDocument::ToRootClass][EFromClass::ToGraph];
-	FHandleInitParams InitParams = { GetGraphAccessPoint(), PathToGraph, RootMetasoundDocument.RootClass.Metadata.NodeName, MakeWeakObjectPtr(this) };
-	return FGraphHandle(FHandleInitParams::PrivateToken, InitParams);
-}
-
-TArray<Metasound::Frontend::FGraphHandle> UMetasound::GetAllSubgraphHandles()
-{
-	using FDescPath = Metasound::Frontend::FDescPath;
-	using FGraphHandle = Metasound::Frontend::FGraphHandle;
-	using FHandleInitParams = Metasound::Frontend::FHandleInitParams;
-	using EFromDocument = Metasound::Frontend::Path::EFromDocument;
-	using EFromClass = Metasound::Frontend::Path::EFromClass;
-
-	TArray<FGraphHandle> OutArray;
-
-	FDescPath RootPathForDependencyGraphs = FDescPath()[EFromDocument::ToDependencies];
-
-	TArray<FMetasoundClassDescription>& DependenciesList = RootMetasoundDocument.Dependencies;
-	for (FMetasoundClassDescription& Dependency : DependenciesList)
-	{
-		const bool bIsSubgraph = Dependency.Metadata.NodeType == EMetasoundClassType::MetasoundGraph && Dependency.Graph.Nodes.Num() > 0;
-		if (bIsSubgraph)
-		{
-			FDescPath SubgraphPath = RootPathForDependencyGraphs[Dependency.UniqueID][EFromClass::ToGraph];
-			FHandleInitParams InitParams = { GetGraphAccessPoint(), SubgraphPath, Dependency.Metadata.NodeName, MakeWeakObjectPtr(this) };
-			OutArray.Emplace(FHandleInitParams::PrivateToken, InitParams);
-		}
-	}
-
-	return OutArray;
-}
-
-bool UMetasound::ExportToJSON(const FString& InAbsolutePath)
-{
-	if (TUniquePtr<FArchive> FileWriter = TUniquePtr<FArchive>(IFileManager::Get().CreateFileWriter(*InAbsolutePath)))
-	{
-		FJsonStructSerializerBackend Backend(*FileWriter, EStructSerializerBackendFlags::Default);
-		FStructSerializer::Serialize<FMetasoundDocument>(RootMetasoundDocument, Backend);
-		FileWriter->Close();
-
-		return true;
-	}
-
-	ensureAlwaysMsgf(false, TEXT("Failed to create a filewriter with the given path."));
-	return false;
-}
diff --git a/Engine/Plugins/Runtime/Metasound/Source/MetasoundEngine/Private/MetasoundEngineModule.cpp b/Engine/Plugins/Runtime/Metasound/Source/MetasoundEngine/Private/MetasoundEngineModule.cpp
--- a/Engine/Plugins/Runtime/Metasound/Source/MetasoundEngine/Private/MetasoundEngineModule.cpp
+++ b/Engine/Plugins/Runtime/Metasound/Source/MetasoundEngine/Private/MetasoundEngineModule.cpp
@@ -9,23 +9,6 @@ DEFINE_LOG_CATEGORY(LogMetasoundEngine);
 
 class FMetasoundEngineModule : public IMetasoundEngineModule
 {
-	virtual UMetasound* DeserializeMetasound(const FString& InPath) override
-	{
-		FMetasoundDocument MetasoundDoc;
-		ensureAlwaysMsgf(false, TEXT("Implement the actual loading part!"));
-
-		UMetasound* NewMetasoundNode = NewObject<UMetasound>();
-		NewMetasoundNode->SetMetasoundDocument(MetasoundDoc);
-
-		return NewMetasoundNode;
-	}
-
-	virtual void SerializeMetasound(const UMetasound& InMetasound, const FString& InPath) override
-	{
-		ensureAlwaysMsgf(false, TEXT("Implement the actual saving part!"));
-		ensureAlwaysMsgf(false, TEXT("Decide whether to delete this graph from the asset that owns it!"));
-	}
-
 	virtual void StartupModule() override
 	{
 		UE_LOG(LogMetasoundEngine, Log, TEXT("Metasound Engine Initialized"));
diff --git a/Engine/Plugins/Runtime/Metasound/Source/MetasoundEngine/Private/MetasoundEngineModule.h b/Engine/Plugins/Runtime/Metasound/Source/MetasoundEngine/Private/MetasoundEngineModule.h
--- a/Engine/Plugins/Runtime/Metasound/Source/MetasoundEngine/Private/MetasoundEngineModule.h
+++ b/Engine/Plugins/Runtime/Metasound/Source/MetasoundEngine/Private/MetasoundEngineModule.h
@@ -13,15 +13,4 @@ DECLARE_LOG_CATEGORY_EXTERN(LogMetasoundEngine, Log, All);
 
 class IMetasoundEngineModule : public IModuleInterface
 {
-	// Deserializes Metasound from the graph at the provided path.
-	// @param InPath Content directory path to load metasound from.
-	// @returns New Metasound object on success, nullptr on failure.
-	virtual UMetasound* DeserializeMetasound(const FString& InPath) = 0;
-
-	// Serializes Metasound to the provided path.
-	// @param InMetasound Metasound to serialize.
-	// @param InPath Path to serialize to.
-	// @todo Decide if this means that we remove the graph in our current asset.
-	//       This seems dangerous since we can fail to save the new asset and scratch the old one.
-	virtual void SerializeMetasound(const UMetasound& InMetasound, const FString& InPath) = 0;
 };
diff --git a/Engine/Plugins/Runtime/Metasound/Source/MetasoundEngine/Public/Metasound.h b/Engine/Plugins/Runtime/Metasound/Source/MetasoundEngine/Public/Metasound.h
--- a/Engine/Plugins/Runtime/Metasound/Source/MetasoundEngine/Public/Metasound.h
+++ b/Engine/Plugins/Runtime/Metasound/Source/MetasoundEngine/Public/Metasound.h
@@ -4,6 +4,8 @@
 #include "CoreMinimal.h"
 #include "MetasoundFrontend.h"
 #include "MetasoundFrontendDataLayout.h"
+#include "MetasoundAssetBase.h"
+
 
 #if WITH_EDITORONLY_DATA
 #include "EdGraph/EdGraph.h"
@@ -17,7 +19,7 @@
  * Because of this, they can have any inputs or outputs they need.
  */
 UCLASS(hidecategories = object, BlueprintType)
-class METASOUNDENGINE_API UMetasound : public UObject
+class METASOUNDENGINE_API UMetasound : public UObject, public FMetasoundAssetBase
 {
 	GENERATED_BODY()
 
@@ -30,8 +32,6 @@ protected:
 	UEdGraph* Graph;
 #endif // WITH_EDITORONLY_DATA
 
-	TSharedPtr<Metasound::Frontend::FDescriptionAccessPoint> AccessPoint;
-
 public:
 	UMetasound(const FObjectInitializer& ObjectInitializer);
 
@@ -41,7 +41,9 @@ public:
 	// Metasound UObject for editor serialization purposes.
 	// @return Editor graph associated with UMetasound.
 	UEdGraph* GetGraph();
+	const UEdGraph* GetGraph() const;
 	UEdGraph& GetGraphChecked();
+	const UEdGraph& GetGraphChecked() const;
 
 	// Sets the graph associated with this Metasound. Graph is required to be referenced on
 	// Metasound UObject for editor serialization purposes.
@@ -49,25 +51,20 @@ public:
 	void SetGraph(UEdGraph* InGraph);
 #endif // WITH_EDITORONLY_DATA
 
-	FMetasoundClassMetadata GetMetadata();
+	FMetasoundDocument& GetDocument() override
+	{
+		return RootMetasoundDocument;
+	}
+
+	UObject* GetOwningAsset() override
+	{
+		return this;
+	}
 
 	// Updates the Metasound's metadata (name, author, etc).
 	// @param InMetadata Metadata containing corrections to the class metadata.
-	void SetMetadata(FMetasoundClassMetadata& InMetadata);
+	void SetMetadata(FMetasoundClassMetadata& InMetadata) override;
 
 	// Deletes Metasound's current metasound document, and replaces it with InClassDescription.
 	void SetMetasoundDocument(const FMetasoundDocument& InDocument);
-
-	// returns a weak pointer that can be used to build a TDescriptionPtr
-	// for direct editing of the FMetasoundClassDescription tree.
-	// For advance use only, and requires knowledge of Metasound::Frontend::FDescPath syntax.
-	// For most use cases, use GetGraphHandle() instead.
-	TWeakPtr<Metasound::Frontend::FDescriptionAccessPoint> GetGraphAccessPoint();
-
-	// Get the handle for the root metasound graph of this asset.
-	Metasound::Frontend::FGraphHandle GetRootGraphHandle();
-
-	TArray<Metasound::Frontend::FGraphHandle> GetAllSubgraphHandles();
-
-	bool ExportToJSON(const FString& InAbsolutePath);
 };
\ No newline at end of file
diff --git a/Engine/Plugins/Runtime/Metasound/Source/MetasoundEngine/Public/MetasoundAssetSerializer.h b/Engine/Plugins/Runtime/Metasound/Source/MetasoundEngine/Public/MetasoundAssetSerializer.h
new file mode 100644
--- /dev/null
+++ b/Engine/Plugins/Runtime/Metasound/Source/MetasoundEngine/Public/MetasoundAssetSerializer.h
@@ -0,0 +1,43 @@
+// Copyright Epic Games, Inc. All Rights Reserved.
+#pragma once
+
+#include "CoreMinimal.h"
+#include "Metasound.h"
+#include "MetasoundFrontendDataLayout.h"
+#include "Misc/AssertionMacros.h"
+#include "UObject/UObjectGlobals.h"
+
+
+namespace Metasound
+{
+	namespace Engine
+	{
+		const FString MetasoundExtension(TEXT(".metasound"));
+
+		class AssetSerializer
+		{
+
+		public:
+			// Deserializes Metasound from the graph at the provided path.
+			// @param InPath Content directory path to load metasound from.
+			// @returns New Metasound object on success, nullptr on failure.
+			static void JSONToMetasound(const FString& InPath, UMetasound& InMetasound)
+			{
+				FMetasoundDocument MetasoundDoc;
+				ensureAlwaysMsgf(false, TEXT("Implement the actual loading part!"));
+
+				InMetasound.SetMetasoundDocument(MetasoundDoc);
+			}
+
+			// Serializes Metasound to the provided path.
+			// @param InMetasound Metasound to serialize.
+			// @todo Decide if this means that we remove the graph in our current asset.
+			//       This seems dangerous since we can fail to save the new asset and scratch the old one.
+			static bool MetasoundToJSON(UMetasound& InMetasound)
+			{
+				const FString Path = FPaths::ProjectIntermediateDir() / TEXT("Metasounds") + FPaths::ChangeExtension(InMetasound.GetPathName(), MetasoundExtension);
+				return InMetasound.GetRootGraphHandle().ExportToJSONAsset(Path);
+			}
+		};
+	} // namespace Engine
+} // namespace Metasound
\ No newline at end of file
diff --git a/Engine/Plugins/Runtime/Metasound/Source/MetasoundFrontend/Private/MetasoundAssetBase.cpp b/Engine/Plugins/Runtime/Metasound/Source/MetasoundFrontend/Private/MetasoundAssetBase.cpp
new file mode 100644
--- /dev/null
+++ b/Engine/Plugins/Runtime/Metasound/Source/MetasoundFrontend/Private/MetasoundAssetBase.cpp
@@ -0,0 +1,84 @@
+// Copyright Epic Games, Inc. All Rights Reserved.
+
+#include "MetasoundAssetBase.h"
+
+#include "Backends/JsonStructSerializerBackend.h"
+#include "IStructSerializerBackend.h"
+#include "StructSerializer.h"
+
+
+void FMetasoundAssetBase::SetMetadata(FMetasoundClassMetadata& InMetadata)
+{
+	InMetadata.NodeType = EMetasoundClassType::MetasoundGraph;
+	GetDocument().RootClass.Metadata = InMetadata;
+}
+
+FMetasoundClassMetadata FMetasoundAssetBase::GetMetadata()
+{
+	return GetDocument().RootClass.Metadata;
+}
+
+TWeakPtr<Metasound::Frontend::FDescriptionAccessPoint> FMetasoundAssetBase::GetGraphAccessPoint()
+{
+	if (!AccessPoint.IsValid())
+	{
+		AccessPoint = MakeShared<Metasound::Frontend::FDescriptionAccessPoint>(GetDocument());
+	}
+
+	return TWeakPtr<Metasound::Frontend::FDescriptionAccessPoint>(AccessPoint);
+}
+
+Metasound::Frontend::FGraphHandle FMetasoundAssetBase::GetRootGraphHandle()
+{
+	using FDescPath = Metasound::Frontend::FDescPath;
+	using FGraphHandle = Metasound::Frontend::FGraphHandle;
+	using FHandleInitParams = Metasound::Frontend::FHandleInitParams;
+	using EFromDocument = Metasound::Frontend::Path::EFromDocument;
+	using EFromClass = Metasound::Frontend::Path::EFromClass;
+
+	FDescPath PathToGraph = FDescPath()[EFromDocument::ToRootClass][EFromClass::ToGraph];
+	FHandleInitParams InitParams = { GetGraphAccessPoint(), PathToGraph, GetDocument().RootClass.Metadata.NodeName, MakeWeakObjectPtr(GetOwningAsset()) };
+	return FGraphHandle(GetPrivateToken(), InitParams);
+}
+
+TArray<Metasound::Frontend::FGraphHandle> FMetasoundAssetBase::GetAllSubgraphHandles()
+{
+	using FDescPath = Metasound::Frontend::FDescPath;
+	using FGraphHandle = Metasound::Frontend::FGraphHandle;
+	using FHandleInitParams = Metasound::Frontend::FHandleInitParams;
+	using EFromDocument = Metasound::Frontend::Path::EFromDocument;
+	using EFromClass = Metasound::Frontend::Path::EFromClass;
+
+	TArray<FGraphHandle> OutArray;
+
+	FDescPath RootPathForDependencyGraphs = FDescPath()[EFromDocument::ToDependencies];
+
+	TArray<FMetasoundClassDescription>& DependenciesList = GetDocument().Dependencies;
+	for (FMetasoundClassDescription& Dependency : DependenciesList)
+	{
+		const bool bIsSubgraph = Dependency.Metadata.NodeType == EMetasoundClassType::MetasoundGraph && Dependency.Graph.Nodes.Num() > 0;
+		if (bIsSubgraph)
+		{
+			FDescPath SubgraphPath = RootPathForDependencyGraphs[Dependency.UniqueID][EFromClass::ToGraph];
+			FHandleInitParams InitParams = { GetGraphAccessPoint(), SubgraphPath, Dependency.Metadata.NodeName, MakeWeakObjectPtr(GetOwningAsset()) };
+			OutArray.Emplace(GetPrivateToken(), InitParams);
+		}
+	}
+
+	return OutArray;
+}
+
+bool FMetasoundAssetBase::ExportToJSON(const FString& InAbsolutePath)
+{
+	if (TUniquePtr<FArchive> FileWriter = TUniquePtr<FArchive>(IFileManager::Get().CreateFileWriter(*InAbsolutePath)))
+	{
+		FJsonStructSerializerBackend Backend(*FileWriter, EStructSerializerBackendFlags::Default);
+		FStructSerializer::Serialize<FMetasoundDocument>(GetDocument(), Backend);
+		FileWriter->Close();
+
+		return true;
+	}
+
+	ensureAlwaysMsgf(false, TEXT("Failed to create a filewriter with the given path."));
+	return false;
+}
\ No newline at end of file
diff --git a/Engine/Plugins/Runtime/Metasound/Source/MetasoundFrontend/Private/MetasoundFrontend.cpp b/Engine/Plugins/Runtime/Metasound/Source/MetasoundFrontend/Private/MetasoundFrontend.cpp
--- a/Engine/Plugins/Runtime/Metasound/Source/MetasoundFrontend/Private/MetasoundFrontend.cpp
+++ b/Engine/Plugins/Runtime/Metasound/Source/MetasoundFrontend/Private/MetasoundFrontend.cpp
@@ -392,7 +392,7 @@ namespace Metasound
 			{
 				FDescPath NodePath = NodePtr.GetPath();
 				FString ClassName = GetNodeClassName();
-				FHandleInitParams InitParams = { NodePtr.GetAccessPoint(), NodePath, ClassName , OwningAsset };
+				FHandleInitParams InitParams = { NodePtr.GetAccessPoint(), NodePath, ClassName, OwningAsset };
 				OutArray.Emplace(FHandleInitParams::PrivateToken, InitParams, InputDescription.Name);
 			}
 
@@ -467,7 +467,7 @@ namespace Metasound
 			return FOutputHandle::InvalidHandle();
 		}
 
-		EMetasoundClassType FNodeHandle::GetNodeType()
+		EMetasoundClassType FNodeHandle::GetNodeType() const
 		{
 			if (!IsValid())
 			{
@@ -477,11 +477,12 @@ namespace Metasound
 			return NodeClass->Metadata.NodeType;
 		}
 
-		FString FNodeHandle::GetNodeClassName()
+		const FString& FNodeHandle::GetNodeClassName() const
 		{
 			if (!IsValid())
 			{
-				return FString();
+				static const FString DefaultClassName;
+				return DefaultClassName;
 			}
 
 			return NodeClass->Metadata.NodeName;
@@ -505,7 +506,7 @@ namespace Metasound
 			OutGraph = FGraphHandle(FHandleInitParams::PrivateToken, InitParams);
 		}
 
-		uint32 FNodeHandle::GetNodeID()
+		uint32 FNodeHandle::GetNodeID() const
 		{
 			if (!IsValid())
 			{
@@ -635,7 +636,7 @@ namespace Metasound
 			return FGraphHandle(FHandleInitParams::Token, InitParams);
 		}
 
-		bool FGraphHandle::IsValid()
+		bool FGraphHandle::IsValid() const
 		{
 			return GraphPtr.IsValid() && GraphsClassDeclaration.IsValid();
 		}
@@ -1012,7 +1013,7 @@ namespace Metasound
 			return FNodeHandle(FHandleInitParams::PrivateToken, InitParams);
 		}
 
-		bool FGraphHandle::RemoveNode(FNodeHandle& InNode, bool bEvenIfInputOrOutputNode /*= false*/)
+		bool FGraphHandle::RemoveNode(const FNodeHandle& InNode, bool bEvenIfInputOrOutputNode /*= false*/)
 		{
 			if (!IsValid())
 			{
@@ -1149,7 +1150,7 @@ namespace Metasound
 			return GraphsClassDeclaration->Metadata;
 		}
 
-		bool FGraphHandle::ExportToJSONAsset(const FString& InAbsolutePath)
+		bool FGraphHandle::ExportToJSONAsset(const FString& InAbsolutePath) const
 		{
 			if (!IsValid())
 			{
diff --git a/Engine/Plugins/Runtime/Metasound/Source/MetasoundFrontend/Public/MetasoundAssetBase.h b/Engine/Plugins/Runtime/Metasound/Source/MetasoundFrontend/Public/MetasoundAssetBase.h
new file mode 100644
--- /dev/null
+++ b/Engine/Plugins/Runtime/Metasound/Source/MetasoundFrontend/Public/MetasoundAssetBase.h
@@ -0,0 +1,53 @@
+// Copyright Epic Games, Inc. All Rights Reserved.
+#pragma once
+
+#include "CoreMinimal.h"
+#include "MetasoundFrontend.h"
+#include "MetasoundFrontendDataLayout.h"
+#include "MetasoundFrontendBaseClasses.h"
+#include "MetasoundGraph.h"
+#include "UObject/WeakObjectPtrTemplates.h"
+
+
+class METASOUNDFRONTEND_API FMetasoundAssetBase
+{
+public:
+	virtual ~FMetasoundAssetBase() = default;
+
+	// Sets/overwrites the root class metadata
+	virtual void SetMetadata(FMetasoundClassMetadata& InMetadata);
+
+	// Returns the root class metadata
+	FMetasoundClassMetadata GetMetadata();
+
+	// Exports the asset to JSON file at provided path
+	bool ExportToJSON(const FString& InAbsolutePath);
+
+	// Returns handle for the root metasound graph of this asset.
+	Metasound::Frontend::FGraphHandle GetRootGraphHandle();
+
+	// Returns all handles for subgraphs referenced
+	TArray<Metasound::Frontend::FGraphHandle> GetAllSubgraphHandles();
+
+protected:
+	// Returns private token allowing implementing asset class to use graph/node handle system
+	static Metasound::Frontend::FHandleInitParams::EPrivateToken GetPrivateToken()
+	{
+		return Metasound::Frontend::FHandleInitParams::PrivateToken;
+	}
+
+	// Returns document object responsible for serializing asset
+	virtual FMetasoundDocument& GetDocument() = 0;
+
+	// Returns the owning asset responsible for transactions applied to metasound
+	virtual UObject* GetOwningAsset() = 0;
+
+	// Returns a weak pointer that can be used to build a TDescriptionPtr
+	// for direct editing of the FMetasoundClassDescription tree.
+	// For advance use only, and requires knowledge of Metasound::Frontend::FDescPath syntax.
+	// For most use cases, use GetRootGraphHandle() instead.
+	TWeakPtr<Metasound::Frontend::FDescriptionAccessPoint> GetGraphAccessPoint();
+
+private:
+	TSharedPtr<Metasound::Frontend::FDescriptionAccessPoint> AccessPoint;
+};
diff --git a/Engine/Plugins/Runtime/Metasound/Source/MetasoundFrontend/Public/MetasoundFrontend.h b/Engine/Plugins/Runtime/Metasound/Source/MetasoundFrontend/Public/MetasoundFrontend.h
--- a/Engine/Plugins/Runtime/Metasound/Source/MetasoundFrontend/Public/MetasoundFrontend.h
+++ b/Engine/Plugins/Runtime/Metasound/Source/MetasoundFrontend/Public/MetasoundFrontend.h
@@ -6,9 +6,13 @@
 #include "MetasoundGraph.h"
 #include "MetasoundFrontendDataLayout.h"
 #include "MetasoundFrontendBaseClasses.h"
+#include "UObject/WeakObjectPtrTemplates.h"
+#include "Serialization/Public/IStructSerializerBackend.h"
+#include "Serialization/Public/Backends/JsonStructSerializerBackend.h"
+#include "Serialization/Public/StructSerializer.h"
 
-
-class UMetasound;
+// Forward Declarations
+class FMetasoundAssetBase;
 
 namespace Metasound
 {	
@@ -32,7 +36,6 @@ namespace Metasound
 	 */
 	namespace Frontend
 	{
-
 		// Struct with the basics of a node class' information,
 		// used to look up that node from our node browser functions,
 		// and also used in FGraphHandle::AddNewNode.
@@ -109,10 +112,10 @@ namespace Metasound
 			friend class FInputHandle;
 			friend class FNodeHandle;
 			friend class FGraphHandle;
-			friend class ::UMetasound;
+			friend class FMetasoundAssetBase;
 		};
 
-		class FOutputHandle : protected ITransactable
+		class METASOUNDFRONTEND_API FOutputHandle : protected ITransactable
 		{
 		public:
 			FOutputHandle() = delete;
@@ -148,7 +151,7 @@ namespace Metasound
 			TDescriptionPtr<FMetasoundOutputDescription> OutputPtr;
 		};
 
-		class FInputHandle : protected ITransactable
+		class METASOUNDFRONTEND_API FInputHandle : protected ITransactable
 		{
 		public:
 			FInputHandle() = delete;
@@ -188,7 +191,7 @@ namespace Metasound
 		// Opaque handle to a single node on a graph.
 		// Can retrieve metadata about that node,
 		// get FNodeHandles for other nodes it is connected to, and validate and modify connections.
-		class FNodeHandle : protected ITransactable
+		class METASOUNDFRONTEND_API FNodeHandle : protected ITransactable
 		{
 		public:
 			FNodeHandle() = delete;
@@ -207,15 +210,15 @@ namespace Metasound
 			FInputHandle GetInputWithName(const FString& InName);
 			FOutputHandle GetOutputWithName(const FString& InName);
 
-			EMetasoundClassType GetNodeType();
-			FString GetNodeClassName();
+			EMetasoundClassType GetNodeType() const;
+			const FString& GetNodeClassName() const;
 
 			// If this node is itself a Metasound,
 			// use this to return the contained graph for the metasound.
 			// Otherwise it will return an invalid FGraphHandle.
 			void GetContainedGraph(FGraphHandle& OutGraph);
 
-			uint32 GetNodeID();
+			uint32 GetNodeID() const;
 			static uint32 GetNodeID(const FDescPath& InNodePath);
 
 		private:
@@ -235,9 +238,18 @@ namespace Metasound
 			// Sole constructor for FGraphHandle. Can only be used by friends of FHandleInitParams.
 			FGraphHandle(FHandleInitParams::EPrivateToken PrivateToken, const FHandleInitParams& InParams);
 
+			static FGraphHandle GetHandle(UObject* InOwner, FMetasoundDocument& InRootMetasoundDocument, TSharedPtr<Metasound::Frontend::FDescriptionAccessPoint>& InAccessPoint)
+			{
+				using namespace Path;
+
+				FDescPath PathToGraph = FDescPath()[EFromDocument::ToRootClass][EFromClass::ToGraph];
+				FHandleInitParams InitParams = { InAccessPoint, PathToGraph, InRootMetasoundDocument.RootClass.Metadata.NodeName, MakeWeakObjectPtr(InOwner) };
+				return FGraphHandle(FHandleInitParams::PrivateToken, InitParams);
+			}
+
 			static FGraphHandle InvalidHandle();
 
-			bool IsValid();
+			bool IsValid() const;
 
 			TArray<FNodeHandle> GetAllNodes();
 			TArray<FNodeHandle> GetOutputNodes();
@@ -256,14 +268,14 @@ namespace Metasound
 
 			// Remove the node corresponding to this node handle.
 			// On success, invalidates the received node handle.
-			bool RemoveNode(FNodeHandle& InNode, bool bEvenIfInputOrOutputNode = false);
+			bool RemoveNode(const FNodeHandle& InNode, bool bEvenIfInputOrOutputNode = false);
 
 			// Returns the metadata for the current graph, including the name, description and author.
 			FMetasoundClassMetadata GetGraphMetadata();
 
 			// Exports this graph to a JSON at the given path.
 			// @returns true on success.
-			bool ExportToJSONAsset(const FString& InAbsolutePath);
+			bool ExportToJSONAsset(const FString& InAbsolutePath) const;
 
 			// If the FNodeHandle given is itself a Metasound graph,
 			// and the FNodeHandle is a direct member of this FGraphHandle,
@@ -311,5 +323,5 @@ namespace Metasound
 		// this can be used to attempt to find which asset that node is implemented in,
 		// and return a graph handle for that node.
 		FGraphHandle GetGraphHandleForClass(const FMetasoundClassDescription& InClass);
-	}	
-}
+	} // namespace Frontend
+} // namespace Metasound
