CIS fix missing typename and defualt constructor

diff --git a/Engine/Plugins/Runtime/Metasound/Source/MetasoundGraphCore/Private/MetasoundAudioFormats.cpp b/Engine/Plugins/Runtime/Metasound/Source/MetasoundGraphCore/Private/MetasoundAudioFormats.cpp
--- a/Engine/Plugins/Runtime/Metasound/Source/MetasoundGraphCore/Private/MetasoundAudioFormats.cpp
+++ b/Engine/Plugins/Runtime/Metasound/Source/MetasoundGraphCore/Private/MetasoundAudioFormats.cpp
@@ -43,6 +43,11 @@ namespace Metasound
 
 	/* FMultichannelAudioFormat */
 
+	FMultichannelAudioFormat::FMultichannelAudioFormat()
+	:	NumChannels(0)
+	{
+	}
+
 	FMultichannelAudioFormat::FMultichannelAudioFormat(int32 InNumFrames, int32 InNumChannels)
 	:	NumChannels(InNumChannels)
 	{
diff --git a/Engine/Plugins/Runtime/Metasound/Source/MetasoundGraphCore/Public/MetasoundAudioFormats.h b/Engine/Plugins/Runtime/Metasound/Source/MetasoundGraphCore/Public/MetasoundAudioFormats.h
--- a/Engine/Plugins/Runtime/Metasound/Source/MetasoundGraphCore/Public/MetasoundAudioFormats.h
+++ b/Engine/Plugins/Runtime/Metasound/Source/MetasoundGraphCore/Public/MetasoundAudioFormats.h
@@ -96,6 +96,8 @@ namespace Metasound
 	class METASOUNDGRAPHCORE_API FMultichannelAudioFormat
 	{
 		public:
+			FMultichannelAudioFormat();
+
 			/** FMultichannelAudioFormat Constructor.
 			 *
 			 * @param InNumFrames - The number of frames per an audio buffer.
diff --git a/Engine/Plugins/Runtime/Metasound/Source/MetasoundGraphCore/Public/MetasoundDataReference.h b/Engine/Plugins/Runtime/Metasound/Source/MetasoundGraphCore/Public/MetasoundDataReference.h
--- a/Engine/Plugins/Runtime/Metasound/Source/MetasoundGraphCore/Public/MetasoundDataReference.h
+++ b/Engine/Plugins/Runtime/Metasound/Source/MetasoundGraphCore/Public/MetasoundDataReference.h
@@ -82,7 +82,7 @@ namespace Metasound
 	template<typename DataType>
 	const FName GetMetasoundDataTypeName() 
 	{
-		static const FName TypeName = FName(TDataReferenceTypeInfo< TDecay<DataType>::Type >::TypeName);
+		static const FName TypeName = FName(TDataReferenceTypeInfo<typename TDecay<DataType>::Type >::TypeName);
 
 		return TypeName;
 	}
@@ -122,6 +122,7 @@ namespace Metasound
 			}
 
 	
+			// TODO: having issues compiling this on linux when no default constructor is available. 
 			/** Construct operator with no arguments if the DataType has a default constructor.  */
 			template< typename = typename TEnableIf< TIsConstructible<DataType>::Value >::Type >
 			TDataReference()
