CIS fixes for android/linux

diff --git a/Engine/Plugins/Runtime/Metasound/Source/MetasoundFrontend/Public/MetasoundFrontend.h b/Engine/Plugins/Runtime/Metasound/Source/MetasoundFrontend/Public/MetasoundFrontend.h
--- a/Engine/Plugins/Runtime/Metasound/Source/MetasoundFrontend/Public/MetasoundFrontend.h
+++ b/Engine/Plugins/Runtime/Metasound/Source/MetasoundFrontend/Public/MetasoundFrontend.h
@@ -112,7 +112,7 @@ namespace Metasound
 			friend class FInputHandle;
 			friend class FNodeHandle;
 			friend class FGraphHandle;
-			friend class FMetasoundAssetBase;
+			friend class ::FMetasoundAssetBase;
 		};
 
 		class METASOUNDFRONTEND_API FOutputHandle : protected ITransactable
diff --git a/Engine/Plugins/Runtime/Metasound/Source/MetasoundGenerator/Public/MetasoundGenerator.h b/Engine/Plugins/Runtime/Metasound/Source/MetasoundGenerator/Public/MetasoundGenerator.h
--- a/Engine/Plugins/Runtime/Metasound/Source/MetasoundGenerator/Public/MetasoundGenerator.h
+++ b/Engine/Plugins/Runtime/Metasound/Source/MetasoundGenerator/Public/MetasoundGenerator.h
@@ -38,12 +38,12 @@ namespace Metasound
 		template<typename DataType>
 		void SetInputValue(const FString& InName, DataType InData)
 		{
-			typedef TDataWriteReference< TDecay<DataType>::Type > FDataWriteRef;
+			typedef TDataWriteReference< typename TDecay<DataType>::Type > FDataWriteRef;
 
 			const FDataReferenceCollection& InputCollection = RootExecuter.GetInputs();
 
 			// Check if an input data reference with the given name and type exist in the graph.
-			bool bContainsWriteReference = InputCollection.ContainsDataWriteReference< TDecay<DataType>::Type >(InName);
+			bool bContainsWriteReference = InputCollection.ContainsDataWriteReference< typename TDecay<DataType>::Type >(InName);
 			if (ensureMsgf(bContainsWriteReference, TEXT("Operator does not contain write reference name \"%s\" of type \"%s\""), *InName, *GetMetasoundDataTypeName<DataType>().ToString()))
 			{
 				FDataWriteRef WriteRef = InputCollection.GetDataWriteReference<DataType>(InName);
@@ -62,12 +62,12 @@ namespace Metasound
 		void ApplyToInputValue(const FString& InName, TUniqueFunction<void(DataType&)> InFunc)
 		{
 			// Get decayed type as InFunc could take a const qualified type.
-			typedef TDataWriteReference< TDecay<DataType>::Type > FDataWriteRef;
+			typedef TDataWriteReference< typename TDecay<DataType>::Type > FDataWriteRef;
 
 			const FDataReferenceCollection& InputCollection = RootExecuter.GetInputs();
 
 			// Check if an input data reference with the given name and type exists in the graph.
-			bool bContainsWriteReference = InputCollection.ContainsDataWriteReference< TDecay<DataType>::Type >(InName);
+			bool bContainsWriteReference = InputCollection.ContainsDataWriteReference< typename TDecay<DataType>::Type >(InName);
 			if (ensureMsgf(bContainsWriteReference, TEXT("Operator does not contain write reference name \"%s\" of type \"%s\""), *InName, *GetMetasoundDataTypeName<DataType>().ToString()))
 			{
 				FDataWriteRef WriteRef = InputCollection.GetDataWriteReference<DataType>(InName);
