Build fix for non-unity builds.

#rnx

diff --git a/Engine/Plugins/Runtime/Metasound/Source/MetasoundFrontend/Private/MetasoundAssetBase.cpp b/Engine/Plugins/Runtime/Metasound/Source/MetasoundFrontend/Private/MetasoundAssetBase.cpp
--- a/Engine/Plugins/Runtime/Metasound/Source/MetasoundFrontend/Private/MetasoundAssetBase.cpp
+++ b/Engine/Plugins/Runtime/Metasound/Source/MetasoundFrontend/Private/MetasoundAssetBase.cpp
@@ -5,6 +5,7 @@
 #include "Backends/JsonStructSerializerBackend.h"
 #include "IStructSerializerBackend.h"
 #include "StructSerializer.h"
+#include "HAL/FileManager.h"
 
 
 void FMetasoundAssetBase::SetMetadata(FMetasoundClassMetadata& InMetadata)
