1) The node registry and METASOUND_REGISTER_NODE(FMyNodeType)
2) The datatype registry and REGISTER_METASOUND_DATATYPE(FMyDataType) (which I have to rename to METASOUND_REGISTER_DATATYPE)
3) Integration of literals, which are handled with FDataTypeLiteralParam and serialized out with FMetasoundLiteralDescription , and can be expanded to pass in a UObject*/TArray<UObject*> in the future
4) And finally, the full construction of an arbitrary graph to build operators out of in Metasound::Frontend::FGraphHandle::BuildOperator, which needs a sample rate/frame count per callback to work.

diff --git a/Engine/Plugins/Runtime/Metasound/Source/MetasoundFrontend/MetasoundFrontend.Build.cs b/Engine/Plugins/Runtime/Metasound/Source/MetasoundFrontend/MetasoundFrontend.Build.cs
--- a/Engine/Plugins/Runtime/Metasound/Source/MetasoundFrontend/MetasoundFrontend.Build.cs
+++ b/Engine/Plugins/Runtime/Metasound/Source/MetasoundFrontend/MetasoundFrontend.Build.cs
@@ -24,6 +24,8 @@ namespace UnrealBuildTool.Rules
 					"Serialization"
 				}
 			);
+
+			PublicDefinitions.Add("WITH_METASOUND_FRONTEND=1");
 		}
 	}
 }
diff --git a/Engine/Plugins/Runtime/Metasound/Source/MetasoundFrontend/Private/MetasoundFrontend.cpp b/Engine/Plugins/Runtime/Metasound/Source/MetasoundFrontend/Private/MetasoundFrontend.cpp
--- a/Engine/Plugins/Runtime/Metasound/Source/MetasoundFrontend/Private/MetasoundFrontend.cpp
+++ b/Engine/Plugins/Runtime/Metasound/Source/MetasoundFrontend/Private/MetasoundFrontend.cpp
@@ -9,6 +9,7 @@
 #include "Modules/ModuleManager.h"
 
 #include "MetasoundAudioFormats.h"
+#include "MetasoundOperatorBuilder.h"
 
 
 static int32 MetasoundUndoRollLimitCvar = 128;
@@ -25,52 +26,190 @@ namespace Metasound
 	{
 		TArray<FNodeClassInfo> GetAllAvailableNodeClasses()
 		{
-			// UEAU-471
-			ensureAlwaysMsgf(false, TEXT("Implement me!"));
-			return TArray<FNodeClassInfo>();
+			TArray<FNodeClassInfo> OutClasses;
+
+			auto& Registry = GetExternalNodeRegistry();
+			for (auto& NodeClassTuple : Registry)
+			{
+				FNodeClassInfo ClassInfo;
+				ClassInfo.NodeName = NodeClassTuple.Key.NodeName.ToString();
+				ClassInfo.NodeType = EMetasoundClassType::External;
+				ClassInfo.LookupKey = NodeClassTuple.Key;
+
+				OutClasses.Add(ClassInfo);
+			}
+
+			return OutClasses;
 		}
 
 		TArray<FNodeClassInfo> GetAllNodeClassesInNamespace(const FString& InNamespace)
 		{
-			// UEAU-471
-			ensureAlwaysMsgf(false, TEXT("Implement me!"));
-			return TArray<FNodeClassInfo>();
+			TArray<FNodeClassInfo> OutClasses;
+
+			auto& Registry = GetExternalNodeRegistry();
+			for (auto& NodeClassTuple : Registry)
+			{
+				// TODO: Build fname namespace tree. as we register nodes and types.
+				// for now we just parse the string here.
+				FString NodeName = NodeClassTuple.Key.NodeName.ToString();
+				if (NodeName.StartsWith(InNamespace))
+				{
+					FNodeClassInfo ClassInfo;
+					ClassInfo.NodeName = NodeName;
+					ClassInfo.NodeType = EMetasoundClassType::External;
+					ClassInfo.LookupKey = NodeClassTuple.Key;
+
+					OutClasses.Add(ClassInfo);
+				}
+			}
+
+			return OutClasses;
 		}
 
 		TArray<FNodeClassInfo> GetAllNodesWhoseNameContains(const FString& InSubstring)
 		{
-			// UEAU-471
-			ensureAlwaysMsgf(false, TEXT("Implement me!"));
-			return TArray<FNodeClassInfo>();
+			TArray<FNodeClassInfo> OutClasses;
+
+			auto& Registry = GetExternalNodeRegistry();
+			for (auto& NodeClassTuple : Registry)
+			{
+				FString NodeName = NodeClassTuple.Key.NodeName.ToString();
+				if (NodeName.Contains(InSubstring))
+				{
+					FNodeClassInfo ClassInfo;
+					ClassInfo.NodeName = NodeName;
+					ClassInfo.NodeType = EMetasoundClassType::External;
+					ClassInfo.LookupKey = NodeClassTuple.Key;
+
+					OutClasses.Add(ClassInfo);
+				}
+			}
+
+			return OutClasses;
 		}
 
 		TArray<FNodeClassInfo> GetAllNodesWithAnOutputOfType(const FName& InType)
 		{
-			// UEAU-471
-			ensureAlwaysMsgf(false, TEXT("Implement me!"));
-			return TArray<FNodeClassInfo>();
+			TArray<FNodeClassInfo> OutClasses;
+
+			auto& Registry = GetExternalNodeRegistry();
+			for (auto& NodeClassTuple : Registry)
+			{
+				if (NodeClassTuple.Value.OutputTypes.Contains(InType))
+				{
+					FNodeClassInfo ClassInfo;
+					ClassInfo.NodeName = NodeClassTuple.Key.NodeName.ToString();
+					ClassInfo.NodeType = EMetasoundClassType::External;
+					ClassInfo.LookupKey = NodeClassTuple.Key;
+
+					OutClasses.Add(ClassInfo);
+					break;
+				}
+			}
+
+			return OutClasses;
 		}
 
 		TArray<FNodeClassInfo> GetAllNodesWithAnInputOfType(const FName& InType)
 		{
-			// UEAU-471
-			ensureAlwaysMsgf(false, TEXT("Implement me!"));
-			return TArray<FNodeClassInfo>();
+			TArray<FNodeClassInfo> OutClasses;
+
+			auto& Registry = GetExternalNodeRegistry();
+			for (auto& NodeClassTuple : Registry)
+			{
+				if (NodeClassTuple.Value.InputTypes.Contains(InType))
+				{
+					FNodeClassInfo ClassInfo;
+					ClassInfo.NodeName = NodeClassTuple.Key.NodeName.ToString();
+					ClassInfo.NodeType = EMetasoundClassType::External;
+					ClassInfo.LookupKey = NodeClassTuple.Key;
+
+					OutClasses.Add(ClassInfo);
+					break;
+				}
+			}
+
+			return OutClasses;
 		}
 
 		// gets all metadata (name, description, author, what to say if it's missing) for a given node.
-		FMetasoundClassMetadata GetMetadataForNode(const FNodeClassInfo& InInfo)
+		FMetasoundClassMetadata GenerateMetadataForNode(const FNodeClassInfo& InInfo)
 		{
-			// UEAU-471
-			ensureAlwaysMsgf(false, TEXT("Implement me!"));
-			return FMetasoundClassMetadata();
+			return GenerateClassDescriptionForNode(InInfo).Metadata;
 		}
 
-		FMetasoundClassDescription GetClassDescriptionForNode(const FNodeClassInfo& InInfo)
+		FMetasoundClassDescription GenerateClassDescriptionForNode(const FNodeClassInfo& InInfo)
 		{
-			// UEAU-471
-			ensureAlwaysMsgf(false, TEXT("Implement me!"));
-			return FMetasoundClassDescription();
+			FNodeInitData DummyInitData;
+			DummyInitData.InstanceName = FString(TEXT("Unused node for registration"));
+
+			auto& Registry = GetExternalNodeRegistry();
+			if (Registry.Contains(InInfo.LookupKey))
+			{
+				TUniquePtr<Metasound::INode> DummyNode = Registry[InInfo.LookupKey].GetterCallback(DummyInitData);
+
+				if (!ensureAlwaysMsgf(DummyNode.IsValid(), TEXT("Node class %s failed to return a valid node. Likely something is wrong with the METASOUND_REGISTER_NODE macro.")))
+				{
+					return FMetasoundClassDescription();
+				}
+
+				const FName NodeName = DummyNode->GetClassName();
+				const FInputDataVertexCollection& Inputs = DummyNode->GetInputDataVertices();
+				const FOutputDataVertexCollection& Outputs = DummyNode->GetOutputDataVertices();
+
+				FMetasoundClassMetadata NodeMetadata;
+				NodeMetadata.NodeName = NodeName.ToString();
+				NodeMetadata.NodeType = EMetasoundClassType::External;
+				NodeMetadata.AuthorName = DummyNode->GetAuthorName();
+				NodeMetadata.MetasoundDescription = DummyNode->GetDescription();
+				NodeMetadata.PromptIfMissing = DummyNode->GetPromptIfMissing();
+
+				FMetasoundClassDescription ClassDescription;
+				ClassDescription.Metadata = NodeMetadata;
+
+				// External metasounds aren't dependent on any other nodes by definition, so all we need to do
+				// is populate the Input and Output sets.
+				for (auto& InputTuple : Inputs)
+				{
+					FMetasoundInputDescription InputDescription;
+					InputDescription.Name = InputTuple.Value.VertexName;
+					InputDescription.TypeName = InputTuple.Value.DataReferenceTypeName;
+					InputDescription.ToolTip = InputTuple.Value.Description;
+
+					ClassDescription.Inputs.Add(InputDescription);
+				}
+
+				for (auto& OutputTuple : Outputs)
+				{
+					FMetasoundOutputDescription OutputDescription;
+					OutputDescription.Name = OutputTuple.Value.VertexName;
+					OutputDescription.TypeName = OutputTuple.Value.DataReferenceTypeName;
+					OutputDescription.ToolTip = OutputTuple.Value.Description;
+
+					ClassDescription.Outputs.Add(OutputDescription);
+				}
+
+				// Populate lookup data.
+				ClassDescription.ExternalNodeClassLookupInfo.ExternalNodeClassName = InInfo.LookupKey.NodeName;
+				ClassDescription.ExternalNodeClassLookupInfo.ExternalNodeClassHash = InInfo.LookupKey.NodeHash;
+
+				return ClassDescription;
+			}
+			else
+			{
+				ensureAlwaysMsgf(false, TEXT("Tried to get Class Description for unknown node %s!"), *InInfo.NodeName);
+				return FMetasoundClassDescription();
+			}
+		}
+
+		TArray<FName> GetAllAvailableDataTypes()
+		{
+			return FMetasoundFrontendRegistryContainer::Get()->GetAllValidDataTypes();
+		}
+
+		bool GetTraitsForDataType(FName InDataType, FDataTypeRegistryInfo& OutInfo)
+		{
+			return FMetasoundFrontendRegistryContainer::Get()->GetInfoForDataType(InDataType, OutInfo);
 		}
 
 		FInputHandle::FInputHandle(FHandleInitParams::EPrivateToken PrivateToken, const FHandleInitParams& InParams, const FString& InputName)
@@ -78,6 +217,7 @@ namespace Metasound
 			, NodePtr(InParams.InAccessPoint, InParams.InPath)
 			, NodeClass(InParams.InAccessPoint, Path::GetDependencyPath(InParams.InClassName))
 			, InputPtr(InParams.InAccessPoint, NodeClass.GetPath()[Path::EFromClass::ToInputs][*InputName])
+			, OutputNodePtr(nullptr, FDescPath())
 		{
 			if (InParams.InAccessPoint.IsValid())
 			{
@@ -86,6 +226,20 @@ namespace Metasound
 			}
 		}
 
+		FInputHandle::FInputHandle(FHandleInitParams::EPrivateToken PrivateToken, const FHandleInitParams& InParams)
+			: ITransactable(MetasoundUndoRollLimitCvar, InParams.InOwningAsset)
+			, NodePtr(InParams.InAccessPoint, InParams.InPath)
+			, NodeClass(nullptr, FDescPath())
+			, InputPtr(nullptr, FDescPath())
+			, OutputNodePtr(InParams.InAccessPoint, Path::GetOutputDescriptionPath(InParams.InPath, InParams.InClassName))
+		{
+			if (InParams.InAccessPoint.IsValid())
+			{
+				// Test both pointers to the graph and it's owning class description.
+				ensureAlwaysMsgf(NodePtr.IsValid() && OutputNodePtr.IsValid(), TEXT("Tried to build FGraphHandle with Invalid Path: %s"), *Path::GetPrintableString(InParams.InPath));
+			}
+		}
+
 		FInputHandle FInputHandle::InvalidHandle()
 		{
 			FDescPath NullPath = FDescPath();
@@ -96,7 +250,8 @@ namespace Metasound
 
 		bool FInputHandle::IsValid() const
 		{
-			return NodePtr.IsValid() && NodeClass.IsValid() && InputPtr.IsValid();
+			return (NodePtr.IsValid() && NodeClass.IsValid() && InputPtr.IsValid())
+				|| (NodePtr.IsValid() && OutputNodePtr.IsValid());
 		}
 
 		bool FInputHandle::IsConnected() const
@@ -124,7 +279,19 @@ namespace Metasound
 				return FName();
 			}
 
-			return InputPtr->TypeName;
+			if (InputPtr.IsValid())
+			{
+				return InputPtr->TypeName;
+			}
+			else if (OutputNodePtr.IsValid())
+			{
+				return OutputNodePtr->TypeName;
+			}
+			else
+			{
+				checkNoEntry();
+				return FName();
+			}
 		}
 
 		FString FInputHandle::GetInputName() const
@@ -134,17 +301,41 @@ namespace Metasound
 				return FString();
 			}
 
-			return InputPtr->Name;
+			if (InputPtr.IsValid())
+			{
+				return InputPtr->Name;
+			}
+			else if (OutputNodePtr.IsValid())
+			{
+				return OutputNodePtr->Name;
+			}
+			else
+			{
+				checkNoEntry();
+				return FString();
+			}
 		}
 
-		FString FInputHandle::GetInputTooltip() const
+		FText FInputHandle::GetInputTooltip() const
 		{
 			if (!IsValid())
 			{
-				return FString();
+				return FText();
 			}
 
-			return InputPtr->ToolTip;
+			if (InputPtr.IsValid())
+			{
+				return InputPtr->ToolTip;
+			}
+			else if (OutputNodePtr.IsValid())
+			{
+				return OutputNodePtr->ToolTip;
+			}
+			else
+			{
+				checkNoEntry();
+				return FText();
+			}
 		}
 
 		FOutputHandle FInputHandle::GetCurrentlyConnectedOutput() const
@@ -162,15 +353,28 @@ namespace Metasound
 			const FMetasoundNodeConnectionDescription* Connection = GetConnectionDescription();
 			check(Connection);
 
-			FString OutputName = Connection->OutputNameOrLiteral;
+			FString OutputName = Connection->OutputName;
 			uint32 OutputNodeID = Connection->NodeID;
 
 			// All node connections are in the same graph, so we just need to go up one level to the Nodes array
 			// and look up the node by it's unique ID.
 			FDescPath OutputNodePath = (NodePtr.GetPath() << 1)[OutputNodeID];
-			FHandleInitParams InitParams = { NodePtr.GetAccessPoint(), OutputNodePath, NodeClass->Metadata.NodeName, OwningAsset };
 
-			return FOutputHandle(FHandleInitParams::PrivateToken, InitParams,  OutputName);
+			if (NodeClass.IsValid())
+			{
+				FHandleInitParams InitParams = { NodePtr.GetAccessPoint(), OutputNodePath, NodeClass->Metadata.NodeName, OwningAsset };
+				return FOutputHandle(FHandleInitParams::PrivateToken, InitParams, OutputName);
+			}
+			else if (OutputNodePtr.IsValid())
+			{
+				FHandleInitParams InitParams = { NodePtr.GetAccessPoint(), OutputNodePath, OutputNodePtr->Name, OwningAsset };
+				return FOutputHandle(FHandleInitParams::PrivateToken, InitParams, OutputName);
+			}
+			else
+			{
+				checkNoEntry();
+				return FOutputHandle::InvalidHandle();
+			}
 		}
 
 		FConnectability FInputHandle::CanConnectTo(const FOutputHandle& InHandle) const
@@ -214,7 +418,7 @@ namespace Metasound
 			}
 
 			Connection->NodeID = OutputNodeID;
-			Connection->OutputNameOrLiteral = OutputName;
+			Connection->OutputName = OutputName;
 
 			return true;
 		}
@@ -233,7 +437,7 @@ namespace Metasound
 				return nullptr;
 			}
 
-			const TArray<FMetasoundNodeConnectionDescription>& Connections = (*NodePtr).InputConnections;
+			const TArray<FMetasoundNodeConnectionDescription>& Connections = NodePtr->InputConnections;
 
 			for (const FMetasoundNodeConnectionDescription& Connection : Connections)
 			{
@@ -253,7 +457,7 @@ namespace Metasound
 				return nullptr;
 			}
 
-			TArray<FMetasoundNodeConnectionDescription>& Connections = (*NodePtr).InputConnections;
+			TArray<FMetasoundNodeConnectionDescription>& Connections = NodePtr->InputConnections;
 
 			for (FMetasoundNodeConnectionDescription& Connection : Connections)
 			{
@@ -271,6 +475,7 @@ namespace Metasound
 			, NodePtr(InParams.InAccessPoint, InParams.InPath)
 			, NodeClass(InParams.InAccessPoint, Path::GetDependencyPath(InParams.InClassName))
 			, OutputPtr(InParams.InAccessPoint, NodeClass.GetPath()[Path::EFromClass::ToOutputs][*InOutputName])
+			, InputNodePtr(nullptr, FDescPath())
 		{
 			if (InParams.InAccessPoint.IsValid())
 			{
@@ -279,6 +484,20 @@ namespace Metasound
 			}
 		}
 
+		FOutputHandle::FOutputHandle(FHandleInitParams::EPrivateToken InToken, const FHandleInitParams& InParams)
+			: ITransactable(MetasoundUndoRollLimitCvar, InParams.InOwningAsset)
+			, NodePtr(InParams.InAccessPoint, InParams.InPath)
+			, NodeClass(nullptr, FDescPath())
+			, OutputPtr(nullptr, FDescPath())
+			, InputNodePtr(InParams.InAccessPoint, Path::GetInputDescriptionPath(InParams.InPath, InParams.InClassName))
+		{
+			if (InParams.InAccessPoint.IsValid())
+			{
+				// Test both pointers to the graph and it's owning class description.
+				ensureAlwaysMsgf(NodePtr.IsValid() && InputNodePtr.IsValid(), TEXT("Tried to build FGraphHandle with Invalid Path: %s"), *Path::GetPrintableString(InParams.InPath));
+			}
+		}
+
 		FOutputHandle FOutputHandle::InvalidHandle()
 		{
 			FString NullString;
@@ -289,7 +508,8 @@ namespace Metasound
 
 		bool FOutputHandle::IsValid() const
 		{
-			return NodePtr.IsValid() && NodeClass.IsValid() && OutputPtr.IsValid();
+			return (NodePtr.IsValid() && NodeClass.IsValid() && OutputPtr.IsValid())
+				|| (NodePtr.IsValid() && InputNodePtr.IsValid());
 		}
 
 		FName FOutputHandle::GetOutputType() const
@@ -299,7 +519,19 @@ namespace Metasound
 				return FName();
 			}
 
-			return OutputPtr->TypeName;
+			if (OutputPtr.IsValid())
+			{
+				return OutputPtr->TypeName;
+			}
+			else if (InputNodePtr.IsValid())
+			{
+				return InputNodePtr->TypeName;
+			}
+			else
+			{
+				checkNoEntry();
+				return FName();
+			}
 		}
 
 		FString FOutputHandle::GetOutputName() const
@@ -309,17 +541,41 @@ namespace Metasound
 				return FString();
 			}
 
-			return OutputPtr->Name;
+			if (OutputPtr.IsValid())
+			{
+				return OutputPtr->Name;
+			}
+			else if (InputNodePtr.IsValid())
+			{
+				return InputNodePtr->Name;
+			}
+			else
+			{
+				checkNoEntry();
+				return FString();
+			}
 		}
 
-		FString FOutputHandle::GetOutputTooltip() const
+		FText FOutputHandle::GetOutputTooltip() const
 		{
 			if (!IsValid())
 			{
-				return FString();
+				return FText();
 			}
 
-			return OutputPtr->ToolTip;
+			if (OutputPtr.IsValid())
+			{
+				return OutputPtr->ToolTip;
+			}
+			else if (InputNodePtr.IsValid())
+			{
+				return InputNodePtr->ToolTip;
+			}
+			else
+			{
+				checkNoEntry();
+				return FText();
+			}
 		}
 
 		uint32 FOutputHandle::GetOwningNodeID() const
@@ -352,15 +608,16 @@ namespace Metasound
 			return InHandle.ConnectWithConverterNode(*this, InNodeClassName);
 		}
 
-		FNodeHandle::FNodeHandle(FHandleInitParams::EPrivateToken PrivateToken, const FHandleInitParams& InParams)
+		FNodeHandle::FNodeHandle(FHandleInitParams::EPrivateToken PrivateToken, const FHandleInitParams& InParams, EMetasoundClassType InNodeClassType)
 			: ITransactable(MetasoundUndoRollLimitCvar, InParams.InOwningAsset)
 			, NodePtr(InParams.InAccessPoint, InParams.InPath)
 			, NodeClass(InParams.InAccessPoint, Path::GetDependencyPath(InParams.InClassName))
+			, NodeClassType(InNodeClassType)
 		{
 			if (InParams.InAccessPoint.IsValid())
 			{
 				// Test both pointers to the graph and it's owning class description.
-				ensureAlwaysMsgf(NodePtr.IsValid() && NodeClass.IsValid(), TEXT("Tried to build FGraphHandle with Invalid Path: %s"), *Path::GetPrintableString(InParams.InPath));
+				ensureAlwaysMsgf(IsValid(), TEXT("Tried to build FGraphHandle with Invalid Path: %s"), *Path::GetPrintableString(InParams.InPath));
 			}
 		}
 
@@ -369,31 +626,46 @@ namespace Metasound
 			FDescPath InvalidPath;
 			FString InvalidClassName;
 			FHandleInitParams InitParams = { nullptr, InvalidPath, InvalidClassName, nullptr };
-			return FNodeHandle(FHandleInitParams::PrivateToken, InitParams);
+			return FNodeHandle(FHandleInitParams::PrivateToken, InitParams, EMetasoundClassType::Invalid);
 		}
 
 		bool FNodeHandle::IsValid() const
 		{
-			return NodePtr.IsValid() && NodeClass.IsValid();
+			const bool bNeedsNodeClass = NodeClassType == EMetasoundClassType::External || NodeClassType == EMetasoundClassType::MetasoundGraph;
+			return NodePtr.IsValid() && (!bNeedsNodeClass || NodeClass.IsValid());
 		}
 
 		TArray<FInputHandle> FNodeHandle::GetAllInputs()
 		{
 			TArray<FInputHandle> OutArray;
 
-			if (!IsValid())
+			if (!IsValid() || NodeClassType == EMetasoundClassType::Input)
 			{
 				return OutArray;
 			}
 
-			TArray<FMetasoundInputDescription>& InputDescriptions = NodeClass->Inputs;
-
-			for (FMetasoundInputDescription& InputDescription : InputDescriptions)
+			if (NodeClassType == EMetasoundClassType::Output)
 			{
+				// Output nodes only have one input- the outgoing parameter.
 				FDescPath NodePath = NodePtr.GetPath();
-				FString ClassName = GetNodeClassName();
-				FHandleInitParams InitParams = { NodePtr.GetAccessPoint(), NodePath, ClassName, OwningAsset };
-				OutArray.Emplace(FHandleInitParams::PrivateToken, InitParams, InputDescription.Name);
+
+				const FString& NodeClassName = NodePtr->Name;
+
+				FHandleInitParams InitParams = { NodePtr.GetAccessPoint(), NodePath, NodeClassName, OwningAsset };
+				OutArray.Emplace(FHandleInitParams::PrivateToken, InitParams);
+			}
+			else
+			{
+				// Iterate over our input descriptions and emplace a new handle for each of them.
+
+				TArray<FMetasoundInputDescription>& InputDescriptions = NodeClass->Inputs;
+				for (FMetasoundInputDescription& InputDescription : InputDescriptions)
+				{
+					FDescPath NodePath = NodePtr.GetPath();
+					FString ClassName = GetNodeClassName();
+					FHandleInitParams InitParams = { NodePtr.GetAccessPoint(), NodePath, ClassName, OwningAsset };
+					OutArray.Emplace(FHandleInitParams::PrivateToken, InitParams, InputDescription.Name);
+				}
 			}
 
 			return OutArray;
@@ -403,19 +675,32 @@ namespace Metasound
 		{
 			TArray<FOutputHandle> OutArray;
 
-			if (!IsValid())
+			if (!IsValid() || NodeClassType == EMetasoundClassType::Output)
 			{
 				return OutArray;
 			}
 
-			TArray<FMetasoundOutputDescription>& OutputDescriptions = NodeClass->Outputs;
-
-			for (FMetasoundOutputDescription& OutputDescription : OutputDescriptions)
+			if (NodeClassType == EMetasoundClassType::Input)
 			{
+				// Input nodes nodes only have one output- the incoming parameter.
 				FDescPath NodePath = NodePtr.GetPath();
-				FString NodeClassName = GetNodeClassName();
+
+				const FString& NodeClassName = NodePtr->Name;
+
 				FHandleInitParams InitParams = { NodePtr.GetAccessPoint(), NodePath, NodeClassName, OwningAsset };
-				OutArray.Emplace(FHandleInitParams::PrivateToken, InitParams, OutputDescription.Name);
+				OutArray.Emplace(FHandleInitParams::PrivateToken, InitParams);
+			}
+			else
+			{
+				TArray<FMetasoundOutputDescription>& OutputDescriptions = NodeClass->Outputs;
+
+				for (FMetasoundOutputDescription& OutputDescription : OutputDescriptions)
+				{
+					FDescPath NodePath = NodePtr.GetPath();
+					FString NodeClassName = GetNodeClassName();
+					FHandleInitParams InitParams = { NodePtr.GetAccessPoint(), NodePath, NodeClassName, OwningAsset };
+					OutArray.Emplace(FHandleInitParams::PrivateToken, InitParams, OutputDescription.Name);
+				}
 			}
 
 			return OutArray;
@@ -423,11 +708,21 @@ namespace Metasound
 
 		FInputHandle FNodeHandle::GetInputWithName(const FString& InName)
 		{
-			if (!IsValid())
+			if (!IsValid() || NodeClassType == EMetasoundClassType::Input)
 			{
 				return FInputHandle::InvalidHandle();
 			}
 
+			if (NodeClassType == EMetasoundClassType::Output)
+			{
+				const FString& NodeClassName = NodePtr->Name;
+				ensureAlwaysMsgf(InName == NodeClassName, TEXT("An output node's input connection should always be the same as it's class name!"));
+
+				FDescPath NodePath = NodePtr.GetPath();
+				FHandleInitParams InitParams = { NodePtr.GetAccessPoint(), NodePath, NodeClassName , OwningAsset };
+				return FInputHandle(FHandleInitParams::PrivateToken, InitParams);
+			}
+
 			TArray<FMetasoundInputDescription>& InputDescriptions = NodeClass->Inputs;
 
 			for (FMetasoundInputDescription& InputDescription : InputDescriptions)
@@ -435,7 +730,8 @@ namespace Metasound
 				if (InputDescription.Name == InName)
 				{
 					FString ClassName = GetNodeClassName();
-					FHandleInitParams InitParams = { NodePtr.GetAccessPoint(), NodePtr.GetPath(), ClassName, OwningAsset };
+					FDescPath NodePath = NodePtr.GetPath();
+					FHandleInitParams InitParams = { NodePtr.GetAccessPoint(), NodePath, ClassName, OwningAsset };
 					return FInputHandle(FHandleInitParams::PrivateToken, InitParams, InputDescription.Name);
 				}
 			}
@@ -446,11 +742,21 @@ namespace Metasound
 
 		FOutputHandle FNodeHandle::GetOutputWithName(const FString& InName)
 		{
-			if (!IsValid())
+			if (!IsValid() || NodeClassType == EMetasoundClassType::Output)
 			{
 				return FOutputHandle::InvalidHandle();
 			}
 
+			// all input nodes have one connectable output, which is the input param they represent.
+			if (NodeClassType == EMetasoundClassType::Input)
+			{
+				const FString& NodeClassName = NodePtr->Name;
+				ensureAlwaysMsgf(InName == NodeClassName, TEXT("An input node's output connection should always be the same as it's class name!"));
+				FDescPath NodePath = NodePtr.GetPath();
+				FHandleInitParams InitParams = { NodePtr.GetAccessPoint(), NodePath, NodeClassName , OwningAsset };
+				return FOutputHandle(FHandleInitParams::PrivateToken, InitParams);
+			}
+
 			TArray<FMetasoundOutputDescription>& OutputDescriptions = NodeClass->Outputs;
 
 			for (FMetasoundOutputDescription& OutputDescription : OutputDescriptions)
@@ -458,7 +764,8 @@ namespace Metasound
 				if (OutputDescription.Name == InName)
 				{
 					FString NodeClassName = GetNodeClassName();
-					FHandleInitParams InitParams = { NodePtr.GetAccessPoint(), NodePtr.GetPath(), NodeClassName , OwningAsset };
+					FDescPath NodePath = NodePtr.GetPath();
+					FHandleInitParams InitParams = { NodePtr.GetAccessPoint(), NodePath, NodeClassName , OwningAsset };
 					return FOutputHandle(FHandleInitParams::PrivateToken, InitParams, OutputDescription.Name);
 				}
 			}
@@ -656,12 +963,34 @@ namespace Metasound
 				FDescPath NodePath = GraphPtr.GetPath()[Path::EFromGraph::ToNodes][NodeDescription.UniqueID];
 
 				FHandleInitParams InitParams = { GraphPtr.GetAccessPoint(), NodePath, NodeDescription.Name, OwningAsset };
-				OutArray.Emplace(FHandleInitParams::PrivateToken, InitParams);
+				OutArray.Emplace(FHandleInitParams::PrivateToken, InitParams, NodeDescription.ObjectTypeOfNode);
 			}
 
 			return OutArray;
 		}
 
+		FNodeHandle FGraphHandle::GetNodeWithId(uint32 InNodeId)
+		{
+			if (!IsValid())
+			{
+				return FNodeHandle::InvalidHandle();
+			}
+
+			const TArray<FMetasoundNodeDescription>& NodeDescriptions = GraphPtr->Nodes;
+			for (const FMetasoundNodeDescription& NodeDescription : NodeDescriptions)
+			{
+				if (NodeDescription.UniqueID == InNodeId)
+				{
+					FDescPath NodePath = GraphPtr.GetPath()[Path::EFromGraph::ToNodes][InNodeId];
+					FHandleInitParams InitParams = { GraphPtr.GetAccessPoint(), NodePath, NodeDescription.Name, OwningAsset };
+					return FNodeHandle(FHandleInitParams::PrivateToken, InitParams, NodeDescription.ObjectTypeOfNode);
+				}
+			}
+
+			ensureAlwaysMsgf(false, TEXT("Couldn't find node in graph with ID %u!"), InNodeId);
+			return FNodeHandle::InvalidHandle();
+		}
+
 		TArray<FNodeHandle> FGraphHandle::GetOutputNodes()
 		{
 			TArray<FNodeHandle> OutArray;
@@ -678,7 +1007,7 @@ namespace Metasound
 				{
 					FDescPath NodePath = GraphPtr.GetPath()[Path::EFromGraph::ToNodes][NodeDescription.UniqueID];
 					FHandleInitParams InitParams = { GraphPtr.GetAccessPoint(), NodePath, NodeDescription.Name, OwningAsset };
-					OutArray.Emplace(FHandleInitParams::PrivateToken, InitParams);
+					OutArray.Emplace(FHandleInitParams::PrivateToken, InitParams, NodeDescription.ObjectTypeOfNode);
 				}
 			}
 
@@ -701,7 +1030,7 @@ namespace Metasound
 				{
 					FDescPath NodePath = GraphPtr.GetPath()[Path::EFromGraph::ToNodes][NodeDescription.UniqueID];
 					FHandleInitParams InitParams = { GraphPtr.GetAccessPoint(), NodePath, NodeDescription.Name, OwningAsset };
-					OutArray.Emplace(FHandleInitParams::PrivateToken, InitParams);
+					OutArray.Emplace(FHandleInitParams::PrivateToken, InitParams, NodeDescription.ObjectTypeOfNode);
 				}
 			}
 
@@ -723,7 +1052,7 @@ namespace Metasound
 					FDescPath NodePath = GraphPtr.GetPath()[Path::EFromGraph::ToNodes][NodeDescription.UniqueID];
 
 					FHandleInitParams InitParams = { GraphPtr.GetAccessPoint(), NodePath, InName, OwningAsset };
-					return FNodeHandle(FHandleInitParams::PrivateToken, InitParams);
+					return FNodeHandle(FHandleInitParams::PrivateToken, InitParams, NodeDescription.ObjectTypeOfNode);
 				}
 			}
 
@@ -747,7 +1076,7 @@ namespace Metasound
 
 					// todo: add special enum for input/output nodes
 					FHandleInitParams InitParams = { GraphPtr.GetAccessPoint(), NodePath, InName, OwningAsset };
-					return FNodeHandle(FHandleInitParams::PrivateToken, InitParams);
+					return FNodeHandle(FHandleInitParams::PrivateToken, InitParams, NodeDescription.ObjectTypeOfNode);
 				}
 			}
 
@@ -796,7 +1125,7 @@ namespace Metasound
 			FDescPath NodePath = GraphPtr.GetPath()[Path::EFromGraph::ToNodes][NewNodeDescription.UniqueID];
 			FHandleInitParams InitParams = { GraphPtr.GetAccessPoint(), NodePath, NewNodeDescription.Name, OwningAsset };
 			// todo: add special enum for input and output nodes
-			return FNodeHandle(FHandleInitParams::PrivateToken, InitParams);
+			return FNodeHandle(FHandleInitParams::PrivateToken, InitParams, NewNodeDescription.ObjectTypeOfNode);
 		}
 
 		bool FGraphHandle::RemoveInput(const FString& InputName)
@@ -889,7 +1218,7 @@ namespace Metasound
 
 			// todo: add special enum or class for input/output nodes
 			FHandleInitParams InitParams = { GraphPtr.GetAccessPoint(), NodePath, NewNodeDescription.Name, OwningAsset };
-			return FNodeHandle(FHandleInitParams::PrivateToken, InitParams);
+			return FNodeHandle(FHandleInitParams::PrivateToken, InitParams, NewNodeDescription.ObjectTypeOfNode);
 		}
 
 		bool FGraphHandle::RemoveOutput(const FString& OutputName)
@@ -984,7 +1313,7 @@ namespace Metasound
 			// If we haven't added a node of this class to the graph yet, add it to the dependencies for this class.
 			if (!bFoundMatchingDependencyInDocument)
 			{
-				FMetasoundClassDescription NewDependencyClassDescription = GetClassDescriptionForNode(InNodeClass);
+				FMetasoundClassDescription NewDependencyClassDescription = GenerateClassDescriptionForNode(InNodeClass);
 				NewDependencyClassDescription.UniqueID = FindNewUniqueDependencyId();
 				Dependencies.Add(NewDependencyClassDescription);
 				GraphsClassDeclaration->DependencyIDs.Add(NewDependencyClassDescription.UniqueID);
@@ -1010,7 +1339,7 @@ namespace Metasound
 
 			FDescPath NodePath = GraphPtr.GetPath()[Path::EFromGraph::ToNodes][NewNodeDescription.UniqueID];
 			FHandleInitParams InitParams = { GraphPtr.GetAccessPoint(), NodePath, NewNodeDescription.Name, OwningAsset };
-			return FNodeHandle(FHandleInitParams::PrivateToken, InitParams);
+			return FNodeHandle(FHandleInitParams::PrivateToken, InitParams, NewNodeDescription.ObjectTypeOfNode);
 		}
 
 		bool FGraphHandle::RemoveNode(const FNodeHandle& InNode, bool bEvenIfInputOrOutputNode /*= false*/)
@@ -1246,16 +1575,295 @@ namespace Metasound
 
 			FDescPath NodePath = GraphPtr.GetPath()[Path::EFromGraph::ToNodes][NewNodeDescription.UniqueID];
 			FHandleInitParams NodeInitParams = { GraphPtr.GetAccessPoint(), NodePath, InInfo.NodeName, OwningAsset };
-			FNodeHandle SubgraphNode = FNodeHandle(FHandleInitParams::PrivateToken, NodeInitParams);
+			FNodeHandle SubgraphNode = FNodeHandle(FHandleInitParams::PrivateToken, NodeInitParams, NewNodeDescription.ObjectTypeOfNode);
 
 			return TTuple<FGraphHandle, FNodeHandle>(SubgraphHandle, SubgraphNode);
 		}
 
-		TUniquePtr<Metasound::IOperator> FGraphHandle::BuildOperator()
+		TUniquePtr<IOperator> FGraphHandle::BuildOperator(const FOperatorSettings& InSettings, TArray<IOperatorBuilder::FBuildErrorPtr>& OutBuildErrors) const
 		{
-			// @todo: bring over operator builder from previous frontend draft with this new implementation.
-			ensureAlwaysMsgf(false, TEXT("Implement Me!"));
-			return nullptr;
+			if (!IsValid())
+			{
+				return nullptr;
+			}
+
+			// TODO: Implement inflation step here.
+
+			// At this point, we should be left with a flat graph of externally implemented nodes.
+			using FGraph = Metasound::FGraph;
+			using FDataEdge = Metasound::FDataEdge;
+			using INodePtr = TUniquePtr<INode>;
+
+			TArray<FName> DataTypes = GetAllAvailableDataTypes();
+
+			auto IsInputNode = [](const FMetasoundNodeDescription& InDescription) -> bool
+			{
+				return InDescription.ObjectTypeOfNode == EMetasoundClassType::Input;
+			};
+
+			auto IsOutputNode = [](const FMetasoundNodeDescription& InDescription) -> bool
+			{
+				return InDescription.ObjectTypeOfNode == EMetasoundClassType::Output;
+			};
+
+			int32 TotalNodesGenerated = 0;
+
+			// These maps are used to fix up input and output destinations after we have fully connected the graph.
+			TMap<int32, uint32> InputIndexToInputNodeIDs;
+			TMap<int32, uint32> OutputIndexToInputNodeIDs;
+
+			// Helper for auto-generating nodes.
+			auto GetNodePtr = [&](const FMetasoundNodeDescription& InDescription) -> INodePtr
+			{
+				if (IsInputNode(InDescription))
+				{
+					// Find this in the graph's class description's list of inputs, until we find a match.
+					const TArray<FMetasoundInputDescription>& InputDescriptions = GraphsClassDeclaration->Inputs;
+
+					int32 InputDescriptionIndex = 0;
+
+					for (const FMetasoundInputDescription& InputDescription : InputDescriptions)
+					{
+						if (InputDescription.Name == InDescription.Name)
+						{
+							// We found a match. now we just need to create the input node.
+							FDataTypeLiteralParam LiteralParam = GetLiteralParamForDataType(InputDescription.TypeName, InputDescription.LiteralValue);
+
+							if (!ensureAlwaysMsgf(DoesDataTypeSupportLiteralType(InputDescription.TypeName, InputDescription.LiteralValue.LiteralType), TEXT("Tried to use an unsupported literal type!")))
+							{
+								// print out some info about the type.
+								UE_LOG(LogTemp, Display, TEXT("Data Type %s supports the following literal types:"), *InputDescription.TypeName.ToString());
+
+								FDataTypeRegistryInfo DataTypeInfo;
+
+								// this shouldn't hit at all, because this should be a registered data type.
+								ensure(GetTraitsForDataType(InputDescription.TypeName, DataTypeInfo));
+
+								if (DataTypeInfo.bIsBoolParsable)
+								{
+									UE_LOG(LogTemp, Display, TEXT("   Boolean"));
+								}
+
+								if (DataTypeInfo.bIsIntParsable)
+								{
+									UE_LOG(LogTemp, Display, TEXT("    Integer"));
+								}
+
+								if (DataTypeInfo.bIsFloatParsable)
+								{
+									UE_LOG(LogTemp, Display, TEXT("    Float"));
+								}
+
+								if (DataTypeInfo.bIsStringParsable)
+								{
+									UE_LOG(LogTemp, Display, TEXT("    String"));
+								}
+
+								return nullptr;
+							}
+
+							FInputNodeConstructorParams InitParams =
+							{
+								InputDescription.Name,
+								InputDescription.Name,
+								InSettings,
+								LiteralParam
+							};
+
+							INodePtr InputNode = ConstructInputNode(InputDescription.TypeName, InitParams);
+							if (!ensureAlwaysMsgf(InputNode.IsValid(), TEXT("Failed to construct a valid input node for Data Type %s!"), *InputDescription.TypeName.ToString()))
+							{
+								return nullptr;
+							}
+
+							TotalNodesGenerated++;
+
+							InputIndexToInputNodeIDs.Add(InputDescriptionIndex, InDescription.UniqueID);
+
+							return InputNode;
+						}
+						else
+						{
+							InputDescriptionIndex++;
+						}
+					}
+
+					// if we hit this, the document has been corrupted in some way, because we didn't have a matching Input Description for this node.
+					ensureAlwaysMsgf(false, TEXT("Document corrupted! found input node %s but couldn't find a matching Input Description for it in the Class Description."), *InDescription.Name);
+					return nullptr;
+				}
+				else if (IsOutputNode(InDescription))
+				{
+					int32 OutputDescriptionIndex = 0;
+
+					// Find this in the graph's class description's list of inputs, until we find a match.
+					const TArray<FMetasoundOutputDescription>& OutputDescriptions = GraphsClassDeclaration->Outputs;
+					for (const FMetasoundOutputDescription& OutputDescription : OutputDescriptions)
+					{
+						if (OutputDescription.Name == InDescription.Name)
+						{
+							// We found a match. now we just need to create the input node.
+							FOutputNodeConstrutorParams InitParams =
+							{
+								OutputDescription.Name,
+								OutputDescription.Name,
+							};
+
+							INodePtr OutputNode = ConstructOutputNode(OutputDescription.TypeName, InitParams);
+							if (!ensureAlwaysMsgf(OutputNode.IsValid(), TEXT("Failed to construct a valid input node for Data Type %s!"), *OutputDescription.TypeName.ToString()))
+							{
+								return nullptr;
+							}
+
+							TotalNodesGenerated++;
+
+							OutputIndexToInputNodeIDs.Add(OutputDescriptionIndex, InDescription.UniqueID);
+
+							return OutputNode;
+						}
+						else
+						{
+							OutputDescriptionIndex++;
+						}
+					}
+
+					// if we hit this, the document has been corrupted in some way, because we didn't have a matching Input Description for this node.
+					ensureAlwaysMsgf(false, TEXT("Document corrupted! found input node %s but couldn't find a matching Input Description for it in the Class Description."), *InDescription.Name);
+					return nullptr;
+				}
+				else
+				{
+					if (!ensureAlwaysMsgf(InDescription.ObjectTypeOfNode == EMetasoundClassType::External, TEXT("At this point in construction, we should only need to look up external nodes.")))
+					{
+						return nullptr;
+					}
+
+					// Find the node class in the dependencies.
+					const TArray<FMetasoundClassDescription>& DependencyDescriptions = OwningDocument->Dependencies;
+					for (const FMetasoundClassDescription& DependencyDescription : DependencyDescriptions)
+					{
+						// TODO: Add the dependency ID to the node so that we can look it up directly.
+						if (DependencyDescription.Metadata.NodeName == InDescription.Name)
+						{
+							// We found a match. now we just need to create the input node.
+							const FMetasoundExternalClassLookupInfo& LookupInfo = DependencyDescription.ExternalNodeClassLookupInfo;
+							FNodeInitData InitData;
+							InitData.InstanceName.Append(InDescription.Name);
+							InitData.InstanceName.AppendChar('_');
+							InitData.InstanceName.AppendInt(InDescription.UniqueID);
+							
+							// Copy over our initialization params.
+							for (auto& StaticParamTuple : InDescription.StaticParameters)
+							{
+								FDataTypeLiteralParam LiteralParam = GetLiteralParam(StaticParamTuple.Value);
+
+								if (LiteralParam.IsValid())
+								{
+									InitData.ParamMap.Add(StaticParamTuple.Key, LiteralParam);
+								}
+							}
+
+							INodePtr ExternalNode = ConstructExternalNode(LookupInfo.ExternalNodeClassName, LookupInfo.ExternalNodeClassHash, InitData);
+							if (!ensureAlwaysMsgf(ExternalNode.IsValid(), TEXT("Failed to construct a valid external node for Node Class %s!"), *DependencyDescription.Metadata.NodeName))
+							{
+								return nullptr;
+							}
+
+							TotalNodesGenerated++;
+							return ExternalNode;
+						}
+					}
+
+					ensureAlwaysMsgf(false, TEXT("Document corrupted! found node %s but couldn't find a matching Class Description for it in the Dependencies."), *InDescription.Name);
+					return nullptr;
+				}
+			};
+
+			FGraph GraphToBuild = FGraph(GraphsClassDeclaration->Metadata.NodeName);
+
+			TMap<uint32, INodePtr> GraphNodes;
+
+			// Step 1: Initialize Nodes
+			const TArray<FMetasoundNodeDescription>& NodeDescriptions = GraphPtr->Nodes;
+
+			for (const FMetasoundNodeDescription& NodeDescription : NodeDescriptions)
+			{
+				INodePtr NodePtr = GetNodePtr(NodeDescription);
+				
+				if (!NodePtr.IsValid())
+				{
+					return nullptr;
+				}
+				else
+				{
+					GraphNodes.Add(NodeDescription.UniqueID, MoveTemp(NodePtr));
+				}
+			}
+
+			// Sanity check that we created enough input nodes and output nodes.
+			const int32 NumInputsInMetasoundClass = GraphsClassDeclaration->Inputs.Num();
+			if (InputIndexToInputNodeIDs.Num() != NumInputsInMetasoundClass)
+			{
+				ensureAlwaysMsgf(false, TEXT("Mismatch between number of inputs in the metasound graph (%d) and number of inputs declared in it's class description (%d)!"), InputIndexToInputNodeIDs.Num(), NumInputsInMetasoundClass);
+				return nullptr;
+			}
+
+			const int32 NumOutputsInMetasoundClass = GraphsClassDeclaration->Outputs.Num();
+			if (OutputIndexToInputNodeIDs.Num() != NumOutputsInMetasoundClass)
+			{
+				ensureAlwaysMsgf(false, TEXT("Mismatch between number of outputs in the metasound graph (%d) and number of outputs declared in it's class description (%d)!"), OutputIndexToInputNodeIDs.Num(), NumOutputsInMetasoundClass);
+				return nullptr;
+			}
+
+			if (TotalNodesGenerated != GraphPtr->Nodes.Num())
+			{
+				ensureAlwaysMsgf(false, TEXT("Created %d of %d needed nodes!"), TotalNodesGenerated, GraphPtr->Nodes.Num());
+				return nullptr;
+			}
+
+			// Step 2: Connect Nodes Inside The Graph
+			for (const FMetasoundNodeDescription& NodeDescription : NodeDescriptions)
+			{
+				// TODO: create a INode type that houses literals.
+
+				INodePtr& NodeToConnectTo = GraphNodes[NodeDescription.UniqueID];
+
+				for (const FMetasoundNodeConnectionDescription& InputConnection : NodeDescription.InputConnections)
+				{
+					const bool bCanMakeConnection = InputConnection.NodeID != FMetasoundNodeConnectionDescription::DisconnectedNodeID
+						&& ensureAlwaysMsgf(GraphNodes.Contains(InputConnection.NodeID), TEXT("Connection in document describes a node ID that doesn't exist!"));
+					
+					if (bCanMakeConnection)
+					{
+						INodePtr& NodeToConnectFrom = GraphNodes[InputConnection.NodeID];
+						GraphToBuild.AddDataEdge(*NodeToConnectFrom, *InputConnection.OutputName, *NodeToConnectTo, *InputConnection.InputName);
+					}
+				}
+			}
+
+			// Step 3: Declare our Input Destinations
+			const TArray<FMetasoundInputDescription>& InputDescriptions = GraphsClassDeclaration->Inputs;
+			for (auto& InputTuple : InputIndexToInputNodeIDs)
+			{
+				INodePtr& InputNodePtr = GraphNodes[InputTuple.Value];
+				const FString& InVertexName = InputDescriptions[InputTuple.Key].Name;
+				GraphToBuild.AddInputDataDestination(*InputNodePtr, *InVertexName);
+			}
+
+			// Step 4: Declare our Output Destinations
+			const TArray<FMetasoundOutputDescription>& OutputDescriptions = GraphsClassDeclaration->Outputs;
+			for (auto& OutputTuple : OutputIndexToInputNodeIDs)
+			{
+				INodePtr& OutputNodePtr = GraphNodes[OutputTuple.Value];
+				const FString& OutVertexName = OutputDescriptions[OutputTuple.Key].Name;
+				GraphToBuild.AddOutputDataSource(*OutputNodePtr, *OutVertexName);
+			}
+
+			/** NOTE: In the future- we should split steps 1-4 above, and 5 below. GraphToBuild and the node map can be cached on the graph handle. */
+
+			// Step 5: Invoke Operator Builder
+			FOperatorBuilder Builder(InSettings);
+			return Builder.BuildGraphOperator(GraphToBuild, OutBuildErrors);
 		}
 
 		FGraphHandle GetGraphHandleForClass(const FMetasoundClassDescription& InClass)
@@ -1271,4 +1879,12 @@ namespace Metasound
 	}
 }
 
-IMPLEMENT_MODULE(FDefaultModuleImpl, MetasoundFrontend);
+class FMetasoundFrontendModule : public IModuleInterface
+{
+	virtual void StartupModule() override
+	{
+		Metasound::Frontend::InitializeFrontend();
+	}
+};
+
+IMPLEMENT_MODULE(FMetasoundFrontendModule, MetasoundFrontend);
diff --git a/Engine/Plugins/Runtime/Metasound/Source/MetasoundFrontend/Private/MetasoundFrontendBaseClasses.cpp b/Engine/Plugins/Runtime/Metasound/Source/MetasoundFrontend/Private/MetasoundFrontendBaseClasses.cpp
--- a/Engine/Plugins/Runtime/Metasound/Source/MetasoundFrontend/Private/MetasoundFrontendBaseClasses.cpp
+++ b/Engine/Plugins/Runtime/Metasound/Source/MetasoundFrontend/Private/MetasoundFrontendBaseClasses.cpp
@@ -1,11 +1,240 @@
 // Copyright Epic Games, Inc. All Rights Reserved.
 
 #include "MetasoundFrontendBaseClasses.h"
+#include "MetasoundDataReference.h"
+
+#include "MetasoundFrontendRegistries.h"
+
+
 
 namespace Metasound
 {
 	namespace Frontend
 	{
+		
+
+		TMap<FNodeRegistryKey, FNodeRegistryElement>& GetExternalNodeRegistry()
+		{
+			return FMetasoundFrontendRegistryContainer::Get()->GetExternalNodeRegistry();
+		}
+
+		TUniquePtr<INode> ConstructInputNode(const FName& InInputType, const FInputNodeConstructorParams& InParams)
+		{
+			return FMetasoundFrontendRegistryContainer::Get()->ConstructInputNode(InInputType, InParams);
+		}
+
+		TUniquePtr<INode> ConstructOutputNode(const FName& InOutputType, const FOutputNodeConstrutorParams& InParams)
+		{
+			return FMetasoundFrontendRegistryContainer::Get()->ConstructOutputNode(InOutputType, InParams);
+		}
+
+		TUniquePtr<INode> ConstructExternalNode(const FName& InNodeType, uint32 InNodeHash, const FNodeInitData& InInitData)
+		{
+			return FMetasoundFrontendRegistryContainer::Get()->ConstructExternalNode(InNodeType, InNodeHash, InInitData);
+		}
+
+		void SetLiteralDescription(FMetasoundLiteralDescription& OutDescription, bool InValue)
+		{
+			OutDescription.LiteralType = EMetasoundLiteralType::Bool;
+			OutDescription.AsBool = InValue;
+			OutDescription.AsInteger = 0;
+			OutDescription.AsFloat = 0.0f;
+			OutDescription.AsString.Empty();
+		}
+
+		void SetLiteralDescription(FMetasoundLiteralDescription& OutDescription, int32 InValue)
+		{
+			OutDescription.LiteralType = EMetasoundLiteralType::Integer;
+			OutDescription.AsBool = false;
+			OutDescription.AsInteger = InValue;
+			OutDescription.AsFloat = 0.0f;
+			OutDescription.AsString.Empty();
+		}
+
+		void SetLiteralDescription(FMetasoundLiteralDescription& OutDescription, float InValue)
+		{
+			OutDescription.LiteralType = EMetasoundLiteralType::Float;
+			OutDescription.AsBool = false;
+			OutDescription.AsInteger = 0;
+			OutDescription.AsFloat = InValue;
+			OutDescription.AsString.Empty();
+		}
+
+		void SetLiteralDescription(FMetasoundLiteralDescription& OutDescription, const FString& InValue)
+		{
+			OutDescription.LiteralType = EMetasoundLiteralType::String;
+			OutDescription.AsBool = false;
+			OutDescription.AsInteger = 0;
+			OutDescription.AsFloat = 0.0f;
+			OutDescription.AsString = InValue;
+		}
+
+		void ClearLiteralDescription(FMetasoundLiteralDescription& OutDescription)
+		{
+			OutDescription.LiteralType = EMetasoundLiteralType::None;
+			OutDescription.AsBool = false;
+			OutDescription.AsInteger = 0;
+			OutDescription.AsFloat = 0.0f;
+			OutDescription.AsString.Empty();
+		}
+
+		FDataTypeLiteralParam GetLiteralParamForDataType(FName InDataType, const FMetasoundLiteralDescription& InDescription)
+		{
+			EMetasoundLiteralType LiteralType = InDescription.LiteralType;
+			FMetasoundFrontendRegistryContainer* Registry = FMetasoundFrontendRegistryContainer::Get();
+
+			switch (LiteralType)
+			{
+				case EMetasoundLiteralType::None:
+				{
+					return GetDefaultParamForDataType(InDataType);
+				}
+				case EMetasoundLiteralType::Bool:
+				{
+					if (!Registry->DoesDataTypeSupportLiteralType(InDataType, ELiteralArgType::Boolean))
+					{
+						return FDataTypeLiteralParam::InvalidParam();
+					}
+					else
+					{
+						return FDataTypeLiteralParam(InDescription.AsBool);
+					}
+				}
+				case EMetasoundLiteralType::Float:
+				{
+					if (!Registry->DoesDataTypeSupportLiteralType(InDataType, ELiteralArgType::Float))
+					{
+						return FDataTypeLiteralParam::InvalidParam();
+					}
+					else
+					{
+						return FDataTypeLiteralParam(InDescription.AsFloat);
+					}
+				}
+				case EMetasoundLiteralType::Integer:
+				{
+					if (!Registry->DoesDataTypeSupportLiteralType(InDataType, ELiteralArgType::Integer))
+					{
+						return FDataTypeLiteralParam::InvalidParam();
+					}
+					else
+					{
+						return FDataTypeLiteralParam(InDescription.AsInteger);
+					}
+				}
+				case EMetasoundLiteralType::String:
+				{
+					if (!Registry->DoesDataTypeSupportLiteralType(InDataType, ELiteralArgType::String))
+					{
+						return FDataTypeLiteralParam::InvalidParam();
+					}
+					else
+					{
+						return FDataTypeLiteralParam(InDescription.AsString);
+					}
+				}
+				case EMetasoundLiteralType::Invalid:
+				default:
+				{
+					return FDataTypeLiteralParam::InvalidParam();
+				}
+			}
+		}
+
+		bool DoesDataTypeSupportLiteralType(FName InDataType, EMetasoundLiteralType InLiteralType)
+		{
+			switch (InLiteralType)
+			{
+			case EMetasoundLiteralType::None:
+				return DoesDataTypeSupportLiteralType(InDataType, ELiteralArgType::None);
+			case EMetasoundLiteralType::Bool:
+				return DoesDataTypeSupportLiteralType(InDataType, ELiteralArgType::Boolean);
+			case EMetasoundLiteralType::Float:
+				return DoesDataTypeSupportLiteralType(InDataType, ELiteralArgType::Float);
+			case EMetasoundLiteralType::Integer:
+				return DoesDataTypeSupportLiteralType(InDataType, EMetasoundLiteralType::Integer);
+			case EMetasoundLiteralType::String:
+				return DoesDataTypeSupportLiteralType(InDataType, EMetasoundLiteralType::String);
+			case EMetasoundLiteralType::Invalid:
+			default:
+				return false;
+			}
+		}
+
+		bool DoesDataTypeSupportLiteralType(FName InDataType, ELiteralArgType InLiteralType)
+		{
+			FMetasoundFrontendRegistryContainer* Registry = FMetasoundFrontendRegistryContainer::Get();
+
+			return Registry->DoesDataTypeSupportLiteralType(InDataType, InLiteralType);
+		}
+
+		Metasound::FDataTypeLiteralParam GetLiteralParam(const FMetasoundLiteralDescription& InDescription)
+		{
+			EMetasoundLiteralType LiteralType = InDescription.LiteralType;
+
+			switch (LiteralType)
+			{
+				case EMetasoundLiteralType::Bool:
+				{
+					return FDataTypeLiteralParam(InDescription.AsBool);
+				}
+				case EMetasoundLiteralType::Float:
+				{
+					return FDataTypeLiteralParam(InDescription.AsFloat);
+				}
+				case EMetasoundLiteralType::Integer:
+				{
+					return FDataTypeLiteralParam(InDescription.AsInteger);
+				}
+				case EMetasoundLiteralType::String:
+				{
+					return FDataTypeLiteralParam(InDescription.AsString);
+				}
+				case EMetasoundLiteralType::None:
+				case EMetasoundLiteralType::Invalid:
+				default:
+				{
+					return FDataTypeLiteralParam::InvalidParam();
+				}
+			}
+		}
+
+		Metasound::FDataTypeLiteralParam GetDefaultParamForDataType(FName InDataType)
+		{
+			FMetasoundFrontendRegistryContainer* Registry = FMetasoundFrontendRegistryContainer::Get();
+
+			ELiteralArgType DesiredArgType = Registry->GetDesiredLiteralTypeForDataType(InDataType);
+
+			switch (DesiredArgType)
+			{
+				case Metasound::ELiteralArgType::Boolean:
+				{
+					return FDataTypeLiteralParam(false);
+				}
+				case Metasound::ELiteralArgType::Integer:
+				{
+					return FDataTypeLiteralParam(0);
+				}
+				case Metasound::ELiteralArgType::Float:
+				{
+					return FDataTypeLiteralParam(0.0f);
+				}
+				case Metasound::ELiteralArgType::String:
+				{
+					return FDataTypeLiteralParam(FString());
+				}
+				case Metasound::ELiteralArgType::None:
+				{
+					return FDataTypeLiteralParam();
+				}
+				case Metasound::ELiteralArgType::Invalid:
+				default:
+				{
+					return FDataTypeLiteralParam::InvalidParam();
+				}
+			}
+
+		}
 
 		ITransactable::ITransactable(uint32 InUndoLimit, TWeakObjectPtr<UObject> InOwningAsset /* = nullptr */)
 			: UndoLimit(InUndoLimit)
@@ -377,7 +606,7 @@ namespace Metasound
 				CurrentStep = GoToNext(CurrentPath, CurrentStep);
 			}
 
-			if (ensureAlwaysMsgf(CurrentPath.Num() == 0, TEXT("Couldn't resolve part of the path.")))
+			if (ensureAlwaysMsgf(CurrentStep.Type != Path::EDescType::Invalid && CurrentPath.Num() == 0, TEXT("Couldn't resolve part of the path.")))
 			{
 				return CurrentStep.DescriptionStructPtr.Get<FClassDependencyIDs*>();
 			}
@@ -558,7 +787,7 @@ namespace Metasound
 							TArray<FMetasoundOutputDescription>& OutputsList = ClassDescription->Outputs;
 							for (FMetasoundOutputDescription& Output : OutputsList)
 							{
-								if (Output.Name == NextStep.LookupName)
+								if (Output.Name == OutputElement.LookupName)
 								{
 									FDescriptionUnwindStep UnwindStep;
 									UnwindStep.DescriptionStructPtr.Set<FMetasoundOutputDescription*>(&Output);
@@ -690,6 +919,16 @@ namespace Metasound
 			return FDescPath()[Path::EFromDocument::ToDependencies][InDependencyID];
 		}
 
+		FDescPath Path::GetInputDescriptionPath(FDescPath InPathForInputNode, const FString& InputName)
+		{
+			return (InPathForInputNode << 3)[Path::EFromClass::ToInputs][*InputName];
+		}
+
+		FDescPath Path::GetOutputDescriptionPath(FDescPath InPathForOutputNode, const FString& OutputName)
+		{
+			return (InPathForOutputNode << 3)[Path::EFromClass::ToOutputs][*OutputName];
+		}
+
 		FString Path::GetPrintableString(FDescPath InPath)
 		{
 			FString OutString = FString(TEXT("//"));
@@ -787,5 +1026,13 @@ namespace Metasound
 			return OutString;
 		}
 
+		
+
+		void InitializeFrontend()
+		{
+			FMetasoundFrontendRegistryContainer::Get()->InitializeFrontend();
+		}
 	}
 }
+
+
diff --git a/Engine/Plugins/Runtime/Metasound/Source/MetasoundFrontend/Public/MetasoundFrontend.h b/Engine/Plugins/Runtime/Metasound/Source/MetasoundFrontend/Public/MetasoundFrontend.h
--- a/Engine/Plugins/Runtime/Metasound/Source/MetasoundFrontend/Public/MetasoundFrontend.h
+++ b/Engine/Plugins/Runtime/Metasound/Source/MetasoundFrontend/Public/MetasoundFrontend.h
@@ -6,14 +6,13 @@
 #include "MetasoundGraph.h"
 #include "MetasoundFrontendDataLayout.h"
 #include "MetasoundFrontendBaseClasses.h"
+#include "MetasoundBuilderInterface.h"
 #include "UObject/WeakObjectPtrTemplates.h"
-#include "Serialization/Public/IStructSerializerBackend.h"
-#include "Serialization/Public/Backends/JsonStructSerializerBackend.h"
-#include "Serialization/Public/StructSerializer.h"
 
 // Forward Declarations
 class FMetasoundAssetBase;
 
+
 namespace Metasound
 {	
 	/**
@@ -41,9 +40,14 @@ namespace Metasound
 		// and also used in FGraphHandle::AddNewNode.
 		struct FNodeClassInfo
 		{
+			// The descriptive name of this node class.
 			FString NodeName;
+
+			// The type for this node.
 			EMetasoundClassType NodeType;
 
+			// The lookup key used for the internal node registry.
+			FNodeRegistryKey LookupKey;
 		};
 
 		// Get all available nodes of any type.
@@ -61,11 +65,18 @@ namespace Metasound
 		TArray<FNodeClassInfo> GetAllNodesWithAnInputOfType(const FName& InType);
 
 		// gets all metadata (name, description, author, what to say if it's missing) for a given node.
-		FMetasoundClassMetadata GetMetadataForNode(const FNodeClassInfo& InInfo);
+		FMetasoundClassMetadata GenerateMetadataForNode(const FNodeClassInfo& InInfo);
 
 
 		// Generates a new FMetasoundClassDescription for a given node class. Only used by classes that manipulate Metasound Description data directly.
-		FMetasoundClassDescription GetClassDescriptionForNode(const FNodeClassInfo& InInfo);
+		FMetasoundClassDescription GenerateClassDescriptionForNode(const FNodeClassInfo& InInfo);
+
+		// Returns a list of all available data types.
+		TArray<FName> GetAllAvailableDataTypes();
+
+		// outputs the traits for a given data type.
+		// returns false if InDataType couldn't be found.
+		bool GetTraitsForDataType(FName InDataType, FDataTypeRegistryInfo& OutInfo);
 
 		// Struct that indicates whether an input and an output can be connected,
 		// and whether an intermediate node is necessary to connect the two.
@@ -122,6 +133,9 @@ namespace Metasound
 
 			explicit FOutputHandle(FHandleInitParams::EPrivateToken InToken, const FHandleInitParams& InParams, const FString& InOutputName);
 
+			// Constructor used for the outgoing connection from an input node.
+			explicit FOutputHandle(FHandleInitParams::EPrivateToken InToken, const FHandleInitParams& InParams);
+
 			~FOutputHandle() = default;
 
 			static FOutputHandle InvalidHandle();
@@ -134,7 +148,7 @@ namespace Metasound
 
 			FName GetOutputType() const;
 			FString GetOutputName() const;
-			FString GetOutputTooltip() const;
+			FText GetOutputTooltip() const;
 			uint32 GetOwningNodeID() const;
 
 			// @todo: with the current spec, we'd have to scan all nodes in the graph for this output to solve this.
@@ -148,7 +162,13 @@ namespace Metasound
 		private:
 			TDescriptionPtr<FMetasoundNodeDescription> NodePtr;
 			TDescriptionPtr<FMetasoundClassDescription> NodeClass;
+
+			// owning output description ptr for the node's class.
+			// This can be none if the node that owns this output is an input node itself.
 			TDescriptionPtr<FMetasoundOutputDescription> OutputPtr;
+
+			// Optional description pointer used when this output connection is owned by an input node.
+			TDescriptionPtr<FMetasoundInputDescription> InputNodePtr;
 		};
 
 		class METASOUNDFRONTEND_API FInputHandle : protected ITransactable
@@ -160,6 +180,9 @@ namespace Metasound
 			// Sole constructor for FInputHandle. Can only be used by friend classes for FHandleInitParams.
 			explicit FInputHandle(FHandleInitParams::EPrivateToken PrivateToken, const FHandleInitParams& InParams, const FString& InputName);
 
+			// constructor used exclusively by output nodes.
+			explicit FInputHandle(FHandleInitParams::EPrivateToken PrivateToken, const FHandleInitParams& InParams);
+
 			static FInputHandle InvalidHandle();
 
 			bool IsValid() const;
@@ -167,7 +190,7 @@ namespace Metasound
 
 			FName GetInputType() const;
 			FString GetInputName() const;
-			FString GetInputTooltip() const;
+			FText GetInputTooltip() const;
 
 			FOutputHandle GetCurrentlyConnectedOutput() const;
 
@@ -177,9 +200,18 @@ namespace Metasound
 
 		private:
 			TDescriptionPtr<FMetasoundNodeDescription> NodePtr;
+
+			// owning node class for the node that owns this input connection.
+			// This can be none if the node that owns this output is an output node itself.
 			TDescriptionPtr<FMetasoundClassDescription> NodeClass;
+
+			// owning input description ptr for the node's class.
+			// This can be none if the node that owns this output is an output node itself.
 			TDescriptionPtr<FMetasoundInputDescription> InputPtr;
 
+			// Optional description pointer used when this input connection is owned by an output node.
+			TDescriptionPtr<FMetasoundOutputDescription> OutputNodePtr;
+
 			FString InputName;
 
 			const FMetasoundNodeConnectionDescription* GetConnectionDescription() const;
@@ -198,7 +230,7 @@ namespace Metasound
 			~FNodeHandle() = default;
 
 			// Sole constructor for FNodeHandle. Can only be used by friends of FHandleInitParams.
-			explicit FNodeHandle(FHandleInitParams::EPrivateToken PrivateToken, const FHandleInitParams& InParams);
+			explicit FNodeHandle(FHandleInitParams::EPrivateToken PrivateToken, const FHandleInitParams& InParams, EMetasoundClassType InNodeClassType);
 
 			static FNodeHandle InvalidHandle();
 
@@ -226,6 +258,10 @@ namespace Metasound
 			TDescriptionPtr<FMetasoundNodeDescription> NodePtr;
 			TDescriptionPtr<FMetasoundClassDescription> NodeClass;
 
+			// whether this node is an input or output to it's owning graph,
+			// an externally implemented node, or a metasound graph itself.
+			EMetasoundClassType NodeClassType;
+
 			uint32 NodeID;
 		};
 
@@ -252,6 +288,7 @@ namespace Metasound
 			bool IsValid() const;
 
 			TArray<FNodeHandle> GetAllNodes();
+			FNodeHandle GetNodeWithId(uint32 InNodeId);
 			TArray<FNodeHandle> GetOutputNodes();
 			TArray<FNodeHandle> GetInputNodes();
 
@@ -296,9 +333,8 @@ namespace Metasound
 			// This invokes the Metasound Builder to synchronously compile a 
 			// metasound operator, which can then be used for playback.
 			// @returns nullptr on failure.
-			// @todo: forward errors from the Metasound Builder.
 			// @todo: make an API for doing this async.
-			TUniquePtr<IOperator> BuildOperator();
+			TUniquePtr<IOperator> BuildOperator(const FOperatorSettings& InSettings, TArray<IOperatorBuilder::FBuildErrorPtr>& OutBuildErrors) const;
 
 		private:
 
diff --git a/Engine/Plugins/Runtime/Metasound/Source/MetasoundFrontend/Public/MetasoundFrontendBaseClasses.h b/Engine/Plugins/Runtime/Metasound/Source/MetasoundFrontend/Public/MetasoundFrontendBaseClasses.h
--- a/Engine/Plugins/Runtime/Metasound/Source/MetasoundFrontend/Public/MetasoundFrontendBaseClasses.h
+++ b/Engine/Plugins/Runtime/Metasound/Source/MetasoundFrontend/Public/MetasoundFrontendBaseClasses.h
@@ -5,12 +5,53 @@
 #include "CoreMinimal.h"
 #include "Misc/TVariant.h"
 #include "MetasoundFrontendDataLayout.h"
+#include "MetasoundFrontendRegistries.h"
+
 
 
 namespace Metasound
-{	
+{
 	namespace Frontend
 	{
+		// This is called by FMetasoundFrontendModule, and flushes any node or datatype registration that was done prior to boot.
+		void InitializeFrontend();
+
+		/** Use this to get a reference to the static registry of all nodes implemented in C++. */
+		TMap<FNodeRegistryKey, FNodeRegistryElement>& GetExternalNodeRegistry();
+
+		// Convenience functions to create an INode corresponding to a specific input or output for a metasound graph.
+		// @returns nullptr if the type given wasn't found.
+		TUniquePtr<INode> ConstructInputNode(const FName& InInputType, const FInputNodeConstructorParams& InParams);
+		TUniquePtr<INode> ConstructOutputNode(const FName& InOutputType, const FOutputNodeConstrutorParams& InParams);
+
+		// Convenience functions to create an INodeB corresponding to a specific externally declared node type.
+		// InNodeType and InNodeHash can be retrieved from the FNodeClassInfo generated from the node registry queries in the metasound frontend (GetAllAvailableNodeClasses, GetAllNodeClassesInNamespace, etc.)
+		// @returns nullptr if the type given wasn't found.
+		TUniquePtr<INode> ConstructExternalNode(const FName& InNodeType, uint32 InNodeHash, const FNodeInitData& InInitData);
+
+		// Utility functions for setting serialized literals.
+		void SetLiteralDescription(FMetasoundLiteralDescription& OutDescription, bool InValue);
+		void SetLiteralDescription(FMetasoundLiteralDescription& OutDescription, int32 InValue);
+		void SetLiteralDescription(FMetasoundLiteralDescription& OutDescription, float InValue);
+		void SetLiteralDescription(FMetasoundLiteralDescription& OutDescription, const FString& InValue);
+		void ClearLiteralDescription(FMetasoundLiteralDescription& OutDescription);
+
+		// Utility functions for building a ::Metasound::FDataInitParam corresponding to a literal.
+
+		// Return the literal description parsed into a init param. 
+		// @Returns an invalid init param if the data type couldn't be found, or if the literal type was incompatible with the data type.
+		FDataTypeLiteralParam GetLiteralParamForDataType(FName InDataType, const FMetasoundLiteralDescription& InDescription);
+
+		bool DoesDataTypeSupportLiteralType(FName InDataType, EMetasoundLiteralType InLiteralType);
+		bool DoesDataTypeSupportLiteralType(FName InDataType, ELiteralArgType InLiteralType);
+		
+		// Returns a literal param without any type checking.
+		FDataTypeLiteralParam GetLiteralParam(const FMetasoundLiteralDescription& InDescription);
+
+		// Returns the defaulted version of a literal param for the given data type.
+		// @Returns an invalid init param if the data type couldn't be found.
+		FDataTypeLiteralParam GetDefaultParamForDataType(FName InDataType);
+
 		// Utility base class for classes that want to support naive multi-level undo/redo.
 		class METASOUNDFRONTEND_API ITransactable : public TSharedFromThis<ITransactable>
 		{
@@ -459,6 +500,11 @@ namespace Metasound
 
 			FDescPath GetDependencyPath(uint32 InDependencyID);
 
+			// Given the path to an input node, returns the path to the corresponding  input description.
+			FDescPath GetInputDescriptionPath(FDescPath InPathForInputNode, const FString& InputName);
+
+			FDescPath GetOutputDescriptionPath(FDescPath InPathForOutputNode, const FString& OutputName);
+
 			// Generates a human readable string for a path into a Metasound description. Used for debugging.
 			FString GetPrintableString(FDescPath InPath);
 		}
diff --git a/Engine/Plugins/Runtime/Metasound/Source/MetasoundFrontend/Public/MetasoundFrontendDataLayout.h b/Engine/Plugins/Runtime/Metasound/Source/MetasoundFrontend/Public/MetasoundFrontendDataLayout.h
--- a/Engine/Plugins/Runtime/Metasound/Source/MetasoundFrontend/Public/MetasoundFrontendDataLayout.h
+++ b/Engine/Plugins/Runtime/Metasound/Source/MetasoundFrontend/Public/MetasoundFrontendDataLayout.h
@@ -9,10 +9,7 @@
 UENUM()
 enum class EMetasoundClassType : uint8
 {
-	// Objects with the type Primitive are always available as part of the Metasound specification.
-	// The Metasound Builder can optimize usage of these types of nodes.
-	Primitive,
-	// Objects with the type External are implemented in compiled code.
+	// Any classes of the external type are implemented in C++..
 	External,
 	// Objects with this type are Metasound graphs unto themselves,
 	// they can be used like any other node.
@@ -24,6 +21,18 @@ enum class EMetasoundClassType : uint8
 	Invalid UMETA(Hidden)
 };
 
+// The type of a given literal for an input value.
+UENUM()
+enum class EMetasoundLiteralType : uint8
+{
+	None,
+	Bool,
+	Float,
+	Integer,
+	String,
+	Invalid UMETA(Hidden)
+};
+
 USTRUCT()
 struct FMetasoundClassMetadata
 {
@@ -53,6 +62,31 @@ struct FMetasoundClassMetadata
 	FString AuthorName;
 };
 
+
+// This ustruct represents the serialized version of our variant literal type.
+// Currently, we only support bools, integers, floats, and strings as literals.
+USTRUCT()
+struct FMetasoundLiteralDescription
+{
+	GENERATED_BODY()
+
+	// The actual type of this literal.
+	UPROPERTY()
+	EMetasoundLiteralType LiteralType;
+
+	UPROPERTY()
+	bool AsBool;
+
+	UPROPERTY();
+	int32 AsInteger;
+
+	UPROPERTY()
+	float AsFloat;
+
+	UPROPERTY()
+	FString AsString;
+};
+
 USTRUCT()
 struct FMetasoundInputDescription
 {
@@ -68,7 +102,11 @@ struct FMetasoundInputDescription
 
 	// Optional description text about this input.
 	UPROPERTY()
-	FString ToolTip;
+	FText ToolTip;
+
+	// The optional literal value, if we have one.
+	UPROPERTY()
+	FMetasoundLiteralDescription LiteralValue;
 };
 
 USTRUCT()
@@ -86,7 +124,7 @@ struct FMetasoundOutputDescription
 
 	// Optional description text about this output.
 	UPROPERTY()
-	FString ToolTip;
+	FText ToolTip;
 };
 
 // This is used to describe a single incoming connection or literal for a given input on a FMetasoundNodeDescription.
@@ -106,9 +144,13 @@ struct FMetasoundNodeConnectionDescription
 	uint32 NodeID;
 
 	// FMetasoundOutputDescription::Name of the output this connection is from if NodeID is valid,
-	// or a literal value if NodeID is FMetasoundNodeConnectionDescription::DisconnectedNodeID.
+	// if FMetasoundNodeConnectionDescription::DisconnectedNodeID.
 	UPROPERTY()
-	FString OutputNameOrLiteral;
+	FString OutputName;
+
+	// The optional literal value, if we have one.
+	UPROPERTY()
+	FMetasoundLiteralDescription LiteralValue;
 };
 
 // description for an individual node within the FMetasoundGraphDescription.
@@ -137,6 +179,10 @@ struct FMetasoundNodeDescription
 	// list of connections or literal values for each connected input pin on this node.
 	UPROPERTY()
 	TArray<FMetasoundNodeConnectionDescription> InputConnections;
+
+	// todo: create way to control init params for specific node instances.
+	UPROPERTY()
+	TMap<FName, FMetasoundLiteralDescription> StaticParameters;
 };
 
 USTRUCT()
@@ -144,8 +190,17 @@ struct FMetasoundGraphDescription
 {
 	GENERATED_BODY()
 
-		UPROPERTY()
-		TArray<FMetasoundNodeDescription> Nodes;
+	UPROPERTY()
+	TArray<FMetasoundNodeDescription> Nodes;
+};
+
+USTRUCT()
+struct FMetasoundExternalClassLookupInfo
+{
+	GENERATED_BODY()
+
+	FName ExternalNodeClassName;
+	uint32 ExternalNodeClassHash;
 };
 
 // Full saved out description of a Metasound object,
@@ -183,6 +238,10 @@ struct FMetasoundClassDescription
 	// we will attempt to find the implemented version of this Metasound in the asset registry using Metasound::Frontend::FindGraphHandleForClass.
 	UPROPERTY()
 	FMetasoundGraphDescription Graph;
+
+	// If this object is an external node type (implemented in C++), this name is used to quickly look up the class in the Node Registry.
+	UPROPERTY()
+	FMetasoundExternalClassLookupInfo ExternalNodeClassLookupInfo;
 };
 
 // This is used to describe the required inputs and outputs for a metasound, and is used to make sure we can use a metasound graph for specific applications.
@@ -197,10 +256,12 @@ struct FMetasoundArchetype
 	FName ArchetypeName;
 
 	// whatever inputs are required to support this archetype.
+	// call Special Inputs
 	UPROPERTY()
 	TArray<FMetasoundInputDescription> RequiredInputs;
 
 	// whatever outputs are required to support this archetype.
+	// call Special Outputs
 	UPROPERTY()
 	TArray<FMetasoundOutputDescription> RequiredOutputs;
 };
diff --git a/Engine/Plugins/Runtime/Metasound/Source/MetasoundGraphCore/Private/MetasoundADSRNode.cpp b/Engine/Plugins/Runtime/Metasound/Source/MetasoundGraphCore/Private/MetasoundADSRNode.cpp
--- a/Engine/Plugins/Runtime/Metasound/Source/MetasoundGraphCore/Private/MetasoundADSRNode.cpp
+++ b/Engine/Plugins/Runtime/Metasound/Source/MetasoundGraphCore/Private/MetasoundADSRNode.cpp
@@ -8,8 +8,12 @@
 
 #define LOCTEXT_NAMESPACE "MetasoundADSRNode"
 
+
+
 namespace Metasound
 {
+	METASOUND_REGISTER_NODE(FADSRNode)
+
 	struct FADSRDataReferences
 	{
 		FFloatTimeReadRef Attack;
@@ -26,7 +30,7 @@ namespace Metasound
 			,	Bop(InBop)
 			,	StopEnvelopePos(-1)
 			,	ADSRDataReferences(InADSRData)
-			,	EnvelopeBuffer(InSettings.FramesPerExecute)
+			,	EnvelopeBuffer(FAudioBufferWriteRef::CreateNew(InSettings.FramesPerExecute))
 			{
 				check(EnvelopeBuffer->Num() == InSettings.FramesPerExecute);
 
@@ -128,17 +132,17 @@ namespace Metasound
 	{
 		const FADSRNode& ADSRNode = static_cast<const FADSRNode&>(InNode);
 
-		FBopReadRef Bop;
+		FBopReadRef Bop = FBopReadRef::CreateNew();
 
 		// TODO: Could no-op this if the bop is not connected.
 		SetReadableRefIfInCollection(TEXT("Bop"), InInputDataReferences, Bop);
 
 		FADSRDataReferences ADSRDataReferences = 
 			{
-				FFloatTimeReadRef(ADSRNode.GetDefaultAttackMs(), ETimeResolution::Milliseconds),
-				FFloatTimeReadRef(ADSRNode.GetDefaultDecayMs(), ETimeResolution::Milliseconds),
-				FFloatTimeReadRef(ADSRNode.GetDefaultSustainMs(), ETimeResolution::Milliseconds),
-				FFloatTimeReadRef(ADSRNode.GetDefaultReleaseMs(), ETimeResolution::Milliseconds)
+				FFloatTimeReadRef::CreateNew(ADSRNode.GetDefaultAttackMs(), ETimeResolution::Milliseconds),
+				FFloatTimeReadRef::CreateNew(ADSRNode.GetDefaultDecayMs(), ETimeResolution::Milliseconds),
+				FFloatTimeReadRef::CreateNew(ADSRNode.GetDefaultSustainMs(), ETimeResolution::Milliseconds),
+				FFloatTimeReadRef::CreateNew(ADSRNode.GetDefaultReleaseMs(), ETimeResolution::Milliseconds)
 			};
 
 		// TODO: If none of these are connected, could pregenerate ADSR envelope and return a different operator. 
@@ -166,6 +170,11 @@ namespace Metasound
 		AddOutputDataVertex<FAudioBuffer>(TEXT("Envelope"), LOCTEXT("EnvelopeTooltip", "The output envelope"));
 	}
 
+	FADSRNode::FADSRNode(const FNodeInitData& InitData)
+		: FADSRNode(InitData.InstanceName, 10.0f, 20.0f, 50.0f, 20.0f)
+	{
+	}
+
 	FADSRNode::~FADSRNode()
 	{
 	}
diff --git a/Engine/Plugins/Runtime/Metasound/Source/MetasoundGraphCore/Private/MetasoundAudioFormats.cpp b/Engine/Plugins/Runtime/Metasound/Source/MetasoundGraphCore/Private/MetasoundAudioFormats.cpp
--- a/Engine/Plugins/Runtime/Metasound/Source/MetasoundGraphCore/Private/MetasoundAudioFormats.cpp
+++ b/Engine/Plugins/Runtime/Metasound/Source/MetasoundGraphCore/Private/MetasoundAudioFormats.cpp
@@ -8,6 +8,10 @@
 
 namespace Metasound
 {
+	REGISTER_METASOUND_DATATYPE(FAudioBuffer)
+	REGISTER_METASOUND_DATATYPE(FUnformattedAudio)
+	REGISTER_METASOUND_DATATYPE(FMultichannelAudioFormat)
+
 	/* FUnformattedAudio */
 
 	FUnformattedAudio::FUnformattedAudio(int32 InNumFrames, int32 InNumChannels, int32 InMaxNumChannels)
@@ -20,7 +24,7 @@ namespace Metasound
 
 		for (int32 i = 0; i < MaxNumChannels; i++)
 		{
-			FAudioBufferWriteRef Audio(NumFrames);
+			FAudioBufferWriteRef Audio = FAudioBufferWriteRef::CreateNew(NumFrames);
 			Audio->Zero();
 
 			WritableBufferStorage.Add(Audio);
@@ -30,6 +34,11 @@ namespace Metasound
 		SetNumChannels(InNumChannels);
 	}
 
+	FUnformattedAudio::FUnformattedAudio(int32 InInitialNumChannels, const FOperatorSettings& InSettings)
+		: FUnformattedAudio(InSettings.FramesPerExecute, InInitialNumChannels, 8)
+	{
+	}
+
 	int32 FUnformattedAudio::SetNumChannels(int32 InNumChannels)
 	{
 		NumChannels = FMath::Max(0, FMath::Min(InNumChannels, MaxNumChannels));
@@ -56,7 +65,7 @@ namespace Metasound
 
 		for (int32 i = 0; i < NumChannels; i++)
 		{
-			FAudioBufferWriteRef Audio(InNumFrames);
+			FAudioBufferWriteRef Audio = FAudioBufferWriteRef::CreateNew(InNumFrames);
 			Audio->Zero();
 
 			WritableBufferStorage.Add(Audio);
@@ -67,6 +76,11 @@ namespace Metasound
 		ReadableBuffers = ReadableBufferStorage;
 	}
 
+	FMultichannelAudioFormat::FMultichannelAudioFormat(int32 InNumChannels, const FOperatorSettings& InSettings)
+		: FMultichannelAudioFormat(InSettings.FramesPerExecute, InNumChannels)
+	{}
+
+
 	FMultichannelAudioFormat::FMultichannelAudioFormat(TArrayView<const FAudioBufferWriteRef> InWriteRefs)
 	:	NumChannels(InWriteRefs.Num())
 	{
diff --git a/Engine/Plugins/Runtime/Metasound/Source/MetasoundGraphCore/Private/MetasoundAudioMultiplyNode.cpp b/Engine/Plugins/Runtime/Metasound/Source/MetasoundGraphCore/Private/MetasoundAudioMultiplyNode.cpp
--- a/Engine/Plugins/Runtime/Metasound/Source/MetasoundGraphCore/Private/MetasoundAudioMultiplyNode.cpp
+++ b/Engine/Plugins/Runtime/Metasound/Source/MetasoundGraphCore/Private/MetasoundAudioMultiplyNode.cpp
@@ -9,6 +9,7 @@
 
 namespace Metasound
 {
+	METASOUND_REGISTER_NODE(FAudioMultiplyNode);
 
 	class FAudioMultiplyOperator : public TExecutableOperator<FAudioMultiplyOperator>
 	{
@@ -17,7 +18,7 @@ namespace Metasound
 			:	OperatorSettings(InSettings)
 			,	InputBuffer1(InBuffer1)
 			,	InputBuffer2(InBuffer2)
-			,	OutputBuffer(InSettings.FramesPerExecute)
+			,	OutputBuffer(FAudioBufferWriteRef::CreateNew(InSettings))
 			{
 				check(OutputBuffer->Num() == InSettings.FramesPerExecute);
 				check(InputBuffer1->Num() == InSettings.FramesPerExecute);
@@ -62,8 +63,8 @@ namespace Metasound
 	{
 		const FAudioMultiplyNode& AudioMultiplyNode = static_cast<const FAudioMultiplyNode&>(InNode);
 
-		FAudioBufferReadRef InputBuffer1;
-		FAudioBufferReadRef InputBuffer2;
+		FAudioBufferReadRef InputBuffer1 = FAudioBufferReadRef::CreateNew();
+		FAudioBufferReadRef InputBuffer2 = FAudioBufferReadRef::CreateNew();
 
 		FAudioBuffer Ones(InOperatorSettings.FramesPerExecute);
 
@@ -76,12 +77,12 @@ namespace Metasound
 
 		if (!SetReadableRefIfInCollection(TEXT("InputBuffer1"), InInputDataReferences, InputBuffer1))
 		{
-			InputBuffer1 = FAudioBufferReadRef(Ones);
+			InputBuffer1 = FAudioBufferReadRef::CreateNew(Ones);
 		}
 
 		if (!SetReadableRefIfInCollection(TEXT("InputBuffer2"), InInputDataReferences, InputBuffer2))
 		{
-			InputBuffer2 = FAudioBufferReadRef(Ones);
+			InputBuffer2 = FAudioBufferReadRef::CreateNew(Ones);
 		}
 
 		return MakeUnique<FAudioMultiplyOperator>(InOperatorSettings, InputBuffer1, InputBuffer2);
@@ -96,6 +97,10 @@ namespace Metasound
 		AddOutputDataVertex<FAudioBuffer>(TEXT("Audio"), LOCTEXT("OutpuBufferTooltip", "The output audio."));
 	}
 
+	FAudioMultiplyNode::FAudioMultiplyNode(const FNodeInitData& InInitData)
+		: FAudioMultiplyNode(InInitData.InstanceName)
+	{}
+
 	FAudioMultiplyNode::~FAudioMultiplyNode()
 	{
 	}
diff --git a/Engine/Plugins/Runtime/Metasound/Source/MetasoundGraphCore/Private/MetasoundFrontendRegistries.cpp b/Engine/Plugins/Runtime/Metasound/Source/MetasoundGraphCore/Private/MetasoundFrontendRegistries.cpp
new file mode 100644
--- /dev/null
+++ b/Engine/Plugins/Runtime/Metasound/Source/MetasoundGraphCore/Private/MetasoundFrontendRegistries.cpp
@@ -0,0 +1,319 @@
+// Copyright Epic Games, Inc. All Rights Reserved.
+
+#include "MetasoundFrontendRegistries.h"
+#include "CoreMinimal.h"
+
+#ifndef WITH_METASOUND_FRONTEND
+#define WITH_METASOUND_FRONTEND 0
+#endif
+
+FMetasoundFrontendRegistryContainer::FMetasoundFrontendRegistryContainer()
+	: bHasModuleBeenInitialized(false)
+{
+}
+
+FMetasoundFrontendRegistryContainer* FMetasoundFrontendRegistryContainer::LazySingleton = nullptr;
+
+FMetasoundFrontendRegistryContainer* FMetasoundFrontendRegistryContainer::Get()
+{
+	if (!LazySingleton)
+	{
+		LazySingleton = new FMetasoundFrontendRegistryContainer();
+	}
+
+	return LazySingleton;
+}
+
+void FMetasoundFrontendRegistryContainer::ShutdownMetasoundFrontend()
+{
+	if (LazySingleton)
+	{
+		delete LazySingleton;
+		LazySingleton = nullptr;
+	}
+}
+
+void FMetasoundFrontendRegistryContainer::InitializeFrontend()
+{
+	FScopeLock ScopeLock(&LazyInitCommandCritSection);
+
+	if (bHasModuleBeenInitialized)
+	{
+		// this function should only be called once.
+		checkNoEntry();
+	}
+
+	UE_LOG(LogTemp, Display, TEXT("Initializing Metasounds Frontend."));
+	uint64 CurrentTime = FPlatformTime::Cycles64();
+
+	for (TUniqueFunction<void()>& Command : LazyInitCommands)
+	{
+		Command();
+	}
+
+	LazyInitCommands.Empty();
+	bHasModuleBeenInitialized = true;
+
+	uint64 CyclesUsed = FPlatformTime::Cycles64() - CurrentTime;
+	UE_LOG(LogTemp, Display, TEXT("Initializing Metasounds Frontend took %f seconds."), FPlatformTime::ToSeconds64(CyclesUsed));
+}
+
+bool FMetasoundFrontendRegistryContainer::EnqueueInitCommand(TUniqueFunction<void()>&& InFunc)
+{
+	// if the module has been initalized already, we can safely call this function now.
+	if (bHasModuleBeenInitialized)
+	{
+		InFunc();
+	}
+
+	// otherwise, we enqueue the function to be executed after the frontend module has been initialized.
+
+	if (LazyInitCommands.Num() >= MaxNumNodesAndDatatypesToInitialize)
+	{
+		UE_LOG(LogTemp, Warning, TEXT("Registering more that %d nodes and datatypes for metasounds! Consider increasing MetasoundFrontendRegistryContainer::MaxNumNodesAndDatatypesToInitialize."));
+	}
+
+	LazyInitCommands.Add(MoveTemp(InFunc));
+	return true;
+}
+
+TMap<FMetasoundFrontendRegistryContainer::FNodeRegistryKey, FMetasoundFrontendRegistryContainer::FNodeRegistryElement>& FMetasoundFrontendRegistryContainer::GetExternalNodeRegistry()
+{
+	return ExternalNodeRegistry;
+}
+
+TUniquePtr<Metasound::INode> FMetasoundFrontendRegistryContainer::ConstructInputNode(const FName& InInputType, const Metasound::FInputNodeConstructorParams& InParams)
+{
+	if (ensureAlwaysMsgf(DataTypeRegistry.Contains(InInputType), TEXT("Couldn't find data type %s!"), *InInputType.ToString()))
+	{
+		return DataTypeRegistry[InInputType].InputNodeConstructor(InParams);
+	}
+	else
+	{
+		return nullptr;
+	}
+}
+
+TUniquePtr<Metasound::INode> FMetasoundFrontendRegistryContainer::ConstructOutputNode(const FName& InOutputType, const Metasound::FOutputNodeConstrutorParams& InParams)
+{
+	if (ensureAlwaysMsgf(DataTypeRegistry.Contains(InOutputType), TEXT("Couldn't find data type %s!"), *InOutputType.ToString()))
+	{
+		return DataTypeRegistry[InOutputType].OutputNodeConstructor(InParams);
+	}
+	else
+	{
+		return nullptr;
+	}
+}
+
+TUniquePtr<Metasound::INode> FMetasoundFrontendRegistryContainer::ConstructExternalNode(const FName& InNodeType, uint32 InNodeHash, const Metasound::FNodeInitData& InInitData)
+{
+	Metasound::Frontend::FNodeRegistryKey RegistryKey;
+	RegistryKey.NodeName = InNodeType;
+	RegistryKey.NodeHash = InNodeHash;
+
+	if (!ExternalNodeRegistry.Contains(RegistryKey))
+	{
+		return nullptr;
+	}
+	else
+	{
+		return ExternalNodeRegistry[RegistryKey].GetterCallback(InInitData);
+	}
+}
+
+Metasound::ELiteralArgType FMetasoundFrontendRegistryContainer::GetDesiredLiteralTypeForDataType(FName InDataType) const
+{
+	if (!DataTypeRegistry.Contains(InDataType))
+	{
+		return Metasound::ELiteralArgType::Invalid;
+	}
+
+	const FDataTypeRegistryElement& DataTypeInfo = DataTypeRegistry[InDataType];
+	
+	// If there's a designated preferred literal type for this datatype, use that.
+	if (DataTypeInfo.Info.PreferredLiteralType != Metasound::ELiteralArgType::None)
+	{
+		return DataTypeInfo.Info.PreferredLiteralType;
+	}
+
+	// Otherwise, we opt for the highest precision construction option available.
+	if (DataTypeInfo.Info.bIsStringParsable)
+	{
+		return Metasound::ELiteralArgType::String;
+	}
+	else if (DataTypeInfo.Info.bIsFloatParsable)
+	{
+		return Metasound::ELiteralArgType::Float;
+	}
+	else if (DataTypeInfo.Info.bIsIntParsable)
+	{
+		return Metasound::ELiteralArgType::Integer;
+	}
+	else if (DataTypeInfo.Info.bIsBoolParsable)
+	{
+		return Metasound::ELiteralArgType::Boolean;
+	}
+	else if (DataTypeInfo.Info.bIsConstructableWithSettings || DataTypeInfo.Info.bIsDefaultConstructible)
+	{
+		return Metasound::ELiteralArgType::None;
+	}
+	else
+	{
+		// if we ever hit this, something has gone terribly wrong with the REGISTER_METASOUND_DATATYPE macro.
+		// we should have failed to compile if any of these are false.
+		checkNoEntry();
+		return Metasound::ELiteralArgType::Invalid;
+	}
+}
+
+bool FMetasoundFrontendRegistryContainer::DoesDataTypeSupportLiteralType(FName InDataType, Metasound::ELiteralArgType InLiteralType) const
+{
+	if (!DataTypeRegistry.Contains(InDataType))
+	{
+		ensureAlwaysMsgf(false, TEXT("couldn't find DataType %s in the registry."), *InDataType.ToString());
+		return false;
+	}
+
+	const FDataTypeRegistryElement& DataTypeInfo = DataTypeRegistry[InDataType];
+	
+	switch (InLiteralType)
+	{
+		case Metasound::ELiteralArgType::Boolean:
+		{
+			return DataTypeInfo.Info.bIsBoolParsable;
+		}
+		case Metasound::ELiteralArgType::Integer:
+		{
+			return DataTypeInfo.Info.bIsIntParsable;
+		}
+		case Metasound::ELiteralArgType::Float:
+		{
+			return DataTypeInfo.Info.bIsFloatParsable;
+		}
+		case Metasound::ELiteralArgType::String:
+		{
+			return DataTypeInfo.Info.bIsStringParsable;
+		}
+		case Metasound::ELiteralArgType::None:
+		{
+			return DataTypeInfo.Info.bIsConstructableWithSettings || DataTypeInfo.Info.bIsDefaultConstructible;
+		}
+		case Metasound::ELiteralArgType::Invalid:
+		default:
+		{
+			return false;
+		}
+	}
+}
+
+bool FMetasoundFrontendRegistryContainer::RegisterDataType(const ::Metasound::FDataTypeRegistryInfo& InDataInfo, FInputNodeConstructorCallback&& InputNodeConstructor, FOutputNodeConstructorCallback&& OutputNodeConstructor)
+{
+	if (!ensureAlwaysMsgf(!DataTypeRegistry.Contains(InDataInfo.DataTypeName), TEXT("Name collision when trying to register Metasound Data Type %s! Make sure that you created a unique name for your data type, and that REGISTER_METASOUND_DATATYPE isn't called in a public header."), *InDataInfo.DataTypeName.ToString()))
+	{
+		// todo: capture callstack for previous declaration for non-shipping builds to help clarify who already registered this name for a type.
+		return false;
+	}
+	else
+	{
+		FDataTypeRegistryElement InElement = { MoveTemp(InputNodeConstructor), MoveTemp(OutputNodeConstructor), InDataInfo };
+		DataTypeRegistry.Add(InDataInfo.DataTypeName, MoveTemp(InElement));
+		UE_LOG(LogTemp, Display, TEXT("Registered Metasound Datatype %s."), *InDataInfo.DataTypeName.ToString());
+		return true;
+	}
+}
+
+bool FMetasoundFrontendRegistryContainer::RegisterExternalNode(FNodeGetterCallback&& InCallback)
+{
+	using FInputDataVertexCollection = Metasound::FInputDataVertexCollection;
+	using FOutputDataVertexCollection = Metasound::FOutputDataVertexCollection;
+
+	Metasound::FNodeInitData DummyInitData;
+	TUniquePtr<Metasound::INode> DummyNodePtr = InCallback(DummyInitData);
+	if (!ensureAlwaysMsgf(DummyNodePtr.IsValid(), TEXT("Invalid getter registered!")))
+	{
+		return false;
+	}
+
+	Metasound::INode& DummyNode = *DummyNodePtr;
+
+	// First, build the key.
+	const FName NodeName = DummyNode.GetClassName();
+	const FInputDataVertexCollection& Inputs = DummyNode.GetInputDataVertices();
+	const FOutputDataVertexCollection& Outputs = DummyNode.GetOutputDataVertices();
+
+	// Construct a hash using a combination of the class name, input names and output names.
+	uint32 NodeHash = FCrc::StrCrc32(*NodeName.ToString());
+
+	TArray<FName> InputTypes;
+	TArray<FName> OutputTypes;
+
+	for (auto& InputTuple : Inputs)
+	{
+		HashCombine(NodeHash, FCrc::StrCrc32(*InputTuple.Value.VertexName));
+		InputTypes.Add(InputTuple.Value.DataReferenceTypeName);
+	}
+
+	for (auto& OutputTuple : Outputs)
+	{
+		HashCombine(NodeHash, FCrc::StrCrc32(*OutputTuple.Value.VertexName));
+		OutputTypes.Add(OutputTuple.Value.DataReferenceTypeName);
+	}
+
+	FNodeRegistryKey InKey = { NodeName, NodeHash };
+
+	// check to see if an identical node was already registered, and log
+	ensureAlwaysMsgf(!ExternalNodeRegistry.Contains(InKey), TEXT("Node with identical name, inputs and outputs to node %s was already registered. The previously registered node will be overwritten. This could also happen because METASOUND_REGISTER_NODE is in a public header."), *NodeName.ToString());
+
+	UE_LOG(LogTemp, Display, TEXT("Registered Metasound Node %s"), *NodeName.ToString());
+
+	const bool bShouldLogInputsAndOutputs = true;
+	if (bShouldLogInputsAndOutputs)
+	{
+		UE_LOG(LogTemp, Display, TEXT("    %d inputs:"), Inputs.Num());
+		for (auto& InputTuple : Inputs)
+		{
+			UE_LOG(LogTemp, Display, TEXT("      %s (of type %s)"), *InputTuple.Value.VertexName, *InputTuple.Value.DataReferenceTypeName.ToString());
+		}
+
+		UE_LOG(LogTemp, Display, TEXT("    %d outputs:"), Outputs.Num());
+		for (auto& OutputTuple : Outputs)
+		{
+			UE_LOG(LogTemp, Display, TEXT("      %s (of type %s)"), *OutputTuple.Value.VertexName, *OutputTuple.Value.DataReferenceTypeName.ToString());
+		}
+	}
+
+	FNodeRegistryElement RegistryElement = FNodeRegistryElement(MoveTemp(InCallback));
+	RegistryElement.InputTypes = MoveTemp(InputTypes);
+	RegistryElement.OutputTypes = MoveTemp(OutputTypes);
+
+	ExternalNodeRegistry.Add(MoveTemp(InKey), MoveTemp(RegistryElement));
+
+	return true;
+}
+
+TArray<FName> FMetasoundFrontendRegistryContainer::GetAllValidDataTypes()
+{
+	TArray<FName> OutDataTypes;
+
+	for (auto& DataTypeTuple : DataTypeRegistry)
+	{
+		OutDataTypes.Add(DataTypeTuple.Key);
+	}
+
+	return OutDataTypes;
+}
+
+bool FMetasoundFrontendRegistryContainer::GetInfoForDataType(FName InDataType, Metasound::FDataTypeRegistryInfo& OutInfo)
+{
+	if (!DataTypeRegistry.Contains(InDataType))
+	{
+		return false;
+	}
+	else
+	{
+		OutInfo = DataTypeRegistry[InDataType].Info;
+		return true;
+	}
+}
+
diff --git a/Engine/Plugins/Runtime/Metasound/Source/MetasoundGraphCore/Private/MetasoundGraph.cpp b/Engine/Plugins/Runtime/Metasound/Source/MetasoundGraphCore/Private/MetasoundGraph.cpp
--- a/Engine/Plugins/Runtime/Metasound/Source/MetasoundGraphCore/Private/MetasoundGraph.cpp
+++ b/Engine/Plugins/Runtime/Metasound/Source/MetasoundGraphCore/Private/MetasoundGraph.cpp
@@ -7,7 +7,7 @@ namespace Metasound
 	const FName FGraph::ClassName = FName(TEXT("Graph"));
 
 	FGraph::FGraph(const FString& InDescription)
-	:	Description(InDescription)
+		: Description(InDescription)
 	{
 	}
 
@@ -16,6 +16,12 @@ namespace Metasound
 	}
 
 	const FString& FGraph::GetDescription() const
+	{
+		static FString ClassDescription(TEXT("Unused"));
+		return ClassDescription;
+	}
+
+	const FString& FGraph::GetInstanceName() const
 	{
 		return Description;
 	}
@@ -25,6 +31,19 @@ namespace Metasound
 		return ClassName;
 	}
 
+
+	const FString& FGraph::GetAuthorName() const
+	{
+		static FString Author(TEXT("Unused"));
+		return Author;
+	}
+
+	const FString& FGraph::GetPromptIfMissing() const
+	{
+		static FString Prompt(TEXT("Unused"));
+		return Prompt;
+	}
+
 	const FInputDataVertexCollection& FGraph::GetInputDataVertices() const
 	{
 		return InputVertices;
diff --git a/Engine/Plugins/Runtime/Metasound/Source/MetasoundGraphCore/Private/MetasoundGraphCoreModule.cpp b/Engine/Plugins/Runtime/Metasound/Source/MetasoundGraphCore/Private/MetasoundGraphCoreModule.cpp
--- a/Engine/Plugins/Runtime/Metasound/Source/MetasoundGraphCore/Private/MetasoundGraphCoreModule.cpp
+++ b/Engine/Plugins/Runtime/Metasound/Source/MetasoundGraphCore/Private/MetasoundGraphCoreModule.cpp
@@ -5,9 +5,28 @@
 #include "CoreMinimal.h"
 #include "Modules/ModuleInterface.h"
 #include "Modules/ModuleManager.h"
+#include "MetasoundPrimitives.h"
+#include "MetasoundDataReferenceTypes.h"
+#include "MetasoundTime.h"
+#include "MetasoundFrequency.h"
 
 DEFINE_LOG_CATEGORY(LogMetasoundGraphCore);
 
+REGISTER_METASOUND_DATATYPE(bool, ::Metasound::ELiteralArgType::Boolean)
+REGISTER_METASOUND_DATATYPE(int32, ::Metasound::ELiteralArgType::Integer)
+REGISTER_METASOUND_DATATYPE(int64, ::Metasound::ELiteralArgType::Integer)
+REGISTER_METASOUND_DATATYPE(float, ::Metasound::ELiteralArgType::Float)
+REGISTER_METASOUND_DATATYPE(double, ::Metasound::ELiteralArgType::Float)
+
+namespace Metasound
+{
+	REGISTER_METASOUND_DATATYPE(FBop)
+	REGISTER_METASOUND_DATATYPE(FFloatTime)
+	REGISTER_METASOUND_DATATYPE(FDoubleTime)
+	REGISTER_METASOUND_DATATYPE(FSampleTime)
+	REGISTER_METASOUND_DATATYPE(FFrequency)
+}
+
 namespace Metasound 
 {
 	class FMetasoundGraphCoreModule : public IMetasoundGraphCoreModule
@@ -17,4 +36,3 @@ namespace Metasound
 
 IMPLEMENT_MODULE(Metasound::FMetasoundGraphCoreModule, MetasoundGraphCore);
 
-
diff --git a/Engine/Plugins/Runtime/Metasound/Source/MetasoundGraphCore/Private/MetasoundNode.cpp b/Engine/Plugins/Runtime/Metasound/Source/MetasoundGraphCore/Private/MetasoundNode.cpp
--- a/Engine/Plugins/Runtime/Metasound/Source/MetasoundGraphCore/Private/MetasoundNode.cpp
+++ b/Engine/Plugins/Runtime/Metasound/Source/MetasoundGraphCore/Private/MetasoundNode.cpp
@@ -15,7 +15,7 @@ namespace Metasound
 	{
 	}
 
-	const FString& FNode::GetDescription() const
+	const FString& FNode::GetInstanceName() const
 	{
 		return Description;
 	}
diff --git a/Engine/Plugins/Runtime/Metasound/Source/MetasoundGraphCore/Private/MetasoundOscNode.cpp b/Engine/Plugins/Runtime/Metasound/Source/MetasoundGraphCore/Private/MetasoundOscNode.cpp
--- a/Engine/Plugins/Runtime/Metasound/Source/MetasoundGraphCore/Private/MetasoundOscNode.cpp
+++ b/Engine/Plugins/Runtime/Metasound/Source/MetasoundGraphCore/Private/MetasoundOscNode.cpp
@@ -8,6 +8,8 @@
 
 namespace Metasound
 {
+	METASOUND_REGISTER_NODE(FOscNode)
+
 	class FOscOperator : public TExecutableOperator<FOscOperator>
 	{
 		public:
@@ -16,7 +18,7 @@ namespace Metasound
 			,	TwoPi(2.f * PI)
 			,	Phase(0.f)
 			,	Frequency(InFrequency)
-			,	AudioBuffer(InSettings.FramesPerExecute)
+			,	AudioBuffer(FAudioBufferWriteRef::CreateNew(InSettings))
 			{
 				check(AudioBuffer->Num() == InSettings.FramesPerExecute);
 
@@ -64,7 +66,7 @@ namespace Metasound
 	TUniquePtr<IOperator> FOscNode::FOperatorFactory::CreateOperator(const INode& InNode, const FOperatorSettings& InOperatorSettings, const FDataReferenceCollection& InInputDataReferences, TArray<TUniquePtr<IOperatorBuildError>>& OutErrors) 
 	{
 		const FOscNode& OscNode = static_cast<const FOscNode&>(InNode);
-		FFrequencyReadRef Frequency(OscNode.GetDefaultFrequency(), EFrequencyResolution::Hertz);
+		FFrequencyReadRef Frequency = FFrequencyReadRef::CreateNew(OscNode.GetDefaultFrequency(), EFrequencyResolution::Hertz);
 
 		if (InInputDataReferences.ContainsDataReadReference<FFrequency>(TEXT("Frequency")))
 		{
@@ -75,13 +77,18 @@ namespace Metasound
 	}
 
 	FOscNode::FOscNode(const FString& InName, float InDefaultFrequency)
-	:	FNode(InName)
-	,	DefaultFrequency(InDefaultFrequency)
+		:	FNode(InName)
+		,	DefaultFrequency(InDefaultFrequency)
 	{
 		AddInputDataVertex<FFrequency>(TEXT("Frequency"), LOCTEXT("FrequencyTooltip", "The frequency of oscillator."));
 		AddOutputDataVertex<FAudioBuffer>(TEXT("Audio"), LOCTEXT("AudioTooltip", "The output audio"));
 	}
 
+	FOscNode::FOscNode(const FNodeInitData& InInitData)
+		: FOscNode(InInitData.InstanceName, 440.0f)
+	{
+	}
+
 	FOscNode::~FOscNode()
 	{
 	}
@@ -93,7 +100,7 @@ namespace Metasound
 
 	const FName& FOscNode::GetClassName() const
 	{
-		return ClassName;
+		return ::Metasound::FOscNode::ClassName;
 	}
 
 	IOperatorFactory& FOscNode::GetDefaultOperatorFactory() 
diff --git a/Engine/Plugins/Runtime/Metasound/Source/MetasoundGraphCore/Private/MetasoundPeriodicBopNode.cpp b/Engine/Plugins/Runtime/Metasound/Source/MetasoundGraphCore/Private/MetasoundPeriodicBopNode.cpp
--- a/Engine/Plugins/Runtime/Metasound/Source/MetasoundGraphCore/Private/MetasoundPeriodicBopNode.cpp
+++ b/Engine/Plugins/Runtime/Metasound/Source/MetasoundGraphCore/Private/MetasoundPeriodicBopNode.cpp
@@ -15,6 +15,8 @@
 
 namespace Metasound
 {
+	METASOUND_REGISTER_NODE(FPeriodicBopNode)
+
 	class FPeriodicBopOperator : public TExecutableOperator<FPeriodicBopOperator>
 	{
 		public:
@@ -23,6 +25,7 @@ namespace Metasound
 
 			FPeriodicBopOperator(const FOperatorSettings& InSettings, const FFloatTimeReadRef& InPeriod)
 			:	OperatorSettings(InSettings)
+			,	Bop(FBopWriteRef::CreateNew())
 			,	Period(InPeriod)
 			,	ExecuteDurationInSamples(InSettings.FramesPerExecute)
 			,	SampleCountdown(0.f)
@@ -76,7 +79,9 @@ namespace Metasound
 	TUniquePtr<IOperator> FPeriodicBopNode::FOperatorFactory::CreateOperator(const INode& InNode, const FOperatorSettings& InOperatorSettings, const FDataReferenceCollection& InInputDataReferences, TArray<TUniquePtr<IOperatorBuildError>>& OutErrors)
 	{
 		const FPeriodicBopNode& PeriodicBopNode = static_cast<const FPeriodicBopNode&>(InNode);
-		FFloatTimeReadRef Period(PeriodicBopNode.GetDefaultPeriodInSeconds(), ETimeResolution::Seconds);
+
+		// Note- won't this cause the Period input to not update after we build the operator?
+		FFloatTimeReadRef Period = FFloatTimeReadRef::CreateNew(PeriodicBopNode.GetDefaultPeriodInSeconds(), ETimeResolution::Seconds);
 
 		if (InInputDataReferences.ContainsDataReadReference<FFloatTime>(TEXT("Period")))
 		{
@@ -96,6 +101,10 @@ namespace Metasound
 		AddOutputDataVertex<FBop>(TEXT("Bop"), LOCTEXT("BopTooltip", "The output bop"));
 	}
 
+	FPeriodicBopNode::FPeriodicBopNode(const FNodeInitData& InInitData)
+		: FPeriodicBopNode(InInitData.InstanceName, 1.0f)
+	{}
+
 	FPeriodicBopNode::~FPeriodicBopNode()
 	{
 	}
diff --git a/Engine/Plugins/Runtime/Metasound/Source/MetasoundGraphCore/Public/MetasoundADSRNode.h b/Engine/Plugins/Runtime/Metasound/Source/MetasoundGraphCore/Public/MetasoundADSRNode.h
--- a/Engine/Plugins/Runtime/Metasound/Source/MetasoundGraphCore/Public/MetasoundADSRNode.h
+++ b/Engine/Plugins/Runtime/Metasound/Source/MetasoundGraphCore/Public/MetasoundADSRNode.h
@@ -16,6 +16,7 @@ namespace Metasound
 	 */
 	class METASOUNDGRAPHCORE_API FADSRNode : public FNode
 	{
+
 			// The operator factory for this node.
 			class FOperatorFactory : public IOperatorFactory
 			{
@@ -48,6 +49,11 @@ namespace Metasound
 			 */
 			FADSRNode(const FString& InName, float InDefaultAttackMs, float InDefaultDecayMs, float InDefaultSustainMs, float InDefaultReleaseMs);
 
+			/**
+			 * Constructor used by the Metasound Frontend.
+			 */
+			FADSRNode(const FNodeInitData& InitData);
+
 			virtual ~FADSRNode();
 
 			/** Return default attack time in milliseconds */
@@ -62,6 +68,24 @@ namespace Metasound
 			/** Returns the type name of the FADSRNode */
 			virtual const FName& GetClassName() const override;
 
+			virtual const FString& GetDescription() const override
+			{
+				static FString StaticDescription = TEXT("This node, when bopped, emits an envelope that can be applied to .");
+				return StaticDescription;
+			}
+
+			virtual const FString& GetAuthorName() const override
+			{
+				static FString Author = TEXT("Epic Games");
+				return Author;
+			}
+
+			virtual const FString& GetPromptIfMissing() const override
+			{
+				static FString Prompt = TEXT("Make sure that the Metasound plugin is loaded.");
+				return Prompt;
+			}
+
 			/** Return a factory for building a metasound operator. */
 			virtual IOperatorFactory& GetDefaultOperatorFactory() override;
 
@@ -74,4 +98,6 @@ namespace Metasound
 
 			FOperatorFactory Factory;
 	};
+
+
 }
diff --git a/Engine/Plugins/Runtime/Metasound/Source/MetasoundGraphCore/Public/MetasoundAudioBuffer.h b/Engine/Plugins/Runtime/Metasound/Source/MetasoundGraphCore/Public/MetasoundAudioBuffer.h
--- a/Engine/Plugins/Runtime/Metasound/Source/MetasoundGraphCore/Public/MetasoundAudioBuffer.h
+++ b/Engine/Plugins/Runtime/Metasound/Source/MetasoundGraphCore/Public/MetasoundAudioBuffer.h
@@ -6,6 +6,7 @@
 #include "DSP/BufferVectorOperations.h"
 #include "MetasoundGraphCoreModule.h"
 #include "MetasoundDataReference.h"
+#include "MetasoundDataTypeRegistrationMacro.h"
 
 /** Define which determines whether to check that the size of the audio buffer has not changed since initialization */
 #define METASOUNDGRAPHCORE_CHECKAUDIONUM !(UE_BUILD_SHIPPING || UE_BUILD_TEST)
@@ -19,14 +20,14 @@ namespace Metasound
 	 * arguments via an implicit conversion operator which exposes the underlying AlignedFloatBuffer container.
 	 *
 	 */
-	class FAudioBuffer
+	class METASOUNDGRAPHCORE_API FAudioBuffer
 	{
 		public:
 			/** Create an FAudioBuffer with a specific number of samples. 
 			 *
 			 * @param InNumSamples - Number of samples in buffer.
 			 */
-			FAudioBuffer(int32 InNumSamples)
+			explicit FAudioBuffer(int32 InNumSamples)
 			{
 				Buffer.AddUninitialized(InNumSamples);
 
@@ -36,8 +37,22 @@ namespace Metasound
 #endif				
 			}
 
+			/**
+			 * This is the constructor used by the frontend.
+			 */
+			FAudioBuffer(const FOperatorSettings& InSettings)
+			{
+				Buffer.AddUninitialized(InSettings.FramesPerExecute);
+
+#if METASOUNDGRAPHCORE_CHECKAUDIONUM
+				// InitialNum should not change during the life of an FAudioBuffer
+				InitialNum = Buffer.Num();
+#endif	
+			}
+
 			/** Allow all other constructors of Audio::AlignedFloatBuffer to be forwarded to
 			 * underlying container */
+			/*
 			template <
 				typename... ArgTypes,
 				typename = typename TEnableIf<
@@ -64,6 +79,7 @@ namespace Metasound
 				InitialNum = Buffer.Num();
 #endif				
 			}
+			*/
 
 			FAudioBuffer()
 			{
diff --git a/Engine/Plugins/Runtime/Metasound/Source/MetasoundGraphCore/Public/MetasoundAudioFormats.h b/Engine/Plugins/Runtime/Metasound/Source/MetasoundGraphCore/Public/MetasoundAudioFormats.h
--- a/Engine/Plugins/Runtime/Metasound/Source/MetasoundGraphCore/Public/MetasoundAudioFormats.h
+++ b/Engine/Plugins/Runtime/Metasound/Source/MetasoundGraphCore/Public/MetasoundAudioFormats.h
@@ -5,6 +5,7 @@
 #include "CoreMinimal.h"
 #include "MetasoundDataReference.h"
 #include "MetasoundAudioBuffer.h"
+#include "MetasoundDataTypeRegistrationMacro.h"
 
 namespace Metasound
 {
@@ -29,6 +30,15 @@ namespace Metasound
 			 */
 			FUnformattedAudio(int32 InNumFrames, int32 InNumChannels, int32 InMaxNumChannels);
 
+			/**
+			 * FUnformattedAudio Constructor used by the metasound frontend.
+			 *
+			 * @param InNumChannels - initial number of audio channels.
+			 * @param InSettings - Operator Settings passed in on construction.
+			 */
+			explicit FUnformattedAudio(int32 InNumChannels, const FOperatorSettings& InSettings);
+
+
 			/** Sets the number of audio channels.
 			 * 
 			 * The requested number of channels is internally clamped to the supported range 
@@ -105,6 +115,14 @@ namespace Metasound
 			 */
 			FMultichannelAudioFormat(int32 InNumFrames, int32 InNumChannels);
 
+			/**
+			 * FMultichannelAudioFormat Constructor used by the metasound frontend.
+			 *
+			 * @param InNumChannels - initial number of audio channels.
+			 * @param InSettings - Operator Settings passed in on construction.
+			 */
+			explicit FMultichannelAudioFormat(int32 InNumChannels, const FOperatorSettings& InSettings);
+
 			/** FMultichannelAudioFormat Constructor.
 			 *
 			 * This constructor accepts an array of writable audio buffer references. Each 
@@ -118,13 +136,13 @@ namespace Metasound
 			FMultichannelAudioFormat(const FMultichannelAudioFormat& InOther) = default;
 
 			// Disable move constructor as incoming object should not be altered. 
-			FMultichannelAudioFormat(FMultichannelAudioFormat&& InOther) = delete;
+			//FMultichannelAudioFormat(FMultichannelAudioFormat&& InOther) = delete;
 
 			// Disable equal operator so channel count does not change
 			FMultichannelAudioFormat& operator=(const FMultichannelAudioFormat& Other) = delete;
 
 			// Disable move operator so channel count does not change
-			FMultichannelAudioFormat& operator=(FMultichannelAudioFormat&& Other) = delete;
+			//FMultichannelAudioFormat& operator=(FMultichannelAudioFormat&& Other) = delete;
 
 			/** Return the number of audio channels. */
 			int32 GetNumChannels() const { return NumChannels; }
@@ -202,7 +220,7 @@ namespace Metasound
 
 				for (int32 i = 0; i < NumChannels; i++)
 				{
-					FAudioBufferWriteRef Audio(InNumFrames);
+					FAudioBufferWriteRef Audio = FAudioBufferWriteRef::CreateNew(InNumFrames);
 					Audio->Zero();
 
 					WritableBufferStorage.Add(Audio);
@@ -439,12 +457,16 @@ namespace Metasound
 
 
 	DECLARE_METASOUND_DATA_REFERENCE_TYPES(FUnformattedAudio, "Audio:Unformatted", 0xd78a3ed1 , FUnformattedAudioTypeInfo, FUnformattedAudioReadRef, FUnformattedAudioWriteRef);
+	
 
 	DECLARE_METASOUND_DATA_REFERENCE_TYPES(FMultichannelAudioFormat, "Audio:Multichannel", 0x56bdcbe0 , FMultichannelAudioFormatTypeInfo, FMultichannelAudioFormatReadRef, FMultichannelAudioFormatWriteRef);
+	
 
 	DECLARE_METASOUND_DATA_REFERENCE_TYPES(FMonoAudioFormat, "Audio:Mono", 0x6f468c8c, FMonoAudioFormatTypeInfo, FMonoAudioFormatReadRef, FMonoAudioFormatWriteRef);
 
 	DECLARE_METASOUND_DATA_REFERENCE_TYPES(FStereoAudioFormat, "Audio:Stereo", 0xb55304e2 , FStereoAudioFormatTypeInfo, FStereoAudioFormatReadRef, FStereoAudioFormatWriteRef);
-
-
+	
+	// TODO: Neither of these are trivially constructible, becuase they are actually views into a preexisting buffer.
+	//REGISTER_METASOUND_DATATYPE(FMonoAudioFormat)
+	//REGISTER_METASOUND_DATATYPE(FStereoAudioFormat)
 }
diff --git a/Engine/Plugins/Runtime/Metasound/Source/MetasoundGraphCore/Public/MetasoundAudioMultiplyNode.h b/Engine/Plugins/Runtime/Metasound/Source/MetasoundGraphCore/Public/MetasoundAudioMultiplyNode.h
--- a/Engine/Plugins/Runtime/Metasound/Source/MetasoundGraphCore/Public/MetasoundAudioMultiplyNode.h
+++ b/Engine/Plugins/Runtime/Metasound/Source/MetasoundGraphCore/Public/MetasoundAudioMultiplyNode.h
@@ -12,6 +12,7 @@ namespace Metasound
 {
 	class METASOUNDGRAPHCORE_API FAudioMultiplyNode : public FNode
 	{
+
 			class FOperatorFactory : public IOperatorFactory
 			{
 				public:
@@ -37,14 +38,37 @@ namespace Metasound
 
 			FAudioMultiplyNode(const FString& InName);
 
+			// Constructor used by the Metasound Frontend.
+			FAudioMultiplyNode(const FNodeInitData& InInitData);
+
 			virtual ~FAudioMultiplyNode();
 
 			virtual const FName& GetClassName() const override;
 
 			virtual IOperatorFactory& GetDefaultOperatorFactory() override;
 
+			virtual const FString& GetDescription() const override
+			{
+				static FString StaticDescription = TEXT("This node multiplies to audio signals together. This is useful for amplitude modulation and other applications.");
+				return StaticDescription;
+			}
+
+			virtual const FString& GetAuthorName() const override
+			{
+				static FString Author = TEXT("Epic Games");
+				return Author;
+			}
+
+			virtual const FString& GetPromptIfMissing() const override
+			{
+				static FString Prompt = TEXT("Make sure that the Metasound plugin is loaded.");
+				return Prompt;
+			}
+
 		private:
 
 			FOperatorFactory Factory;
 	};
+
+	
 }
diff --git a/Engine/Plugins/Runtime/Metasound/Source/MetasoundGraphCore/Public/MetasoundDataReference.h b/Engine/Plugins/Runtime/Metasound/Source/MetasoundGraphCore/Public/MetasoundDataReference.h
--- a/Engine/Plugins/Runtime/Metasound/Source/MetasoundGraphCore/Public/MetasoundDataReference.h
+++ b/Engine/Plugins/Runtime/Metasound/Source/MetasoundGraphCore/Public/MetasoundDataReference.h
@@ -3,27 +3,31 @@
 #pragma once
 
 #include "CoreMinimal.h"
+#include "Misc/TVariant.h"
 
-
-/** Macros to make declaring a metasound parameter simple.  */
-
+/** Macro to make declaring a metasound parameter simple.  */
 // Declares a metasound parameter type by
 // - Adding typedefs for commonly used template types.
 // - Defining parameter type traits.
 #define DECLARE_METASOUND_DATA_REFERENCE_TYPES(DataType, DataTypeName, DataTypeMagicNumber, DataTypeInfoTypeName, DataReadReferenceTypeName, DataWriteReferenceTypeName) \
 	template<> \
-	struct TDataReferenceTypeInfo<DataType> \
+	struct ::Metasound::TDataReferenceTypeInfo<DataType> \
 	{ \
 		static constexpr const TCHAR* TypeName = TEXT(DataTypeName); \
-		static constexpr FDataTypeMagicNumber MagicNumber = (DataTypeMagicNumber); \
+		static constexpr ::Metasound::FDataTypeMagicNumber MagicNumber = (DataTypeMagicNumber); \
+		static constexpr bool bIsStringParsable = TTestIfDataTypeCtorIsImplemented<DataType, const FString&>::Value; \
+		static constexpr bool bIsBoolParsable = TTestIfDataTypeCtorIsImplemented<DataType, bool>::Value; \
+		static constexpr bool bIsIntParsable = TTestIfDataTypeCtorIsImplemented<DataType, int32>::Value; \
+		static constexpr bool bIsFloatParsable = TTestIfDataTypeCtorIsImplemented<DataType, float>::Value; \
+		static constexpr bool bIsConstructableWithSettings = TTestIfDataTypeDefaultCtorIsImplemented<DataType>::Value; \
+		static constexpr bool bCanUseDefaultConstructor = TIsConstructible<DataType>::Value; \
+		static constexpr bool bIsValidSpecialization = true; \
 	}; \
 	\
-	typedef TDataReferenceTypeInfo<DataType> DataTypeInfoTypeName; \
+	typedef ::Metasound::TDataReferenceTypeInfo<DataType> DataTypeInfoTypeName; \
 	\
-	typedef TDataReadReference<DataType> DataReadReferenceTypeName; \
-	typedef TDataWriteReference<DataType> DataWriteReferenceTypeName; \
-
-
+	typedef ::Metasound::TDataReadReference<DataType> DataReadReferenceTypeName; \
+	typedef ::Metasound::TDataWriteReference<DataType> DataWriteReferenceTypeName; \
 
 namespace Metasound
 {
@@ -77,6 +81,14 @@ namespace Metasound
 		// This static assert is triggered if TDataReferenceTypeInfo is used 
 		// without specialization.
 		static_assert(TSpecializationHelper<DataType>::Value, "TDataReferenceTypeInfo must be specialized.  Use macro DECLARE_METASOUND_DATA_REFERENCE_TYPES");
+	
+		static constexpr bool bIsStringParsable = false;
+		static constexpr bool bIsBoolParsable = false;
+		static constexpr bool bIsIntParsable = false;
+		static constexpr bool bIsFloatParsable = false;
+		static constexpr bool bIsConstructableWithSettings = false;
+		static constexpr bool bCanUseDefaultConstructor = false;
+		static constexpr bool bIsValidSpecialization = false;
 	};
 
 	template<typename DataType>
@@ -87,6 +99,14 @@ namespace Metasound
 		return TypeName;
 	}
 	
+
+	// This enum is used as a token to explicitly delineate when we should create a new object for the reference,
+	// or use a different constructor.
+	enum class EDataRefShouldConstruct
+	{
+		NewObject
+	};
+
 	/** Template class for a paramter reference. 
 	 *
 	 * This fulfills the IParamterRef interface, utilizing TDataReferenceTypeInfo to
@@ -95,15 +115,27 @@ namespace Metasound
 	template <typename DataType>
 	class TDataReference : public IDataReference
 	{
+	protected:
+		/**
+		 * This constructor forwards arguments to an underlying constructor.
+		 */
+		template <typename... ArgTypes>
+		TDataReference(EDataRefShouldConstruct InToken, ArgTypes&&... Args)
+			: ObjectReference(MakeShared<DataType, ESPMode::NotThreadSafe>(Forward<ArgTypes>(Args)...))
+		{
+		}
+
 		public:
 			typedef TSharedRef<DataType, ESPMode::NotThreadSafe> FRefType;
 
 			typedef TDataReferenceTypeInfo<DataType> FInfoType;
 
-			/** This constructor forwards arguments to the constructor of the underlying DataType.
+
+			/*
+			 * This constructor forwards arguments to the constructor of the underlying DataType.
 			 *
 			 * SFINAE is used to disable copy and move constructors as those are defined separately. 
-			 */
+			 
 			template <
 				typename... ArgTypes,
 				typename = typename TEnableIf<
@@ -111,7 +143,7 @@ namespace Metasound
 						sizeof...(ArgTypes) != 0,
 						TOrValue<
 							sizeof...(ArgTypes) != 1,
-							TNot< TIsDerivedFrom< typename TDecay< typename TNthTypeFromParameterPack< 0, ArgTypes... >::Type >::Type, /*typename*/ TDataReference<DataType> > >
+							TNot< TIsDerivedFrom< typename TDecay< typename TNthTypeFromParameterPack< 0, ArgTypes... >::Type >::Type,  TDataReference<DataType> > >
 						>
 					>::Value
 				>::Type
@@ -121,14 +153,24 @@ namespace Metasound
 			{
 			}
 
-	
+			
+			
 			// TODO: having issues compiling this on linux when no default constructor is available. 
-			/** Construct operator with no arguments if the DataType has a default constructor.  */
+			// Construct operator with no arguments if the DataType has a default constructor.
 			template< typename = typename TEnableIf< TIsConstructible<DataType>::Value >::Type >
 			TDataReference()
 			:	ObjectReference(MakeShared<DataType, ESPMode::NotThreadSafe>())
 			{
 			}
+			*/
+
+			// This should be used to construct a new DataType object and return this TDataReference as a wrapper around it.
+			template <typename... ArgTypes>
+			static TDataReference<DataType> CreateNew(ArgTypes&&... Args)
+			{
+				static_assert(TIsConstructible<DataType, ArgTypes...>::Value, "Tried to call TDataReference::CreateNew with args that don't match any constructor for an underlying type!");
+				return TDataReference<DataType>(EDataRefShouldConstruct::NewObject, Forward<ArgTypes>(Args)...);
+			}
 
 			/** Enable copy constructor */
 			TDataReference(const TDataReference<DataType>& Other) = default;
@@ -180,10 +222,18 @@ namespace Metasound
 	template <typename DataType>
 	class TDataWriteReference : public TDataReference<DataType>
 	{
+		// Construct operator with no arguments if the DataType has a default constructor.
+		template <typename... ArgTypes>
+		TDataWriteReference(EDataRefShouldConstruct InToken, ArgTypes&&... Args)
+			: FDataReference(InToken, Forward<ArgTypes>(Args)...)
+		{
+		}
+
 		public:
 			typedef TDataReference<DataType> FDataReference;
 
-			/** Construct operator type with arguments to "DataType" constructor. */
+			/*
+			// Construct operator type with arguments to "DataType" constructor. 
 			template<
 				typename... ArgTypes,
 				typename = typename TEnableIf<
@@ -191,7 +241,7 @@ namespace Metasound
 						sizeof...(ArgTypes) != 0,
 						TOrValue<
 							sizeof...(ArgTypes) != 1,
-							TNot< TIsDerivedFrom< typename TDecay< typename TNthTypeFromParameterPack< 0, ArgTypes... >::Type >::Type, /*typename*/ TDataReference<DataType> > >
+							TNot< TIsDerivedFrom< typename TDecay< typename TNthTypeFromParameterPack< 0, ArgTypes... >::Type >::Type, TDataReference<DataType> > >
 						>
 					>::Value
 				>::Type
@@ -201,14 +251,24 @@ namespace Metasound
 			{
 			}
 
-	
-			/** Construct operator with no arguments if the DataType has a default constructor.  */
+			
+			// Construct operator with no arguments if the DataType has a default constructor.
 			template< typename = typename TEnableIf< TIsConstructible<DataType>::Value >::Type >
 			TDataWriteReference()
 			:	FDataReference()
 			{
 			}
 
+			*/
+
+			// This should be used to construct a new DataType object and return this TDataWriteReference as a wrapper around it.
+			template <typename... ArgTypes>
+			static TDataWriteReference<DataType> CreateNew(ArgTypes&&... Args)
+			{
+				static_assert(TIsConstructible<DataType, ArgTypes...>::Value, "Tried to call TDataWriteReference::CreateNew with args that don't match any constructor for an underlying type!");
+				return TDataWriteReference<DataType>(EDataRefShouldConstruct::NewObject, Forward<ArgTypes>(Args)...);
+			}
+
 			/** Enable copy constructor */
 			TDataWriteReference(const TDataWriteReference<DataType>& Other) = default;
 
@@ -289,10 +349,18 @@ namespace Metasound
 	template <typename DataType>
 	class TDataReadReference : public TDataReference<DataType>
 	{
+		// Construct operator with no arguments if the DataType has a default constructor.
+		template <typename... ArgTypes>
+		TDataReadReference(EDataRefShouldConstruct InToken, ArgTypes&&... Args)
+			: FDataReference(InToken, Forward<ArgTypes>(Args)...)
+		{
+		}
+
 		public:
 			typedef TDataReference<DataType> FDataReference;
 
-			/** Construct operator type with arguments to "DataType" constructor. */
+			/*
+			// Construct operator type with arguments to "DataType" constructor. 
 			template<
 				typename... ArgTypes,
 				typename = typename TEnableIf<
@@ -300,7 +368,7 @@ namespace Metasound
 						sizeof...(ArgTypes) != 0,
 						TOrValue<
 							sizeof...(ArgTypes) != 1,
-							TNot< TIsDerivedFrom< typename TDecay< typename TNthTypeFromParameterPack< 0, ArgTypes... >::Type >::Type, /*typename*/ TDataReference<DataType> > >
+							TNot< TIsDerivedFrom< typename TDecay< typename TNthTypeFromParameterPack< 0, ArgTypes... >::Type >::Type, TDataReference<DataType> > >
 						>
 					>::Value
 				>::Type
@@ -310,12 +378,22 @@ namespace Metasound
 			{
 			}
 
-			/** Construct operator with no arguments if the DataType has a default constructor.  */
+			
+			// Construct operator with no arguments if the DataType has a default constructor.
 			template< typename = typename TEnableIf< TIsConstructible<DataType>::Value >::Type >
 			TDataReadReference()
 			:	FDataReference()
 			{
 			}
+			*/
+
+			// This should be used to construct a new DataType object and return this TDataReadReference as a wrapper around it.
+			template <typename... ArgTypes>
+			static TDataReadReference<DataType> CreateNew(ArgTypes&&... Args)
+			{
+				static_assert(TIsConstructible<DataType, ArgTypes...>::Value, "Tried to call TDataReadReference::CreateNew with args that don't match any constructor for an underlying type!");
+				return TDataReadReference<DataType>(EDataRefShouldConstruct::NewObject, Forward<ArgTypes>(Args)...);
+			}
 
 			TDataReadReference(const TDataReadReference& Other) = default;
 
@@ -361,4 +439,45 @@ namespace Metasound
 				return MakeUnique< FDataReadReference >(*this);
 			}
 	};
-}
+
+	enum class ELiteralArgType : uint8
+	{
+		Boolean,
+		Integer,
+		Float,
+		String,
+		None, // If this is set, we will invoke TType(const FOperatorSettings&) If that constructor exists, or the default constructor if not.
+		Invalid,
+	};
+
+	// Various elements that we pass to the frontend registry based on templated type traits.
+	struct FDataTypeRegistryInfo
+	{
+		// The name of the data type itself.
+		FName DataTypeName;
+
+		// What type we should default to using for literals.
+		ELiteralArgType PreferredLiteralType;
+
+		// These bools signify what basic
+		// UProperty primitives we can use to describe this data type as a literal in a document.
+		bool bIsBoolParsable;
+		bool bIsIntParsable;
+		bool bIsFloatParsable;
+		bool bIsStringParsable;
+
+		// This indicates the type can only be constructed with FOperatorSettings and no other args.
+		// TODO: these can be consolidated to a single bool, since FDataTypeLiteralParam automatically falls back to not using the settings.
+		bool bIsConstructableWithSettings;
+		bool bIsDefaultConstructible;
+
+		FDataTypeRegistryInfo()
+			: bIsBoolParsable(false)
+			, bIsIntParsable(false)
+			, bIsFloatParsable(false)
+			, bIsStringParsable(false)
+			, bIsConstructableWithSettings(false)
+			, bIsDefaultConstructible(false)
+		{}
+	};
+}
\ No newline at end of file
diff --git a/Engine/Plugins/Runtime/Metasound/Source/MetasoundGraphCore/Public/MetasoundDataReferenceCollection.h b/Engine/Plugins/Runtime/Metasound/Source/MetasoundGraphCore/Public/MetasoundDataReferenceCollection.h
--- a/Engine/Plugins/Runtime/Metasound/Source/MetasoundGraphCore/Public/MetasoundDataReferenceCollection.h
+++ b/Engine/Plugins/Runtime/Metasound/Source/MetasoundGraphCore/Public/MetasoundDataReferenceCollection.h
@@ -211,7 +211,7 @@ namespace Metasound
 				}
 				else
 				{
-					return FDataRefType(Forward<ConstructorArgTypes>(ConstructorArgs)...);
+					return FDataRefType::CreateNew(Forward<ConstructorArgTypes>(ConstructorArgs)...);
 				}
 			}
 
diff --git a/Engine/Plugins/Runtime/Metasound/Source/MetasoundGraphCore/Public/MetasoundDataReferenceTypes.h b/Engine/Plugins/Runtime/Metasound/Source/MetasoundGraphCore/Public/MetasoundDataReferenceTypes.h
--- a/Engine/Plugins/Runtime/Metasound/Source/MetasoundGraphCore/Public/MetasoundDataReferenceTypes.h
+++ b/Engine/Plugins/Runtime/Metasound/Source/MetasoundGraphCore/Public/MetasoundDataReferenceTypes.h
@@ -4,6 +4,7 @@
 
 #include "MetasoundDataReference.h"
 #include "MetasoundAudioBuffer.h"
+#include "MetasoundDataTypeRegistrationMacro.h"
 
 namespace Metasound
 {
diff --git a/Engine/Plugins/Runtime/Metasound/Source/MetasoundGraphCore/Public/MetasoundDataTypeRegistrationMacro.h b/Engine/Plugins/Runtime/Metasound/Source/MetasoundGraphCore/Public/MetasoundDataTypeRegistrationMacro.h
new file mode 100644
--- /dev/null
+++ b/Engine/Plugins/Runtime/Metasound/Source/MetasoundGraphCore/Public/MetasoundDataTypeRegistrationMacro.h
@@ -0,0 +1,70 @@
+// Copyright Epic Games, Inc. All Rights Reserved.
+
+#pragma once
+
+#include "CoreMinimal.h"
+
+#include "MetasoundDataReference.h"
+#include "MetasoundOutputNode.h"
+#include "MetasoundInputNode.h"
+#include "MetasoundOperatorInterface.h"
+#include "MetasoundFrontendRegistries.h"
+
+typedef TUniqueFunction<TUniquePtr<Metasound::INode>(const ::Metasound::FInputNodeConstructorParams&)> FInputNodeConstructorCallback;
+typedef TUniqueFunction<TUniquePtr<Metasound::INode>(const ::Metasound::FOutputNodeConstrutorParams&)> FOutputNodeConstructorCallback;
+
+template<typename TDataType>
+bool RegisterDataTypeWithFrontend(::Metasound::ELiteralArgType PreferredArgType = ::Metasound::ELiteralArgType::None)
+{
+	// if we reenter this code (because DECLARE_METASOUND_DATA_REFERENCE_TYPES was called twice with the same type),
+	// we catch it here.
+	static bool bAlreadyRegisteredThisDataType = false;
+	if (bAlreadyRegisteredThisDataType)
+	{
+		UE_LOG(LogTemp, Display, TEXT("Tried to call REGISTER_METASOUND_DATATYPE twice with the same class %s. ignoring the second call. Likely because REGISTER_METASOUND_DATATYPE is in a header that's used in multiple modules. Consider moving it to a private header or cpp file."), ::Metasound::TDataReferenceTypeInfo<TDataType>::TypeName)
+		return false;
+	}
+
+	bAlreadyRegisteredThisDataType = true;
+
+	// Lambdas that generate our template-instantiated input and output nodes:
+	FInputNodeConstructorCallback InputNodeConstructor = [](const ::Metasound::FInputNodeConstructorParams& InParams) -> TUniquePtr<Metasound::INode>
+	{
+		return TUniquePtr<Metasound::INode>(new ::Metasound::TInputNode<TDataType>(InParams.InNodeName, InParams.InVertexName, InParams.InitParam, InParams.InSettings));
+	};
+
+	FOutputNodeConstructorCallback OutputNodeConstructor = [](const ::Metasound::FOutputNodeConstrutorParams& InParams) -> TUniquePtr<Metasound::INode>
+	{
+		return TUniquePtr<Metasound::INode>(new ::Metasound::TOutputNode<TDataType>(InParams.InNodeName, InParams.InVertexName));
+	};
+	
+	static FName DataTypeName = FName(::Metasound::TDataReferenceTypeInfo<TDataType>::TypeName);
+
+	::Metasound::FDataTypeRegistryInfo RegistryInfo;
+	RegistryInfo.DataTypeName = DataTypeName;
+	RegistryInfo.PreferredLiteralType = PreferredArgType;
+	RegistryInfo.bIsBoolParsable = ::Metasound::TDataReferenceTypeInfo<TDataType>::bIsBoolParsable;
+	RegistryInfo.bIsIntParsable = ::Metasound::TDataReferenceTypeInfo<TDataType>::bIsIntParsable;
+	RegistryInfo.bIsFloatParsable = ::Metasound::TDataReferenceTypeInfo<TDataType>::bIsFloatParsable;
+	RegistryInfo.bIsStringParsable = ::Metasound::TDataReferenceTypeInfo<TDataType>::bIsStringParsable;
+	RegistryInfo.bIsConstructableWithSettings = ::Metasound::TDataReferenceTypeInfo<TDataType>::bIsConstructableWithSettings;
+	RegistryInfo.bIsDefaultConstructible = ::Metasound::TDataReferenceTypeInfo<TDataType>::bCanUseDefaultConstructor;
+
+	bool bSucceeded = FMetasoundFrontendRegistryContainer::Get()->RegisterDataType(RegistryInfo, MoveTemp(InputNodeConstructor), MoveTemp(OutputNodeConstructor));
+	ensureAlwaysMsgf(bSucceeded, TEXT("Failed to register data type %s in the node registry!"), ::Metasound::TDataReferenceTypeInfo<TDataType>::TypeName);
+	return bSucceeded;
+}
+
+// This should be used to expose a datatype as a potential input or output for a metasound graph.
+// The first argument to the macro is the class to expose.
+// Optionally, a Metasound::ELiteralArgType can be passed in to designate a preferred literal type-
+// For example, if Metasound::ELiteralArgType::Float is passed in, we will default to using a float parameter to create this datatype.
+// If no argument is passed in, we will infer a literal type to use.
+// Metasound::ELiteralArgType::Invalid can be used to enforce that we don't provide space for a literal, in which case you should have a default constructor or a constructor that takes [const FOperatorSettings&] implemented.
+// If you pass in a preferred arg type, please make sure that the passed in datatype has a matching constructor, since we won't check this until runtime.
+#define REGISTER_METASOUND_DATATYPE(DataType, ...) \
+	static_assert(::Metasound::TDataReferenceTypeInfo<DataType>::bIsValidSpecialization, "Please call DECLARE_METASOUND_DATA_REFERENCE_TYPES(" #DataType "...) before calling REGISTER_METASOUND_DATATYPE(" #DataType")."); \
+	static constexpr bool bCanRegister##DataType = ::Metasound::TDataReferenceTypeInfo<DataType>::bIsStringParsable || ::Metasound::TDataReferenceTypeInfo<DataType>::bIsBoolParsable || ::Metasound::TDataReferenceTypeInfo<DataType>::bIsIntParsable || ::Metasound::TDataReferenceTypeInfo<DataType>::bIsFloatParsable || ::Metasound::TDataReferenceTypeInfo<DataType>::bIsConstructableWithSettings || ::Metasound::TDataReferenceTypeInfo<DataType>::bCanUseDefaultConstructor; \
+	static_assert(bCanRegister##DataType , "To register " #DataType " to be used as a Metasounds input or output type, it needs a default constructor or one of the following constructors must be implemented:  " #DataType "(const ::Metasound::FOperatorSettings& InSettings), " #DataType "(int32 InValue, const ::Metasound::FOperatorSettings& InSettings), " #DataType "(float InValue, const ::Metasound::FOperatorSettings& InSettings), " #DataType "(bool InValue, const ::Metasound::FOperatorSettings& InSettings), or " #DataType "(const FString& InString, const ::Metasound::FOperatorSettings& InSettings)."); \
+	static const bool bSuccessfullyRegistered##DataType = FMetasoundFrontendRegistryContainer::Get()->EnqueueInitCommand([](){ ::RegisterDataTypeWithFrontend<DataType>(__VA_ARGS__); }); // This static bool is useful for debugging, but also is the only way the compiler will let us call this function outside of an expression.
+
diff --git a/Engine/Plugins/Runtime/Metasound/Source/MetasoundGraphCore/Public/MetasoundFrequency.h b/Engine/Plugins/Runtime/Metasound/Source/MetasoundGraphCore/Public/MetasoundFrequency.h
--- a/Engine/Plugins/Runtime/Metasound/Source/MetasoundGraphCore/Public/MetasoundFrequency.h
+++ b/Engine/Plugins/Runtime/Metasound/Source/MetasoundGraphCore/Public/MetasoundFrequency.h
@@ -3,6 +3,7 @@
 #pragma once
 
 #include "MetasoundDataReference.h"
+#include "MetasoundDataTypeRegistrationMacro.h"
 
 namespace Metasound
 {
@@ -58,6 +59,13 @@ namespace Metasound
 				}
 			}
 
+			/**
+			 * FFrequency constructor used to pass in float literals from the metasound frontend.
+			 */
+			FFrequency(float InValue)
+				: FFrequency(InValue, EFrequencyResolution::Hertz)
+			{}
+
 			/** Set the frequency in hertz. */
 			void SetHertz(float InHz)
 			{
@@ -143,5 +151,4 @@ namespace Metasound
 	};
 
 	DECLARE_METASOUND_DATA_REFERENCE_TYPES(FFrequency, "Primitive:Frequency", 0xf1c430a4 , FFrequencyTypeInfo, FFrequencyReadRef, FFrequencyWriteRef);
-
 }
diff --git a/Engine/Plugins/Runtime/Metasound/Source/MetasoundGraphCore/Public/MetasoundFrontendRegistries.h b/Engine/Plugins/Runtime/Metasound/Source/MetasoundGraphCore/Public/MetasoundFrontendRegistries.h
new file mode 100644
--- /dev/null
+++ b/Engine/Plugins/Runtime/Metasound/Source/MetasoundGraphCore/Public/MetasoundFrontendRegistries.h
@@ -0,0 +1,131 @@
+// Copyright Epic Games, Inc. All Rights Reserved.
+
+#pragma once
+
+#include "CoreMinimal.h"
+#include "MetasoundNodeInterface.h"
+#include "MetasoundDataReference.h"
+#include "MetasoundOperatorInterface.h"
+
+typedef TUniqueFunction<TUniquePtr<Metasound::INode>(const ::Metasound::FInputNodeConstructorParams&)> FInputNodeConstructorCallback;
+typedef TUniqueFunction<TUniquePtr<Metasound::INode>(const ::Metasound::FOutputNodeConstrutorParams&)> FOutputNodeConstructorCallback;
+
+typedef TUniqueFunction<TUniquePtr<Metasound::INode>(const Metasound::FNodeInitData&)> FNodeGetterCallback;
+
+namespace Metasound
+{
+	namespace Frontend
+	{
+		struct METASOUNDGRAPHCORE_API FNodeRegistryKey
+		{
+			// The class name for the node.
+			FName NodeName;
+
+			// A hash generated from the input types and output types for this node.
+			// TODO: Write up some tests to ensure this is deterministic.
+			uint32 NodeHash;
+
+			FORCEINLINE bool operator==(const FNodeRegistryKey& Other) const
+			{
+				return NodeHash == Other.NodeHash && NodeName == Other.NodeName;
+			}
+
+			friend uint32 GetTypeHash(const Metasound::Frontend::FNodeRegistryKey& InKey)
+			{
+				return InKey.NodeHash;
+			}
+		};
+
+
+		struct METASOUNDGRAPHCORE_API FNodeRegistryElement
+		{
+			// This lambda can be used to get an INodeBase for this specific node class.
+			FNodeGetterCallback GetterCallback;
+
+			TArray<FName> InputTypes;
+			TArray<FName> OutputTypes;
+
+			FNodeRegistryElement(FNodeGetterCallback&& InCallback)
+				: GetterCallback(MoveTemp(InCallback))
+			{
+			}
+		};
+	}
+}
+
+/*
+static uint32 METASOUNDGRAPHCORE_API GetTypeHash(const Metasound::Frontend::FNodeRegistryKey& InKey)
+{
+	return InKey.NodeHash;
+}
+*/
+
+/**
+ * Singleton registry for all types and nodes.
+ */
+class METASOUNDGRAPHCORE_API FMetasoundFrontendRegistryContainer
+{
+public:
+	static FMetasoundFrontendRegistryContainer* Get();
+	static void ShutdownMetasoundFrontend();
+
+	FMetasoundFrontendRegistryContainer(const FMetasoundFrontendRegistryContainer&) = delete;
+
+	void InitializeFrontend();
+	bool EnqueueInitCommand(TUniqueFunction<void()>&& InFunc);
+
+	using FNodeRegistryKey = Metasound::Frontend::FNodeRegistryKey;
+	using FNodeRegistryElement = Metasound::Frontend::FNodeRegistryElement;
+
+	TMap<FNodeRegistryKey, FNodeRegistryElement>& GetExternalNodeRegistry();
+
+	TUniquePtr<Metasound::INode> ConstructInputNode(const FName& InInputType, const Metasound::FInputNodeConstructorParams& InParams);
+	TUniquePtr<Metasound::INode> ConstructOutputNode(const FName& InOutputType, const Metasound::FOutputNodeConstrutorParams& InParams);
+
+	TUniquePtr<Metasound::INode> ConstructExternalNode(const FName& InNodeType, uint32 InNodeHash, const Metasound::FNodeInitData& InInitData);
+
+	// Get the desired kind of literal for a given data type. Returns EConstructorArgType::Invalid if the data type couldn't be found.
+	Metasound::ELiteralArgType GetDesiredLiteralTypeForDataType(FName InDataType) const;
+
+	// Get whether we can build a literal of this specific type for InDataType.
+	bool DoesDataTypeSupportLiteralType(FName InDataType, Metasound::ELiteralArgType InLiteralType) const;
+
+	bool RegisterDataType(const ::Metasound::FDataTypeRegistryInfo& InDataInfo, FInputNodeConstructorCallback&& InputNodeConstructor, FOutputNodeConstructorCallback&& OutputNodeConstructor);
+	bool RegisterExternalNode(FNodeGetterCallback&& InCallback);
+
+	// Return any data types that can be used as a metasound input type or output type.
+	TArray<FName> GetAllValidDataTypes();
+
+	// Get info about a specific data type (what kind of literals we can use, etc.)
+	// @returns false if InDataType wasn't found in the registry. 
+	bool GetInfoForDataType(FName InDataType, Metasound::FDataTypeRegistryInfo& OutInfo);
+
+private:
+	FMetasoundFrontendRegistryContainer();
+
+	static FMetasoundFrontendRegistryContainer* LazySingleton;
+
+	// This buffer is used to enqueue nodes and datatypes to register when the module has been initialized,
+	// in order to avoid bad behavior with ensures, logs, etc. on static initialization.
+	// The bad news is that TInlineAllocator is the safest allocator to use on static init.
+	// The good news is that none of these lambdas typically have captures, so this should have low memory overhead.
+	static constexpr int32 MaxNumNodesAndDatatypesToInitialize = 8192;
+	TArray<TUniqueFunction<void()>/*, TInlineAllocator<MaxNumNodesAndDatatypesToInitialize>*/> LazyInitCommands;
+	
+	FCriticalSection LazyInitCommandCritSection;
+	bool bHasModuleBeenInitialized;
+
+	// Registry in which we keep all information about nodes implemented in C++.
+	TMap<FNodeRegistryKey, FNodeRegistryElement> ExternalNodeRegistry;
+
+	struct FDataTypeRegistryElement
+	{
+		FInputNodeConstructorCallback InputNodeConstructor;
+		FOutputNodeConstructorCallback  OutputNodeConstructor;
+		Metasound::FDataTypeRegistryInfo Info;
+	};
+
+	TMap<FName, FDataTypeRegistryElement> DataTypeRegistry;
+};
+
+
diff --git a/Engine/Plugins/Runtime/Metasound/Source/MetasoundGraphCore/Public/MetasoundGraph.h b/Engine/Plugins/Runtime/Metasound/Source/MetasoundGraphCore/Public/MetasoundGraph.h
--- a/Engine/Plugins/Runtime/Metasound/Source/MetasoundGraphCore/Public/MetasoundGraph.h
+++ b/Engine/Plugins/Runtime/Metasound/Source/MetasoundGraphCore/Public/MetasoundGraph.h
@@ -17,15 +17,23 @@ namespace Metasound
 		public:
 			static const FName ClassName;
 
-			FGraph(const FString& InDescription);
+			FGraph(const FString& InInstanceName);
 			virtual ~FGraph();
 
-			/** Return the name of this graph. */
-			virtual const FString& GetDescription() const override;
+			/** Return the name of this specific instance of the node class. */
+			virtual const FString& GetInstanceName() const override;
 
-			/** Return the type name of this graph. */
+			/** Return the type name of this node. */
 			virtual const FName& GetClassName() const override;
 
+			/** Return a longer text description describing how this node is used. */
+			virtual const FString& GetDescription() const override;
+
+			/** Return the original author of this node class. */
+			virtual const FString& GetAuthorName() const override;
+
+			virtual const FString& GetPromptIfMissing() const override;
+
 			/** Return a collection of input parameter descriptions for this graph. */
 			virtual const FInputDataVertexCollection& GetInputDataVertices() const override;
 
diff --git a/Engine/Plugins/Runtime/Metasound/Source/MetasoundGraphCore/Public/MetasoundInputNode.h b/Engine/Plugins/Runtime/Metasound/Source/MetasoundGraphCore/Public/MetasoundInputNode.h
--- a/Engine/Plugins/Runtime/Metasound/Source/MetasoundGraphCore/Public/MetasoundInputNode.h
+++ b/Engine/Plugins/Runtime/Metasound/Source/MetasoundGraphCore/Public/MetasoundInputNode.h
@@ -11,7 +11,8 @@ namespace Metasound
 	template<typename DataType>
 	class TInputNode : public INode
 	{
-		
+			static_assert(TDataReferenceTypeInfo<DataType>::bIsValidSpecialization, "Please use DECLARE_METASOUND_DATA_REFERENCE_TYPES with this class before trying to create an input node with it.");
+
 			class FInputOperator : public IOperator
 			{
 				public:
@@ -50,26 +51,35 @@ namespace Metasound
 			// or pointers to objects which are no longer valid. 
 			class FCopyOperatorFactory : public IOperatorFactory
 			{
+				static_assert(TIsConstructible<DataType, const DataType&>::Value, "Data Type must be copy constructible!");
+				
 				public:
 					using FDataWriteReference = TDataWriteReference<DataType>;
 					using FInputNodeType = TInputNode<DataType>;
 
-					template<typename... ArgTypes, typename = typename TEnableIf< TIsConstructible<DataType, ArgTypes...>::Value >::Type >
+					template< typename... ArgTypes >
 					FCopyOperatorFactory(ArgTypes&&... Args)
 					:	Data(Forward<ArgTypes>(Args)...)
 					{
 					}
 
+					/*
 					template<typename = typename TEnableIf< TIsConstructible<DataType>::Value >::Type >
 					FCopyOperatorFactory()
 					:	Data()
 					{
 					}
+					*/
+
+					FCopyOperatorFactory(FDataTypeLiteralParam InitParam, const FOperatorSettings& InSettings)
+						: Data(InitParam.ParseTo<DataType>(InSettings))
+					{
+					}
 
 					virtual TUniquePtr<IOperator> CreateOperator(const INode& InNode, const FOperatorSettings& InOperatorSettings, const FDataReferenceCollection& InInputDataReferences, TArray<TUniquePtr<IOperatorBuildError>>& OutErrors) override
 					{
 						// Use copy constructor to create a new parameter reference.
-						FDataWriteReference DataRef(Data);
+						FDataWriteReference DataRef = FDataWriteReference::CreateNew(Data);
 						// TODO: Write special version of this for audio types since they will need to be constructed based upon the operator settings. 
 						// TODO: Need to ponder how inputs are initialized, but template specialization might do the trick. 
 
@@ -83,12 +93,14 @@ namespace Metasound
 			};
 
 		public:
-
-
-			template<typename = typename TEnableIf< TIsConstructible<DataType>::Value >::Type >
+			
+			// Constructors that need to get working on linux:
+			/*
+			template<typename = typename TEnableIf< TIsConstructible<DataType>::Value >::Type>
 			TInputNode(const FString& InNodeDescription, const FString& InVertexName)
 			:	NodeDescription(InNodeDescription)
 			,	VertexName(InVertexName)
+			,   Factory()
 			{
 				FOutputDataVertex OutputVertex = MakeOutputDataVertex<DataType>(VertexName, FText::GetEmpty());
 				FDataVertexKey OutputVertexKey = MakeDataVertexKey(OutputVertex);
@@ -101,13 +113,16 @@ namespace Metasound
 				Inputs.Add(InputVertexKey, InputVertex);
 
 			}
+			*/
 
-			template<typename... ArgTypes, typename = typename TEnableIf< TIsConstructible<DataType, ArgTypes...>::Value >::Type >
+			template<typename... ArgTypes>
 			TInputNode(const FString& InNodeDescription, const FString& InVertexName, ArgTypes&&... Args)
 			:	NodeDescription(InNodeDescription)
 			,	VertexName(InVertexName)
 			,	Factory(Forward<ArgTypes>(Args)...)
 			{
+				static_assert(TIsConstructible<DataType, ArgTypes...>::Value, "Tried to construct TInputNode<DataType> with arguments that don't match any constructor for DataType.");
+
 				FOutputDataVertex OutputVertex = MakeOutputDataVertex<DataType>(VertexName, FText::GetEmpty());
 				FDataVertexKey OutputVertexKey = MakeDataVertexKey(OutputVertex);
 
@@ -119,19 +134,55 @@ namespace Metasound
 				Inputs.Add(InputVertexKey, InputVertex);
 
 			}
+			
 
-			virtual const FString& GetDescription() const override
+			explicit TInputNode(const FString& InNodeDescription, const FString& InVertexName, FDataTypeLiteralParam InParam, const FOperatorSettings& InSettings)
+				: NodeDescription(InNodeDescription)
+				, VertexName(InVertexName)
+				, Factory(InParam, InSettings)
+			{
+				FOutputDataVertex OutputVertex = MakeOutputDataVertex<DataType>(VertexName, FText::GetEmpty());
+				FDataVertexKey OutputVertexKey = MakeDataVertexKey(OutputVertex);
+
+				Outputs.Add(OutputVertexKey, OutputVertex);
+
+				FInputDataVertex InputVertex = MakeInputDataVertex<DataType>(VertexName, FText::GetEmpty());
+				FDataVertexKey InputVertexKey = MakeDataVertexKey(InputVertex);
+
+				Inputs.Add(InputVertexKey, InputVertex);
+			}
+
+			virtual const FString& GetInstanceName() const override
 			{
 				return NodeDescription;
 			}
 
 			virtual const FName& GetClassName() const override
 			{
-				static const FName ClassName = FName(FString(TEXT("Input_")) + FString(TDataReferenceTypeInfo<DataType>::TypeName));
-
+				// TODO: although this is ok with MSVC's lax template instantiation, every other compiler will complain about TDataReferenceTypeInfo.
+				//static const FName ClassName = FName(FString(TEXT("Input_")) + FString(TDataReferenceTypeInfo<DataType>::TypeName));
+				static const FName ClassName = FName(TEXT("InputNode"));
 				return ClassName;
 			}
 
+			virtual const FString& GetDescription() const override
+			{
+				static FString Description = TEXT("An Input into the current Metasound graph.");
+				return Description;
+			}
+
+			virtual const FString& GetAuthorName() const override
+			{
+				static FString Author = TEXT("Epic Games");
+				return Author;
+			}
+
+			virtual const FString& GetPromptIfMissing() const override
+			{
+				static FString Prompt = TEXT("Make sure that the Metasound plugin is loaded.");
+				return Prompt;
+			}
+
 			const FString& GetVertexName() const
 			{
 				return VertexName;
diff --git a/Engine/Plugins/Runtime/Metasound/Source/MetasoundGraphCore/Public/MetasoundNode.h b/Engine/Plugins/Runtime/Metasound/Source/MetasoundGraphCore/Public/MetasoundNode.h
--- a/Engine/Plugins/Runtime/Metasound/Source/MetasoundGraphCore/Public/MetasoundNode.h
+++ b/Engine/Plugins/Runtime/Metasound/Source/MetasoundGraphCore/Public/MetasoundNode.h
@@ -4,6 +4,7 @@
 
 #include "CoreMinimal.h"
 #include "MetasoundNodeInterface.h"
+#include "MetasoundNodeRegistrationMacro.h"
 
 namespace Metasound
 {
@@ -13,7 +14,7 @@ namespace Metasound
 			FNode(const FString& InDescription);
 			virtual ~FNode();
 
-			virtual const FString& GetDescription() const override;
+			virtual const FString& GetInstanceName() const override;
 
 			virtual const FInputDataVertexCollection& GetInputDataVertices() const override;
 			virtual const FOutputDataVertexCollection& GetOutputDataVertices() const override;
diff --git a/Engine/Plugins/Runtime/Metasound/Source/MetasoundGraphCore/Public/MetasoundNodeInterface.h b/Engine/Plugins/Runtime/Metasound/Source/MetasoundGraphCore/Public/MetasoundNodeInterface.h
--- a/Engine/Plugins/Runtime/Metasound/Source/MetasoundGraphCore/Public/MetasoundNodeInterface.h
+++ b/Engine/Plugins/Runtime/Metasound/Source/MetasoundGraphCore/Public/MetasoundNodeInterface.h
@@ -112,6 +112,26 @@ namespace Metasound
 		return InVertex.VertexName;
 	}
 
+	/**
+	 * This struct is used to pass in any arguments required for constructing a single node instance.
+	 * because of this, all FNode implementations have to implement a constructor that takes an FNodeInitData instance.
+	 */
+	struct FNodeInitData
+	{
+		FString InstanceName;
+		TMap<FName, FDataTypeLiteralParam> ParamMap;
+
+		template<typename ParamType>
+		ParamType GetParamValue(FName ParamName)
+		{
+			checkf(ParamMap.Contains(ParamName), TEXT("Tried to use node initialization parameter that didn't exist!"));
+
+			// If this hits, check that you have a constructor for DataType that can interpret the parameter.
+			checkf(ParamMap[ParamName].IsCompatibleWithType<ParamType>(), TEXT("Tried to parse initialization parameter to an invalid type!"));
+			return ParamMap[ParamName].ParseTo<ParamType>();
+		}
+	};
+
 	/** INodeBase
 	 * 
 	 * Interface for all nodes that can describe their name, type, inputs and outputs.
@@ -121,12 +141,24 @@ namespace Metasound
 		public:
 			virtual ~INodeBase() {}
 
-			/** Return the name of this node. */
-			virtual const FString& GetDescription() const = 0;
+			/** Return the name of this specific instance of the node class. */
+			virtual const FString& GetInstanceName() const = 0;
 
 			/** Return the type name of this node. */
 			virtual const FName& GetClassName() const = 0;
 
+			/** Return a longer text description describing how this node is used. */
+			virtual const FString& GetDescription() const = 0;
+
+			/** Return the original author of this node class. */
+			virtual const FString& GetAuthorName() const = 0;
+
+			/** 
+			 *  Return an optional prompt on how users can get the plugin this node is in,
+			 *  if they have found a metasound that uses this node but don't have this plugin downloaded or enabled.
+			 */
+			virtual const FString& GetPromptIfMissing() const = 0;
+
 			/** Return a collection of input parameter descriptions for this node. */
 			virtual const FInputDataVertexCollection& GetInputDataVertices() const = 0;
 
diff --git a/Engine/Plugins/Runtime/Metasound/Source/MetasoundGraphCore/Public/MetasoundNodeRegistrationMacro.h b/Engine/Plugins/Runtime/Metasound/Source/MetasoundGraphCore/Public/MetasoundNodeRegistrationMacro.h
new file mode 100644
--- /dev/null
+++ b/Engine/Plugins/Runtime/Metasound/Source/MetasoundGraphCore/Public/MetasoundNodeRegistrationMacro.h
@@ -0,0 +1,52 @@
+// Copyright Epic Games, Inc. All Rights Reserved.
+
+#pragma once
+
+#include "CoreMinimal.h"
+#include "MetasoundNodeInterface.h"
+#include "Templates/IntegralConstant.h"
+#include "MetasoundFrontendRegistries.h"
+
+// Utility class to ensure that a node class can use the constructor the frontend uses.
+template <typename NodeClass>
+struct ConstructorTakesNodeInitData {
+
+	// Use SFINAE trick to see if we have a valid constructor:
+	template<typename T>
+	static uint16 TestForConstructor(decltype(T(::Metasound::FNodeInitData()))*);
+
+	template<typename T>
+	static uint8 TestForConstructor(...);
+
+	static const bool Value = sizeof(TestForConstructor<NodeClass>(nullptr)) == sizeof(uint16);
+};
+
+template <typename FNodeType>
+bool RegisterNodeWithFrontend()
+{
+	// if we reenter this code (because DECLARE_METASOUND_DATA_REFERENCE_TYPES was called twice with the same type),
+	// we catch it here.
+	static bool bAlreadyRegisteredThisDataType = false;
+	if (bAlreadyRegisteredThisDataType)
+	{
+		UE_LOG(LogTemp, Display, TEXT("Tried to call REGISTER_METASOUND_DATATYPE twice with the same class. ignoring the second call. Likely because METASOUND_REGISTER_NODE is in a header that's used in multiple modules. Consider moving it to a private header or cpp file."));
+		return false;
+	}
+
+	bAlreadyRegisteredThisDataType = true;
+
+	bool bSuccessfullyRegisteredNode = FMetasoundFrontendRegistryContainer::Get()->RegisterExternalNode(
+		[](const Metasound::FNodeInitData& InInitData) -> TUniquePtr<Metasound::INode>
+		{
+			return TUniquePtr<Metasound::INode>(new FNodeType(InInitData));
+		});
+
+	ensureAlwaysMsgf(bSuccessfullyRegisteredNode, TEXT("Registering node class failed. Please check the logs."));
+	return bSuccessfullyRegisteredNode;
+}
+
+#define METASOUND_REGISTER_NODE(NodeClass) \
+	 static_assert(std::is_base_of<::Metasound::INodeBase, NodeClass>::value, "To be registered as a  Metasound Node," #NodeClass "need to be a derived class from Metasound::INodeBase, Metasound::INode, or Metasound::FNode."); \
+	 static_assert(::ConstructorTakesNodeInitData<NodeClass>::Value, "In order to be registered as a Metasound Node, " #NodeClass " needs to implement the following public constructor: " #NodeClass "(const Metasound::FNodeInitData& InInitData);"); \
+	 static bool bSuccessfullyRegistered##NodeClass  = FMetasoundFrontendRegistryContainer::Get()->EnqueueInitCommand([](){ ::RegisterNodeWithFrontend<NodeClass>(); }); // This static bool is useful for debugging, but also is the only way the compiler will let us call this function outside of an expression.
+
diff --git a/Engine/Plugins/Runtime/Metasound/Source/MetasoundGraphCore/Public/MetasoundOperatorInterface.h b/Engine/Plugins/Runtime/Metasound/Source/MetasoundGraphCore/Public/MetasoundOperatorInterface.h
--- a/Engine/Plugins/Runtime/Metasound/Source/MetasoundGraphCore/Public/MetasoundOperatorInterface.h
+++ b/Engine/Plugins/Runtime/Metasound/Source/MetasoundGraphCore/Public/MetasoundOperatorInterface.h
@@ -14,47 +14,47 @@ namespace Metasound
 	 */
 	class IOperator
 	{
-		public:
-			/** Pointer to execute function for an operator. 
-			 *
-			 * @param IOperator* - The operator associated with the function pointer.
-			 */
-			typedef void(*FExecuteFunction)(IOperator*);
-
-			virtual ~IOperator() {}
-
-			/** Return the input parameters associated with this IOperator. 
-			 *
-			 * Implementations of IOperator should populate and return their input parameters
-			 * if they want to enable callers to query and write to their parameters. Most IOperator
-			 * implementations will return an empty FDataReferenceCollection because thier inputs are 
-			 * set during the IOperator's construction and do not need to be updated afterwards. The
-			 * exceptions are input operators and graph operators which need to interface with external
-			 * systems even after the operator is created.
-			 */
-			virtual const FDataReferenceCollection& GetInputs() const = 0;
-
-			/** Return the output parameters associated with this IOperator.
-			 *
-			 * Implementations of IOperator should return a collection of their output parameters 
-			 * which other nodes can read.
-			 */
-			virtual const FDataReferenceCollection& GetOutputs() const = 0;
-
-			/** Return the execution function to call during graph execution.
-			 * 
-			 * The IOperator* argument to the FExecutionFunction will be the same IOperator instance
-			 * which returned the execution function. 
-			 *
-			 * nullptr return values are valid and signal an IOperator which does not need to be 
-			 * executed.
-			 */
-			virtual FExecuteFunction GetExecuteFunction() = 0;
+	public:
+		/** Pointer to execute function for an operator.
+		 *
+		 * @param IOperator* - The operator associated with the function pointer.
+		 */
+		typedef void(*FExecuteFunction)(IOperator*);
+
+		virtual ~IOperator() {}
+
+		/** Return the input parameters associated with this IOperator.
+		 *
+		 * Implementations of IOperator should populate and return their input parameters
+		 * if they want to enable callers to query and write to their parameters. Most IOperator
+		 * implementations will return an empty FDataReferenceCollection because thier inputs are
+		 * set during the IOperator's construction and do not need to be updated afterwards. The
+		 * exceptions are input operators and graph operators which need to interface with external
+		 * systems even after the operator is created.
+		 */
+		virtual const FDataReferenceCollection& GetInputs() const = 0;
+
+		/** Return the output parameters associated with this IOperator.
+		 *
+		 * Implementations of IOperator should return a collection of their output parameters
+		 * which other nodes can read.
+		 */
+		virtual const FDataReferenceCollection& GetOutputs() const = 0;
+
+		/** Return the execution function to call during graph execution.
+		 *
+		 * The IOperator* argument to the FExecutionFunction will be the same IOperator instance
+		 * which returned the execution function.
+		 *
+		 * nullptr return values are valid and signal an IOperator which does not need to be
+		 * executed.
+		 */
+		virtual FExecuteFunction GetExecuteFunction() = 0;
 	};
 
 	/** FOperatorSettings
 	 *
-	 * Audio settings for operators. 
+	 * Audio settings for operators.
 	 */
 	struct FOperatorSettings
 	{
@@ -65,12 +65,42 @@ namespace Metasound
 		int32 FramesPerExecute = 0;
 
 		FOperatorSettings(float InSampleRate, uint32 InFramesPerExecute)
-		:	SampleRate(InSampleRate)
-		,	FramesPerExecute(InFramesPerExecute)
+			: SampleRate(InSampleRate)
+			, FramesPerExecute(InFramesPerExecute)
 		{
 		}
 	};
+}
+
+template <typename TDataType, typename TTypeToParse>
+struct TTestIfDataTypeCtorIsImplemented
+{
+private:
+	static constexpr bool bSupportsConstructionWithSettings = 
+		TIsConstructible<TDataType, TTypeToParse, const ::Metasound::FOperatorSettings&>::Value
+		|| TIsConstructible<TDataType, TTypeToParse, ::Metasound::FOperatorSettings&>::Value
+		|| TIsConstructible<TDataType, TTypeToParse, ::Metasound::FOperatorSettings>::Value;
 
+	static constexpr bool bSupportsConstructionWithoutSettings = TIsConstructible<TDataType, TTypeToParse>::Value;
+
+public:
+	/*
+	static constexpr bool Value = TIsConstructible<TDataType, TTypeToParse, const ::Metasound::FOperatorSettings&>::Value
+		|| TIsConstructible<TDataType, TTypeToParse>::Value;
+		*/
+
+	static constexpr bool Value = bSupportsConstructionWithSettings || bSupportsConstructionWithoutSettings;
+};
+
+template <typename TDataType>
+struct TTestIfDataTypeDefaultCtorIsImplemented
+{
+public:
+	static constexpr bool Value = TIsConstructible<TDataType, const ::Metasound::FOperatorSettings&>::Value;
+};
+
+namespace Metasound
+{
 	/*
 	class IGraphOperator : public IOperator
 	{
@@ -79,4 +109,336 @@ namespace Metasound
 			virtual ~IGraphOperator() {}
 	};
 	*/
+
+	/*
+	const TCHAR* ToString(EConstructorArgType InType)
+	{
+		switch (InType)
+		{
+		case Metasound::EConstructorArgType::Boolean:
+			return TEXT("Boolean");
+			break;
+		case Metasound::EConstructorArgType::Integer:
+			return TEXT("Integer");
+			break;
+		case Metasound::EConstructorArgType::Float:
+			return TEXT("Float");
+			break;
+		case Metasound::EConstructorArgType::String:
+			return TEXT("String");
+			break;
+		default:
+			return TEXT("");
+			break;
+		}
+	}
+	*/
+
+	/**
+	 * Convenience wrapper for safely invoking the correct constructor.
+	 * This can be used for int32, float, bool, or FString.
+	 *
+	 * For example:
+	 * // Somewhere before DECLARE_METASOUND_DATA_REFERENCE_TYPES(FAudioBuffer...), ParseFrom<FAudioBuffer>(int32) is defined
+	 * // which means this data type can be created from an integer.
+	 * Metasound::FAudioBuffer ::Metasound::ParseFrom<Metasound::FAudioBuffer>(int32 InNumFrames, const ::Metasound::FOperatorSettings&)
+	 * {
+	 *     return Metasound::FAudioBuffer(InNumFrames);
+	 * }
+	 *
+	 * //...
+	 * // In the frontend, we know that int and pass it to somewhere in MetasoundGraphCore...
+	 * FDataTypeLiteralParam InitParam(512);
+	 *
+	 * //...
+	 * // In the backend, we can safely construct an FAudioBuffer.
+	 * if(InitParam.IsCompatibleWithType<Metasound::FAudioBuffer>())
+	 * {
+	 *     Metasound::FAudioBuffer AudioBuffer = InitParam.ParseTo<Metasound::FAudioBuffer>();
+	 *     //...
+	 * }
+	 */
+	struct FDataTypeLiteralParam
+	{
+		ELiteralArgType ConstructorArgType;
+		TVariant<bool, int, float, FString> ConstructorArg;
+
+		// builds an invalid FDataTypeLiteralParam.
+		static FDataTypeLiteralParam InvalidParam()
+		{
+			FDataTypeLiteralParam InitParam;
+			InitParam.ConstructorArgType = ELiteralArgType::Invalid;
+			return InitParam;
+		}
+
+		bool IsValid() const
+		{
+			return ConstructorArgType != ELiteralArgType::Invalid;
+		}
+
+		FDataTypeLiteralParam()
+			: ConstructorArgType(ELiteralArgType::None)
+		{}
+
+		FDataTypeLiteralParam(bool bInValue)
+			: ConstructorArgType(ELiteralArgType::Boolean)
+		{
+			ConstructorArg.Set<bool>(bInValue);
+		}
+
+		FDataTypeLiteralParam(int32 InValue)
+			: ConstructorArgType(ELiteralArgType::Integer)
+		{
+			ConstructorArg.Set<int32>(InValue);
+		}
+
+		FDataTypeLiteralParam(float InValue)
+			: ConstructorArgType(ELiteralArgType::Float)
+		{
+			ConstructorArg.Set<float>(InValue);
+		}
+
+		FDataTypeLiteralParam(const FString& InString)
+			: ConstructorArgType(ELiteralArgType::String)
+		{
+			ConstructorArg.Set<FString>(InString);
+		}
+
+		template<typename DataType>
+		bool IsCompatibleWithType()
+		{
+			switch (ConstructorArgType)
+			{
+				case ELiteralArgType::Float:
+				{
+					return TTestIfDataTypeCtorIsImplemented<DataType, float>::Value;
+				}
+				case ELiteralArgType::Integer:
+				{
+					return TTestIfDataTypeCtorIsImplemented<DataType, int32>::Value;
+				}
+				case ELiteralArgType::Boolean:
+				{
+					return TTestIfDataTypeCtorIsImplemented<DataType, bool>::Value;
+				}
+				case ELiteralArgType::String:
+				{
+					return TTestIfDataTypeCtorIsImplemented<DataType, const FString&>::Value;
+				}
+				case ELiteralArgType::None:
+				{
+					return TTestIfDataTypeDefaultCtorIsImplemented<DataType>::Value || TIsConstructible<DataType>::Value;
+				}
+				default:
+				{
+					checkNoEntry();
+					return false;
+				}
+			}
+		}
+
+	private:
+		// This is a bucket constructor for when ConstructorArgType is unsupported by the data type we're trying to generate.
+		// This will only be invoked if something has gone terribly wrong in the frontend.
+		template<typename DataType>
+		DataType InvalidConstructor()
+		{
+			checkNoEntry();
+
+			// We don't know which constructors DataType supports, and we should never enter this code,
+			// so we just dereference garbage memory here.
+			return *reinterpret_cast<DataType*>(0x1);
+		}
+
+		// These are SFINAE-specialized to ensure we safely invoke the correct function,
+		// or invoke InvalidConstructor if the correct function doesn't exist.
+		// For each literal type, there are three SFINAE variants from which we generate a function at compile time:
+		// 1: If there's a constructor that can take the literal type and FOperatorSettings, we use that.
+		// 2: Otherwise, if there's a constructor that just takes the literal type, we'll use that.
+		// 3: If neither constructor exists, we crash. Ideally this has already been caught by the static assert in the REGISTER_METASOUND_DATATYPE macro.
+
+		template<typename DataType, typename TEnableIf< 
+			TIsConstructible<DataType, const FString&, const FOperatorSettings&>::Value
+			, bool >::Type = true >
+		DataType ParseAsString(const FOperatorSettings& WithOperatorSettings)
+		{
+			check(ConstructorArgType == ELiteralArgType::String);
+			FString StringToParse = ConstructorArg.Get<FString>();
+			return DataType(StringToParse, WithOperatorSettings);
+		}
+
+		template<typename DataType, typename TEnableIf< 
+			TIsConstructible<DataType, const FString&>::Value && !TIsConstructible<DataType, const FString&, const FOperatorSettings&>::Value
+			, bool >::Type = true >
+		DataType ParseAsString(const FOperatorSettings& WithOperatorSettings)
+		{
+			check(ConstructorArgType == ELiteralArgType::String);
+			FString StringToParse = ConstructorArg.Get<FString>();
+			return DataType(StringToParse);
+		}
+
+		template<typename DataType>
+		DataType ParseAsString(...)
+		{
+			return InvalidConstructor<DataType>();
+		}
+
+		template<typename DataType, typename TEnableIf< 
+			TIsConstructible<DataType, float, const FOperatorSettings&>::Value
+			, bool >::Type = true >
+		DataType ParseAsFloat(const FOperatorSettings& WithOperatorSettings)
+		{
+			check(ConstructorArgType == ELiteralArgType::Float);
+			float ValueToParse = ConstructorArg.Get<float>();
+			return DataType(ValueToParse, WithOperatorSettings);
+		}
+
+		template<typename DataType, typename TEnableIf< 
+			TIsConstructible<DataType, float >::Value && !TIsConstructible<DataType, float, const FOperatorSettings&>::Value
+			, bool >::Type = true >
+		DataType ParseAsFloat(const FOperatorSettings& WithOperatorSettings)
+		{
+			check(ConstructorArgType == ELiteralArgType::Float);
+			float ValueToParse = ConstructorArg.Get<float>();
+			return DataType(ValueToParse);
+		}
+
+		template<typename DataType>
+		DataType ParseAsFloat(...)
+		{
+			return InvalidConstructor<DataType>();
+		}
+
+		template<typename DataType, typename TEnableIf< 
+			TIsConstructible<DataType, int32, const FOperatorSettings&>::Value
+			, bool >::Type = true >
+		DataType ParseAsInt(const FOperatorSettings& WithOperatorSettings)
+		{
+			check(ConstructorArgType == ELiteralArgType::Integer);
+			int32 ValueToParse = ConstructorArg.Get<int32>();
+			return DataType(ValueToParse, WithOperatorSettings);
+		}
+
+		template<typename DataType, typename TEnableIf< 
+			TIsConstructible<DataType, int32>::Value && !TIsConstructible<DataType, int32, const FOperatorSettings&>::Value
+			, bool >::Type = true >
+		DataType ParseAsInt(const FOperatorSettings& WithOperatorSettings)
+		{
+			check(ConstructorArgType == ELiteralArgType::Integer);
+			int32 ValueToParse = ConstructorArg.Get<int32>();
+			return DataType(ValueToParse);
+		}
+
+		template<typename DataType>
+		DataType ParseAsInt(...)
+		{
+			return InvalidConstructor<DataType>();
+		}
+
+		template<typename DataType, typename TEnableIf< 
+			TIsConstructible<DataType, bool, const FOperatorSettings&>::Value
+			, bool >::Type = true >
+		DataType ParseAsBool(const FOperatorSettings& WithOperatorSettings)
+		{
+			check(ConstructorArgType == ELiteralArgType::Boolean);
+			bool ValueToParse = ConstructorArg.Get<bool>();
+			return DataType(ValueToParse, WithOperatorSettings);
+		}
+
+		template<typename DataType, typename TEnableIf< 
+			TIsConstructible<DataType, bool>::Value && !TIsConstructible<DataType, bool, const FOperatorSettings&>::Value
+			, bool >::Type = true >
+		DataType ParseAsBool(const FOperatorSettings& WithOperatorSettings)
+		{
+			check(ConstructorArgType == ELiteralArgType::Boolean);
+			bool ValueToParse = ConstructorArg.Get<bool>();
+			return DataType(ValueToParse);
+		}
+
+		template<typename DataType>
+		DataType ParseAsBool(...)
+		{
+			return InvalidConstructor<DataType>();
+		}
+
+
+		template<typename DataType, typename TEnableIf< 
+			TIsConstructible<DataType, const FOperatorSettings&>::Value
+			, bool >::Type = true >
+		DataType ParseWithDefaultConstructor(const FOperatorSettings& WithOperatorSettings)
+		{
+			check(ConstructorArgType == ELiteralArgType::None);
+			return DataType(WithOperatorSettings);
+		}
+
+		template<typename DataType, typename TEnableIf< 
+			TIsConstructible<DataType>::Value && !TIsConstructible<DataType, const FOperatorSettings&>::Value
+			, bool>::Type = true >
+		DataType ParseWithDefaultConstructor(const FOperatorSettings& WithOperatorSettings)
+		{
+			check(ConstructorArgType == ELiteralArgType::None);
+			return DataType();
+		}
+
+		template<typename DataType>
+		DataType ParseWithDefaultConstructor(...)
+		{
+			return InvalidConstructor<DataType>();
+		}
+
+	public:
+
+		// This can be used to invoke the correct constructor for this init param for a data type.
+		template<typename DataType>
+		DataType ParseTo(const FOperatorSettings& WithOperatorSettings)
+		{
+			switch (ConstructorArgType)
+			{
+				case ELiteralArgType::Float:
+				{
+					return ParseAsFloat<DataType>(WithOperatorSettings);
+				}
+				case ELiteralArgType::Integer:
+				{
+					return ParseAsInt<DataType>(WithOperatorSettings);
+				}
+				case ELiteralArgType::Boolean:
+				{
+					return ParseAsBool<DataType>(WithOperatorSettings);
+				}
+				case ELiteralArgType::String:
+				{
+					return ParseAsString<DataType>(WithOperatorSettings);
+				}
+				case ELiteralArgType::None:
+				{
+					return ParseWithDefaultConstructor<DataType>(WithOperatorSettings);
+				}
+				default:
+				{
+					return InvalidConstructor<DataType>();
+				}
+			}
+		}
+	};
+
+	// This struct is passed to FInputNodeConstructorCallback and FOutputNodeConstructorCallback.
+	struct FInputNodeConstructorParams
+	{
+		// the instance name and name of the specific connection that we should use.
+		const FString& InNodeName;
+		const FString& InVertexName;
+
+		// the operator settings that we will compile with
+		const FOperatorSettings& InSettings;
+
+		FDataTypeLiteralParam InitParam;
+	};
+
+	struct FOutputNodeConstrutorParams
+	{
+		// the instance name and name of the specific connection that we should use.
+		const FString& InNodeName;
+		const FString& InVertexName;
+	};
 }
diff --git a/Engine/Plugins/Runtime/Metasound/Source/MetasoundGraphCore/Public/MetasoundOscNode.h b/Engine/Plugins/Runtime/Metasound/Source/MetasoundGraphCore/Public/MetasoundOscNode.h
--- a/Engine/Plugins/Runtime/Metasound/Source/MetasoundGraphCore/Public/MetasoundOscNode.h
+++ b/Engine/Plugins/Runtime/Metasound/Source/MetasoundGraphCore/Public/MetasoundOscNode.h
@@ -22,6 +22,9 @@ namespace Metasound
 
 			FOscNode(const FString& InName, float InDefaultFrequency);
 
+			// constructor used by the Metasound Frontend.
+			FOscNode(const FNodeInitData& InInitData);
+
 			virtual ~FOscNode();
 
 			float GetDefaultFrequency() const;
@@ -30,6 +33,24 @@ namespace Metasound
 
 			virtual IOperatorFactory& GetDefaultOperatorFactory() override;
 
+			virtual const FString& GetDescription() const override
+			{
+				static FString StaticDescription = TEXT("This node emits an audio signal of a sinusoid.");
+				return StaticDescription;
+			}
+
+			virtual const FString& GetAuthorName() const override
+			{
+				static FString Author = TEXT("Epic Games");
+				return Author;
+			}
+
+			virtual const FString& GetPromptIfMissing() const override
+			{
+				static FString Prompt = TEXT("Make sure that the Metasound plugin is loaded.");
+				return Prompt;
+			}
+
 		private:
 
 			float DefaultFrequency;
diff --git a/Engine/Plugins/Runtime/Metasound/Source/MetasoundGraphCore/Public/MetasoundOutputNode.h b/Engine/Plugins/Runtime/Metasound/Source/MetasoundGraphCore/Public/MetasoundOutputNode.h
--- a/Engine/Plugins/Runtime/Metasound/Source/MetasoundGraphCore/Public/MetasoundOutputNode.h
+++ b/Engine/Plugins/Runtime/Metasound/Source/MetasoundGraphCore/Public/MetasoundOutputNode.h
@@ -4,6 +4,7 @@
 
 #include "CoreMinimal.h"
 #include "MetasoundOperatorInterface.h"
+#include "MetasoundBuilderInterface.h"
 #include "MetasoundNodeInterface.h"
 
 #define LOCTEXT_NAMESPACE "MetasoundOutputNode"
@@ -13,6 +14,8 @@ namespace Metasound
 	template<typename DataType>
 	class TOutputNode : public INode
 	{
+			static_assert(TDataReferenceTypeInfo<DataType>::bIsValidSpecialization, "Please use DECLARE_METASOUND_DATA_REFERENCE_TYPES with this class before trying to create an output node with it.");
+			
 			class FOutputOperator : public IOperator
 			{
 				public:
@@ -87,17 +90,37 @@ namespace Metasound
 				Outputs.Add(OutputKey, OutputVertex);
 			}
 
-			virtual const FString& GetDescription() const override
+			virtual const FName& GetClassName() const override
+			{
+				// TODO: Any special formatting for these node type names?
+				// TODO: although this is ok with MSVC's lax template instantiation, every other compiler will complain about TDataReferenceTypeInfo.
+				//static const FName ClassName = FName(FString(TEXT("Input_")) + FString(TDataReferenceTypeInfo<DataType>::TypeName));
+				static const FName ClassName = FName(TEXT("OutputNode"));
+
+				return ClassName;
+			}
+
+			virtual const FString& GetInstanceName() const override
 			{
 				return NodeDescription;
 			}
 
-			virtual const FName& GetClassName() const override
+			virtual const FString& GetDescription() const override
 			{
-				// TODO: Any special formatting for these node type names?
-				static const FName ClassName = FName(FString(TEXT("Output_")) + FString(TDataReferenceTypeInfo<DataType>::TypeName));
+				static FString Description = TEXT("This node can be used as an output for a given type.");
+				return Description;
+			}
 
-				return ClassName;
+			virtual const FString& GetAuthorName() const override
+			{
+				static FString Author = TEXT("Epic Games");
+				return Author;
+			}
+
+			virtual const FString& GetPromptIfMissing() const override
+			{
+				static FString Prompt = TEXT("Make sure that the Metasound plugin is loaded.");
+				return Prompt;
 			}
 
 			virtual const FInputDataVertexCollection& GetInputDataVertices() const override
diff --git a/Engine/Plugins/Runtime/Metasound/Source/MetasoundGraphCore/Public/MetasoundPeriodicBopNode.h b/Engine/Plugins/Runtime/Metasound/Source/MetasoundGraphCore/Public/MetasoundPeriodicBopNode.h
--- a/Engine/Plugins/Runtime/Metasound/Source/MetasoundGraphCore/Public/MetasoundPeriodicBopNode.h
+++ b/Engine/Plugins/Runtime/Metasound/Source/MetasoundGraphCore/Public/MetasoundPeriodicBopNode.h
@@ -14,6 +14,7 @@ namespace Metasound
 {
 	class METASOUNDGRAPHCORE_API FPeriodicBopNode : public FNode
 	{
+
 			class FOperatorFactory : public IOperatorFactory
 			{
 				virtual TUniquePtr<IOperator> CreateOperator(const INode& InNode, const FOperatorSettings& InOperatorSettings, const FDataReferenceCollection& InInputDataReferences, TArray<TUniquePtr<IOperatorBuildError>>& OutErrors) override;
@@ -23,6 +24,7 @@ namespace Metasound
 			static const FName ClassName;
 
 			FPeriodicBopNode(const FString& InName, float InDefaultPeriodInSeconds);
+			FPeriodicBopNode(const FNodeInitData& InInitData);
 
 			virtual ~FPeriodicBopNode();
 
@@ -32,6 +34,25 @@ namespace Metasound
 
 			virtual IOperatorFactory& GetDefaultOperatorFactory() override;
 
+			virtual const FString& GetDescription() const override
+			{
+				static FString StaticDescription = TEXT("This node emits a bop periodically, based on the duration given.");
+				return StaticDescription;
+			}
+
+			virtual const FString& GetAuthorName() const override
+			{
+				static FString Author = TEXT("Epic Games");
+				return Author;
+			}
+
+			virtual const FString& GetPromptIfMissing() const override
+			{
+				static FString Prompt = TEXT("Make sure that the Metasound plugin is loaded.");
+				return Prompt;
+			}
+
+
 		private:
 			float DefaultPeriod;
 			FOperatorFactory Factory;
diff --git a/Engine/Plugins/Runtime/Metasound/Source/MetasoundGraphCore/Public/MetasoundTime.h b/Engine/Plugins/Runtime/Metasound/Source/MetasoundGraphCore/Public/MetasoundTime.h
--- a/Engine/Plugins/Runtime/Metasound/Source/MetasoundGraphCore/Public/MetasoundTime.h
+++ b/Engine/Plugins/Runtime/Metasound/Source/MetasoundGraphCore/Public/MetasoundTime.h
@@ -3,6 +3,7 @@
 #pragma once
 
 #include "MetasoundDataReference.h"
+#include "MetasoundDataTypeRegistrationMacro.h"
 
 namespace Metasound
 { 
@@ -69,6 +70,11 @@ namespace Metasound
 				}
 			}
 
+			/** Constructor used by the metasound backend. */
+			TTime(QuantizationType InSeconds)
+				: TTime(InSeconds, ETimeResolution::Seconds)
+			{}
+
 			/** Returns the time as seconds. */
 			TimeType GetSeconds() const 
 			{ 
@@ -290,6 +296,16 @@ namespace Metasound
 				}
 			}
 
+			/**
+			 * Constructor used by the Metasound Frontend.
+			 *
+			 * @param InSeconds The number in time to initialize to.
+			 * @param InSettings the operator settings to use.
+			 */
+			TTime(float InSeconds, const FOperatorSettings& InSettings)
+				: TTime(InSeconds, ETimeResolution::Seconds, InSettings.SampleRate)
+			{}
+
 			/** Construct a sample count based TTime using a number of samples
 			 * and a sample rate.
 			 *
