Add APIs to FGraphHandle to push literals to inputs.

diff --git a/Engine/Plugins/Runtime/Metasound/Source/MetasoundFrontend/Private/MetasoundFrontend.cpp b/Engine/Plugins/Runtime/Metasound/Source/MetasoundFrontend/Private/MetasoundFrontend.cpp
--- a/Engine/Plugins/Runtime/Metasound/Source/MetasoundFrontend/Private/MetasoundFrontend.cpp
+++ b/Engine/Plugins/Runtime/Metasound/Source/MetasoundFrontend/Private/MetasoundFrontend.cpp
@@ -935,6 +935,50 @@ namespace Metasound
 			return INDEX_NONE;
 		}
 
+		FMetasoundLiteralDescription* FGraphHandle::GetLiteralDescriptionForInput(const FString& InInputName, FName& OutDataType) const
+		{
+			if (!IsValid())
+			{
+				return nullptr;
+			}
+
+			// scan through our inputs to find a match.
+			TArray<FMetasoundInputDescription>& Inputs = GraphsClassDeclaration->Inputs;
+			for (FMetasoundInputDescription& Input : Inputs)
+			{
+				if (Input.Name == InInputName)
+				{
+					OutDataType = Input.TypeName;
+					return &(Input.LiteralValue);
+				}
+			}
+
+			ensureAlwaysMsgf(false, TEXT("Couldn't find Input of name %s in this Metasoud graph!"), *InInputName);
+			return nullptr;
+		}
+
+		bool FGraphHandle::GetDataTypeForInput(const FString& InInputName, FName& OutDataType)
+		{
+			if (!IsValid())
+			{
+				return nullptr;
+			}
+
+			// scan through our inputs to find a match.
+			TArray<FMetasoundInputDescription>& Inputs = GraphsClassDeclaration->Inputs;
+			for (FMetasoundInputDescription& Input : Inputs)
+			{
+				if (Input.Name == InInputName)
+				{
+					OutDataType = Input.TypeName;
+					return true;
+				}
+			}
+
+			ensureAlwaysMsgf(false, TEXT("Couldn't find Input of name %s in this Metasoud graph!"), *InInputName);
+			return false;
+		}
+
 		FGraphHandle FGraphHandle::InvalidHandle()
 		{
 			FDescPath InvalidPath;
@@ -1268,6 +1312,102 @@ namespace Metasound
 			return true;
 		}
 
+		ELiteralArgType FGraphHandle::GetPreferredLiteralTypeForInput(const FString& InInputName)
+		{
+			FName DataType;
+			if (GetDataTypeForInput(InInputName, DataType))
+			{
+				FMetasoundFrontendRegistryContainer* Registry = FMetasoundFrontendRegistryContainer::Get();
+
+				return Registry->GetDesiredLiteralTypeForDataType(DataType);
+			}
+			else
+			{
+				return ELiteralArgType::Invalid;
+			}
+		}
+
+		bool FGraphHandle::SetInputToLiteral(const FString& InInputName, bool bInValue)
+		{
+			FName DataType;
+			if (FMetasoundLiteralDescription* Literal = GetLiteralDescriptionForInput(InInputName, DataType))
+			{
+				if (!ensureAlwaysMsgf(DoesDataTypeSupportLiteralType(DataType, ELiteralArgType::Boolean), TEXT("Tried to set Data Type %s to an unsupported literal type (Boolean)")))
+				{
+					return false;
+				}
+
+				SetLiteralDescription(*Literal, bInValue);
+				return true;
+			}
+
+			return false;
+		}
+
+		bool FGraphHandle::SetInputToLiteral(const FString& InInputName, int32 InValue)
+		{
+			FName DataType;
+			if (FMetasoundLiteralDescription* Literal = GetLiteralDescriptionForInput(InInputName, DataType))
+			{
+				if (!ensureAlwaysMsgf(DoesDataTypeSupportLiteralType(DataType, ELiteralArgType::Integer), TEXT("Tried to set Data Type %s to an unsupported literal type (Integer)")))
+				{
+					return false;
+				}
+
+				SetLiteralDescription(*Literal, InValue);
+				return true;
+			}
+
+			return false;
+		}
+
+		bool FGraphHandle::SetInputToLiteral(const FString& InInputName, float InValue)
+		{
+			FName DataType;
+			if (FMetasoundLiteralDescription* Literal = GetLiteralDescriptionForInput(InInputName, DataType))
+			{
+				if (!ensureAlwaysMsgf(DoesDataTypeSupportLiteralType(DataType, ELiteralArgType::Float), TEXT("Tried to set Data Type %s to an unsupported literal type (Float)")))
+				{
+					return false;
+				}
+
+				SetLiteralDescription(*Literal, InValue);
+				return true;
+			}
+
+			return false;
+		}
+
+		bool FGraphHandle::SetInputToLiteral(const FString& InInputName, const FString& InValue)
+		{
+			FName DataType;
+			if (FMetasoundLiteralDescription* Literal = GetLiteralDescriptionForInput(InInputName, DataType))
+			{
+				if (!ensureAlwaysMsgf(DoesDataTypeSupportLiteralType(DataType, ELiteralArgType::String), TEXT("Tried to set Data Type %s to an unsupported literal type (String)")))
+				{
+					return false;
+				}
+
+				SetLiteralDescription(*Literal, InValue);
+				return true;
+			}
+
+			return false;
+		}
+
+		bool FGraphHandle::ClearLiteralForInput(const FString& InInputName)
+		{
+			FName DataType;
+			if (FMetasoundLiteralDescription* Literal = GetLiteralDescriptionForInput(InInputName, DataType))
+			{
+
+				ClearLiteralDescription(*Literal);
+				return true;
+			}
+
+			return false;
+		}
+
 		FNodeHandle FGraphHandle::AddNewNode(const FNodeClassInfo& InNodeClass)
 		{
 			if (!IsValid())
diff --git a/Engine/Plugins/Runtime/Metasound/Source/MetasoundFrontend/Public/MetasoundFrontend.h b/Engine/Plugins/Runtime/Metasound/Source/MetasoundFrontend/Public/MetasoundFrontend.h
--- a/Engine/Plugins/Runtime/Metasound/Source/MetasoundFrontend/Public/MetasoundFrontend.h
+++ b/Engine/Plugins/Runtime/Metasound/Source/MetasoundFrontend/Public/MetasoundFrontend.h
@@ -301,6 +301,21 @@ namespace Metasound
 			FNodeHandle AddNewOutput(const FMetasoundOutputDescription& OutDescription);
 			bool RemoveOutput(const FString& OutputName);
 
+			// This can be used to determine what kind of property editor we should use for the data type of a given input.
+			// Will return Invalid if the input couldn't be found, or if the input doesn't support any kind of literals.
+			ELiteralArgType GetPreferredLiteralTypeForInput(const FString& InInputName);
+
+			// These can be used to set the default value for a given input on this graph.
+			// @returns false if the input name couldn't be found, or if the literal type was incompatible with the Data Type of this input.
+			bool SetInputToLiteral(const FString& InInputName, bool bInValue);
+			bool SetInputToLiteral(const FString& InInputName, int32 InValue);
+			bool SetInputToLiteral(const FString& InInputName, float InValue);
+			bool SetInputToLiteral(const FString& InInputName, const FString& InValue);
+
+			// This can be used to clear the current literal for a given input.
+			// @returns false if the input name couldn't be found.
+			bool ClearLiteralForInput(const FString& InInputName);
+
 			FNodeHandle AddNewNode(const FNodeClassInfo& InNodeClass);
 
 			// Remove the node corresponding to this node handle.
@@ -352,6 +367,10 @@ namespace Metasound
 
 			// Scans all existing dependency IDs in the root dependency list to guarantee a new unique ID when adding a dependency.
 			uint32 FindNewUniqueDependencyId();
+
+			FMetasoundLiteralDescription* GetLiteralDescriptionForInput(const FString& InInputName, FName& OutDataType) const;
+
+			bool GetDataTypeForInput(const FString& InInputName, FName& OutDataType);
 		};
 
 		// If there's a Metasound node referenced by a class in it's dependencies list,
