CIS compile fix for MetasoundFrontend.cpp

diff --git a/Engine/Plugins/Runtime/Metasound/Source/MetasoundFrontend/Private/MetasoundFrontend.cpp b/Engine/Plugins/Runtime/Metasound/Source/MetasoundFrontend/Private/MetasoundFrontend.cpp
--- a/Engine/Plugins/Runtime/Metasound/Source/MetasoundFrontend/Private/MetasoundFrontend.cpp
+++ b/Engine/Plugins/Runtime/Metasound/Source/MetasoundFrontend/Private/MetasoundFrontend.cpp
@@ -961,7 +961,7 @@ namespace Metasound
 		{
 			if (!IsValid())
 			{
-				return nullptr;
+				return false;
 			}
 
 			// scan through our inputs to find a match.
