Fix missing includes for windows strict compile.

diff --git a/Engine/Plugins/Runtime/Metasound/Source/MetasoundGraphCore/Private/MetasoundFrontendRegistries.cpp b/Engine/Plugins/Runtime/Metasound/Source/MetasoundGraphCore/Private/MetasoundFrontendRegistries.cpp
--- a/Engine/Plugins/Runtime/Metasound/Source/MetasoundGraphCore/Private/MetasoundFrontendRegistries.cpp
+++ b/Engine/Plugins/Runtime/Metasound/Source/MetasoundGraphCore/Private/MetasoundFrontendRegistries.cpp
@@ -2,6 +2,8 @@
 
 #include "MetasoundFrontendRegistries.h"
 #include "CoreMinimal.h"
+#include "Misc/ScopeLock.h"
+#include "HAL/PlatformTime.h"
 
 #ifndef WITH_METASOUND_FRONTEND
 #define WITH_METASOUND_FRONTEND 0
