Fix for linker error on string literal on linux. Move implementation of the string literal from DECLARE_METASOUND_DATA_TYPE to REGISTER_METASOUND_DATATYPE.

diff --git a/Engine/Plugins/Runtime/Metasound/Source/MetasoundFrontend/Private/MetasoundFrontend.cpp b/Engine/Plugins/Runtime/Metasound/Source/MetasoundFrontend/Private/MetasoundFrontend.cpp
--- a/Engine/Plugins/Runtime/Metasound/Source/MetasoundFrontend/Private/MetasoundFrontend.cpp
+++ b/Engine/Plugins/Runtime/Metasound/Source/MetasoundFrontend/Private/MetasoundFrontend.cpp
@@ -24,6 +24,8 @@ namespace Metasound
 {
 	namespace Frontend
 	{
+		const FHandleInitParams::EPrivateToken FHandleInitParams::PrivateToken = FHandleInitParams::EPrivateToken::Token;
+
 		TArray<FNodeClassInfo> GetAllAvailableNodeClasses()
 		{
 			TArray<FNodeClassInfo> OutClasses;
diff --git a/Engine/Plugins/Runtime/Metasound/Source/MetasoundFrontend/Public/MetasoundFrontend.h b/Engine/Plugins/Runtime/Metasound/Source/MetasoundFrontend/Public/MetasoundFrontend.h
--- a/Engine/Plugins/Runtime/Metasound/Source/MetasoundFrontend/Public/MetasoundFrontend.h
+++ b/Engine/Plugins/Runtime/Metasound/Source/MetasoundFrontend/Public/MetasoundFrontend.h
@@ -117,7 +117,7 @@ namespace Metasound
 		private:
 			// This is used to grant specific classes the ability to create handles.
 			enum EPrivateToken { Token };
-			static const EPrivateToken PrivateToken = EPrivateToken::Token;
+			static const EPrivateToken PrivateToken;
 
 			friend class FOutputHandle;
 			friend class FInputHandle;
diff --git a/Engine/Plugins/Runtime/Metasound/Source/MetasoundGraphCore/Private/MetasoundAudioFormats.cpp b/Engine/Plugins/Runtime/Metasound/Source/MetasoundGraphCore/Private/MetasoundAudioFormats.cpp
--- a/Engine/Plugins/Runtime/Metasound/Source/MetasoundGraphCore/Private/MetasoundAudioFormats.cpp
+++ b/Engine/Plugins/Runtime/Metasound/Source/MetasoundGraphCore/Private/MetasoundAudioFormats.cpp
@@ -6,11 +6,29 @@
 #include "MetasoundDataReference.h"
 #include "MetasoundAudioBuffer.h"
 
+struct FMyType
+{
+public:
+	FMyType() = delete;
+
+private:
+	FMyType Get();
+	friend class FMyOtherType;
+};
+
 namespace Metasound
 {
-	REGISTER_METASOUND_DATATYPE(FAudioBuffer)
-	REGISTER_METASOUND_DATATYPE(FUnformattedAudio)
-	REGISTER_METASOUND_DATATYPE(FMultichannelAudioFormat)
+	REGISTER_METASOUND_DATATYPE(FAudioBuffer, "Audio:Buffer")
+	REGISTER_METASOUND_DATATYPE(FUnformattedAudio, "Audio:Unformatted")
+	REGISTER_METASOUND_DATATYPE(FMultichannelAudioFormat, "Audio:Multichannel")
+
+
+
+	// TODO: Neither of these are trivially constructible, becuase they are actually views into a preexisting buffer.
+	//REGISTER_METASOUND_DATATYPE(FMonoAudioFormat)
+	//REGISTER_METASOUND_DATATYPE(FStereoAudioFormat)
+	IMPL_METASOUND_DATA_TYPE(FMonoAudioFormat, "Audio:Mono")
+	IMPL_METASOUND_DATA_TYPE(FStereoAudioFormat, "Audio:Stereo")
 
 	/* FUnformattedAudio */
 
diff --git a/Engine/Plugins/Runtime/Metasound/Source/MetasoundGraphCore/Private/MetasoundGraphCoreModule.cpp b/Engine/Plugins/Runtime/Metasound/Source/MetasoundGraphCore/Private/MetasoundGraphCoreModule.cpp
--- a/Engine/Plugins/Runtime/Metasound/Source/MetasoundGraphCore/Private/MetasoundGraphCoreModule.cpp
+++ b/Engine/Plugins/Runtime/Metasound/Source/MetasoundGraphCore/Private/MetasoundGraphCoreModule.cpp
@@ -12,19 +12,23 @@
 
 DEFINE_LOG_CATEGORY(LogMetasoundGraphCore);
 
-REGISTER_METASOUND_DATATYPE(bool, ::Metasound::ELiteralArgType::Boolean)
-REGISTER_METASOUND_DATATYPE(int32, ::Metasound::ELiteralArgType::Integer)
-REGISTER_METASOUND_DATATYPE(int64, ::Metasound::ELiteralArgType::Integer)
-REGISTER_METASOUND_DATATYPE(float, ::Metasound::ELiteralArgType::Float)
-REGISTER_METASOUND_DATATYPE(double, ::Metasound::ELiteralArgType::Float)
+REGISTER_METASOUND_DATATYPE(bool, "Primitive:Bool", ::Metasound::ELiteralArgType::Boolean)
+
+REGISTER_METASOUND_DATATYPE(int32, "Primitive:Int32", ::Metasound::ELiteralArgType::Integer)
+
+REGISTER_METASOUND_DATATYPE(int64, "Primitive:Int64", ::Metasound::ELiteralArgType::Integer)
+
+REGISTER_METASOUND_DATATYPE(float, "Primitive:Float", ::Metasound::ELiteralArgType::Float)
+
+REGISTER_METASOUND_DATATYPE(double, "Primitive:Double", ::Metasound::ELiteralArgType::Float)
 
 namespace Metasound
 {
-	REGISTER_METASOUND_DATATYPE(FBop)
-	REGISTER_METASOUND_DATATYPE(FFloatTime)
-	REGISTER_METASOUND_DATATYPE(FDoubleTime)
-	REGISTER_METASOUND_DATATYPE(FSampleTime)
-	REGISTER_METASOUND_DATATYPE(FFrequency)
+	REGISTER_METASOUND_DATATYPE(FBop, "Primitive:Bop")
+	REGISTER_METASOUND_DATATYPE(FFloatTime, "Primitive:Time")
+	REGISTER_METASOUND_DATATYPE(FDoubleTime, "Primitive:Time:HighResolution")
+	REGISTER_METASOUND_DATATYPE(FSampleTime, "Primitive:Time:SampleResolution")
+	REGISTER_METASOUND_DATATYPE(FFrequency, "Primitive:Frequency")
 }
 
 namespace Metasound 
diff --git a/Engine/Plugins/Runtime/Metasound/Source/MetasoundGraphCore/Public/MetasoundAudioBuffer.h b/Engine/Plugins/Runtime/Metasound/Source/MetasoundGraphCore/Public/MetasoundAudioBuffer.h
--- a/Engine/Plugins/Runtime/Metasound/Source/MetasoundGraphCore/Public/MetasoundAudioBuffer.h
+++ b/Engine/Plugins/Runtime/Metasound/Source/MetasoundGraphCore/Public/MetasoundAudioBuffer.h
@@ -154,5 +154,5 @@ namespace Metasound
 #endif				
 	};
 
-	DECLARE_METASOUND_DATA_REFERENCE_TYPES(FAudioBuffer, "Audio:Buffer", 0x1b02dac2 , FAudioBufferTypeInfo, FAudioBufferReadRef, FAudioBufferWriteRef);
+	DECLARE_METASOUND_DATA_REFERENCE_TYPES(FAudioBuffer, 0x1b02dac2 , FAudioBufferTypeInfo, FAudioBufferReadRef, FAudioBufferWriteRef);
 };
diff --git a/Engine/Plugins/Runtime/Metasound/Source/MetasoundGraphCore/Public/MetasoundAudioFormats.h b/Engine/Plugins/Runtime/Metasound/Source/MetasoundGraphCore/Public/MetasoundAudioFormats.h
--- a/Engine/Plugins/Runtime/Metasound/Source/MetasoundGraphCore/Public/MetasoundAudioFormats.h
+++ b/Engine/Plugins/Runtime/Metasound/Source/MetasoundGraphCore/Public/MetasoundAudioFormats.h
@@ -456,17 +456,15 @@ namespace Metasound
 	};
 
 
-	DECLARE_METASOUND_DATA_REFERENCE_TYPES(FUnformattedAudio, "Audio:Unformatted", 0xd78a3ed1 , FUnformattedAudioTypeInfo, FUnformattedAudioReadRef, FUnformattedAudioWriteRef);
+	DECLARE_METASOUND_DATA_REFERENCE_TYPES(FUnformattedAudio, 0xd78a3ed1 , FUnformattedAudioTypeInfo, FUnformattedAudioReadRef, FUnformattedAudioWriteRef);
 	
 
-	DECLARE_METASOUND_DATA_REFERENCE_TYPES(FMultichannelAudioFormat, "Audio:Multichannel", 0x56bdcbe0 , FMultichannelAudioFormatTypeInfo, FMultichannelAudioFormatReadRef, FMultichannelAudioFormatWriteRef);
+	DECLARE_METASOUND_DATA_REFERENCE_TYPES(FMultichannelAudioFormat, 0x56bdcbe0 , FMultichannelAudioFormatTypeInfo, FMultichannelAudioFormatReadRef, FMultichannelAudioFormatWriteRef);
 	
 
-	DECLARE_METASOUND_DATA_REFERENCE_TYPES(FMonoAudioFormat, "Audio:Mono", 0x6f468c8c, FMonoAudioFormatTypeInfo, FMonoAudioFormatReadRef, FMonoAudioFormatWriteRef);
+	DECLARE_METASOUND_DATA_REFERENCE_TYPES(FMonoAudioFormat, 0x6f468c8c, FMonoAudioFormatTypeInfo, FMonoAudioFormatReadRef, FMonoAudioFormatWriteRef);
 
-	DECLARE_METASOUND_DATA_REFERENCE_TYPES(FStereoAudioFormat, "Audio:Stereo", 0xb55304e2 , FStereoAudioFormatTypeInfo, FStereoAudioFormatReadRef, FStereoAudioFormatWriteRef);
+	DECLARE_METASOUND_DATA_REFERENCE_TYPES(FStereoAudioFormat, 0xb55304e2 , FStereoAudioFormatTypeInfo, FStereoAudioFormatReadRef, FStereoAudioFormatWriteRef);
+	
 	
-	// TODO: Neither of these are trivially constructible, becuase they are actually views into a preexisting buffer.
-	//REGISTER_METASOUND_DATATYPE(FMonoAudioFormat)
-	//REGISTER_METASOUND_DATATYPE(FStereoAudioFormat)
 }
diff --git a/Engine/Plugins/Runtime/Metasound/Source/MetasoundGraphCore/Public/MetasoundDataReference.h b/Engine/Plugins/Runtime/Metasound/Source/MetasoundGraphCore/Public/MetasoundDataReference.h
--- a/Engine/Plugins/Runtime/Metasound/Source/MetasoundGraphCore/Public/MetasoundDataReference.h
+++ b/Engine/Plugins/Runtime/Metasound/Source/MetasoundGraphCore/Public/MetasoundDataReference.h
@@ -9,11 +9,11 @@
 // Declares a metasound parameter type by
 // - Adding typedefs for commonly used template types.
 // - Defining parameter type traits.
-#define DECLARE_METASOUND_DATA_REFERENCE_TYPES(DataType, DataTypeName, DataTypeMagicNumber, DataTypeInfoTypeName, DataReadReferenceTypeName, DataWriteReferenceTypeName) \
+#define DECLARE_METASOUND_DATA_REFERENCE_TYPES(DataType, DataTypeMagicNumber, DataTypeInfoTypeName, DataReadReferenceTypeName, DataWriteReferenceTypeName) \
 	template<> \
 	struct ::Metasound::TDataReferenceTypeInfo<DataType> \
 	{ \
-		static constexpr const TCHAR* TypeName = TEXT(DataTypeName); \
+		static const TCHAR* TypeName; \
 		static constexpr ::Metasound::FDataTypeMagicNumber MagicNumber = (DataTypeMagicNumber); \
 		static constexpr bool bIsStringParsable = TTestIfDataTypeCtorIsImplemented<DataType, const FString&>::Value; \
 		static constexpr bool bIsBoolParsable = TTestIfDataTypeCtorIsImplemented<DataType, bool>::Value; \
@@ -29,6 +29,10 @@
 	typedef ::Metasound::TDataReadReference<DataType> DataReadReferenceTypeName; \
 	typedef ::Metasound::TDataWriteReference<DataType> DataWriteReferenceTypeName; \
 
+// This only needs to be called if you don't plan on calling REGISTER_METASOUND_DATATYPE.
+#define IMPL_METASOUND_DATA_TYPE(DataType, DataTypeName) \
+	const TCHAR* ::Metasound::TDataReferenceTypeInfo<DataType>::TypeName = TEXT(DataTypeName);
+
 namespace Metasound
 {
 	/** ID for parameter type. Should be unique for different C++ types. */
@@ -73,7 +77,7 @@ namespace Metasound
 		 * Users should be aware of FName case-insensitivity and other 
 		 * limitiations. 
 		 */
-		static constexpr const TCHAR* TypeName = nullptr;
+		static constexpr const TCHAR TypeName[] = TEXT("");
 
 		/** Magic number used to check the data type when casting. */
 		static constexpr const FDataTypeMagicNumber MagicNumber = -1;
@@ -480,4 +484,9 @@ namespace Metasound
 			, bIsDefaultConstructible(false)
 		{}
 	};
-}
\ No newline at end of file
+}
+
+struct FMyRandomDataType
+{
+	explicit FMyRandomDataType(bool bInValue);
+};
\ No newline at end of file
diff --git a/Engine/Plugins/Runtime/Metasound/Source/MetasoundGraphCore/Public/MetasoundDataReferenceTypes.h b/Engine/Plugins/Runtime/Metasound/Source/MetasoundGraphCore/Public/MetasoundDataReferenceTypes.h
--- a/Engine/Plugins/Runtime/Metasound/Source/MetasoundGraphCore/Public/MetasoundDataReferenceTypes.h
+++ b/Engine/Plugins/Runtime/Metasound/Source/MetasoundGraphCore/Public/MetasoundDataReferenceTypes.h
@@ -60,5 +60,5 @@ namespace Metasound
 			TArray<uint32> Frames;
 	};
 
-	DECLARE_METASOUND_DATA_REFERENCE_TYPES(FBop, "Bop", 2, FBopTypeInfo, FBopReadRef, FBopWriteRef);
+	DECLARE_METASOUND_DATA_REFERENCE_TYPES(FBop, 2, FBopTypeInfo, FBopReadRef, FBopWriteRef);
 }
diff --git a/Engine/Plugins/Runtime/Metasound/Source/MetasoundGraphCore/Public/MetasoundDataTypeRegistrationMacro.h b/Engine/Plugins/Runtime/Metasound/Source/MetasoundGraphCore/Public/MetasoundDataTypeRegistrationMacro.h
--- a/Engine/Plugins/Runtime/Metasound/Source/MetasoundGraphCore/Public/MetasoundDataTypeRegistrationMacro.h
+++ b/Engine/Plugins/Runtime/Metasound/Source/MetasoundGraphCore/Public/MetasoundDataTypeRegistrationMacro.h
@@ -28,14 +28,14 @@ bool RegisterDataTypeWithFrontend(::Metasound::ELiteralArgType PreferredArgType
 	bAlreadyRegisteredThisDataType = true;
 
 	// Lambdas that generate our template-instantiated input and output nodes:
-	FInputNodeConstructorCallback InputNodeConstructor = [](const ::Metasound::FInputNodeConstructorParams& InParams) -> TUniquePtr<Metasound::INode>
+	FInputNodeConstructorCallback InputNodeConstructor = [](const ::Metasound::FInputNodeConstructorParams& InParams) -> TUniquePtr<::Metasound::INode>
 	{
-		return TUniquePtr<Metasound::INode>(new ::Metasound::TInputNode<TDataType>(InParams.InNodeName, InParams.InVertexName, InParams.InitParam, InParams.InSettings));
+		return TUniquePtr<::Metasound::INode>(new ::Metasound::TInputNode<TDataType>(InParams.InNodeName, InParams.InVertexName, InParams.InitParam, InParams.InSettings));
 	};
 
-	FOutputNodeConstructorCallback OutputNodeConstructor = [](const ::Metasound::FOutputNodeConstrutorParams& InParams) -> TUniquePtr<Metasound::INode>
+	FOutputNodeConstructorCallback OutputNodeConstructor = [](const ::Metasound::FOutputNodeConstrutorParams& InParams) -> TUniquePtr<::Metasound::INode>
 	{
-		return TUniquePtr<Metasound::INode>(new ::Metasound::TOutputNode<TDataType>(InParams.InNodeName, InParams.InVertexName));
+		return TUniquePtr<::Metasound::INode>(new ::Metasound::TOutputNode<TDataType>(InParams.InNodeName, InParams.InVertexName));
 	};
 	
 	static FName DataTypeName = FName(::Metasound::TDataReferenceTypeInfo<TDataType>::TypeName);
@@ -62,7 +62,8 @@ bool RegisterDataTypeWithFrontend(::Metasound::ELiteralArgType PreferredArgType
 // If no argument is passed in, we will infer a literal type to use.
 // Metasound::ELiteralArgType::Invalid can be used to enforce that we don't provide space for a literal, in which case you should have a default constructor or a constructor that takes [const FOperatorSettings&] implemented.
 // If you pass in a preferred arg type, please make sure that the passed in datatype has a matching constructor, since we won't check this until runtime.
-#define REGISTER_METASOUND_DATATYPE(DataType, ...) \
+#define REGISTER_METASOUND_DATATYPE(DataType, DataTypeName, ...) \
+	IMPL_METASOUND_DATA_TYPE(DataType, DataTypeName) \
 	static_assert(::Metasound::TDataReferenceTypeInfo<DataType>::bIsValidSpecialization, "Please call DECLARE_METASOUND_DATA_REFERENCE_TYPES(" #DataType "...) before calling REGISTER_METASOUND_DATATYPE(" #DataType")."); \
 	static constexpr bool bCanRegister##DataType = ::Metasound::TDataReferenceTypeInfo<DataType>::bIsStringParsable || ::Metasound::TDataReferenceTypeInfo<DataType>::bIsBoolParsable || ::Metasound::TDataReferenceTypeInfo<DataType>::bIsIntParsable || ::Metasound::TDataReferenceTypeInfo<DataType>::bIsFloatParsable || ::Metasound::TDataReferenceTypeInfo<DataType>::bIsConstructableWithSettings || ::Metasound::TDataReferenceTypeInfo<DataType>::bCanUseDefaultConstructor; \
 	static_assert(bCanRegister##DataType , "To register " #DataType " to be used as a Metasounds input or output type, it needs a default constructor or one of the following constructors must be implemented:  " #DataType "(const ::Metasound::FOperatorSettings& InSettings), " #DataType "(int32 InValue, const ::Metasound::FOperatorSettings& InSettings), " #DataType "(float InValue, const ::Metasound::FOperatorSettings& InSettings), " #DataType "(bool InValue, const ::Metasound::FOperatorSettings& InSettings), or " #DataType "(const FString& InString, const ::Metasound::FOperatorSettings& InSettings)."); \
diff --git a/Engine/Plugins/Runtime/Metasound/Source/MetasoundGraphCore/Public/MetasoundFrequency.h b/Engine/Plugins/Runtime/Metasound/Source/MetasoundGraphCore/Public/MetasoundFrequency.h
--- a/Engine/Plugins/Runtime/Metasound/Source/MetasoundGraphCore/Public/MetasoundFrequency.h
+++ b/Engine/Plugins/Runtime/Metasound/Source/MetasoundGraphCore/Public/MetasoundFrequency.h
@@ -150,5 +150,5 @@ namespace Metasound
 			float FrequencyInHz;
 	};
 
-	DECLARE_METASOUND_DATA_REFERENCE_TYPES(FFrequency, "Primitive:Frequency", 0xf1c430a4 , FFrequencyTypeInfo, FFrequencyReadRef, FFrequencyWriteRef);
+	DECLARE_METASOUND_DATA_REFERENCE_TYPES(FFrequency, 0xf1c430a4, FFrequencyTypeInfo, FFrequencyReadRef, FFrequencyWriteRef)
 }
diff --git a/Engine/Plugins/Runtime/Metasound/Source/MetasoundGraphCore/Public/MetasoundInputNode.h b/Engine/Plugins/Runtime/Metasound/Source/MetasoundGraphCore/Public/MetasoundInputNode.h
--- a/Engine/Plugins/Runtime/Metasound/Source/MetasoundGraphCore/Public/MetasoundInputNode.h
+++ b/Engine/Plugins/Runtime/Metasound/Source/MetasoundGraphCore/Public/MetasoundInputNode.h
@@ -49,6 +49,8 @@ namespace Metasound
 			// TODO: Would be nice if we didn't utilize the copy constructor by default here. But it's a nice catch-all
 			// for our use case.  Could store constructor parameters of the paramtype, but run the risk of holding rvalues
 			// or pointers to objects which are no longer valid. 
+
+			// TODO: instead of this copy factory, register a DataType factory lambda that uses a FDataTypeLiteralParam
 			class FCopyOperatorFactory : public IOperatorFactory
 			{
 				static_assert(TIsConstructible<DataType, const DataType&>::Value, "Data Type must be copy constructible!");
diff --git a/Engine/Plugins/Runtime/Metasound/Source/MetasoundGraphCore/Public/MetasoundPrimitives.h b/Engine/Plugins/Runtime/Metasound/Source/MetasoundGraphCore/Public/MetasoundPrimitives.h
--- a/Engine/Plugins/Runtime/Metasound/Source/MetasoundGraphCore/Public/MetasoundPrimitives.h
+++ b/Engine/Plugins/Runtime/Metasound/Source/MetasoundGraphCore/Public/MetasoundPrimitives.h
@@ -12,14 +12,14 @@
 
 namespace Metasound
 {
-	DECLARE_METASOUND_DATA_REFERENCE_TYPES(bool, "Primitive:Bool", 0xcc9d715e, FBoolTypeInfo, FBoolReadRef, FBoolWriteRef);
+	DECLARE_METASOUND_DATA_REFERENCE_TYPES(bool, 0xcc9d715e, FBoolTypeInfo, FBoolReadRef, FBoolWriteRef);
 	
-	DECLARE_METASOUND_DATA_REFERENCE_TYPES(int32, "Primitive:Int32", 0xf9920034, FInt32TypeInfo, FInt32ReadRef, FInt32WriteRef);
+	DECLARE_METASOUND_DATA_REFERENCE_TYPES(int32, 0xf9920034, FInt32TypeInfo, FInt32ReadRef, FInt32WriteRef);
 
-	DECLARE_METASOUND_DATA_REFERENCE_TYPES(int64, "Primitive:Int64", 0xfc8f8ada, FInt64TypeInfo, FInt64ReadRef, FInt64WriteRef);
+	DECLARE_METASOUND_DATA_REFERENCE_TYPES(int64, 0xfc8f8ada, FInt64TypeInfo, FInt64ReadRef, FInt64WriteRef);
 
-	DECLARE_METASOUND_DATA_REFERENCE_TYPES(float, "Primitive:Float", 0xe7adc4db, FFloatTypeInfo, FFloatReadRef, FFloatWriteRef);
+	DECLARE_METASOUND_DATA_REFERENCE_TYPES(float, 0xe7adc4db, FFloatTypeInfo, FFloatReadRef, FFloatWriteRef);
 
-	DECLARE_METASOUND_DATA_REFERENCE_TYPES(double, "Primitive:Double", 0x8f48a053, FDoubleTypeInfo, FDoubleReadRef, FDoubleWriteRef);
+	DECLARE_METASOUND_DATA_REFERENCE_TYPES(double, 0x8f48a053, FDoubleTypeInfo, FDoubleReadRef, FDoubleWriteRef);
 }
 
diff --git a/Engine/Plugins/Runtime/Metasound/Source/MetasoundGraphCore/Public/MetasoundTime.h b/Engine/Plugins/Runtime/Metasound/Source/MetasoundGraphCore/Public/MetasoundTime.h
--- a/Engine/Plugins/Runtime/Metasound/Source/MetasoundGraphCore/Public/MetasoundTime.h
+++ b/Engine/Plugins/Runtime/Metasound/Source/MetasoundGraphCore/Public/MetasoundTime.h
@@ -624,19 +624,19 @@ namespace Metasound
 	 */
 	typedef TTime<float, float> FFloatTime;
 
-	DECLARE_METASOUND_DATA_REFERENCE_TYPES(FFloatTime, "Primitive:Time", 0x2490848e, FFloatTimeTypeInfo, FFloatTimeReadRef, FFloatTimeWriteRef);
+	DECLARE_METASOUND_DATA_REFERENCE_TYPES(FFloatTime, 0x2490848e, FFloatTimeTypeInfo, FFloatTimeReadRef, FFloatTimeWriteRef);
 
 	/** FDoubleTime stores and represents all time values as 64bit floating
 	 * point values.
 	 */
 	typedef TTime<double, double> FDoubleTime;
 
-	DECLARE_METASOUND_DATA_REFERENCE_TYPES(FDoubleTime, "Primitive:Time:HighResolution", 0x407fcff1, FDoubleTimeTypeInfo, FDoubleTimeReadRef, FDoubleTimeWriteRef);
+	DECLARE_METASOUND_DATA_REFERENCE_TYPES(FDoubleTime, 0x407fcff1, FDoubleTimeTypeInfo, FDoubleTimeReadRef, FDoubleTimeWriteRef);
 
 	/** FSampleTime stores time quantized to sample boundaries. It returns time
 	 * as 32bit floating point values.
 	 */
 	typedef TTime<FTimeSampleCounter, float> FSampleTime;
 
-	DECLARE_METASOUND_DATA_REFERENCE_TYPES(FSampleTime, "Primitive:Time:SampleResolution", 0xf1de8f2d, FSampleTimeTypeInfo, FSampleTimeReadRef, FSampleTimeWriteRef);
+	DECLARE_METASOUND_DATA_REFERENCE_TYPES(FSampleTime, 0xf1de8f2d, FSampleTimeTypeInfo, FSampleTimeReadRef, FSampleTimeWriteRef);
 }
