Add import logic for JSON parsing.

diff --git a/Engine/Plugins/Runtime/Metasound/Source/MetasoundFrontend/Private/MetasoundFrontend.cpp b/Engine/Plugins/Runtime/Metasound/Source/MetasoundFrontend/Private/MetasoundFrontend.cpp
--- a/Engine/Plugins/Runtime/Metasound/Source/MetasoundFrontend/Private/MetasoundFrontend.cpp
+++ b/Engine/Plugins/Runtime/Metasound/Source/MetasoundFrontend/Private/MetasoundFrontend.cpp
@@ -3,7 +3,9 @@
 #include "MetasoundFrontend.h"
 
 #include "Backends/JsonStructSerializerBackend.h"
+#include "Backends/JsonStructDeserializerBackend.h"
 #include "StructSerializer.h"
+#include "StructDeserializer.h"
 #include "HAL/FileManager.h"
 #include "HAL/IConsoleManager.h"
 #include "Modules/ModuleManager.h"
@@ -214,6 +216,20 @@ namespace Metasound
 			return FMetasoundFrontendRegistryContainer::Get()->GetInfoForDataType(InDataType, OutInfo);
 		}
 
+		bool ImportJSONToMetasound(const FString& InPath, FMetasoundDocument& OutMetasoundDocument)
+		{
+			if (TUniquePtr<FArchive> FileReader = TUniquePtr<FArchive>(IFileManager::Get().CreateFileReader(*InPath)))
+			{
+				FJsonStructDeserializerBackend Backend(*FileReader);
+				bool DeserializeResult = FStructDeserializer::Deserialize(OutMetasoundDocument, Backend);
+
+				FileReader->Close();
+				return DeserializeResult && !FileReader->IsError();
+			}
+
+			return false;
+		}
+
 		FInputHandle::FInputHandle(FHandleInitParams::EPrivateToken PrivateToken, const FHandleInitParams& InParams, const FString& InputName)
 			: ITransactable(MetasoundUndoRollLimitCvar, InParams.InOwningAsset)
 			, NodePtr(InParams.InAccessPoint, InParams.InPath)
diff --git a/Engine/Plugins/Runtime/Metasound/Source/MetasoundFrontend/Public/MetasoundFrontend.h b/Engine/Plugins/Runtime/Metasound/Source/MetasoundFrontend/Public/MetasoundFrontend.h
--- a/Engine/Plugins/Runtime/Metasound/Source/MetasoundFrontend/Public/MetasoundFrontend.h
+++ b/Engine/Plugins/Runtime/Metasound/Source/MetasoundFrontend/Public/MetasoundFrontend.h
@@ -78,6 +78,11 @@ namespace Metasound
 		// returns false if InDataType couldn't be found.
 		bool GetTraitsForDataType(FName InDataType, FDataTypeRegistryInfo& OutInfo);
 
+		// This will open a json document at the given absolute path and convert it into a 
+		// metasound document struct. 
+		// @returns false if the file couldn't be found or parsed into a document.
+		bool ImportJSONToMetasound(const FString& InPath, FMetasoundDocument& OutMetasoundDocument);
+
 		// Struct that indicates whether an input and an output can be connected,
 		// and whether an intermediate node is necessary to connect the two.
 		struct FConnectability
diff --git a/Engine/Plugins/Runtime/Metasound/Source/MetasoundFrontend/Public/MetasoundFrontendDataLayout.h b/Engine/Plugins/Runtime/Metasound/Source/MetasoundFrontend/Public/MetasoundFrontendDataLayout.h
--- a/Engine/Plugins/Runtime/Metasound/Source/MetasoundFrontend/Public/MetasoundFrontendDataLayout.h
+++ b/Engine/Plugins/Runtime/Metasound/Source/MetasoundFrontend/Public/MetasoundFrontendDataLayout.h
@@ -199,7 +199,10 @@ struct FMetasoundExternalClassLookupInfo
 {
 	GENERATED_BODY()
 
+	UPROPERTY()
 	FName ExternalNodeClassName;
+
+	UPROPERTY()
 	uint32 ExternalNodeClassHash;
 };
 
