fixup for missing dllexport.

diff --git a/Engine/Plugins/Runtime/Metasound/Source/MetasoundGraphCore/Public/MetasoundDataReference.h b/Engine/Plugins/Runtime/Metasound/Source/MetasoundGraphCore/Public/MetasoundDataReference.h
--- a/Engine/Plugins/Runtime/Metasound/Source/MetasoundGraphCore/Public/MetasoundDataReference.h
+++ b/Engine/Plugins/Runtime/Metasound/Source/MetasoundGraphCore/Public/MetasoundDataReference.h
@@ -9,11 +9,12 @@
 // Declares a metasound parameter type by
 // - Adding typedefs for commonly used template types.
 // - Defining parameter type traits.
+// - TODO- add argument for module api.
 #define DECLARE_METASOUND_DATA_REFERENCE_TYPES(DataType, DataTypeMagicNumber, DataTypeInfoTypeName, DataReadReferenceTypeName, DataWriteReferenceTypeName) \
 	template<> \
 	struct ::Metasound::TDataReferenceTypeInfo<DataType> \
 	{ \
-		static const TCHAR* TypeName; \
+		static METASOUNDGRAPHCORE_API const TCHAR* TypeName; \
 		static constexpr ::Metasound::FDataTypeMagicNumber MagicNumber = (DataTypeMagicNumber); \
 		static constexpr bool bIsStringParsable = TTestIfDataTypeCtorIsImplemented<DataType, const FString&>::Value; \
 		static constexpr bool bIsBoolParsable = TTestIfDataTypeCtorIsImplemented<DataType, bool>::Value; \
