Implement Translation Layer between Metasound Document Model & EdGraph Part 1
- Build EdGraphNodes from registry
- Display Metasound Metadata in custom details panel
- Fix namespacing
- Add input/output nodes scaffolding to MetasoundEditorGraph

diff --git a/Engine/Plugins/Runtime/Metasound/Source/MetasoundEditor/Private/MetasoundAssetTypeActions.cpp b/Engine/Plugins/Runtime/Metasound/Source/MetasoundEditor/Private/MetasoundAssetTypeActions.cpp
--- a/Engine/Plugins/Runtime/Metasound/Source/MetasoundEditor/Private/MetasoundAssetTypeActions.cpp
+++ b/Engine/Plugins/Runtime/Metasound/Source/MetasoundEditor/Private/MetasoundAssetTypeActions.cpp
@@ -8,35 +8,38 @@
 
 #define LOCTEXT_NAMESPACE "AssetTypeActions"
 
-namespace MetasoundEditorUtils
+namespace Metasound
 {
-	static const TArray<FText> AssetTypeActionSubMenu
+	namespace Editor
 	{
-		LOCTEXT("AssetMetasoundSubMenu", "Metasound")
-	};
-} // namespace MetasoundEditorUtils
-
-UClass* FAssetTypeActions_Metasound::GetSupportedClass() const
-{
-	return UMetasound::StaticClass();
-}
+		static const TArray<FText> AssetTypeActionSubMenu
+		{
+			LOCTEXT("AssetMetasoundSubMenu", "Metasound")
+		};
 
-const TArray<FText>& FAssetTypeActions_Metasound::GetSubMenus() const
-{
-	return MetasoundEditorUtils::AssetTypeActionSubMenu;
-}
+		UClass* FAssetTypeActions_Metasound::GetSupportedClass() const
+		{
+			return UMetasound::StaticClass();
+		}
 
-void FAssetTypeActions_Metasound::OpenAssetEditor(const TArray<UObject*>& InObjects, TSharedPtr<IToolkitHost> ToolkitHost)
-{
-	const EToolkitMode::Type Mode = ToolkitHost.IsValid() ? EToolkitMode::WorldCentric : EToolkitMode::Standalone;
+		const TArray<FText>& FAssetTypeActions_Metasound::GetSubMenus() const
+		{
+			return AssetTypeActionSubMenu;
+		}
 
-	for (UObject* Object : InObjects)
-	{
-		if (UMetasound* Metasound = Cast<UMetasound>(Object))
+		void FAssetTypeActions_Metasound::OpenAssetEditor(const TArray<UObject*>& InObjects, TSharedPtr<IToolkitHost> ToolkitHost)
 		{
-			TSharedRef<FMetasoundEditor> NewEditor = MakeShared<FMetasoundEditor>();
-			NewEditor->InitMetasoundEditor(Mode, ToolkitHost, Metasound);
+			const EToolkitMode::Type Mode = ToolkitHost.IsValid() ? EToolkitMode::WorldCentric : EToolkitMode::Standalone;
+
+			for (UObject* Object : InObjects)
+			{
+				if (UMetasound* Metasound = Cast<UMetasound>(Object))
+				{
+					TSharedRef<FEditor> NewEditor = MakeShared<FEditor>();
+					NewEditor->InitMetasoundEditor(Mode, ToolkitHost, Metasound);
+				}
+			}
 		}
-	}
-}
+	} // namespace Editor
+} // namespace Metasound
 #undef LOCTEXT_NAMESPACE
diff --git a/Engine/Plugins/Runtime/Metasound/Source/MetasoundEditor/Private/MetasoundAssetTypeActions.h b/Engine/Plugins/Runtime/Metasound/Source/MetasoundEditor/Private/MetasoundAssetTypeActions.h
--- a/Engine/Plugins/Runtime/Metasound/Source/MetasoundEditor/Private/MetasoundAssetTypeActions.h
+++ b/Engine/Plugins/Runtime/Metasound/Source/MetasoundEditor/Private/MetasoundAssetTypeActions.h
@@ -8,16 +8,21 @@
 // Forward Declarations
 class IToolkitHost;
 
-
-class FAssetTypeActions_Metasound : public FAssetTypeActions_Base
+namespace Metasound
 {
-public:
-	// IAssetTypeActions Implementation
-	virtual FText GetName() const override { return NSLOCTEXT("AssetTypeActions", "AssetTypeActions_Metasound", "Metasound"); }
-	virtual FColor GetTypeColor() const override { return FColor(220, 160, 0); }
-	virtual UClass* GetSupportedClass() const override;
-	virtual uint32 GetCategories() override { return EAssetTypeCategories::Sounds; }
-	virtual void OpenAssetEditor(const TArray<UObject*>& InObjects, TSharedPtr<IToolkitHost> ToolkitHost) override;
+	namespace Editor
+	{
+		class FAssetTypeActions_Metasound : public FAssetTypeActions_Base
+		{
+		public:
+			// IAssetTypeActions Implementation
+			virtual FText GetName() const override { return NSLOCTEXT("AssetTypeActions", "AssetTypeActions_Metasound", "Metasound"); }
+			virtual FColor GetTypeColor() const override { return FColor(220, 160, 0); }
+			virtual UClass* GetSupportedClass() const override;
+			virtual uint32 GetCategories() override { return EAssetTypeCategories::Sounds; }
+			virtual void OpenAssetEditor(const TArray<UObject*>& InObjects, TSharedPtr<IToolkitHost> ToolkitHost) override;
 
-	virtual const TArray<FText>& GetSubMenus() const override;
-};
+			virtual const TArray<FText>& GetSubMenus() const override;
+		};
+	} // namespace Editor
+} // namespace Metasound
diff --git a/Engine/Plugins/Runtime/Metasound/Source/MetasoundEditor/Private/MetasoundDetailCustomization.cpp b/Engine/Plugins/Runtime/Metasound/Source/MetasoundEditor/Private/MetasoundDetailCustomization.cpp
new file mode 100644
--- /dev/null
+++ b/Engine/Plugins/Runtime/Metasound/Source/MetasoundEditor/Private/MetasoundDetailCustomization.cpp
@@ -0,0 +1,28 @@
+// Copyright Epic Games, Inc. All Rights Reserved.
+#include "MetasoundDetailCustomization.h"
+
+#include "CoreMinimal.h"
+#include "DetailCategoryBuilder.h"
+#include "DetailLayoutBuilder.h"
+#include "Metasound.h"
+#include "PropertyEditorDelegates.h"
+#include "PropertyHandle.h"
+#include "Templates/Casts.h"
+#include "Templates/SharedPointer.h"
+
+
+void FMetasoundDetailCustomization::CustomizeDetails(IDetailLayoutBuilder& DetailLayout)
+{
+	IDetailCategoryBuilder& Category = DetailLayout.EditCategory("General");
+
+	TSharedPtr<IPropertyHandle> AuthorHandle = DetailLayout.GetProperty("RootMetasoundDocument.RootClass.Metadata.AuthorName");
+	TSharedPtr<IPropertyHandle> DescHandle = DetailLayout.GetProperty("RootMetasoundDocument.RootClass.Metadata.MetasoundDescription");
+	TSharedPtr<IPropertyHandle> NodeTypeHandle = DetailLayout.GetProperty("RootMetasoundDocument.RootClass.Metadata.NodeType");
+
+	Category.AddProperty(NodeTypeHandle);
+	Category.AddProperty(AuthorHandle);
+	Category.AddProperty(DescHandle);
+
+	// Hack to hide parent structs for nested metadata properties
+	DetailLayout.HideCategory("Hidden");
+}
diff --git a/Engine/Plugins/Runtime/Metasound/Source/MetasoundEditor/Private/MetasoundDetailCustomization.h b/Engine/Plugins/Runtime/Metasound/Source/MetasoundEditor/Private/MetasoundDetailCustomization.h
new file mode 100644
--- /dev/null
+++ b/Engine/Plugins/Runtime/Metasound/Source/MetasoundEditor/Private/MetasoundDetailCustomization.h
@@ -0,0 +1,16 @@
+// Copyright Epic Games, Inc. All Rights Reserved.
+#pragma once
+
+#include "IDetailCustomization.h"
+
+// Forward Declarations
+class IDetailLayoutBuilder;
+
+
+class FMetasoundDetailCustomization : public IDetailCustomization
+{
+public:
+	// IDetailCustomization interface
+	virtual void CustomizeDetails(IDetailLayoutBuilder& DetailLayout) override;
+	// End of IDetailCustomization interface
+};
\ No newline at end of file
diff --git a/Engine/Plugins/Runtime/Metasound/Source/MetasoundEditor/Private/MetasoundEditor.cpp b/Engine/Plugins/Runtime/Metasound/Source/MetasoundEditor/Private/MetasoundEditor.cpp
--- a/Engine/Plugins/Runtime/Metasound/Source/MetasoundEditor/Private/MetasoundEditor.cpp
+++ b/Engine/Plugins/Runtime/Metasound/Source/MetasoundEditor/Private/MetasoundEditor.cpp
@@ -15,10 +15,13 @@
 #include "HAL/PlatformApplicationMisc.h"
 #include "IDetailsView.h"
 #include "Kismet2/BlueprintEditorUtils.h"
+#include "MetasoundAssetSerializer.h"
 #include "MetasoundEditorCommands.h"
 #include "MetasoundEditorGraph.h"
+#include "MetasoundEditorGraphBuilder.h"
 #include "MetasoundEditorGraphNode.h"
 #include "MetasoundEditorGraphSchema.h"
+#include "MetasoundEditorTabFactory.h"
 #include "Misc/Attribute.h"
 #include "Modules/ModuleManager.h"
 #include "PropertyEditorModule.h"
@@ -29,874 +32,759 @@
 
 #define LOCTEXT_NAMESPACE "MetasoundEditor"
 
-const FName FMetasoundEditor::GraphCanvasTabId(TEXT("MetasoundEditor_GraphCanvas"));
-const FName FMetasoundEditor::PropertiesTabId(TEXT("MetasoundEditor_Properties"));
-const FName FMetasoundEditor::PaletteTabId(TEXT("MetasoundEditor_Palette"));
-
-FMetasoundEditor::FMetasoundEditor()
-	: Metasound(nullptr)
-{
-}
-
-void FMetasoundEditor::RegisterTabSpawners(const TSharedRef<FTabManager>& InTabManager)
+namespace Metasound
 {
-	WorkspaceMenuCategory = InTabManager->AddLocalWorkspaceMenuCategory(LOCTEXT("WorkspaceMenu_MetasoundEditor", "Metasound Editor"));
-	auto WorkspaceMenuCategoryRef = WorkspaceMenuCategory.ToSharedRef();
-
-	FAssetEditorToolkit::RegisterTabSpawners(InTabManager);
-
-	InTabManager->RegisterTabSpawner(GraphCanvasTabId, FOnSpawnTab::CreateSP(this, &FMetasoundEditor::SpawnTab_GraphCanvas))
-		.SetDisplayName(LOCTEXT("GraphCanvasTab", "Viewport"))
-		.SetGroup(WorkspaceMenuCategoryRef)
-		.SetIcon(FSlateIcon(FEditorStyle::GetStyleSetName(), "GraphEditor.EventGraph_16x"));
-
-	InTabManager->RegisterTabSpawner(PropertiesTabId, FOnSpawnTab::CreateSP(this, &FMetasoundEditor::SpawnTab_Properties))
-		.SetDisplayName(LOCTEXT("DetailsTab", "Details"))
-		.SetGroup(WorkspaceMenuCategoryRef)
-		.SetIcon(FSlateIcon(FEditorStyle::GetStyleSetName(), "LevelEditor.Tabs.Details"));
+	namespace Editor
+	{
+		void FEditor::RegisterTabSpawners(const TSharedRef<FTabManager>& InTabManager)
+		{
+			WorkspaceMenuCategory = InTabManager->AddLocalWorkspaceMenuCategory(LOCTEXT("WorkspaceMenu_MetasoundEditor", "Metasound Editor"));
+			auto WorkspaceMenuCategoryRef = WorkspaceMenuCategory.ToSharedRef();
 
-	InTabManager->RegisterTabSpawner(PaletteTabId, FOnSpawnTab::CreateSP(this, &FMetasoundEditor::SpawnTab_Palette))
-		.SetDisplayName(LOCTEXT("PaletteTab", "Palette"))
-		.SetGroup(WorkspaceMenuCategoryRef)
-		.SetIcon(FSlateIcon(FEditorStyle::GetStyleSetName(), "Kismet.Tabs.Palette"));
-}
+			FAssetEditorToolkit::RegisterTabSpawners(InTabManager);
 
-void FMetasoundEditor::UnregisterTabSpawners(const TSharedRef<FTabManager>& InTabManager)
-{
-	FAssetEditorToolkit::UnregisterTabSpawners(InTabManager);
+			InTabManager->RegisterTabSpawner(TabFactory::Names::GraphCanvas, FOnSpawnTab::CreateLambda([InMetasoundGraphEditor = MetasoundGraphEditor](const FSpawnTabArgs& Args)
+			{
+				return TabFactory::CreateGraphCanvasTab(InMetasoundGraphEditor, Args);
+			}))
+			.SetDisplayName(LOCTEXT("GraphCanvasTab", "Viewport"))
+			.SetGroup(WorkspaceMenuCategoryRef)
+			.SetIcon(FSlateIcon(FEditorStyle::GetStyleSetName(), "GraphEditor.EventGraph_16x"));
 
-	InTabManager->UnregisterTabSpawner(GraphCanvasTabId);
-	InTabManager->UnregisterTabSpawner(PropertiesTabId);
-	InTabManager->UnregisterTabSpawner(PaletteTabId);
-}
+			InTabManager->RegisterTabSpawner(TabFactory::Names::Properties, FOnSpawnTab::CreateLambda([InMetasoundProperties = MetasoundProperties](const FSpawnTabArgs& Args)
+			{
+				return TabFactory::CreatePropertiesTab(InMetasoundProperties, Args);
+			}))
+			.SetDisplayName(LOCTEXT("DetailsTab", "Details"))
+			.SetGroup(WorkspaceMenuCategoryRef)
+			.SetIcon(FSlateIcon(FEditorStyle::GetStyleSetName(), "LevelEditor.Tabs.Details"));
 
-FMetasoundEditor::~FMetasoundEditor()
-{
-	// Stop any playing sounds when the editor closes
-	UAudioComponent* Component = GEditor->GetPreviewAudioComponent();
-	if (Component && Component->IsPlaying())
-	{
-		Stop();
-	}
+			InTabManager->RegisterTabSpawner(TabFactory::Names::Palette, FOnSpawnTab::CreateLambda([InPalette = Palette](const FSpawnTabArgs& Args)
+			{
+				return TabFactory::CreatePaletteTab(InPalette, Args);
+			}))
+			.SetDisplayName(LOCTEXT("PaletteTab", "Palette"))
+			.SetGroup(WorkspaceMenuCategoryRef)
+			.SetIcon(FSlateIcon(FEditorStyle::GetStyleSetName(), "Kismet.Tabs.Palette"));
+		}
 
-	check(GEditor);
-	GEditor->UnregisterForUndo(this);
-}
+		void FEditor::UnregisterTabSpawners(const TSharedRef<FTabManager>& InTabManager)
+		{
+			using namespace Metasound::Editor;
 
-void FMetasoundEditor::InitMetasoundEditor(const EToolkitMode::Type Mode, const TSharedPtr<IToolkitHost>& InitToolkitHost, UObject* ObjectToEdit)
-{
-	Metasound = CastChecked<UMetasound>(ObjectToEdit);
+			FAssetEditorToolkit::UnregisterTabSpawners(InTabManager);
 
-	// Support undo/redo
-	Metasound->SetFlags(RF_Transactional);
-	
-	if (!Metasound->GetGraph())
-	{
-		UMetasoundEditorGraph* Graph = NewObject<UMetasoundEditorGraph>();
-		Graph->Metasound = Metasound;
-		Graph->Schema = UMetasoundEditorGraphSchema::StaticClass();
-		Metasound->SetGraph(Graph);
-	}
+			InTabManager->UnregisterTabSpawner(TabFactory::Names::GraphCanvas);
+			InTabManager->UnregisterTabSpawner(TabFactory::Names::Palette);
+			InTabManager->UnregisterTabSpawner(TabFactory::Names::Properties);
+		}
 
-	GEditor->RegisterForUndo(this);
-
-	FGraphEditorCommands::Register();
-	FMetasoundEditorCommands::Register();
-
-	BindGraphCommands();
-
-	CreateInternalWidgets();
-
-	const TSharedRef<FTabManager::FLayout> StandaloneDefaultLayout = FTabManager::NewLayout("Standalone_MetasoundEditor_Layout_v1")
-	->AddArea
-	(
-		FTabManager::NewPrimaryArea()
-		->SetOrientation(Orient_Vertical)
-		->Split
-		(
-			FTabManager::NewStack()
-			->SetSizeCoefficient(0.1f)
-			->SetHideTabWell(true)
-			->AddTab(GetToolbarTabId(), ETabState::OpenedTab)
-		)
-		->Split(FTabManager::NewSplitter()
-			->SetOrientation(Orient_Horizontal)
-			->SetSizeCoefficient(0.9f)
-			->Split
-			(
-				FTabManager::NewStack()
-				->SetSizeCoefficient(0.225f)
-				->SetHideTabWell(true)
-				->AddTab(PropertiesTabId, ETabState::OpenedTab)
-			)
-			->Split
-			(
-				FTabManager::NewStack()
-				->SetSizeCoefficient(0.65f)
-				->SetHideTabWell(true)
-				->AddTab(GraphCanvasTabId, ETabState::OpenedTab)
-			)
-			->Split
-			(
-				FTabManager::NewStack()
-				->SetSizeCoefficient(0.125f)
-				->SetHideTabWell(true)
-				->AddTab(PaletteTabId, ETabState::OpenedTab)
-			)
-		)
-	);
-
-	const bool bCreateDefaultStandaloneMenu = true;
-	const bool bCreateDefaultToolbar = true;
-	FAssetEditorToolkit::InitAssetEditor(Mode, InitToolkitHost, TEXT("MetasoundEditorApp"), StandaloneDefaultLayout, bCreateDefaultStandaloneMenu, bCreateDefaultToolbar, ObjectToEdit, false);
-
-	ExtendToolbar();
-	RegenerateMenusAndToolbars();
-}
+		FEditor::~FEditor()
+		{
+			// Stop any playing sounds when the editor closes
+			UAudioComponent* Component = GEditor->GetPreviewAudioComponent();
+			if (Component && Component->IsPlaying())
+			{
+				Stop();
+			}
 
-UMetasound* FMetasoundEditor::GetMetasound() const
-{
-	return Metasound;
-}
+			check(GEditor);
+			GEditor->UnregisterForUndo(this);
+		}
 
-void FMetasoundEditor::SetSelection(TArray<UObject*> SelectedObjects)
-{
-	if (MetasoundProperties.IsValid())
-	{
-		MetasoundProperties->SetObjects(SelectedObjects);
-	}
-}
+		void FEditor::InitMetasoundEditor(const EToolkitMode::Type Mode, const TSharedPtr<IToolkitHost>& InitToolkitHost, UObject* ObjectToEdit)
+		{
 
-bool FMetasoundEditor::GetBoundsForSelectedNodes(FSlateRect& Rect, float Padding)
-{
-	return MetasoundGraphEditor->GetBoundsForSelectedNodes(Rect, Padding);
-}
+			Metasound = CastChecked<UMetasound>(ObjectToEdit);
 
-int32 FMetasoundEditor::GetNumberOfSelectedNodes() const
-{
-	return MetasoundGraphEditor->GetSelectedNodes().Num();
-}
+			// Support undo/redo
+			Metasound->SetFlags(RF_Transactional);
 
-FName FMetasoundEditor::GetToolkitFName() const
-{
-	return FName("MetasoundEditor");
-}
+			if (!Metasound->GetGraph())
+			{
+				UMetasoundEditorGraph* Graph = NewObject<UMetasoundEditorGraph>(Metasound);
+				Graph->ParentMetasound = Metasound;
+				Graph->Schema = UMetasoundEditorGraphSchema::StaticClass();
+				Metasound->SetGraph(Graph);
+			}
 
-FText FMetasoundEditor::GetBaseToolkitName() const
-{
-	return LOCTEXT("AppLabel", "Metasound Editor");
-}
+			GEditor->RegisterForUndo(this);
+
+			FGraphEditorCommands::Register();
+			FEditorCommands::Register();
+
+			BindGraphCommands();
+
+			CreateInternalWidgets();
+
+			const TSharedRef<FTabManager::FLayout> StandaloneDefaultLayout = FTabManager::NewLayout("Standalone_MetasoundEditor_Layout_v1")
+				->AddArea
+				(
+					FTabManager::NewPrimaryArea()
+					->SetOrientation(Orient_Vertical)
+					->Split
+					(
+						FTabManager::NewStack()
+						->SetSizeCoefficient(0.1f)
+						->SetHideTabWell(true)
+						->AddTab(GetToolbarTabId(), ETabState::OpenedTab)
+					)
+					->Split(FTabManager::NewSplitter()
+						->SetOrientation(Orient_Horizontal)
+						->SetSizeCoefficient(0.9f)
+						->Split
+						(
+							FTabManager::NewStack()
+							->SetSizeCoefficient(0.225f)
+							->SetHideTabWell(true)
+							->AddTab(TabFactory::Names::Properties, ETabState::OpenedTab)
+						)
+						->Split
+						(
+							FTabManager::NewStack()
+							->SetSizeCoefficient(0.65f)
+							->SetHideTabWell(true)
+							->AddTab(TabFactory::Names::GraphCanvas, ETabState::OpenedTab)
+						)
+						->Split
+						(
+							FTabManager::NewStack()
+							->SetSizeCoefficient(0.125f)
+							->SetHideTabWell(true)
+							->AddTab(TabFactory::Names::Palette, ETabState::OpenedTab)
+						)
+					)
+				);
+
+			const bool bCreateDefaultStandaloneMenu = true;
+			const bool bCreateDefaultToolbar = true;
+			FAssetEditorToolkit::InitAssetEditor(Mode, InitToolkitHost, TEXT("MetasoundEditorApp"), StandaloneDefaultLayout, bCreateDefaultStandaloneMenu, bCreateDefaultToolbar, ObjectToEdit, false);
+
+			ExtendToolbar();
+			RegenerateMenusAndToolbars();
+		}
 
-FString FMetasoundEditor::GetWorldCentricTabPrefix() const
-{
-	return LOCTEXT("WorldCentricTabPrefix", "Metasound ").ToString();
-}
+		UMetasound* FEditor::GetMetasound() const
+		{
+			return Metasound;
+		}
 
-FLinearColor FMetasoundEditor::GetWorldCentricTabColorScale() const
-{
-	return FLinearColor(0.3f, 0.2f, 0.5f, 0.5f);
-}
+		void FEditor::SetSelection(const TArray<UObject*>& SelectedObjects)
+		{
+			if (MetasoundProperties.IsValid())
+			{
+				MetasoundProperties->SetObjects(SelectedObjects);
+			}
+		}
 
-TSharedRef<SDockTab> FMetasoundEditor::SpawnTab_GraphCanvas(const FSpawnTabArgs& Args)
-{
-	check(Args.GetTabId() == GraphCanvasTabId);
+		bool FEditor::GetBoundsForSelectedNodes(FSlateRect& Rect, float Padding)
+		{
+			return MetasoundGraphEditor->GetBoundsForSelectedNodes(Rect, Padding);
+		}
 
-	TSharedRef<SDockTab> SpawnedTab = SNew(SDockTab)
-		.Label(LOCTEXT("MetasoundGraphCanvasTitle", "Viewport"));
+		FName FEditor::GetToolkitFName() const
+		{
+			return FName("MetasoundEditor");
+		}
 
-	if (MetasoundGraphEditor.IsValid())
-	{
-		SpawnedTab->SetContent(MetasoundGraphEditor.ToSharedRef());
-	}
+		FText FEditor::GetBaseToolkitName() const
+		{
+			return LOCTEXT("AppLabel", "Metasound Editor");
+		}
 
-	return SpawnedTab;
-}
+		FString FEditor::GetWorldCentricTabPrefix() const
+		{
+			return LOCTEXT("WorldCentricTabPrefix", "Metasound ").ToString();
+		}
 
-TSharedRef<SDockTab> FMetasoundEditor::SpawnTab_Properties(const FSpawnTabArgs& Args)
-{
-	check(Args.GetTabId() == PropertiesTabId);
-
-	return SNew(SDockTab)
-		.Icon(FEditorStyle::GetBrush("LevelEditor.Tabs.Details"))
-		.Label(LOCTEXT("MetasoundDetailsTitle", "Details"))
-		[
-			MetasoundProperties.ToSharedRef()
-		];
-}
+		FLinearColor FEditor::GetWorldCentricTabColorScale() const
+		{
+			return FLinearColor(0.3f, 0.2f, 0.5f, 0.5f);
+		}
 
-TSharedRef<SDockTab> FMetasoundEditor::SpawnTab_Palette(const FSpawnTabArgs& Args)
-{
-	check(Args.GetTabId() == PaletteTabId);
-
-	return SNew(SDockTab)
-		.Icon(FEditorStyle::GetBrush("Kismet.Tabs.Palette"))
-		.Label(LOCTEXT("MetasoundPaletteTitle", "Palette"))
-		[
-			Palette.ToSharedRef()
-		];
-}
+		void FEditor::AddReferencedObjects(FReferenceCollector& Collector)
+		{
+			Collector.AddReferencedObject(Metasound);
+		}
 
-void FMetasoundEditor::AddReferencedObjects(FReferenceCollector& Collector)
-{
-	Collector.AddReferencedObject(Metasound);
-}
+		void FEditor::PostUndo(bool bSuccess)
+		{
+			if (MetasoundGraphEditor.IsValid())
+			{
+				MetasoundGraphEditor->ClearSelectionSet();
+				MetasoundGraphEditor->NotifyGraphChanged();
+				FSlateApplication::Get().DismissAllMenus();
+			}
 
-void FMetasoundEditor::PostUndo(bool bSuccess)
-{
-	if (MetasoundGraphEditor.IsValid())
-	{
-		MetasoundGraphEditor->ClearSelectionSet();
-		MetasoundGraphEditor->NotifyGraphChanged();
-		FSlateApplication::Get().DismissAllMenus();
-	}
+		}
 
-}
+		void FEditor::NotifyPostChange(const FPropertyChangedEvent& PropertyChangedEvent, FProperty* PropertyThatChanged)
+		{
+			if (MetasoundGraphEditor.IsValid() && PropertyChangedEvent.ChangeType != EPropertyChangeType::Interactive)
+			{
+				MetasoundGraphEditor->NotifyGraphChanged();
+			}
+		}
 
-void FMetasoundEditor::NotifyPostChange(const FPropertyChangedEvent& PropertyChangedEvent, FProperty* PropertyThatChanged)
-{
-	if (MetasoundGraphEditor.IsValid() && PropertyChangedEvent.ChangeType != EPropertyChangeType::Interactive)
-	{
-		MetasoundGraphEditor->NotifyGraphChanged();
-	}
-}
+		void FEditor::CreateInternalWidgets()
+		{
+			MetasoundGraphEditor = CreateGraphEditorWidget();
 
-void FMetasoundEditor::CreateInternalWidgets()
-{
-	MetasoundGraphEditor = CreateGraphEditorWidget();
+			FDetailsViewArgs Args;
+			Args.bHideSelectionTip = true;
+			Args.NotifyHook = this;
 
-	FDetailsViewArgs Args;
-	Args.bHideSelectionTip = true;
-	Args.NotifyHook = this;
+			FPropertyEditorModule& PropertyModule = FModuleManager::LoadModuleChecked<FPropertyEditorModule>("PropertyEditor");
+			MetasoundProperties = PropertyModule.CreateDetailView(Args);
+			MetasoundProperties->SetObject(Metasound);
 
-	FPropertyEditorModule& PropertyModule = FModuleManager::LoadModuleChecked<FPropertyEditorModule>("PropertyEditor");
-	MetasoundProperties = PropertyModule.CreateDetailView(Args);
-	MetasoundProperties->SetObject(Metasound);
+			Palette = SNew(SMetasoundPalette);
+		}
 
-	Palette = SNew(SMetasoundPalette);
-}
+		void FEditor::ExtendToolbar()
+		{
+			TSharedPtr<FExtender> ToolbarExtender = MakeShared<FExtender>();
+			ToolbarExtender->AddToolBarExtension
+			(
+				"Asset",
+				EExtensionHook::After,
+				GetToolkitCommands(),
+				FToolBarExtensionDelegate::CreateLambda([this](FToolBarBuilder& ToolbarBuilder)
+				{
+					ToolbarBuilder.BeginSection("Compile");
+					{
+						ToolbarBuilder.AddToolBarButton
+						(
+							FEditorCommands::Get().Compile,
+							NAME_None,
+							TAttribute<FText>(),
+							TAttribute<FText>(),
+							TAttribute<FSlateIcon>::Create([this]() { return GetCompileStatusImage(); }),
+							"CompileMetasound"
+						);
+					}
+					ToolbarBuilder.EndSection();
+
+					ToolbarBuilder.BeginSection("Audition");
+					{
+						ToolbarBuilder.AddToolBarButton(FEditorCommands::Get().Play);
+						ToolbarBuilder.AddToolBarButton(FEditorCommands::Get().Stop);
+					}
+					ToolbarBuilder.EndSection();
+				})
+			);
+
+			AddToolbarExtender(ToolbarExtender);
+		}
 
-void FMetasoundEditor::ExtendToolbar()
-{
-	TSharedPtr<FExtender> ToolbarExtender = MakeShared<FExtender>();
-	ToolbarExtender->AddToolBarExtension
-	(
-		"Asset",
-		EExtensionHook::After,
-		GetToolkitCommands(),
-		FToolBarExtensionDelegate::CreateLambda([](FToolBarBuilder& ToolbarBuilder)
-		{
-			ToolbarBuilder.BeginSection("Audition");
+		FSlateIcon FEditor::GetCompileStatusImage() const
+		{
+			FName IconName = "MetasoundEditor.Compile";
+			if (!bPassedValidation)
 			{
-				ToolbarBuilder.AddToolBarButton(FMetasoundEditorCommands::Get().Play);
-				ToolbarBuilder.AddToolBarButton(FMetasoundEditorCommands::Get().Stop);
+				IconName = "MetasoundEditor.CompileError";
 			}
-			ToolbarBuilder.EndSection();
-		})
-	);
 
-	AddToolbarExtender(ToolbarExtender);
-}
-
-void FMetasoundEditor::BindGraphCommands()
-{
-	const FMetasoundEditorCommands& Commands = FMetasoundEditorCommands::Get();
+			return FSlateIcon("MetasoundStyle", IconName);
+		}
 
-	ToolkitCommands->MapAction(
-		Commands.Play,
-		FExecuteAction::CreateSP(this, &FMetasoundEditor::Play));
+		void FEditor::BindGraphCommands()
+		{
+			const FEditorCommands& Commands = FEditorCommands::Get();
 
-	ToolkitCommands->MapAction(
-		Commands.Stop,
-		FExecuteAction::CreateSP(this, &FMetasoundEditor::Stop));
+			ToolkitCommands->MapAction(
+				Commands.Play,
+				FExecuteAction::CreateSP(this, &FEditor::Play));
 
-	ToolkitCommands->MapAction(
-		Commands.TogglePlayback,
-		FExecuteAction::CreateSP(this, &FMetasoundEditor::TogglePlayback));
+			ToolkitCommands->MapAction(
+				Commands.Stop,
+				FExecuteAction::CreateSP(this, &FEditor::Stop));
 
-	ToolkitCommands->MapAction(
-		FGenericCommands::Get().Undo,
-		FExecuteAction::CreateSP(this, &FMetasoundEditor::UndoGraphAction));
+			ToolkitCommands->MapAction(
+				Commands.Compile,
+				FExecuteAction::CreateSP(this, &FEditor::Compile));
 
-	ToolkitCommands->MapAction(
-		FGenericCommands::Get().Redo,
-		FExecuteAction::CreateSP(this, &FMetasoundEditor::RedoGraphAction));
-}
+			ToolkitCommands->MapAction(
+				Commands.TogglePlayback,
+				FExecuteAction::CreateSP(this, &FEditor::TogglePlayback));
 
-void FMetasoundEditor::Play()
-{
-// TODO: Implement play
-
-// 	GEditor->PlayPreviewSound(Metasound);
-
-// 	MetasoundGraphEditor->RegisterActiveTimer(0.0f, 
-// 		FWidgetActiveTimerDelegate::CreateLambda([](double InCurrentTime, float InDeltaTime)
-// 		{
-// 			UAudioComponent* PreviewComp = GEditor->GetPreviewAudioComponent();
-// 			if (PreviewComp && PreviewComp->IsPlaying())
-// 			{
-// 				return EActiveTimerReturnType::Continue;
-// 			}
-// 			else
-// 			{
-// 				return EActiveTimerReturnType::Stop;
-// 			}
-// 		})
-// 	);
-}
+			ToolkitCommands->MapAction(
+				FGenericCommands::Get().Undo,
+				FExecuteAction::CreateSP(this, &FEditor::UndoGraphAction));
 
-void FMetasoundEditor::PlayNode()
-{
-	// already checked that only one node is selected
-	const FGraphPanelSelectionSet SelectedNodes = GetSelectedNodes();
+			ToolkitCommands->MapAction(
+				FGenericCommands::Get().Redo,
+				FExecuteAction::CreateSP(this, &FEditor::RedoGraphAction));
+		}
 
-	for (FGraphPanelSelectionSet::TConstIterator NodeIt(SelectedNodes); NodeIt; ++NodeIt)
-	{
-		PlaySingleNode(CastChecked<UEdGraphNode>(*NodeIt));
-	}
-}
+		void FEditor::Compile()
+		{
+			if (Metasound)
+			{
+				const UMetasoundEditorGraph* Graph = CastChecked<UMetasoundEditorGraph>(Metasound->GetGraph());
+				bPassedValidation = Metasound::Engine::AssetSerializer::MetasoundToJSON(*Metasound);
+			}
+		}
 
-bool FMetasoundEditor::CanPlayNode() const
-{
-	return false;
-}
+		void FEditor::Play()
+		{
+			// 	TODO: Implement play
+			// 	check(GEditor);
+			// 	GEditor->PlayPreviewSound(Metasound);
+			// 
+			// 	MetasoundGraphEditor->RegisterActiveTimer(0.0f, 
+			// 		FWidgetActiveTimerDelegate::CreateLambda([](double InCurrentTime, float InDeltaTime)
+			// 		{
+			// 			UAudioComponent* PreviewComp = GEditor->GetPreviewAudioComponent();
+			// 			if (PreviewComp && PreviewComp->IsPlaying())
+			// 			{
+			// 				return EActiveTimerReturnType::Continue;
+			// 			}
+			// 			else
+			// 			{
+			// 				return EActiveTimerReturnType::Stop;
+			// 			}
+			// 		})
+			// 	);
+		}
 
-void FMetasoundEditor::Stop()
-{
-	check(GEditor);
-	GEditor->ResetPreviewAudioComponent();
-}
+		void FEditor::PlayNode()
+		{
+			const FGraphPanelSelectionSet SelectedNodes = MetasoundGraphEditor->GetSelectedNodes();
+			for (FGraphPanelSelectionSet::TConstIterator NodeIt(SelectedNodes); NodeIt; ++NodeIt)
+			{
+				PlaySingleNode(CastChecked<UEdGraphNode>(*NodeIt));
+			}
+		}
 
-void FMetasoundEditor::TogglePlayback()
-{
-	check(GEditor);
+		bool FEditor::CanPlayNode() const
+		{
+			// TODO: Implement node playback
+			return false;
+		}
 
-	UAudioComponent* Component = GEditor->GetPreviewAudioComponent();
-	if (Component && Component->IsPlaying())
-	{
-		Stop();
-	}
-	else
-	{
-		Play();
-	}
-}
+		void FEditor::Stop()
+		{
+			check(GEditor);
+			GEditor->ResetPreviewAudioComponent();
+		}
 
-void FMetasoundEditor::PlaySingleNode(UEdGraphNode* Node)
-{
-	// TODO: Implement? Will we support single node playback?
-}
+		void FEditor::TogglePlayback()
+		{
+			check(GEditor);
 
-void FMetasoundEditor::SyncInBrowser()
-{
-	TArray<UObject*> ObjectsToSync;
+			UAudioComponent* Component = GEditor->GetPreviewAudioComponent();
+			if (Component && Component->IsPlaying())
+			{
+				Stop();
+			}
+			else
+			{
+				Play();
+			}
+		}
 
-	const FGraphPanelSelectionSet SelectedNodes = GetSelectedNodes();
-	for (FGraphPanelSelectionSet::TConstIterator NodeIt(SelectedNodes); NodeIt; ++NodeIt)
-	{
-		// TODO: Implement sync to referenced Metasound if selected node is a reference to another metasound
-	}
+		void FEditor::PlaySingleNode(UEdGraphNode* Node)
+		{
+			// TODO: Implement? Will we support single node playback?
+		}
 
-	if (!ObjectsToSync.Num())
-	{
-		ObjectsToSync.Add(Metasound);
-	}
+		void FEditor::SyncInBrowser()
+		{
+			TArray<UObject*> ObjectsToSync;
 
-	check(GEditor);
-	GEditor->SyncBrowserToObjects(ObjectsToSync);
-}
+			const FGraphPanelSelectionSet SelectedNodes = MetasoundGraphEditor->GetSelectedNodes();
+			for (FGraphPanelSelectionSet::TConstIterator NodeIt(SelectedNodes); NodeIt; ++NodeIt)
+			{
+				// TODO: Implement sync to referenced Metasound if selected node is a reference to another metasound
+			}
 
-void FMetasoundEditor::AddInput()
-{
-}
+			if (!ObjectsToSync.Num())
+			{
+				ObjectsToSync.Add(Metasound);
+			}
 
-bool FMetasoundEditor::CanAddInput() const
-{
-	return GetSelectedNodes().Num() == 1;
-}
+			check(GEditor);
+			GEditor->SyncBrowserToObjects(ObjectsToSync);
+		}
 
-void FMetasoundEditor::DeleteInput()
-{
-}
+		void FEditor::AddInput()
+		{
+		}
 
-bool FMetasoundEditor::CanDeleteInput() const
-{
-	return true;
-}
+		bool FEditor::CanAddInput() const
+		{
+			return MetasoundGraphEditor->GetSelectedNodes().Num() == 1;
+		}
 
-void FMetasoundEditor::OnCreateComment()
-{
-	Metasound->MarkPackageDirty();
-}
+		void FEditor::DeleteInput()
+		{
+		}
 
-TSharedRef<SGraphEditor> FMetasoundEditor::CreateGraphEditorWidget()
-{
-	if (!GraphEditorCommands.IsValid())
-	{
-		GraphEditorCommands = MakeShared<FUICommandList>();
+		bool FEditor::CanDeleteInput() const
+		{
+			return true;
+		}
 
-		GraphEditorCommands->MapAction(FMetasoundEditorCommands::Get().BrowserSync,
-			FExecuteAction::CreateSP(this, &FMetasoundEditor::SyncInBrowser));
+		void FEditor::OnCreateComment()
+		{
+		}
 
-		GraphEditorCommands->MapAction(FMetasoundEditorCommands::Get().AddInput,
-			FExecuteAction::CreateSP(this, &FMetasoundEditor::AddInput),
-			FCanExecuteAction::CreateSP(this, &FMetasoundEditor::CanAddInput));
+		TSharedRef<SGraphEditor> FEditor::CreateGraphEditorWidget()
+		{
+			if (!GraphEditorCommands.IsValid())
+			{
+				GraphEditorCommands = MakeShared<FUICommandList>();
 
-		GraphEditorCommands->MapAction(FMetasoundEditorCommands::Get().DeleteInput,
-			FExecuteAction::CreateSP(this, &FMetasoundEditor::DeleteInput),
-			FCanExecuteAction::CreateSP(this, &FMetasoundEditor::CanDeleteInput));
+				GraphEditorCommands->MapAction(FEditorCommands::Get().BrowserSync,
+					FExecuteAction::CreateSP(this, &FEditor::SyncInBrowser));
 
-		// Graph Editor Commands
-		GraphEditorCommands->MapAction(FGraphEditorCommands::Get().CreateComment,
-			FExecuteAction::CreateSP(this, &FMetasoundEditor::OnCreateComment));
+				GraphEditorCommands->MapAction(FEditorCommands::Get().AddInput,
+					FExecuteAction::CreateSP(this, &FEditor::AddInput),
+					FCanExecuteAction::CreateSP(this, &FEditor::CanAddInput));
 
-		// Editing commands
-		GraphEditorCommands->MapAction(FGenericCommands::Get().SelectAll,
-			FExecuteAction::CreateSP(this, &FMetasoundEditor::SelectAllNodes),
-			FCanExecuteAction::CreateSP(this, &FMetasoundEditor::CanSelectAllNodes));
+				GraphEditorCommands->MapAction(FEditorCommands::Get().DeleteInput,
+					FExecuteAction::CreateSP(this, &FEditor::DeleteInput),
+					FCanExecuteAction::CreateSP(this, &FEditor::CanDeleteInput));
 
-		GraphEditorCommands->MapAction(FGenericCommands::Get().Delete,
-			FExecuteAction::CreateSP(this, &FMetasoundEditor::DeleteSelectedNodes),
-			FCanExecuteAction::CreateSP(this, &FMetasoundEditor::CanDeleteNodes));
+				// Graph Editor Commands
+				GraphEditorCommands->MapAction(FGraphEditorCommands::Get().CreateComment,
+					FExecuteAction::CreateSP(this, &FEditor::OnCreateComment));
 
-		GraphEditorCommands->MapAction(FGenericCommands::Get().Copy,
-			FExecuteAction::CreateSP(this, &FMetasoundEditor::CopySelectedNodes),
-			FCanExecuteAction::CreateSP(this, &FMetasoundEditor::CanCopyNodes));
+				// Editing commands
+				GraphEditorCommands->MapAction(FGenericCommands::Get().SelectAll,
+					FExecuteAction::CreateLambda([this]() { MetasoundGraphEditor->SelectAllNodes(); }));
 
-		GraphEditorCommands->MapAction(FGenericCommands::Get().Cut,
-			FExecuteAction::CreateSP(this, &FMetasoundEditor::CutSelectedNodes),
-			FCanExecuteAction::CreateSP(this, &FMetasoundEditor::CanCutNodes));
+				GraphEditorCommands->MapAction(FGenericCommands::Get().Delete,
+					FExecuteAction::CreateSP(this, &FEditor::DeleteSelectedNodes),
+					FCanExecuteAction::CreateLambda([this]() { return MetasoundGraphEditor->GetSelectedNodes().Num() > 0; }));
 
-		GraphEditorCommands->MapAction(FGenericCommands::Get().Paste,
-			FExecuteAction::CreateSP(this, &FMetasoundEditor::PasteNodes),
-			FCanExecuteAction::CreateSP(this, &FMetasoundEditor::CanPasteNodes));
+				GraphEditorCommands->MapAction(FGenericCommands::Get().Copy,
+					FExecuteAction::CreateSP(this, &FEditor::CopySelectedNodes),
+					FCanExecuteAction::CreateSP(this, &FEditor::CanCopyNodes));
 
-		GraphEditorCommands->MapAction(FGenericCommands::Get().Duplicate,
-			FExecuteAction::CreateSP(this, &FMetasoundEditor::DuplicateNodes),
-			FCanExecuteAction::CreateSP(this, &FMetasoundEditor::CanDuplicateNodes));
+				GraphEditorCommands->MapAction(FGenericCommands::Get().Cut,
+					FExecuteAction::CreateSP(this, &FEditor::CutSelectedNodes),
+					FCanExecuteAction::CreateSP(this, &FEditor::CanCutNodes));
 
-		// Alignment Commands
-		GraphEditorCommands->MapAction(FGraphEditorCommands::Get().AlignNodesTop,
-			FExecuteAction::CreateSP(this, &FMetasoundEditor::OnAlignTop));
+				GraphEditorCommands->MapAction(FGenericCommands::Get().Paste,
+					FExecuteAction::CreateSP(this, &FEditor::PasteNodes),
+					FCanExecuteAction::CreateSP(this, &FEditor::CanPasteNodes));
 
-		GraphEditorCommands->MapAction(FGraphEditorCommands::Get().AlignNodesMiddle,
-			FExecuteAction::CreateSP(this, &FMetasoundEditor::OnAlignMiddle));
+				GraphEditorCommands->MapAction(FGenericCommands::Get().Duplicate,
+					FExecuteAction::CreateSP(this, &FEditor::DuplicateNodes),
+					FCanExecuteAction::CreateSP(this, &FEditor::CanDuplicateNodes));
 
-		GraphEditorCommands->MapAction(FGraphEditorCommands::Get().AlignNodesBottom,
-			FExecuteAction::CreateSP(this, &FMetasoundEditor::OnAlignBottom));
+				// Alignment Commands
+				GraphEditorCommands->MapAction(FGraphEditorCommands::Get().AlignNodesTop,
+					FExecuteAction::CreateLambda([this]() { MetasoundGraphEditor->OnAlignTop(); }));
 
-		GraphEditorCommands->MapAction(FGraphEditorCommands::Get().AlignNodesLeft,
-			FExecuteAction::CreateSP(this, &FMetasoundEditor::OnAlignLeft));
+				GraphEditorCommands->MapAction(FGraphEditorCommands::Get().AlignNodesMiddle,
+					FExecuteAction::CreateLambda([this]() { MetasoundGraphEditor->OnAlignMiddle(); }));
 
-		GraphEditorCommands->MapAction(FGraphEditorCommands::Get().AlignNodesCenter,
-			FExecuteAction::CreateSP(this, &FMetasoundEditor::OnAlignCenter));
+				GraphEditorCommands->MapAction(FGraphEditorCommands::Get().AlignNodesBottom,
+					FExecuteAction::CreateLambda([this]() { MetasoundGraphEditor->OnAlignBottom(); }));
 
-		GraphEditorCommands->MapAction(FGraphEditorCommands::Get().AlignNodesRight,
-			FExecuteAction::CreateSP(this, &FMetasoundEditor::OnAlignRight));
+				GraphEditorCommands->MapAction(FGraphEditorCommands::Get().AlignNodesLeft,
+					FExecuteAction::CreateLambda([this]() { MetasoundGraphEditor->OnAlignLeft(); }));
 
-		GraphEditorCommands->MapAction(FGraphEditorCommands::Get().StraightenConnections,
-			FExecuteAction::CreateSP(this, &FMetasoundEditor::OnStraightenConnections));
+				GraphEditorCommands->MapAction(FGraphEditorCommands::Get().AlignNodesCenter,
+					FExecuteAction::CreateLambda([this]() { MetasoundGraphEditor->OnAlignCenter(); }));
 
-		// Distribution Commands
-		GraphEditorCommands->MapAction(FGraphEditorCommands::Get().DistributeNodesHorizontally,
-			FExecuteAction::CreateSP(this, &FMetasoundEditor::OnDistributeNodesH));
+				GraphEditorCommands->MapAction(FGraphEditorCommands::Get().AlignNodesRight,
+					FExecuteAction::CreateLambda([this]() { MetasoundGraphEditor->OnAlignRight(); }));
 
-		GraphEditorCommands->MapAction(FGraphEditorCommands::Get().DistributeNodesVertically,
-			FExecuteAction::CreateSP(this, &FMetasoundEditor::OnDistributeNodesV));
-	}
+				GraphEditorCommands->MapAction(FGraphEditorCommands::Get().StraightenConnections,
+					FExecuteAction::CreateLambda([this]() { MetasoundGraphEditor->OnStraightenConnections(); }));
 
-	FGraphAppearanceInfo AppearanceInfo;
-	AppearanceInfo.CornerText = LOCTEXT("AppearanceCornerText_Metasound", "Metasound");
-
-	SGraphEditor::FGraphEditorEvents InEvents;
-	InEvents.OnSelectionChanged = SGraphEditor::FOnSelectionChanged::CreateSP(this, &FMetasoundEditor::OnSelectedNodesChanged);
-	InEvents.OnTextCommitted = FOnNodeTextCommitted::CreateSP(this, &FMetasoundEditor::OnNodeTitleCommitted);
-	InEvents.OnNodeDoubleClicked = FSingleNodeEvent::CreateSP(this, &FMetasoundEditor::PlaySingleNode);
-
-	return SNew(SGraphEditor)
-		.AdditionalCommands(GraphEditorCommands)
-		.IsEditable(true)
-		.Appearance(AppearanceInfo)
-		.GraphToEdit(Metasound->GetGraph())
-		.GraphEvents(InEvents)
-		.AutoExpandActionMenu(true)
-		.ShowGraphStateOverlay(false);
-}
+				// Distribution Commands
+				GraphEditorCommands->MapAction(FGraphEditorCommands::Get().DistributeNodesHorizontally,
+					FExecuteAction::CreateLambda([this]() { MetasoundGraphEditor->OnDistributeNodesH(); }));
 
-FGraphPanelSelectionSet FMetasoundEditor::GetSelectedNodes() const
-{
-	FGraphPanelSelectionSet CurrentSelection;
-	if (MetasoundGraphEditor.IsValid())
-	{
-		CurrentSelection = MetasoundGraphEditor->GetSelectedNodes();
-	}
-	return CurrentSelection;
-}
+				GraphEditorCommands->MapAction(FGraphEditorCommands::Get().DistributeNodesVertically,
+					FExecuteAction::CreateLambda([this]() { MetasoundGraphEditor->OnDistributeNodesV(); }));
+			}
 
-void FMetasoundEditor::OnSelectedNodesChanged(const TSet<UObject*>& NewSelection)
-{
-	TArray<UObject*> Selection;
+			FGraphAppearanceInfo AppearanceInfo;
+			AppearanceInfo.CornerText = LOCTEXT("AppearanceCornerText_Metasound", "Metasound");
+
+			SGraphEditor::FGraphEditorEvents InEvents;
+			InEvents.OnSelectionChanged = SGraphEditor::FOnSelectionChanged::CreateSP(this, &FEditor::OnSelectedNodesChanged);
+			InEvents.OnTextCommitted = FOnNodeTextCommitted::CreateSP(this, &FEditor::OnNodeTitleCommitted);
+			InEvents.OnNodeDoubleClicked = FSingleNodeEvent::CreateSP(this, &FEditor::PlaySingleNode);
+
+			return SNew(SGraphEditor)
+				.AdditionalCommands(GraphEditorCommands)
+				.IsEditable(true)
+				.Appearance(AppearanceInfo)
+				.GraphToEdit(Metasound->GetGraph())
+				.GraphEvents(InEvents)
+				.AutoExpandActionMenu(true)
+				.ShowGraphStateOverlay(false);
+		}
 
-	if(NewSelection.Num())
-	{
-		for(TSet<UObject*>::TConstIterator SetIt(NewSelection); SetIt; ++SetIt)
+		void FEditor::OnSelectedNodesChanged(const TSet<UObject*>& NewSelection)
 		{
-			if (Cast<UMetasoundEditorGraphNode>(*SetIt))
+			TArray<UObject*> Selection;
+
+			if (NewSelection.Num())
 			{
-				Selection.Add(GetMetasound());
+				for (TSet<UObject*>::TConstIterator SetIt(NewSelection); SetIt; ++SetIt)
+				{
+					if (Cast<UMetasoundEditorGraphNode>(*SetIt))
+					{
+						Selection.Add(GetMetasound());
+					}
+					else
+					{
+						Selection.Add(*SetIt);
+					}
+				}
 			}
 			else
 			{
-				Selection.Add(*SetIt);
+				Selection.Add(GetMetasound());
 			}
-		}
-	}
-	else
-	{
-		Selection.Add(GetMetasound());
-	}
-
-	SetSelection(Selection);
-}
-
-void FMetasoundEditor::OnNodeTitleCommitted(const FText& NewText, ETextCommit::Type CommitInfo, UEdGraphNode* NodeBeingChanged)
-{
-	if (NodeBeingChanged)
-	{
-		const FScopedTransaction Transaction(LOCTEXT("RenameNode", "Rename Node"));
-		NodeBeingChanged->Modify();
-		NodeBeingChanged->OnRenameNode(NewText.ToString());
-	}
-}
 
-void FMetasoundEditor::SelectAllNodes()
-{
-	MetasoundGraphEditor->SelectAllNodes();
-}
+			SetSelection(Selection);
+		}
 
-bool FMetasoundEditor::CanSelectAllNodes() const
-{
-	return true;
-}
+		void FEditor::OnNodeTitleCommitted(const FText& NewText, ETextCommit::Type CommitInfo, UEdGraphNode* NodeBeingChanged)
+		{
+			if (NodeBeingChanged)
+			{
+				const FScopedTransaction Transaction(LOCTEXT("RenameNode", "Rename Node"));
+				NodeBeingChanged->Modify();
+				NodeBeingChanged->OnRenameNode(NewText.ToString());
+			}
+		}
 
-void FMetasoundEditor::DeleteSelectedNodes()
-{
-	const FScopedTransaction Transaction(NSLOCTEXT("UnrealEd", "MetasoundEditorDeleteSelectedNode", "Delete Selected Metasound Node"));
+		void FEditor::DeleteSelectedNodes()
+		{
+			using namespace Metasound::Frontend;
 
-	MetasoundGraphEditor->GetCurrentGraph()->Modify();
+			const FScopedTransaction Transaction(LOCTEXT("MetasoundEditorDeleteSelectedNode", "Delete Selected Metasound Node(s)"));
 
-	const FGraphPanelSelectionSet SelectedNodes = GetSelectedNodes();
+			UMetasoundEditorGraph* Graph = CastChecked<UMetasoundEditorGraph>(MetasoundGraphEditor->GetCurrentGraph());
+			Graph->Modify();
 
-	MetasoundGraphEditor->ClearSelectionSet();
+			const FGraphPanelSelectionSet SelectedNodes = MetasoundGraphEditor->GetSelectedNodes();
+			MetasoundGraphEditor->ClearSelectionSet();
 
-	for (FGraphPanelSelectionSet::TConstIterator NodeIt(SelectedNodes); NodeIt; ++NodeIt)
-	{
-		UEdGraphNode* Node = CastChecked<UEdGraphNode>(*NodeIt);
-		if (Node->CanUserDeleteNode())
-		{
-			if (UMetasoundEditorGraphNode* MetasoundGraphNode = Cast<UMetasoundEditorGraphNode>(Node))
+			FGraphHandle GraphHandle = Graph->ParentMetasound->GetRootGraphHandle();
+			for (FGraphPanelSelectionSet::TConstIterator NodeIt(SelectedNodes); NodeIt; ++NodeIt)
 			{
-				// TODO: Implement delete selected here
+				UMetasoundEditorGraphNode* Node = CastChecked<UMetasoundEditorGraphNode>(*NodeIt);
+				FGraphBuilder::DeleteNode(*Node, false /* bInRecordTransaction */);
 			}
 		}
-	}
-}
-
-bool FMetasoundEditor::CanDeleteNodes() const
-{
-	const FGraphPanelSelectionSet SelectedNodes = GetSelectedNodes();
 
-	if (SelectedNodes.Num() == 1)
-	{
-		for (FGraphPanelSelectionSet::TConstIterator NodeIt(SelectedNodes); NodeIt; ++NodeIt)
+		void FEditor::DeleteSelectedDuplicatableNodes()
 		{
-			if (Cast<UMetasoundEditorGraphNode>(*NodeIt))
+			// Cache off the old selection
+			const FGraphPanelSelectionSet OldSelectedNodes = MetasoundGraphEditor->GetSelectedNodes();
+
+			// Clear the selection and only select the nodes that can be duplicated
+			FGraphPanelSelectionSet RemainingNodes;
+			MetasoundGraphEditor->ClearSelectionSet();
+
+			for (FGraphPanelSelectionSet::TConstIterator SelectedIter(OldSelectedNodes); SelectedIter; ++SelectedIter)
 			{
-				// Return false if only root node is selected, as it can't be deleted
-				return false;
+				UEdGraphNode* Node = Cast<UEdGraphNode>(*SelectedIter);
+				if ((Node != nullptr) && Node->CanDuplicateNode())
+				{
+					MetasoundGraphEditor->SetNodeSelection(Node, true);
+				}
+				else
+				{
+					RemainingNodes.Add(Node);
+				}
 			}
-		}
-	}
 
-	return SelectedNodes.Num() > 0;
-}
+			// Delete the duplicatable nodes
+			DeleteSelectedNodes();
 
-void FMetasoundEditor::DeleteSelectedDuplicatableNodes()
-{
-	// Cache off the old selection
-	const FGraphPanelSelectionSet OldSelectedNodes = GetSelectedNodes();
-
-	// Clear the selection and only select the nodes that can be duplicated
-	FGraphPanelSelectionSet RemainingNodes;
-	MetasoundGraphEditor->ClearSelectionSet();
+			// Reselect whatever's left from the original selection after the deletion
+			MetasoundGraphEditor->ClearSelectionSet();
 
-	for (FGraphPanelSelectionSet::TConstIterator SelectedIter(OldSelectedNodes); SelectedIter; ++SelectedIter)
-	{
-		UEdGraphNode* Node = Cast<UEdGraphNode>(*SelectedIter);
-		if ((Node != nullptr) && Node->CanDuplicateNode())
-		{
-			MetasoundGraphEditor->SetNodeSelection(Node, true);
+			for (FGraphPanelSelectionSet::TConstIterator SelectedIter(RemainingNodes); SelectedIter; ++SelectedIter)
+			{
+				if (UEdGraphNode* Node = Cast<UEdGraphNode>(*SelectedIter))
+				{
+					MetasoundGraphEditor->SetNodeSelection(Node, true);
+				}
+			}
 		}
-		else
+
+		void FEditor::CutSelectedNodes()
 		{
-			RemainingNodes.Add(Node);
+			CopySelectedNodes();
+			// Cut should only delete nodes that can be duplicated
+			DeleteSelectedDuplicatableNodes();
 		}
-	}
-
-	// Delete the duplicatable nodes
-	DeleteSelectedNodes();
 
-	// Reselect whatever's left from the original selection after the deletion
-	MetasoundGraphEditor->ClearSelectionSet();
-
-	for (FGraphPanelSelectionSet::TConstIterator SelectedIter(RemainingNodes); SelectedIter; ++SelectedIter)
-	{
-		if (UEdGraphNode* Node = Cast<UEdGraphNode>(*SelectedIter))
+		bool FEditor::CanCutNodes() const
 		{
-			MetasoundGraphEditor->SetNodeSelection(Node, true);
+			return CanCopyNodes();
 		}
-	}
-}
 
-void FMetasoundEditor::CutSelectedNodes()
-{
-	CopySelectedNodes();
-	// Cut should only delete nodes that can be duplicated
-	DeleteSelectedDuplicatableNodes();
-}
+		void FEditor::CopySelectedNodes()
+		{
+			// Export the selected nodes and place the text on the clipboard
+			const FGraphPanelSelectionSet SelectedNodes = MetasoundGraphEditor->GetSelectedNodes();
 
-bool FMetasoundEditor::CanCutNodes() const
-{
-	return CanCopyNodes() && CanDeleteNodes();
-}
+			FString ExportedText;
 
-void FMetasoundEditor::CopySelectedNodes()
-{
-	// Export the selected nodes and place the text on the clipboard
-	const FGraphPanelSelectionSet SelectedNodes = GetSelectedNodes();
+			for (FGraphPanelSelectionSet::TConstIterator SelectedIter(SelectedNodes); SelectedIter; ++SelectedIter)
+			{
+				if (UMetasoundEditorGraphNode* Node = Cast<UMetasoundEditorGraphNode>(*SelectedIter))
+				{
+					Node->PrepareForCopying();
+				}
+			}
 
-	FString ExportedText;
+			FEdGraphUtilities::ExportNodesToText(SelectedNodes, /*out*/ ExportedText);
+			FPlatformApplicationMisc::ClipboardCopy(*ExportedText);
 
-	for (FGraphPanelSelectionSet::TConstIterator SelectedIter(SelectedNodes); SelectedIter; ++SelectedIter)
-	{
-		if(UMetasoundEditorGraphNode* Node = Cast<UMetasoundEditorGraphNode>(*SelectedIter))
-		{
-			Node->PrepareForCopying();
+			// Make sure Metasound remains the owner of the copied nodes
+			for (FGraphPanelSelectionSet::TConstIterator SelectedIter(SelectedNodes); SelectedIter; ++SelectedIter)
+			{
+				if (UMetasoundEditorGraphNode* Node = Cast<UMetasoundEditorGraphNode>(*SelectedIter))
+				{
+					Node->PostCopyNode();
+				}
+			}
 		}
-	}
-
-	FEdGraphUtilities::ExportNodesToText(SelectedNodes, /*out*/ ExportedText);
-	FPlatformApplicationMisc::ClipboardCopy(*ExportedText);
 
-	// Make sure Metasound remains the owner of the copied nodes
-	for (FGraphPanelSelectionSet::TConstIterator SelectedIter(SelectedNodes); SelectedIter; ++SelectedIter)
-	{
-		if (UMetasoundEditorGraphNode* Node = Cast<UMetasoundEditorGraphNode>(*SelectedIter))
+		bool FEditor::CanCopyNodes() const
 		{
-			Node->PostCopyNode();
+			// If any of the nodes can be duplicated then allow copying
+			const FGraphPanelSelectionSet& SelectedNodes = MetasoundGraphEditor->GetSelectedNodes();
+			for (FGraphPanelSelectionSet::TConstIterator SelectedIter(SelectedNodes); SelectedIter; ++SelectedIter)
+			{
+				UEdGraphNode* Node = Cast<UEdGraphNode>(*SelectedIter);
+				if (Node && Node->CanDuplicateNode())
+				{
+					return true;
+				}
+			}
+			return false;
 		}
-	}
-}
 
-bool FMetasoundEditor::CanCopyNodes() const
-{
-	// If any of the nodes can be duplicated then we should allow copying
-	const FGraphPanelSelectionSet SelectedNodes = GetSelectedNodes();
-	for (FGraphPanelSelectionSet::TConstIterator SelectedIter(SelectedNodes); SelectedIter; ++SelectedIter)
-	{
-		UEdGraphNode* Node = Cast<UEdGraphNode>(*SelectedIter);
-		if ((Node != nullptr) && Node->CanDuplicateNode())
+		void FEditor::PasteNodes()
 		{
-			return true;
+			PasteNodesAtLocation(MetasoundGraphEditor->GetPasteLocation());
 		}
-	}
-	return false;
-}
-
-void FMetasoundEditor::PasteNodes()
-{
-	PasteNodesAtLocation(MetasoundGraphEditor->GetPasteLocation());
-}
-
-void FMetasoundEditor::PasteNodesAtLocation(const FVector2D& Location)
-{
-	UEdGraph* Graph = Metasound->GetGraph();
-	if (!Graph)
-	{
-		return;
-	}
-
-	// Undo/Redo support
-	const FScopedTransaction Transaction(NSLOCTEXT("UnrealEd", "MetasoundEditorPaste", "Paste Metasound Node"));
-	Graph->Modify();
-	Metasound->Modify();
-
-	// Clear the selection set (newly pasted stuff will be selected)
-	MetasoundGraphEditor->ClearSelectionSet();
-
-	// Grab the text to paste from the clipboard.
-	FString TextToImport;
-	FPlatformApplicationMisc::ClipboardPaste(TextToImport);
-
-	// Import the nodes
-	TSet<UEdGraphNode*> PastedNodes;
-	FEdGraphUtilities::ImportNodesFromText(Graph, TextToImport, /*out*/ PastedNodes);
-
-	//Average position of nodes so we can move them while still maintaining relative distances to each other
-	FVector2D AvgNodePosition(0.0f,0.0f);
-
-	for (TSet<UEdGraphNode*>::TIterator It(PastedNodes); It; ++It)
-	{
-		UEdGraphNode* Node = *It;
-		AvgNodePosition.X += Node->NodePosX;
-		AvgNodePosition.Y += Node->NodePosY;
-	}
-
-	if (PastedNodes.Num() > 0)
-	{
-		float InvNumNodes = 1.0f/float(PastedNodes.Num());
-		AvgNodePosition.X *= InvNumNodes;
-		AvgNodePosition.Y *= InvNumNodes;
-	}
-
-	for (TSet<UEdGraphNode*>::TIterator It(PastedNodes); It; ++It)
-	{
-		UEdGraphNode* Node = *It;
 
-		if (UMetasoundEditorGraphNode* SoundGraphNode = Cast<UMetasoundEditorGraphNode>(Node))
+		void FEditor::PasteNodesAtLocation(const FVector2D& Location)
 		{
-			// TODO: Add newly referenced nodes to MS reference list
-		}
+			UEdGraph* Graph = Metasound->GetGraph();
+			if (!Graph)
+			{
+				return;
+			}
 
-		// Select the newly pasted stuff
-		MetasoundGraphEditor->SetNodeSelection(Node, true);
+			// Undo/Redo support
+			const FScopedTransaction Transaction(NSLOCTEXT("UnrealEd", "MetasoundEditorPaste", "Paste Metasound Node"));
+			Graph->Modify();
+			Metasound->Modify();
 
-		Node->NodePosX = (Node->NodePosX - AvgNodePosition.X) + Location.X ;
-		Node->NodePosY = (Node->NodePosY - AvgNodePosition.Y) + Location.Y ;
+			// Clear the selection set (newly pasted stuff will be selected)
+			MetasoundGraphEditor->ClearSelectionSet();
 
-		Node->SnapToGrid(SNodePanel::GetSnapGridSize());
+			// Grab the text to paste from the clipboard.
+			FString TextToImport;
+			FPlatformApplicationMisc::ClipboardPaste(TextToImport);
 
-		// Give new node a different Guid from the old one
-		Node->CreateNewGuid();
-	}
+			// Import the nodes
+			TSet<UEdGraphNode*> PastedNodes;
+			FEdGraphUtilities::ImportNodesFromText(Graph, TextToImport, /*out*/ PastedNodes);
 
-	// TODO: Compile metasound here???
+			//Average position of nodes so we can move them while still maintaining relative distances to each other
+			FVector2D AvgNodePosition(0.0f, 0.0f);
 
-	// Update UI
-	MetasoundGraphEditor->NotifyGraphChanged();
+			for (TSet<UEdGraphNode*>::TIterator It(PastedNodes); It; ++It)
+			{
+				UEdGraphNode* Node = *It;
+				AvgNodePosition.X += Node->NodePosX;
+				AvgNodePosition.Y += Node->NodePosY;
+			}
 
-	Metasound->PostEditChange();
-	Metasound->MarkPackageDirty();
-}
+			if (PastedNodes.Num() > 0)
+			{
+				float InvNumNodes = 1.0f / float(PastedNodes.Num());
+				AvgNodePosition.X *= InvNumNodes;
+				AvgNodePosition.Y *= InvNumNodes;
+			}
 
-bool FMetasoundEditor::CanPasteNodes() const
-{
-	FString ClipboardContent;
-	FPlatformApplicationMisc::ClipboardPaste(ClipboardContent);
+			for (TSet<UEdGraphNode*>::TIterator It(PastedNodes); It; ++It)
+			{
+				UEdGraphNode* Node = *It;
 
-	// TODO: Implement
-// 	const bool bCanPasteNodes = FEdGraphUtilities::CanImportNodesFromText(Metasound->MetasoundGraph, ClipboardContent);
-// 	return bCanPasteNodes;
+				if (UMetasoundEditorGraphNode* SoundGraphNode = Cast<UMetasoundEditorGraphNode>(Node))
+				{
+					// TODO: Add newly referenced nodes to MS reference list
+				}
 
-	return false;
-}
+				// Select the newly pasted stuff
+				MetasoundGraphEditor->SetNodeSelection(Node, true);
 
-void FMetasoundEditor::DuplicateNodes()
-{
-	// Copy and paste current selection
-	CopySelectedNodes();
-	PasteNodes();
-}
+				Node->NodePosX = (Node->NodePosX - AvgNodePosition.X) + Location.X;
+				Node->NodePosY = (Node->NodePosY - AvgNodePosition.Y) + Location.Y;
 
-bool FMetasoundEditor::CanDuplicateNodes() const
-{
-	return CanCopyNodes();
-}
+				Node->SnapToGrid(SNodePanel::GetSnapGridSize());
 
-void FMetasoundEditor::UndoGraphAction()
-{
-	check(GEditor);
-	GEditor->UndoTransaction();
-}
+				// Give new node a different Guid from the old one
+				Node->CreateNewGuid();
+			}
 
-void FMetasoundEditor::RedoGraphAction()
-{
-	// Clear selection, to avoid holding refs to nodes that go away
-	MetasoundGraphEditor->ClearSelectionSet();
+			// TODO: Compile metasound here???
 
-	check(GEditor);
-	GEditor->RedoTransaction();
-}
+			// Update UI
+			MetasoundGraphEditor->NotifyGraphChanged();
 
-void FMetasoundEditor::OnAlignTop()
-{
-	if (MetasoundGraphEditor.IsValid())
-	{
-		MetasoundGraphEditor->OnAlignTop();
-	}
-}
+			Metasound->PostEditChange();
+			Metasound->MarkPackageDirty();
+		}
 
-void FMetasoundEditor::OnAlignMiddle()
-{
-	if (MetasoundGraphEditor.IsValid())
-	{
-		MetasoundGraphEditor->OnAlignMiddle();
-	}
-}
+		bool FEditor::CanPasteNodes() const
+		{
+			FString ClipboardContent;
+			FPlatformApplicationMisc::ClipboardPaste(ClipboardContent);
 
-void FMetasoundEditor::OnAlignBottom()
-{
-	if (MetasoundGraphEditor.IsValid())
-	{
-		MetasoundGraphEditor->OnAlignBottom();
-	}
-}
+			// TODO: Implement
+		// 	const bool bCanPasteNodes = FEdGraphUtilities::CanImportNodesFromText(Metasound->MetasoundGraph, ClipboardContent);
+		// 	return bCanPasteNodes;
 
-void FMetasoundEditor::OnAlignLeft()
-{
-	if (MetasoundGraphEditor.IsValid())
-	{
-		MetasoundGraphEditor->OnAlignLeft();
-	}
-}
+			return false;
+		}
 
-void FMetasoundEditor::OnAlignCenter()
-{
-	if (MetasoundGraphEditor.IsValid())
-	{
-		MetasoundGraphEditor->OnAlignCenter();
-	}
-}
+		void FEditor::DuplicateNodes()
+		{
+			// Copy and paste current selection
+			CopySelectedNodes();
+			PasteNodes();
+		}
 
-void FMetasoundEditor::OnAlignRight()
-{
-	if (MetasoundGraphEditor.IsValid())
-	{
-		MetasoundGraphEditor->OnAlignRight();
-	}
-}
+		bool FEditor::CanDuplicateNodes() const
+		{
+			return CanCopyNodes();
+		}
 
-void FMetasoundEditor::OnStraightenConnections()
-{
-	if (MetasoundGraphEditor.IsValid())
-	{
-		MetasoundGraphEditor->OnStraightenConnections();
-	}
-}
+		void FEditor::UndoGraphAction()
+		{
+			check(GEditor);
+			GEditor->UndoTransaction();
+		}
 
-void FMetasoundEditor::OnDistributeNodesH()
-{
-	if (MetasoundGraphEditor.IsValid())
-	{
-		MetasoundGraphEditor->OnDistributeNodesH();
-	}
-}
+		void FEditor::RedoGraphAction()
+		{
+			// Clear selection, to avoid holding refs to nodes that go away
+			MetasoundGraphEditor->ClearSelectionSet();
 
-void FMetasoundEditor::OnDistributeNodesV()
-{
-	if (MetasoundGraphEditor.IsValid())
-	{
-		MetasoundGraphEditor->OnDistributeNodesV();
+			check(GEditor);
+			GEditor->RedoTransaction();
+		}
 	}
 }
-
 #undef LOCTEXT_NAMESPACE
diff --git a/Engine/Plugins/Runtime/Metasound/Source/MetasoundEditor/Private/MetasoundEditor.h b/Engine/Plugins/Runtime/Metasound/Source/MetasoundEditor/Private/MetasoundEditor.h
--- a/Engine/Plugins/Runtime/Metasound/Source/MetasoundEditor/Private/MetasoundEditor.h
+++ b/Engine/Plugins/Runtime/Metasound/Source/MetasoundEditor/Private/MetasoundEditor.h
@@ -1,5 +1,4 @@
 // Copyright Epic Games, Inc. All Rights Reserved.
-
 #pragma once
 
 #include "CoreMinimal.h"
@@ -10,6 +9,7 @@
 #include "IMetasoundEditor.h"
 #include "Misc/NotifyHook.h"
 #include "SMetasoundPalette.h"
+#include "Textures/SlateIcon.h"
 #include "Toolkits/IToolkitHost.h"
 #include "UObject/GCObject.h"
 
@@ -27,192 +27,174 @@ class UMetasound;
 
 struct FPropertyChangedEvent;
 
-
-class FMetasoundEditor : public IMetasoundEditor, public FGCObject, public FNotifyHook, public FEditorUndoClient
+namespace Metasound
 {
-public:
-	FMetasoundEditor();
-
-	virtual void RegisterTabSpawners(const TSharedRef<FTabManager>& TabManager) override;
-	virtual void UnregisterTabSpawners(const TSharedRef<FTabManager>& TabManager) override;
-
-	/** Destructor */
-	virtual ~FMetasoundEditor();
+	namespace Editor
+	{
+		class FEditor : public IMetasoundEditor, public FGCObject, public FNotifyHook, public FEditorUndoClient
+		{
+		public:
+			virtual ~FEditor();
 
-	/** Edits the specified Metasound object */
-	void InitMetasoundEditor(const EToolkitMode::Type Mode, const TSharedPtr<IToolkitHost >& InitToolkitHost, UObject* ObjectToEdit);
+			virtual void RegisterTabSpawners(const TSharedRef<FTabManager>& TabManager) override;
+			virtual void UnregisterTabSpawners(const TSharedRef<FTabManager>& TabManager) override;
 
-	/** IMetasoundEditor interface */
-	virtual UMetasound* GetMetasound() const override;
-	virtual void SetSelection(TArray<UObject*> SelectedObjects) override;
-	virtual bool GetBoundsForSelectedNodes(FSlateRect& Rect, float Padding) override;
-	virtual int32 GetNumberOfSelectedNodes() const override;
-	virtual TSet<UObject*> GetSelectedNodes() const override;
+			/** Edits the specified Metasound object */
+			void InitMetasoundEditor(const EToolkitMode::Type Mode, const TSharedPtr<IToolkitHost>& InitToolkitHost, UObject* ObjectToEdit);
 
-	/** IToolkit interface */
-	virtual FName GetToolkitFName() const override;
-	virtual FText GetBaseToolkitName() const override;
-	virtual FString GetWorldCentricTabPrefix() const override;
-	virtual FLinearColor GetWorldCentricTabColorScale() const override;
+			/** IMetasoundEditor interface */
+			virtual UMetasound* GetMetasound() const override;
+			virtual void SetSelection(const TArray<UObject*>& SelectedObjects) override;
+			virtual bool GetBoundsForSelectedNodes(FSlateRect& Rect, float Padding) override;
 
-	virtual FString GetDocumentationLink() const override
-	{
-		return FString(TEXT("Engine/Audio/Metasounds/Editor"));
-	}
+			/** IToolkit interface */
+			virtual FName GetToolkitFName() const override;
+			virtual FText GetBaseToolkitName() const override;
+			virtual FString GetWorldCentricTabPrefix() const override;
+			virtual FLinearColor GetWorldCentricTabColorScale() const override;
 
-	/** FGCObject interface */
-	virtual void AddReferencedObjects(FReferenceCollector& Collector) override;
+			virtual FString GetDocumentationLink() const override
+			{
+				return FString(TEXT("Engine/Audio/Metasounds/Editor"));
+			}
 
-	//~ Begin FEditorUndoClient Interface
-	virtual void PostUndo(bool bSuccess) override;
-	virtual void PostRedo(bool bSuccess) override { PostUndo(bSuccess); }
-	// End of FEditorUndoClient
+			/** FGCObject interface */
+			virtual void AddReferencedObjects(FReferenceCollector& Collector) override;
 
-	/** Whether pasting the currently selected nodes is permissible */
-	bool CanPasteNodes() const;
+			//~ Begin FEditorUndoClient Interface
+			virtual void PostUndo(bool bSuccess) override;
+			virtual void PostRedo(bool bSuccess) override { PostUndo(bSuccess); }
+			// End of FEditorUndoClient
 
-	/** Paste the contents of the clipboard at the provided location */
-	void PasteNodesAtLocation(const FVector2D& Location);
+			/** Whether pasting the currently selected nodes is permissible */
+			bool CanPasteNodes() const;
 
-private:
-	TSharedRef<SDockTab> SpawnTab_GraphCanvas(const FSpawnTabArgs& Args);
-	TSharedRef<SDockTab> SpawnTab_Properties(const FSpawnTabArgs& Args);
-	TSharedRef<SDockTab> SpawnTab_Palette(const FSpawnTabArgs& Args);
+			/** Paste the contents of the clipboard at the provided location */
+			void PasteNodesAtLocation(const FVector2D& Location);
 
-protected:
-	/** Called when the selection changes in the GraphEditor */
-	void OnSelectedNodesChanged(const TSet<UObject*>& NewSelection);
+			int32 GetNumNodesSelected() const
+			{
+				return MetasoundGraphEditor->GetSelectedNodes().Num();
+			}
 
-	/**
-	 * Called when a node's title is committed for a rename
-	 *
-	 * @param	NewText				New title text
-	 * @param	CommitInfo			How text was committed
-	 * @param	NodeBeingChanged	The node being changed
-	 */
-	void OnNodeTitleCommitted(const FText& NewText, ETextCommit::Type CommitInfo, UEdGraphNode* NodeBeingChanged);
+		protected:
+			/** Called when the selection changes in the GraphEditor */
+			void OnSelectedNodesChanged(const TSet<UObject*>& NewSelection);
 
-	/** Select every node in the graph */
-	void SelectAllNodes();
-	/** Whether we can select every node */
-	bool CanSelectAllNodes() const;
+			/**
+			 * Called when a node's title is committed for a rename
+			 *
+			 * @param	NewText				New title text
+			 * @param	CommitInfo			How text was committed
+			 * @param	NodeBeingChanged	The node being changed
+			 */
+			void OnNodeTitleCommitted(const FText& NewText, ETextCommit::Type CommitInfo, UEdGraphNode* NodeBeingChanged);
 
-	/** Delete the currently selected nodes */
-	void DeleteSelectedNodes();
+			/** Delete the currently selected nodes */
+			void DeleteSelectedNodes();
 
-	/** Whether we are able to delete the currently selected nodes */
-	bool CanDeleteNodes() const;
+			/** Delete only the currently selected nodes that can be duplicated */
+			void DeleteSelectedDuplicatableNodes();
 
-	/** Delete only the currently selected nodes that can be duplicated */
-	void DeleteSelectedDuplicatableNodes();
+			/** Cut the currently selected nodes */
+			void CutSelectedNodes();
 
-	/** Cut the currently selected nodes */
-	void CutSelectedNodes();
+			/** Whether we are able to cut the currently selected nodes */
+			bool CanCutNodes() const;
 
-	/** Whether we are able to cut the currently selected nodes */
-	bool CanCutNodes() const;
+			/** Copy the currently selected nodes */
+			void CopySelectedNodes();
 
-	/** Copy the currently selected nodes */
-	void CopySelectedNodes();
+			/** Whether copying the currently selected nodes is permissible */
+			bool CanCopyNodes() const;
 
-	/** Whether copying the currently selected nodes is permissible */
-	bool CanCopyNodes() const;
+			/** Paste the contents of the clipboard */
+			void PasteNodes();
 
-	/** Paste the contents of the clipboard */
-	void PasteNodes();
+			/** Duplicate the currently selected nodes */
+			void DuplicateNodes();
 
-	/** Duplicate the currently selected nodes */
-	void DuplicateNodes();
+			/** Whether we are able to duplicate the currently selected nodes */
+			bool CanDuplicateNodes() const;
 
-	/** Whether we are able to duplicate the currently selected nodes */
-	bool CanDuplicateNodes() const;
+			/** Called to undo the last action */
+			void UndoGraphAction();
 
-	/** Called to undo the last action */
-	void UndoGraphAction();
+			/** Called to redo the last undone action */
+			void RedoGraphAction();
 
-	/** Called to redo the last undone action */
-	void RedoGraphAction();
+		private:
+			/** FNotifyHook interface */
+			virtual void NotifyPostChange(const FPropertyChangedEvent& PropertyChangedEvent, FProperty* PropertyThatChanged) override;
 
-	void OnAlignTop();
-	void OnAlignMiddle();
-	void OnAlignBottom();
-	void OnAlignLeft();
-	void OnAlignCenter();
-	void OnAlignRight();
+			/** Creates all internal widgets for the tabs to point at */
+			void CreateInternalWidgets();
 
-	void OnStraightenConnections();
+			/** Builds the toolbar widget for the Metasound editor */
+			void ExtendToolbar();
 
-	void OnDistributeNodesH();
-	void OnDistributeNodesV();
+			/** Binds new graph commands to delegates */
+			void BindGraphCommands();
 
-private:
-	/** FNotifyHook interface */
-	virtual void NotifyPostChange(const FPropertyChangedEvent& PropertyChangedEvent, FProperty* PropertyThatChanged) override;
+			FSlateIcon GetCompileStatusImage() const;
 
-	/** Creates all internal widgets for the tabs to point at */
-	void CreateInternalWidgets();
+			/** Toolbar command methods */
+			void Compile();
+			void Play();
+			void PlayNode();
+			void Stop();
 
-	/** Builds the toolbar widget for the Metasound editor */
-	void ExtendToolbar();
+			/** Whether we can play the current selection of nodes */
+			bool CanPlayNode() const;
 
-	/** Binds new graph commands to delegates */
-	void BindGraphCommands();
+			/** Either play the Metasound or stop currently playing sound */
+			void TogglePlayback();
 
-	/** Toolbar command methods */
-	void Play();
-	void PlayNode();
-	
+			/** Plays a single specified node */
+			void PlaySingleNode(UEdGraphNode* Node);
 
-	/** Whether we can play the current selection of nodes */
-	bool CanPlayNode() const;
-	void Stop();
+			/** Sync the content browser to the current selection of nodes */
+			void SyncInBrowser();
 
-	/** Either play the Metasound or stop currently playing sound */
-	void TogglePlayback();
-	
-	/** Plays a single specified node */
-	void PlaySingleNode(UEdGraphNode* Node);
+			/** Add an input to the currently selected node */
+			void AddInput();
+			/** Whether we can add an input to the currently selected node */
+			bool CanAddInput() const;
 
-	/** Sync the content browser to the current selection of nodes */
-	void SyncInBrowser();
+			/** Delete an input from the currently selected node */
+			void DeleteInput();
 
-	/** Add an input to the currently selected node */
-	void AddInput();
-	/** Whether we can add an input to the currently selected node */
-	bool CanAddInput() const;
+			/** Whether we can delete an input from the currently selected node */
+			bool CanDeleteInput() const;
 
-	/** Delete an input from the currently selected node */
-	void DeleteInput();
-	/** Whether we can delete an input from the currently selected node */
-	bool CanDeleteInput() const;
+			/* Create comment node on graph */
+			void OnCreateComment();
 
-	/* Create comment node on graph */
-	void OnCreateComment();
+			/** Create new graph editor widget */
+			TSharedRef<SGraphEditor> CreateGraphEditorWidget();
 
-	/** Create new graph editor widget */
-	TSharedRef<SGraphEditor> CreateGraphEditorWidget();
+		private:
 
-private:
-	/** The Metasound asset being inspected */
-	UMetasound* Metasound;
+			/** List of open tool panels; used to ensure only one exists at any one time */
+			TMap<FName, TWeakPtr<SDockableTab>> SpawnedToolPanels;
 
-	/** List of open tool panels; used to ensure only one exists at any one time */
-	TMap<FName, TWeakPtr<SDockableTab>> SpawnedToolPanels;
+			/** New Graph Editor */
+			TSharedPtr<SGraphEditor> MetasoundGraphEditor;
 
-	/** New Graph Editor */
-	TSharedPtr<SGraphEditor> MetasoundGraphEditor;
+			/** Properties tab */
+			TSharedPtr<IDetailsView> MetasoundProperties;
 
-	/** Properties tab */
-	TSharedPtr<IDetailsView> MetasoundProperties;
+			/** Palette of Sound Node types */
+			TSharedPtr<SMetasoundPalette> Palette;
 
-	/** Palette of Sound Node types */
-	TSharedPtr<SMetasoundPalette> Palette;
+			/** Command list for this editor */
+			TSharedPtr<FUICommandList> GraphEditorCommands;
 
-	/** Command list for this editor */
-	TSharedPtr<FUICommandList> GraphEditorCommands;
+			/** The Metasound asset being edited */
+			UMetasound* Metasound = nullptr;
 
-	/**	The tab ids for all the tabs used */
-	static const FName GraphCanvasTabId;
-	static const FName PropertiesTabId;
-	static const FName PaletteTabId;
-};
+			/** Whether or not metasound being edited is valid */
+			bool bPassedValidation = true;
+		};
+	} // namespace Editor
+} // namespace Metasound
diff --git a/Engine/Plugins/Runtime/Metasound/Source/MetasoundEditor/Private/MetasoundEditorCommands.cpp b/Engine/Plugins/Runtime/Metasound/Source/MetasoundEditor/Private/MetasoundEditorCommands.cpp
--- a/Engine/Plugins/Runtime/Metasound/Source/MetasoundEditor/Private/MetasoundEditorCommands.cpp
+++ b/Engine/Plugins/Runtime/Metasound/Source/MetasoundEditor/Private/MetasoundEditorCommands.cpp
@@ -3,14 +3,22 @@
 
 #define LOCTEXT_NAMESPACE "MetasoundEditorCommands"
 
-void FMetasoundEditorCommands::RegisterCommands()
+namespace Metasound
 {
-	UI_COMMAND(Play, "Play", "Plays (or restarts) the Metasound", EUserInterfaceActionType::Button, FInputChord());
-	UI_COMMAND(Stop, "Stop", "Stops Metasound (If currently playing)", EUserInterfaceActionType::Button, FInputChord());
-	UI_COMMAND(TogglePlayback, "Toggle Playback", "Plays the Metasound or stops the currently playing Metasound", EUserInterfaceActionType::Button, FInputChord(EKeys::SpaceBar));
+	namespace Editor
+	{
+		void FEditorCommands::RegisterCommands()
+		{
+			UI_COMMAND(Play, "Play", "Plays (or restarts) the Metasound", EUserInterfaceActionType::Button, FInputChord());
+			UI_COMMAND(Stop, "Stop", "Stops Metasound (If currently playing)", EUserInterfaceActionType::Button, FInputChord());
+			UI_COMMAND(TogglePlayback, "Toggle Playback", "Plays the Metasound or stops the currently playing Metasound", EUserInterfaceActionType::Button, FInputChord(EKeys::SpaceBar));
 
-	UI_COMMAND(BrowserSync, "Sync to Browser", "Selects the SoundWave in the content browser", EUserInterfaceActionType::Button, FInputChord());
-	UI_COMMAND(AddInput, "Add Input", "Adds an input to the node", EUserInterfaceActionType::Button, FInputChord());
-	UI_COMMAND(DeleteInput, "Delete Input", "Removes an input from the node", EUserInterfaceActionType::Button, FInputChord());
-}
+			UI_COMMAND(Compile, "Compile", "Compiles Metasound", EUserInterfaceActionType::Button, FInputChord());
+
+			UI_COMMAND(BrowserSync, "Browse", "Selects the Metasound in the content browser. If referencing Metasound nodes are selected, selects referenced assets instead.", EUserInterfaceActionType::Button, FInputChord());
+			UI_COMMAND(AddInput, "Add Input", "Adds an input to the node", EUserInterfaceActionType::Button, FInputChord());
+			UI_COMMAND(DeleteInput, "Delete Input", "Removes an input from the node", EUserInterfaceActionType::Button, FInputChord());
+		}
+	} // namespace Editor
+} // namespace Metasound
 #undef LOCTEXT_NAMESPACE
diff --git a/Engine/Plugins/Runtime/Metasound/Source/MetasoundEditor/Private/MetasoundEditorGraph.cpp b/Engine/Plugins/Runtime/Metasound/Source/MetasoundEditor/Private/MetasoundEditorGraph.cpp
new file mode 100644
--- /dev/null
+++ b/Engine/Plugins/Runtime/Metasound/Source/MetasoundEditor/Private/MetasoundEditorGraph.cpp
@@ -0,0 +1,16 @@
+// Copyright Epic Games, Inc. All Rights Reserved.
+#include "MetasoundEditorGraph.h"
+
+#include "EdGraph/EdGraphNode.h"
+
+
+UMetasound* UMetasoundEditorGraph::GetMetasound() const
+{
+	return ParentMetasound;
+}
+
+UMetasound& UMetasoundEditorGraph::GetMetasoundChecked() const
+{
+	check(ParentMetasound);
+	return *ParentMetasound;
+}
diff --git a/Engine/Plugins/Runtime/Metasound/Source/MetasoundEditor/Private/MetasoundEditorGraphBuilder.cpp b/Engine/Plugins/Runtime/Metasound/Source/MetasoundEditor/Private/MetasoundEditorGraphBuilder.cpp
new file mode 100644
--- /dev/null
+++ b/Engine/Plugins/Runtime/Metasound/Source/MetasoundEditor/Private/MetasoundEditorGraphBuilder.cpp
@@ -0,0 +1,212 @@
+// Copyright Epic Games, Inc. All Rights Reserved.
+#include "MetasoundEditorGraphBuilder.h"
+
+#include "EdGraph/EdGraph.h"
+#include "EdGraph/EdGraphPin.h"
+#include "Metasound.h"
+#include "MetasoundEditorGraph.h"
+#include "MetasoundEditorGraphNode.h"
+#include "MetasoundFrontend.h"
+#include "ScopedTransaction.h"
+#include "Templates/Tuple.h"
+#include "MetasoundEditorModule.h"
+#include "Modules/ModuleManager.h"
+
+
+#define LOCTEXT_NAMESPACE "MetasoundEditor"
+
+
+namespace Metasound
+{
+	namespace Editor
+	{
+		const FName FGraphBuilder::PinPrimitiveBoolean = "Boolean";
+		const FName FGraphBuilder::PinPrimitiveFloat = "Float";
+		const FName FGraphBuilder::PinPrimitiveInteger = "Int";
+		const FName FGraphBuilder::PinPrimitiveString = "String";
+
+		UEdGraphNode* FGraphBuilder::AddNode(UMetasound& InMetasound, const FVector2D& Location, Frontend::FNodeHandle& InNodeHandle, bool bInSelectNewNode)
+		{
+			const FScopedTransaction Transaction(LOCTEXT("AddMetasoundGraphNode", "Add Metasound Node"));
+
+			const FString EdNodeName = InNodeHandle.GetNodeClassName() + TEXT("_") + FString::FromInt(InNodeHandle.GetNodeID());
+
+			UEdGraph& Graph = InMetasound.GetGraphChecked();
+			FGraphNodeCreator<UMetasoundEditorGraphNode> NodeCreator(Graph);
+			UMetasoundEditorGraphNode* NewGraphNode = NodeCreator.CreateNode(bInSelectNewNode);
+			NodeCreator.Finalize();
+
+			NewGraphNode->SetNodeID(InNodeHandle.GetNodeID());
+			NewGraphNode->CreateNewGuid();
+			NewGraphNode->NodePosX = Location.X;
+			NewGraphNode->NodePosY = Location.Y;
+
+			RebuildNodePins(*NewGraphNode, InNodeHandle);
+
+			InMetasound.PostEditChange();
+			InMetasound.MarkPackageDirty();
+
+			return NewGraphNode;
+		}
+
+		UEdGraphNode* FGraphBuilder::AddNode(UMetasound& InMetasound, const FVector2D& Location, const Frontend::FNodeClassInfo& InClassInfo, bool bInSelectNewNode)
+		{
+			Frontend::FNodeHandle NodeHandle = InMetasound.GetRootGraphHandle().AddNewNode(InClassInfo);
+			return AddNode(InMetasound, Location, NodeHandle, bInSelectNewNode);
+		}
+
+		UEdGraphNode* FGraphBuilder::AddInput(UMetasound& InMetasound, const FVector2D& Location, const FString& InName, const FName InTypeName, const FText& InToolTip, bool bInSelectNewNode)
+		{
+			FMetasoundInputDescription Description;
+			Description.Name = InName;
+			Description.TypeName = InTypeName;
+			Description.ToolTip = InToolTip;
+
+			Frontend::FNodeHandle NodeHandle = InMetasound.GetRootGraphHandle().AddNewInput(Description);
+			return AddNode(InMetasound, Location, NodeHandle, bInSelectNewNode);
+		}
+
+		UEdGraphNode* FGraphBuilder::AddOutput(UMetasound& InMetasound, const FVector2D& Location, const FString& InName, const FName InTypeName, const FText& InToolTip, bool bInSelectNewNode)
+		{
+			FMetasoundOutputDescription Description;
+			Description.Name = InName;
+			Description.TypeName = InTypeName;
+			Description.ToolTip = InToolTip;
+
+			Frontend::FNodeHandle NodeHandle = InMetasound.GetRootGraphHandle().AddNewOutput(Description);
+			return AddNode(InMetasound, Location, NodeHandle, bInSelectNewNode);
+		}
+
+		void FGraphBuilder::DeleteNode(UMetasoundEditorGraphNode& InNode, bool bInRecordTransaction)
+		{
+			const FScopedTransaction Transaction(LOCTEXT("DeleteMetasoundGraphNode", "Delete Metasound Node"), bInRecordTransaction);
+
+			UMetasoundEditorGraph* Graph = CastChecked<UMetasoundEditorGraph>(InNode.GetGraph());
+			if (InNode.CanUserDeleteNode() && Graph->RemoveNode(&InNode))
+			{
+				Graph->PostEditChange();
+				Graph->MarkPackageDirty();
+			}
+
+			const Frontend::FNodeHandle& NodeHandle = InNode.GetNodeHandle();
+			Frontend::FGraphHandle GraphHandle = Graph->GetMetasoundChecked().GetRootGraphHandle();
+			if (GraphHandle.IsValid() && NodeHandle.IsValid())
+			{
+				GraphHandle.RemoveNode(NodeHandle, true /* bEvenIfInputOrOutputNode */);
+			}
+
+			InNode.PostEditChange();
+			InNode.MarkPackageDirty();
+		}
+
+		void FGraphBuilder::RebuildGraph(UMetasound& InMetasound)
+		{
+			using namespace Frontend;
+
+			FGraphHandle GraphHandle = InMetasound.GetRootGraphHandle();
+			UMetasoundEditorGraph* Graph = CastChecked<UMetasoundEditorGraph>(InMetasound.GetGraph());
+
+			Graph->Nodes.Reset();
+
+			// TODO: Space graph nodes in a readable way
+			FVector2D Location = FVector2D::ZeroVector;
+
+			struct FNodePair
+			{
+				FNodeHandle NodeHandle = FNodeHandle::InvalidHandle();
+				UEdGraphNode* GraphNode = nullptr;
+			};
+
+			TMap<uint32, FNodePair> NewIdNodeMap;
+			TArray<FNodeHandle> NodeHandles = GraphHandle.GetAllNodes();
+			for (FNodeHandle& NodeHandle : NodeHandles)
+			{
+				const FNodeClassInfo ClassInfo = NodeHandle.GetClassInfo();
+				UEdGraphNode* NewNode = AddNode(InMetasound, Location, NodeHandle, false /* bInSelectNewNode */);
+				NewIdNodeMap.Add(NodeHandle.GetNodeID(), FNodePair { NodeHandle, NewNode });
+			}
+
+			for (const TPair<uint32, FNodePair>& IdNodePair : NewIdNodeMap)
+			{
+				UEdGraphNode* GraphNode = IdNodePair.Value.GraphNode;
+				check(GraphNode);
+
+				FNodeHandle NodeHandle = IdNodePair.Value.NodeHandle;
+				TArray<UEdGraphPin*> Pins = GraphNode->GetAllPins();
+
+				const TArray<FInputHandle> NodeInputs = NodeHandle.GetAllInputs();
+
+				int32 InputIndex = 0;
+				for (UEdGraphPin* Pin : Pins)
+				{
+					switch (Pin->Direction)
+					{
+						case EEdGraphPinDirection::EGPD_Input:
+						{
+							FOutputHandle OutputHandle = NodeInputs[InputIndex].GetCurrentlyConnectedOutput();
+							if (OutputHandle.IsValid())
+							{
+								UEdGraphNode* OutputGraphNode = NewIdNodeMap.FindChecked(OutputHandle.GetOwningNodeID()).GraphNode;
+								UEdGraphPin* OutputPin = OutputGraphNode->FindPinChecked(OutputHandle.GetOutputName(), EEdGraphPinDirection::EGPD_Output);
+								Pin->MakeLinkTo(OutputPin);
+							}
+
+							InputIndex++;
+						}
+						break;
+
+						case EEdGraphPinDirection::EGPD_Output:
+							// Do nothing.  Connecting all inputs will naturally connect all outputs where required
+						break;
+					}
+				}
+			}
+
+			InMetasound.PostEditChange();
+			InMetasound.MarkPackageDirty();
+		}
+
+		void FGraphBuilder::RebuildNodePins(UMetasoundEditorGraphNode& InGraphNode, Frontend::FNodeHandle& InNodeHandle, bool bInRecordTransaction)
+		{
+			const FScopedTransaction Transaction(LOCTEXT("RebuildMetasoundGraphNodePins", "Rebuild Metasound Pins"), bInRecordTransaction);
+
+			InGraphNode.Pins.Reset();
+
+			IMetasoundEditorModule& EditorModule = FModuleManager::GetModuleChecked<IMetasoundEditorModule>("MetasoundEditor");
+
+			TArray<Frontend::FInputHandle> InputHandles = InNodeHandle.GetAllInputs();
+			for (int32 i = 0; i < InputHandles.Num(); ++i)
+			{
+				const Frontend::FInputHandle& InputHandle = InputHandles[i];
+				FEdGraphPinType PinType("MetasoundNode", NAME_None, nullptr, EPinContainerType::None, false, FEdGraphTerminalType());
+				UEdGraphPin* NewPin = InGraphNode.CreatePin(EGPD_Input, PinType, *InputHandle.GetInputName(), i);
+				if (ensureAlways(NewPin))
+				{
+					NewPin->PinToolTip = InputHandle.GetInputTooltip().ToString();
+
+					FEditorDataType DataType = EditorModule.FindDataType(InputHandle.GetInputType());
+					NewPin->PinType = DataType.PinType;
+				}
+			}
+
+			TArray<Frontend::FOutputHandle> OutputHandles = InNodeHandle.GetAllOutputs();
+			for (int32 i = 0; i < OutputHandles.Num(); ++i)
+			{
+				const Frontend::FOutputHandle& OutputHandle = OutputHandles[i];
+				FEdGraphPinType PinType("MetasoundNode", NAME_None, nullptr, EPinContainerType::None, false, FEdGraphTerminalType());
+				UEdGraphPin* NewPin = InGraphNode.CreatePin(EGPD_Output, PinType, *OutputHandle.GetOutputName(), i);
+				if (ensureAlways(NewPin))
+				{
+					NewPin->PinToolTip = OutputHandle.GetOutputTooltip().ToString();
+
+					const FName OutputType = OutputHandle.GetOutputType();
+					FEditorDataType DataType = EditorModule.FindDataType(OutputType);
+					NewPin->PinType = DataType.PinType;
+				}
+			}
+
+			InGraphNode.MarkPackageDirty();
+		}
+	} // namespace Editor
+} // namespace Metasound
+#undef LOCTEXT_NAMESPACE
\ No newline at end of file
diff --git a/Engine/Plugins/Runtime/Metasound/Source/MetasoundEditor/Private/MetasoundEditorGraphBuilder.h b/Engine/Plugins/Runtime/Metasound/Source/MetasoundEditor/Private/MetasoundEditorGraphBuilder.h
new file mode 100644
--- /dev/null
+++ b/Engine/Plugins/Runtime/Metasound/Source/MetasoundEditor/Private/MetasoundEditorGraphBuilder.h
@@ -0,0 +1,46 @@
+// Copyright Epic Games, Inc. All Rights Reserved.
+#pragma once
+
+#include "CoreMinimal.h"
+
+// Forward Declarations
+class UMetasound;
+class UMetasoundEditorGraphNode;
+
+namespace Metasound
+{
+	namespace Frontend
+	{
+		// Forward Declarations
+		class FNodeHandle;
+		struct FNodeClassInfo;
+	} // namespace Frontend
+
+	namespace Editor
+	{
+		class FGraphBuilder
+		{
+		public:
+			static const FName PinPrimitiveBoolean;
+			static const FName PinPrimitiveFloat;
+			static const FName PinPrimitiveInteger;
+			static const FName PinPrimitiveString;
+
+			// Adds a node to the editor graph that corresponds to the provided node handle.
+			static UEdGraphNode* AddNode(UMetasound& InMetasound, const FVector2D& Location, Frontend::FNodeHandle& InNodeHandle, bool bInSelectNewNode = true);
+
+			// Adds a node with the given class info to both the editor and document graphs
+			static UEdGraphNode* AddNode(UMetasound& InMetasound, const FVector2D& Location, const Frontend::FNodeClassInfo& InClassInfo, bool bInSelectNewNode = true);
+
+			static UEdGraphNode* AddInput(UMetasound& InMetasound, const FVector2D& Location, const FString& InName, const FName InTypeName, const FText& InToolTip, bool bInSelectNewNode = true);
+
+			static UEdGraphNode* AddOutput(UMetasound& InMetasound, const FVector2D& Location, const FString& InName, const FName InTypeName, const FText& InToolTip, bool bInSelectNewNode = true);
+
+			static void DeleteNode(UMetasoundEditorGraphNode& InNode, bool bInRecordTransaction = true);
+
+			static void RebuildGraph(UMetasound& InMetasound);
+
+			static void RebuildNodePins(UMetasoundEditorGraphNode& InGraphNode, Frontend::FNodeHandle& InNodeHandle, bool bInRecordTransaction = true);
+		};
+	} // namespace Editor
+} // namespace Metasound
diff --git a/Engine/Plugins/Runtime/Metasound/Source/MetasoundEditor/Private/MetasoundEditorGraphNode.cpp b/Engine/Plugins/Runtime/Metasound/Source/MetasoundEditor/Private/MetasoundEditorGraphNode.cpp
--- a/Engine/Plugins/Runtime/Metasound/Source/MetasoundEditor/Private/MetasoundEditorGraphNode.cpp
+++ b/Engine/Plugins/Runtime/Metasound/Source/MetasoundEditor/Private/MetasoundEditorGraphNode.cpp
@@ -2,17 +2,20 @@
 
 #include "MetasoundEditorGraphNode.h"
 
+#include "EdGraph/EdGraphPin.h"
 #include "Editor/EditorEngine.h"
 #include "Engine/Font.h"
 #include "Framework/Commands/GenericCommands.h"
 #include "GraphEditorActions.h"
 #include "Metasound.h"
 #include "MetasoundEditorGraph.h"
+#include "MetasoundEditorGraphBuilder.h"
+#include "MetasoundEditorGraphSchema.h"
 #include "MetasoundEditorCommands.h"
 #include "ScopedTransaction.h"
 #include "ToolMenus.h"
 
-#define LOCTEXT_NAMESPACE "MetasoundEditorGraphNode"
+#define LOCTEXT_NAMESPACE "MetasoundEditor"
 
 
 UMetasoundEditorGraphNode::UMetasoundEditorGraphNode(const FObjectInitializer& ObjectInitializer)
@@ -55,68 +58,115 @@ void UMetasoundEditorGraphNode::CreateInputPin()
 	}
 }
 
-void UMetasoundEditorGraphNode::AddInputPin()
+int32 UMetasoundEditorGraphNode::EstimateNodeWidth() const
 {
-	const FScopedTransaction Transaction( NSLOCTEXT("UnrealEd", "MetasoundEditorAddInput", "Add Metasound Input") );
-	Modify();
-	CreateInputPin();
+	const int32 EstimatedCharWidth = 6;
+	FString NodeTitle = GetNodeTitle(ENodeTitleType::FullTitle).ToString();
+	UFont* Font = GetDefault<UEditorEngine>()->EditorFont;
+	int32 Result = NodeTitle.Len()*EstimatedCharWidth;
 
-// 	UMetasound* Metasound = CastChecked<UMetasoundEditorGraph>(GetGraph())->Metasound;
-	// TODO: Add metasound input pin
+	if (Font)
+	{
+		Result = Font->GetStringSize(*NodeTitle);
+	}
 
-	// Refresh the current graph, so the pins can be updated
-	GetGraph()->NotifyGraphChanged();
+	return Result;
 }
 
-void UMetasoundEditorGraphNode::RemoveInputPin(UEdGraphPin* InGraphPin)
+Metasound::Frontend::FNodeHandle UMetasoundEditorGraphNode::GetNodeHandle() const
 {
-	const FScopedTransaction Transaction( NSLOCTEXT("UnrealEd", "SoundCueEditorDeleteInput", "Delete Metasound Input") );
-	Modify();
-
-	TArray<UEdGraphPin*> InputPins;
-	GetInputPins(InputPins);
+	UMetasoundEditorGraph* EdGraph = CastChecked<UMetasoundEditorGraph>(GetGraph());
+	UMetasound& Metasound = EdGraph->GetMetasoundChecked();
+	return Metasound.GetRootGraphHandle().GetNodeWithId(NodeID);
+}
 
-	for (int32 InputIndex = 0; InputIndex < InputPins.Num(); InputIndex++)
+void UMetasoundEditorGraphNode::IteratePins(TUniqueFunction<void(UEdGraphPin* /* Pin */, int32 /* Index */)> InFunc, EEdGraphPinDirection InPinDirection)
+{
+	for (int32 PinIndex = 0; PinIndex < Pins.Num(); PinIndex++)
 	{
-		if (InGraphPin == InputPins[InputIndex])
+		if (InPinDirection == EGPD_MAX || Pins[PinIndex]->Direction == InPinDirection)
 		{
-			InGraphPin->MarkPendingKill();
-			Pins.Remove(InGraphPin);
-
-			// Remove pin from UMetasound
-			break;
+			InFunc(Pins[PinIndex], PinIndex);
 		}
 	}
+}
+
+void UMetasoundEditorGraphNode::AllocateDefaultPins()
+{
+	ensureAlways(Pins.Num() == 0);
+
+	using namespace Metasound;
 
-	// Refresh the current graph, so the pins can be updated
-	GetGraph()->NotifyGraphChanged();
+	Frontend::FNodeHandle NodeHandle = Frontend::FNodeHandle::InvalidHandle();
+	Editor::FGraphBuilder::RebuildNodePins(*this, NodeHandle, false /* bInRecordTransaction */);
 }
 
-int32 UMetasoundEditorGraphNode::EstimateNodeWidth() const
+void UMetasoundEditorGraphNode::ReconstructNode()
 {
-	const int32 EstimatedCharWidth = 6;
-	FString NodeTitle = GetNodeTitle(ENodeTitleType::FullTitle).ToString();
-	UFont* Font = GetDefault<UEditorEngine>()->EditorFont;
-	int32 Result = NodeTitle.Len()*EstimatedCharWidth;
+	using namespace Metasound;
 
-	if (Font)
+	Frontend::FNodeHandle NodeHandle = Frontend::FNodeHandle::InvalidHandle();
+	Editor::FGraphBuilder::RebuildNodePins(*this, NodeHandle);
+}
+
+void UMetasoundEditorGraphNode::AutowireNewNode(UEdGraphPin* FromPin)
+{
+	if (FromPin)
 	{
-		Result = Font->GetStringSize(*NodeTitle);
+		const UMetasoundEditorGraphSchema* Schema = CastChecked<UMetasoundEditorGraphSchema>(GetSchema());
+
+		TSet<UEdGraphNode*> NodeList;
+
+		// auto-connect from dragged pin to first compatible pin on the new node
+		for (int32 i = 0; i < Pins.Num(); i++)
+		{
+			UEdGraphPin* Pin = Pins[i];
+			check(Pin);
+			FPinConnectionResponse Response = Schema->CanCreateConnection(FromPin, Pin);
+			if (ECanCreateConnectionResponse::CONNECT_RESPONSE_MAKE == Response.Response)
+			{
+				if (Schema->TryCreateConnection(FromPin, Pin))
+				{
+					NodeList.Add(FromPin->GetOwningNode());
+					NodeList.Add(this);
+				}
+				break;
+			}
+			else if (ECanCreateConnectionResponse::CONNECT_RESPONSE_BREAK_OTHERS_A == Response.Response)
+			{
+				// TODO: Implement default connections in GraphBuilder
+				break;
+			}
+		}
+
+		// Send all nodes that received a new pin connection a notification
+		for (auto It = NodeList.CreateConstIterator(); It; ++It)
+		{
+			UEdGraphNode* Node = (*It);
+			Node->NodeConnectionListChanged();
+		}
 	}
+}
 
-	return Result;
+bool UMetasoundEditorGraphNode::CanCreateUnderSpecifiedSchema(const UEdGraphSchema* Schema) const
+{
+	return Schema->IsA(UMetasoundEditorGraphSchema::StaticClass());
 }
 
-bool UMetasoundEditorGraphNode::CanAddInputPin() const
+FString UMetasoundEditorGraphNode::GetDocumentationLink() const
 {
-	// Get from UMetasound
-	return false;
+	return TEXT("Shared/GraphNodes/Metasound");
+}
+
+void UMetasoundEditorGraphNode::SetNodeID(uint32 InNodeID)
+{
+	NodeID = InNodeID;
 }
 
 FText UMetasoundEditorGraphNode::GetNodeTitle(ENodeTitleType::Type TitleType) const
 {
-	// Get from UMetasound
-	return FText::GetEmpty();
+	const Metasound::Frontend::FNodeHandle NodeHandle = GetNodeHandle();
+	return FText::FromString(NodeHandle.GetNodeClassName());
 }
 
 void UMetasoundEditorGraphNode::PrepareForCopying()
@@ -141,20 +191,17 @@ void UMetasoundEditorGraphNode::PostDuplicate(bool bDuplicateForPIE)
 	}
 }
 
-void UMetasoundEditorGraphNode::CreateInputPins()
-{
-	// Get from UMetasound... archetype?
-}
-
 void UMetasoundEditorGraphNode::GetNodeContextMenuActions(UToolMenu* Menu, UGraphNodeContextMenuContext* Context) const
 {
+	using namespace Metasound::Editor;
+
 	if (Context->Pin)
 	{
 		// If on an input that can be deleted, show option
 		if (Context->Pin->Direction == EGPD_Input /*&& SoundNode->ChildNodes.Num() > SoundNode->GetMinChildNodes()*/)
 		{
 			FToolMenuSection& Section = Menu->AddSection("MetasoundEditorGraphDeleteInput");
-			Section.AddMenuEntry(FMetasoundEditorCommands::Get().DeleteInput);
+			Section.AddMenuEntry(FEditorCommands::Get().DeleteInput);
 		}
 	}
 	else if (Context->Node)
@@ -189,14 +236,6 @@ void UMetasoundEditorGraphNode::GetNodeContextMenuActions(UToolMenu* Menu, UGrap
 			Section.AddMenuEntry(FGenericCommands::Get().Copy);
 			Section.AddMenuEntry(FGenericCommands::Get().Duplicate);
 		}
-
-		{
-			FToolMenuSection& Section = Menu->AddSection("MetasoundEditorGraphNodeAddPlaySync");
-			if (CanAddInputPin())
-			{
-				Section.AddMenuEntry(FMetasoundEditorCommands::Get().AddInput);
-			}
-		}
 	}
 }
 
diff --git a/Engine/Plugins/Runtime/Metasound/Source/MetasoundEditor/Private/MetasoundEditorGraphNode_Base.cpp b/Engine/Plugins/Runtime/Metasound/Source/MetasoundEditor/Private/MetasoundEditorGraphNode_Base.cpp
deleted file mode 100644
diff --git a/Engine/Plugins/Runtime/Metasound/Source/MetasoundEditor/Private/MetasoundEditorGraphNode_Root.cpp b/Engine/Plugins/Runtime/Metasound/Source/MetasoundEditor/Private/MetasoundEditorGraphNode_Root.cpp
deleted file mode 100644
diff --git a/Engine/Plugins/Runtime/Metasound/Source/MetasoundEditor/Private/MetasoundEditorGraphSchema.cpp b/Engine/Plugins/Runtime/Metasound/Source/MetasoundEditor/Private/MetasoundEditorGraphSchema.cpp
--- a/Engine/Plugins/Runtime/Metasound/Source/MetasoundEditor/Private/MetasoundEditorGraphSchema.cpp
+++ b/Engine/Plugins/Runtime/Metasound/Source/MetasoundEditor/Private/MetasoundEditorGraphSchema.cpp
@@ -9,62 +9,217 @@
 #include "Framework/MultiBox/MultiBoxBuilder.h"
 #include "GraphEditor.h"
 #include "GraphEditorActions.h"
+#include "GraphEditorSettings.h"
 #include "Layout/SlateRect.h"
 #include "Metasound.h"
 #include "MetasoundEditor.h"
 #include "MetasoundEditorGraph.h"
+#include "MetasoundEditorGraphBuilder.h"
 #include "MetasoundEditorGraphNode.h"
-#include "MetasoundEditorGraphNode_Root.h"
-#include "MetasoundEditorUtilities.h"
+#include "MetasoundEditorModule.h"
+#include "MetasoundFrontend.h"
 #include "ScopedTransaction.h"
+#include "Toolkits/ToolkitManager.h"
 #include "ToolMenus.h"
+#include "UObject/NoExportTypes.h"
 #include "UObject/UObjectHash.h"
 #include "UObject/UObjectIterator.h"
-#include "Toolkits/ToolkitManager.h"
 
-#define LOCTEXT_NAMESPACE "MetasoundSchema"
+#define LOCTEXT_NAMESPACE "MetasoundEditor"
+
 
-namespace MetasoundEditorUtils
+namespace Metasound
 {
-	TSharedPtr<FMetasoundEditor> GetMetasoundEditor(const UMetasound* Metasound)
+	namespace Editor
 	{
-		if (Metasound)
+		TSharedPtr<FEditor> GetEditorForGraph(const UMetasound& Metasound)
 		{
-			TSharedPtr<IToolkit> FoundAssetEditor = FToolkitManager::Get().FindEditorForAsset(CastChecked<const UObject>(Metasound));
-			return StaticCastSharedPtr<FMetasoundEditor, IToolkit>(FoundAssetEditor);
+			TSharedPtr<IToolkit> FoundAssetEditor = FToolkitManager::Get().FindEditorForAsset(CastChecked<const UObject>(&Metasound));
+			return StaticCastSharedPtr<FEditor, IToolkit>(FoundAssetEditor);
 		}
 
-		return nullptr;
-	}
+		TSharedPtr<FEditor> GetEditorForGraph(const UEdGraph& EdGraph)
+		{
+			const UMetasoundEditorGraph* MetasoundGraph = CastChecked<const UMetasoundEditorGraph>(&EdGraph);
+			return GetEditorForGraph(MetasoundGraph->GetMetasoundChecked());
+		}
+
+		FLinearColor GetPinCategoryColor(const FEdGraphPinType& PinType)
+		{
+			const UGraphEditorSettings* Settings = GetDefault<UGraphEditorSettings>();
+			check(Settings);
+
+			if (PinType.PinCategory == FGraphBuilder::PinPrimitiveBoolean)
+			{
+				return Settings->BooleanPinTypeColor;
+			}
+
+			if (PinType.PinCategory == FGraphBuilder::PinPrimitiveFloat)
+			{
+				return Settings->FloatPinTypeColor;
+			}
+
+			if (PinType.PinCategory == FGraphBuilder::PinPrimitiveInteger)
+			{
+				return Settings->IntPinTypeColor;
+			}
+
+			if (PinType.PinCategory == FGraphBuilder::PinPrimitiveString)
+			{
+				return Settings->StringPinTypeColor;
+			}
+
+			return Settings->StructPinTypeColor;
+		}
+
+		FConnectionDrawingPolicy* FGraphConnectionDrawingPolicyFactory::CreateConnectionPolicy(
+			const UEdGraphSchema* InSchema,
+			int32 InBackLayerID,
+			int32 InFrontLayerID,
+			float InZoomFactor,
+			const FSlateRect& InClippingRect,
+			FSlateWindowElementList& InDrawElements,
+			UEdGraph* InGraphObj) const
+		{
+			if (InSchema->IsA(UMetasoundEditorGraphSchema::StaticClass()))
+			{
+				return new FGraphConnectionDrawingPolicy(InBackLayerID, InFrontLayerID, InZoomFactor, InClippingRect, InDrawElements, InGraphObj);
+			}
+			return nullptr;
+		}
+
+		FGraphConnectionDrawingPolicy::FGraphConnectionDrawingPolicy(int32 InBackLayerID, int32 InFrontLayerID, float InZoomFactor, const FSlateRect& InClippingRect, FSlateWindowElementList& InDrawElements, UEdGraph* InGraphObj)
+			: FConnectionDrawingPolicy(InBackLayerID, InFrontLayerID, InZoomFactor, InClippingRect, InDrawElements)
+			, GraphObj(InGraphObj)
+		{
+			ActiveWireThickness = Settings->TraceAttackWireThickness;
+			InactiveWireThickness = Settings->TraceReleaseWireThickness;
+		}
+
+		void FGraphConnectionDrawingPolicy::DetermineWiringStyle(UEdGraphPin* OutputPin, UEdGraphPin* InputPin, FConnectionParams& OutParams)
+		{
+			using namespace Frontend;
+
+			if (!OutputPin || !InputPin || !GraphObj)
+			{
+				return;
+			}
+
+			OutParams.AssociatedPin1 = InputPin;
+			OutParams.AssociatedPin2 = OutputPin;
+
+			// Get the schema and grab the default color from it
+			const UEdGraphSchema* Schema = GraphObj->GetSchema();
+			FNodeHandle InputNodeHandle = CastChecked<UMetasoundEditorGraphNode>(InputPin->GetOwningNode())->GetNodeHandle();
+			FNodeHandle OutputNodeHandle = CastChecked<UMetasoundEditorGraphNode>(OutputPin->GetOwningNode())->GetNodeHandle();
+
+			OutParams.WireColor = GetPinCategoryColor(OutputPin->PinType);
+			bool bExecuted = false;
+
+			// Run through the predecessors, and on
+			if (FExecPairingMap* PredecessorMap = PredecessorNodes.Find(OutputPin->GetOwningNode()))
+			{
+				if (FTimePair* Times = PredecessorMap->Find(InputPin->GetOwningNode()))
+				{
+					bExecuted = true;
+
+					OutParams.WireThickness = ActiveWireThickness;
+					OutParams.bDrawBubbles = true;
+				}
+			}
+
+			if (!bExecuted)
+			{
+				OutParams.WireThickness = InactiveWireThickness;
+			}
+		}
+	} // namespace Editor
+} // namespace Metasound
+
+UEdGraphNode* FMetasoundGraphSchemaAction_NewNode::PerformAction(UEdGraph* ParentGraph, UEdGraphPin* FromPin, const FVector2D Location, bool bSelectNewNode /* = true */)
+{
+	using namespace Metasound::Editor;
+	using namespace Metasound::Frontend;
 
-	TSharedPtr<FMetasoundEditor> GetMetasoundEditor(const UEdGraph& EdGraph)
+	const FScopedTransaction Transaction(LOCTEXT("MetasoundEditorNewNode", "Add New Metasound Node"));
+
+	ParentGraph->Modify();
+
+	UMetasound& ParentMetasound = CastChecked<UMetasoundEditorGraph>(ParentGraph)->GetMetasoundChecked();
+	ParentMetasound.Modify();
+
+	UEdGraphNode* NewGraphNode = FGraphBuilder::AddNode(ParentMetasound, Location, NodeClassInfo);
+
+	return NewGraphNode;
+}
+
+FMetasoundGraphSchemaAction_NewInput::FMetasoundGraphSchemaAction_NewInput(FText InNodeCategory, FText InDisplayName, FName InTypeName, FText InToolTip, const int32 InGrouping)
+	: FEdGraphSchemaAction(MoveTemp(InNodeCategory), MoveTemp(InDisplayName), MoveTemp(InToolTip), InGrouping)
+	, NodeTypeName(InTypeName)
+{
+}
+
+UEdGraphNode* FMetasoundGraphSchemaAction_NewInput::PerformAction(UEdGraph* ParentGraph, UEdGraphPin* FromPin, const FVector2D Location, bool bSelectNewNode /* = true */)
+{
+	using namespace Metasound::Editor;
+	using namespace Metasound::Frontend;
+
+	const FScopedTransaction Transaction(LOCTEXT("MetasoundEditorNewInput", "Add New Metasound Input"));
+
+	ParentGraph->Modify();
+
+	UMetasound& ParentMetasound = CastChecked<UMetasoundEditorGraph>(ParentGraph)->GetMetasoundChecked();
+	ParentMetasound.Modify();
+
+	FString NameBase = FString::Printf(TEXT("%s"), *NodeTypeName.ToString());
+
+	FGraphHandle GraphHandle = ParentMetasound.GetRootGraphHandle();
+	FString NewNodeName = NameBase;
+
+	int32 i = 0;
+	while (GraphHandle.ContainsOutputNodeWithName(NewNodeName))
 	{
-		const UMetasoundEditorGraph* MetasoundGraph = CastChecked<const UMetasoundEditorGraph>(&EdGraph);
-		return GetMetasoundEditor(MetasoundGraph->Metasound);
+		NewNodeName = NameBase + FString::Printf(TEXT("_%u"), ++i);
 	}
-} // namespace MetasoundEditorUtils
 
-UEdGraphNode* FMetasoundGraphSchemaAction_NewNode::PerformAction(UEdGraph* ParentGraph, UEdGraphPin* FromPin, const FVector2D Location, bool bSelectNewNode/* = true*/)
+	UEdGraphNode* NewGraphNode = FGraphBuilder::AddInput(ParentMetasound, Location, NewNodeName, NodeTypeName, FText::GetEmpty());
+	return NewGraphNode;
+}
+
+FMetasoundGraphSchemaAction_NewOutput::FMetasoundGraphSchemaAction_NewOutput(FText InNodeCategory, FText InDisplayName, FName InTypeName, FText InToolTip, const int32 InGrouping)
+	: FEdGraphSchemaAction(MoveTemp(InNodeCategory), MoveTemp(InDisplayName), MoveTemp(InToolTip), InGrouping)
+	, NodeTypeName(InTypeName)
 {
-	UMetasoundEditorGraph* MetasoundGraph = CastChecked<UMetasoundEditorGraph>(ParentGraph);
-	UMetasound* Metasound = MetasoundGraph->Metasound;
-	const FScopedTransaction Transaction(LOCTEXT("MetasoundEditorNewSoundNode", "Metasound Editor: New Metasound Node"));
+}
+
+UEdGraphNode* FMetasoundGraphSchemaAction_NewOutput::PerformAction(UEdGraph* ParentGraph, UEdGraphPin* FromPin, const FVector2D Location, bool bSelectNewNode /* = true */)
+{
+	using namespace Metasound::Editor;
+	using namespace Metasound::Frontend;
+
+	const FScopedTransaction Transaction(LOCTEXT("MetasoundEditorNewOutput", "Add New Metasound Output"));
+
 	ParentGraph->Modify();
 
+	UMetasound& ParentMetasound = CastChecked<UMetasoundEditorGraph>(ParentGraph)->GetMetasoundChecked();
+	ParentMetasound.Modify();
 
-	Metasound->Modify();
+	FString NameBase = FString::Printf(TEXT("%s"), *NodeTypeName.ToString());
 
-	UMetasoundEditorGraphNode_Root* GraphNode = nullptr;
-// TODO: Implement
-// 	UMetasoundEditorGraphNode_Root* GraphNode = MetasoundGraph->CreateBlankNode<MetasoundNodeClass>();
+	FGraphHandle GraphHandle = ParentMetasound.GetRootGraphHandle();
+	FString NewNodeName = NameBase;
 
-	Metasound->PostEditChange();
-	Metasound->MarkPackageDirty();
+	int32 i = 0;
+	while (GraphHandle.ContainsOutputNodeWithName(NewNodeName))
+	{
+		NewNodeName = NameBase + FString::Printf(TEXT("_%u"), ++i);
+	}
 
-	return GraphNode;
+	UEdGraphNode* NewGraphNode = FGraphBuilder::AddOutput(ParentMetasound, Location, NewNodeName, NodeTypeName, FText::GetEmpty());
+	return NewGraphNode;
 }
 
-void FMetasoundGraphSchemaAction_NewNode::ConnectToSelectedNodes(UMetasound* NewNode, UEdGraph* ParentGraph) const
+void FMetasoundGraphSchemaAction_NewNode::ConnectToSelectedNodes(UMetasoundEditorGraphNode* NewGraphNode, UEdGraph* ParentGraph) const
 {
 	// TODO: Implement
 }
@@ -77,11 +232,13 @@ UEdGraphNode* FMetasoundGraphSchemaAction_NewFromSelected::PerformAction(UEdGrap
 
 UEdGraphNode* FMetasoundGraphSchemaAction_NewComment::PerformAction(UEdGraph* ParentGraph, UEdGraphPin* FromPin, const FVector2D Location, bool bSelectNewNode/* = true*/)
 {
+	using namespace Metasound::Editor;
+
 	UEdGraphNode_Comment* CommentTemplate = NewObject<UEdGraphNode_Comment>();
 
 	FSlateRect Bounds;
 	FVector2D SpawnLocation = Location;
-	TSharedPtr<FMetasoundEditor> MetasoundEditor = MetasoundEditorUtils::GetMetasoundEditor(*ParentGraph);
+	TSharedPtr<FEditor> MetasoundEditor = GetEditorForGraph(*ParentGraph);
 
 	if (MetasoundEditor.IsValid() && MetasoundEditor->GetBoundsForSelectedNodes(Bounds, 50.0f))
 	{
@@ -95,7 +252,9 @@ UEdGraphNode* FMetasoundGraphSchemaAction_NewComment::PerformAction(UEdGraph* Pa
 
 UEdGraphNode* FMetasoundGraphSchemaAction_Paste::PerformAction(UEdGraph* ParentGraph, UEdGraphPin* FromPin, const FVector2D Location, bool bSelectNewNode/* = true*/)
 {
-	TSharedPtr<FMetasoundEditor> MetasoundEditor = MetasoundEditorUtils::GetMetasoundEditor(*ParentGraph);
+	using namespace Metasound::Editor;
+
+	TSharedPtr<FEditor> MetasoundEditor = GetEditorForGraph(*ParentGraph);
 	if (MetasoundEditor.IsValid())
 	{
 		MetasoundEditor->PasteNodesAtLocation(Location);
@@ -112,29 +271,28 @@ UMetasoundEditorGraphSchema::UMetasoundEditorGraphSchema(const FObjectInitialize
 bool UMetasoundEditorGraphSchema::ConnectionCausesLoop(const UEdGraphPin* InputPin, const UEdGraphPin* OutputPin) const
 {
 	// TODO: Implement loop check
-
-	// Simple connection to root node
 	return false;
 }
 
 void UMetasoundEditorGraphSchema::GetPaletteActions(FGraphActionMenuBuilder& ActionMenuBuilder) const
 {
-	GetAllMetasoundActions(ActionMenuBuilder, /* bShowSelectedActions*/ false);
+	GetAllMetasoundActions(ActionMenuBuilder, false /* bShowSelectedActions */);
 	GetCommentAction(ActionMenuBuilder);
 }
 
 void UMetasoundEditorGraphSchema::GetGraphContextActions(FGraphContextMenuBuilder& ContextMenuBuilder) const
 {
-	GetAllMetasoundActions(ContextMenuBuilder, /* bShowSelectedActions*/ true);
+	using namespace Metasound::Editor;
 
+	GetAllMetasoundActions(ContextMenuBuilder, true /* bShowSelectedActions */);
 	GetCommentAction(ContextMenuBuilder, ContextMenuBuilder.CurrentGraph);
 
 	if (!ContextMenuBuilder.FromPin)
 	{
-		TSharedPtr<FMetasoundEditor> MetasoundEditor = MetasoundEditorUtils::GetMetasoundEditor(*ContextMenuBuilder.CurrentGraph);
+		TSharedPtr<FEditor> MetasoundEditor = GetEditorForGraph(*ContextMenuBuilder.CurrentGraph);
 		if (MetasoundEditor.IsValid() && MetasoundEditor->CanPasteNodes())
 		{
-			TSharedPtr<FMetasoundGraphSchemaAction_Paste> NewAction(new FMetasoundGraphSchemaAction_Paste(FText::GetEmpty(), LOCTEXT("PasteHereAction", "Paste here"), FText::GetEmpty(), 0));
+			TSharedPtr<FMetasoundGraphSchemaAction_Paste> NewAction = MakeShared<FMetasoundGraphSchemaAction_Paste>(FText::GetEmpty(), LOCTEXT("PasteHereAction", "Paste here"), FText::GetEmpty(), 0);
 			ContextMenuBuilder.AddAction(NewAction);
 		}
 	}
@@ -198,6 +356,26 @@ const FPinConnectionResponse UMetasoundEditorGraphSchema::CanCreateConnection(co
 		return FPinConnectionResponse(CONNECT_RESPONSE_DISALLOW, LOCTEXT("ConnectionLoop", "Connection would cause loop"));
 	}
 
+	if (InputPin->PinType.PinCategory != InputPin->PinType.PinCategory)
+	{
+		return FPinConnectionResponse(CONNECT_RESPONSE_DISALLOW, LOCTEXT("ConnectionTypeIncorrect", "Connection pin types do not match"));
+	}
+
+	UMetasoundEditorGraphNode* InputGraphNode = CastChecked<UMetasoundEditorGraphNode>(InputPin->GetOwningNode());
+	Metasound::Frontend::FNodeHandle InputNodeHandle = InputGraphNode->GetNodeHandle();
+	Metasound::Frontend::FInputHandle InputHandle = InputNodeHandle.GetInputWithName(InputPin->GetName());
+
+	UMetasoundEditorGraphNode* OutputGraphNode = CastChecked<UMetasoundEditorGraphNode>(OutputPin->GetOwningNode());
+	Metasound::Frontend::FNodeHandle OutputNodeHandle = OutputGraphNode->GetNodeHandle();
+	Metasound::Frontend::FOutputHandle OutputHandle = OutputNodeHandle.GetOutputWithName(OutputPin->GetName());
+
+	// TODO: Implement YesWithConverterNode to provide conversion options
+	Metasound::Frontend::FConnectability Connectability = InputHandle.CanConnectTo(OutputHandle);
+	if (Connectability.Connectable != Metasound::Frontend::FConnectability::EConnectable::Yes)
+	{
+		return FPinConnectionResponse(CONNECT_RESPONSE_DISALLOW, LOCTEXT("ConnectionTypeIncompatible", "Connection pin types cannot be connected"));
+	}
+
 	// Break existing connections on inputs only - multiple output connections are acceptable
 	if (InputPin->LinkedTo.Num() > 0)
 	{
@@ -218,12 +396,7 @@ const FPinConnectionResponse UMetasoundEditorGraphSchema::CanCreateConnection(co
 
 bool UMetasoundEditorGraphSchema::TryCreateConnection(UEdGraphPin* PinA, UEdGraphPin* PinB) const
 {
-	bool bModified = UEdGraphSchema::TryCreateConnection(PinA, PinB);
-	if (bModified)
-	{
-		// TODO: Compile Metasound???
-	}
-
+	const bool bModified = UEdGraphSchema::TryCreateConnection(PinA, PinB);
 	return bModified;
 }
 
@@ -235,7 +408,8 @@ bool UMetasoundEditorGraphSchema::ShouldHidePinDefaultValue(UEdGraphPin* Pin) co
 
 FLinearColor UMetasoundEditorGraphSchema::GetPinTypeColor(const FEdGraphPinType& PinType) const
 {
-	return FLinearColor::White;
+	return Metasound::Editor::GetPinCategoryColor(PinType);
+
 }
 
 void UMetasoundEditorGraphSchema::BreakNodeLinks(UEdGraphNode& TargetNode) const
@@ -268,34 +442,87 @@ void UMetasoundEditorGraphSchema::DroppedAssetsOnGraph(const TArray<FAssetData>&
 
 void UMetasoundEditorGraphSchema::DroppedAssetsOnNode(const TArray<FAssetData>& Assets, const FVector2D& GraphPosition, UEdGraphNode* Node) const
 {
-	if (!Node->IsA<UMetasoundEditorGraphNode>())
-	{
-		return;
-	}
-
-	UMetasoundEditorGraphNode* GraphNode = CastChecked<UMetasoundEditorGraphNode>(Node);
-	UMetasoundEditorGraph* MetasoundGraph = CastChecked<UMetasoundEditorGraph>(Node->GetGraph());
-
-	// TODO: Implement adding reference node
-// 	UMetasound* Metasound = MetasoundGraph->Metasound;
-
-	MetasoundGraph->NotifyGraphChanged();
+	// TODO: Implement for metasound references
 }
 
 void UMetasoundEditorGraphSchema::GetAllMetasoundActions(FGraphActionMenuBuilder& ActionMenuBuilder, bool bShowSelectedActions) const
 {
-	// TODO: Add actions for each type of metasound node
+	using namespace Metasound::Editor;
+	using namespace Metasound::Frontend;
+
+	const FText InputMenuName = LOCTEXT("MetasoundAddInputMenu", "Add Input");
+	const FText OutputMenuName = LOCTEXT("MetasoundAddOutputMenu", "Add Output");
+
+	IMetasoundEditorModule& EditorModule = FModuleManager::GetModuleChecked<IMetasoundEditorModule>("MetasoundEditor");
+	EditorModule.IterateDataTypes([InputMenuName, OutputMenuName, InMenuBuilder = &ActionMenuBuilder](const FEditorDataType& DataType)
+	{
+		const FName DataTypeName = DataType.RegistryInfo.DataTypeName;
+		const FText DataTypeTextName = FText::FromName(DataTypeName);
+		const FText MenuJoinFormat = LOCTEXT("MetasoundFormatNodeSubCategory", "{0}|{1}");
+
+		FString SubMenu = DataTypeName.ToString();
+		int32 Index = 0;
+		SubMenu.FindLastChar(':', Index);
+
+		const FText DisplayName = FText::FromString(SubMenu.RightChop(Index + 1));
+
+		SubMenu.RemoveAt(Index, SubMenu.Len() - Index);
+		SubMenu.ReplaceCharInline(':', '|');
+
+		TSharedPtr<FMetasoundGraphSchemaAction_NewInput> AddInputNodeAction = MakeShared<FMetasoundGraphSchemaAction_NewInput>
+		(
+			FText::Format(MenuJoinFormat, InputMenuName, FText::FromString(SubMenu)),
+			DisplayName,
+			DataTypeName,
+			FText::Format(LOCTEXT("MetasoundTooltipAddInputFormat", "Adds an input of type {0} to the Metasound"), DataTypeTextName),
+			0
+		);
+		InMenuBuilder->AddAction(AddInputNodeAction);
+
+		TSharedPtr<FMetasoundGraphSchemaAction_NewOutput> AddOutputNodeAction = MakeShared<FMetasoundGraphSchemaAction_NewOutput>
+		(
+			FText::Format(MenuJoinFormat, OutputMenuName, FText::FromString(SubMenu)),
+			DisplayName,
+			DataTypeName,
+			FText::Format(LOCTEXT("MetasoundTooltipAddOutputFormat", "Adds an output of type {0} to the Metasound"), DataTypeTextName),
+			0
+		);
+		InMenuBuilder->AddAction(AddOutputNodeAction);
+	});
+
+	const FText NodeMenuName = LOCTEXT("MetasoundNodesMenu", "Add Node");
+	const TArray<FNodeClassInfo> ClassInfos = GetAllAvailableNodeClasses();
+	for (const FNodeClassInfo& ClassInfo : ClassInfos)
+	{
+		const FMetasoundClassDescription ClassDescription = GenerateClassDescriptionForNode(ClassInfo);
+		const FMetasoundClassMetadata Metadata = ClassDescription.Metadata;
+		const FText Tooltip = Metadata.AuthorName.IsEmpty()
+			? Metadata.MetasoundDescription
+			: FText::Format(LOCTEXT("MetasoundTooltipAuthorFormat", "{0}\nAuthor: {1}"), Metadata.MetasoundDescription, Metadata.AuthorName);
+
+		TSharedPtr<FMetasoundGraphSchemaAction_NewNode> NewNodeAction = MakeShared<FMetasoundGraphSchemaAction_NewNode>
+		(
+			NodeMenuName,
+			FText::FromString(ClassInfo.NodeName),
+			Tooltip,
+			0
+		);
+
+		NewNodeAction->NodeClassInfo = ClassInfo;
+		ActionMenuBuilder.AddAction(NewNodeAction);
+	}
 }
 
 void UMetasoundEditorGraphSchema::GetCommentAction(FGraphActionMenuBuilder& ActionMenuBuilder, const UEdGraph* CurrentGraph) const
 {
+	using namespace Metasound::Editor;
+
 	if (!ActionMenuBuilder.FromPin && CurrentGraph)
 	{
-		TSharedPtr<FMetasoundEditor> MetasoundEditor = MetasoundEditorUtils::GetMetasoundEditor(*CurrentGraph);
-
+		TSharedPtr<FEditor> MetasoundEditor = GetEditorForGraph(*CurrentGraph);
 		if (MetasoundEditor.IsValid())
 		{
-			const int32 NumSelected = MetasoundEditor->GetNumberOfSelectedNodes();
+			const int32 NumSelected = MetasoundEditor->GetNumNodesSelected();
 			const FText MenuDescription = NumSelected > 0 ? LOCTEXT("CreateCommentAction", "Create Comment from Selection") : LOCTEXT("AddCommentAction", "Add Comment...");
 			const FText ToolTip = LOCTEXT("CreateCommentToolTip", "Creates a comment.");
 
@@ -307,14 +534,12 @@ void UMetasoundEditorGraphSchema::GetCommentAction(FGraphActionMenuBuilder& Acti
 
 int32 UMetasoundEditorGraphSchema::GetNodeSelectionCount(const UEdGraph* Graph) const
 {
-	if (const UMetasound* Metasound = CastChecked<UMetasoundEditorGraph>(Graph)->Metasound)
+	using namespace Metasound::Editor;
+
+	TSharedPtr<FEditor> MetasoundEditor = GetEditorForGraph(*Graph);
+	if (MetasoundEditor.IsValid())
 	{
-		TSharedPtr<IToolkit> FoundAssetEditor = FToolkitManager::Get().FindEditorForAsset(CastChecked<UObject>(Metasound));
-		if (FoundAssetEditor.IsValid())
-		{
-			const FMetasoundEditor* AssetEditor = static_cast<FMetasoundEditor*>(FoundAssetEditor.Get());
-			return AssetEditor->GetNumberOfSelectedNodes();
-		}
+		return MetasoundEditor->GetNumNodesSelected();
 	}
 
 	return 0;
@@ -325,5 +550,4 @@ TSharedPtr<FEdGraphSchemaAction> UMetasoundEditorGraphSchema::GetCreateCommentAc
 	TSharedPtr<FMetasoundGraphSchemaAction_NewComment> Comment = MakeShared<FMetasoundGraphSchemaAction_NewComment>();
 	return StaticCastSharedPtr<FEdGraphSchemaAction, FMetasoundGraphSchemaAction_NewComment>(Comment);
 }
-
 #undef LOCTEXT_NAMESPACE
diff --git a/Engine/Plugins/Runtime/Metasound/Source/MetasoundEditor/Private/MetasoundEditorModule.cpp b/Engine/Plugins/Runtime/Metasound/Source/MetasoundEditor/Private/MetasoundEditorModule.cpp
--- a/Engine/Plugins/Runtime/Metasound/Source/MetasoundEditor/Private/MetasoundEditorModule.cpp
+++ b/Engine/Plugins/Runtime/Metasound/Source/MetasoundEditor/Private/MetasoundEditorModule.cpp
@@ -5,10 +5,20 @@
 #include "AssetTypeActions_Base.h"
 #include "Brushes/SlateImageBrush.h"
 #include "CoreMinimal.h"
+#include "EdGraph/EdGraphNode.h"
+#include "EdGraph/EdGraphPin.h"
+#include "EdGraphUtilities.h"
 #include "EditorStyleSet.h"
+#include "IDetailCustomization.h"
+#include "Metasound.h"
 #include "MetasoundAssetTypeActions.h"
+#include "MetasoundDetailCustomization.h"
+#include "MetasoundEditorGraphBuilder.h"
+#include "MetasoundEditorGraphSchema.h"
 #include "Modules/ModuleInterface.h"
 #include "Modules/ModuleManager.h"
+#include "PropertyEditorDelegates.h"
+#include "PropertyEditorModule.h"
 #include "Styling/CoreStyle.h"
 #include "Styling/SlateStyle.h"
 #include "Styling/SlateStyleMacros.h"
@@ -19,73 +29,167 @@
 DEFINE_LOG_CATEGORY(LogMetasoundEditor);
 
 
-namespace MetasoundEditorUtils
+namespace Metasound
 {
-	static const FName AssetToolName = TEXT("AssetTools");
-
-	template <typename T>
-	void AddAssetAction(IAssetTools& AssetTools, TArray<TSharedPtr<FAssetTypeActions_Base>>& AssetArray)
-	{
-		TSharedPtr<T> AssetAction = MakeShared<T>();
-		TSharedPtr<FAssetTypeActions_Base> AssetActionBase = StaticCastSharedPtr<FAssetTypeActions_Base>(AssetAction);
-		AssetTools.RegisterAssetTypeActions(AssetAction.ToSharedRef());
-		AssetArray.Add(AssetActionBase);
-	}
-} // namespace MetasoundEditorUtils
-
-class FMetasoundSlateStyle : public FSlateStyleSet
-{
-public:
-	FMetasoundSlateStyle()
-		: FSlateStyleSet("MetasoundStyle")
+	namespace Editor
 	{
-		SetParentStyleName(FEditorStyle::GetStyleSetName());
+		static const FName AssetToolName = TEXT("AssetTools");
 
-		SetContentRoot(FPaths::EngineContentDir() / TEXT("Editor/Slate"));
-		SetCoreContentRoot(FPaths::EngineContentDir() / TEXT("Slate"));
+		template <typename T>
+		void AddAssetAction(IAssetTools& AssetTools, TArray<TSharedPtr<FAssetTypeActions_Base>>& AssetArray)
+		{
+			TSharedPtr<T> AssetAction = MakeShared<T>();
+			TSharedPtr<FAssetTypeActions_Base> AssetActionBase = StaticCastSharedPtr<FAssetTypeActions_Base>(AssetAction);
+			AssetTools.RegisterAssetTypeActions(AssetAction.ToSharedRef());
+			AssetArray.Add(AssetActionBase);
+		}
 
-		static const FVector2D Icon20x20(20.0f, 20.0f);
-		static const FVector2D Icon40x40(40.0f, 40.0f);
+		class FSlateStyle : public FSlateStyleSet
+		{
+		public:
+			FSlateStyle()
+				: FSlateStyleSet("MetasoundStyle")
+			{
+				SetParentStyleName(FEditorStyle::GetStyleSetName());
+
+				SetContentRoot(FPaths::EngineContentDir() / TEXT("Editor/Slate"));
+				SetCoreContentRoot(FPaths::EngineContentDir() / TEXT("Slate"));
+
+				static const FVector2D Icon20x20(20.0f, 20.0f);
+				static const FVector2D Icon40x40(40.0f, 40.0f);
+
+				// Metasound Editor
+				{
+					Set("MetasoundEditor.Play", new FSlateImageBrush(RootToContentDir(TEXT("Icons/icon_SCueEd_PlayCue_40x.png")), Icon40x40));
+					Set("MetasoundEditor.Play.Small", new FSlateImageBrush(RootToContentDir(TEXT("Icons/icon_SCueEd_PlayCue_40x.png")), Icon20x20));
+					Set("MetasoundEditor.Stop", new FSlateImageBrush(RootToContentDir(TEXT("Icons/icon_SCueEd_Stop_40x.png")), Icon40x40));
+					Set("MetasoundEditor.Stop.Small", new FSlateImageBrush(RootToContentDir(TEXT("Icons/icon_SCueEd_Stop_40x.png")), Icon20x20));
+					Set("MetasoundEditor.Compile", new FSlateImageBrush(RootToContentDir(TEXT("/Old/Kismet2/compile_40px.png")), Icon40x40));
+					Set("MetasoundEditor.Compile.Small", new FSlateImageBrush(RootToContentDir(TEXT("/Old/Kismet2/compile_40px.png")), Icon20x20));
+					Set("MetasoundEditor.CompileError", new FSlateImageBrush(RootToContentDir(TEXT("/Old/Kismet2/CompileStatus_Fail.png")), Icon40x40));
+					Set("MetasoundEditor.CompileError.Small", new FSlateImageBrush(RootToContentDir(TEXT("/Old/Kismet2/CompileStatus_Broken_Small.png")), Icon20x20));
+				}
+
+				FSlateStyleRegistry::RegisterSlateStyle(*this);
+			}
+		};
 
-		// Metasound Editor
+		class FModule : public IMetasoundEditorModule
 		{
-			Set("MetasoundEditor.Play", new FSlateImageBrush(RootToContentDir(TEXT("Icons/icon_SCueEd_PlayCue_40x.png")), Icon40x40));
-			Set("MetasoundEditor.Play.Small", new FSlateImageBrush(RootToContentDir(TEXT("Icons/icon_SCueEd_PlayCue_40x.png")), Icon20x20));
-			Set("MetasoundEditor.Stop", new FSlateImageBrush(RootToContentDir(TEXT("Icons/icon_SCueEd_Stop_40x.png")), Icon40x40));
-			Set("MetasoundEditor.Stop.Small", new FSlateImageBrush(RootToContentDir(TEXT("Icons/icon_SCueEd_Stop_40x.png")), Icon20x20));
-		}
+			virtual const FEditorDataType& FindDataType(FName InDataTypeName) const override
+			{
+				return DataTypeInfo.FindChecked(InDataTypeName);
+			}
 
-		FSlateStyleRegistry::RegisterSlateStyle(*this);
-	}
-};
+			virtual void IterateDataTypes(TUniqueFunction<void(const FEditorDataType&)> InDataTypeFunction) const override
+			{
+				for (const TPair<FName, FEditorDataType>& Pair : DataTypeInfo)
+				{
+					InDataTypeFunction(Pair.Value);
+				}
+			}
 
-class FMetasoundEditorModule : public IMetasoundEditorModule
-{
-	virtual void StartupModule() override
-	{
-		// Register Metasound asset type actions
-		IAssetTools& AssetTools = FModuleManager::LoadModuleChecked<FAssetToolsModule>(MetasoundEditorUtils::AssetToolName).Get();
-		MetasoundEditorUtils::AddAssetAction<FAssetTypeActions_Metasound>(AssetTools, AssetActions);
+			virtual void RegisterDataType(FName InPinCategoryName, const FDataTypeRegistryInfo& InRegistryInfo) override
+			{
+				DataTypeInfo.Add(InRegistryInfo.DataTypeName, FEditorDataType(InPinCategoryName, InRegistryInfo));
+			}
 
-		StyleSet = MakeShared<FMetasoundSlateStyle>();
-	}
+			void RegisterCoreDataTypes()
+			{
+				TArray<FName> DataTypeNames = Frontend::GetAllAvailableDataTypes();
+				for (FName DataTypeName : DataTypeNames)
+				{
+					FDataTypeRegistryInfo RegistryInfo;
+					Frontend::GetTraitsForDataType(DataTypeName, RegistryInfo);
+
+					FName PinType = NAME_None;
+					switch (RegistryInfo.PreferredLiteralType)
+					{
+						case ELiteralArgType::Boolean:
+						{
+							PinType = FGraphBuilder::PinPrimitiveBoolean;
+						}
+						break;
+
+						case ELiteralArgType::Float:
+						{
+							PinType = FGraphBuilder::PinPrimitiveFloat;
+						}
+						break;
+
+						case ELiteralArgType::Integer:
+						{
+							PinType = FGraphBuilder::PinPrimitiveInteger;
+						}
+						break;
+
+						case ELiteralArgType::String:
+						{
+							PinType = FGraphBuilder::PinPrimitiveString;
+						}
+						break;
+
+						// Register atypical primitives
+						default:
+						case ELiteralArgType::None:
+						case ELiteralArgType::Invalid:
+						{
+							static_assert(static_cast<int32>(ELiteralArgType::Invalid) == 5, "Possible missing binding of pin category to primitive type");
+						}
+						break;
+					}
+
+					DataTypeInfo.Add(DataTypeName, FEditorDataType(PinType, RegistryInfo));
+				}
+			}
 
-	virtual void ShutdownModule() override
-	{
-		if (FModuleManager::Get().IsModuleLoaded(MetasoundEditorUtils::AssetToolName))
-		{
-			IAssetTools& AssetTools = FModuleManager::GetModuleChecked<FAssetToolsModule>(MetasoundEditorUtils::AssetToolName).Get();
-			for (TSharedPtr<FAssetTypeActions_Base>& AssetAction : AssetActions)
+			virtual void StartupModule() override
 			{
-				AssetTools.UnregisterAssetTypeActions(AssetAction.ToSharedRef());
+				// Register Metasound asset type actions
+				IAssetTools& AssetTools = FModuleManager::LoadModuleChecked<FAssetToolsModule>(AssetToolName).Get();
+				AddAssetAction<FAssetTypeActions_Metasound>(AssetTools, AssetActions);
+
+				FPropertyEditorModule& PropertyModule = FModuleManager::LoadModuleChecked<FPropertyEditorModule>("PropertyEditor");
+				PropertyModule.RegisterCustomClassLayout(
+					UMetasound::StaticClass()->GetFName(),
+					FOnGetDetailCustomizationInstance::CreateLambda([]() { return MakeShared<FMetasoundDetailCustomization>(); }));
+
+				StyleSet = MakeShared<FSlateStyle>();
+
+				RegisterCoreDataTypes();
+
+				GraphConnectionFactory = MakeShared<FGraphConnectionDrawingPolicyFactory>();
+				FEdGraphUtilities::RegisterVisualPinConnectionFactory(GraphConnectionFactory);
+
+			}
+
+			virtual void ShutdownModule() override
+			{
+				if (FModuleManager::Get().IsModuleLoaded(AssetToolName))
+				{
+					IAssetTools& AssetTools = FModuleManager::GetModuleChecked<FAssetToolsModule>(AssetToolName).Get();
+					for (TSharedPtr<FAssetTypeActions_Base>& AssetAction : AssetActions)
+					{
+						AssetTools.UnregisterAssetTypeActions(AssetAction.ToSharedRef());
+					}
+				}
+
+				if (GraphConnectionFactory.IsValid())
+				{
+					FEdGraphUtilities::UnregisterVisualPinConnectionFactory(GraphConnectionFactory);
+				}
+
+				AssetActions.Reset();
+				DataTypeInfo.Reset();
 			}
-		}
-		AssetActions.Reset();
-	}
 
-	TArray<TSharedPtr<FAssetTypeActions_Base>> AssetActions;
-	TSharedPtr<FSlateStyleSet> StyleSet;
-};
+			TArray<TSharedPtr<FAssetTypeActions_Base>> AssetActions;
+			TMap<FName, FEditorDataType> DataTypeInfo;
 
+			TSharedPtr<FGraphPanelPinConnectionFactory> GraphConnectionFactory;
+			TSharedPtr<FSlateStyleSet> StyleSet;
+		};
+	} // namespace Editor
+} // namespace Metasound
 
-IMPLEMENT_MODULE(FMetasoundEditorModule, MetasoundEditor);
+IMPLEMENT_MODULE(Metasound::Editor::FModule, MetasoundEditor);
diff --git a/Engine/Plugins/Runtime/Metasound/Source/MetasoundEditor/Private/MetasoundEditorModule.h b/Engine/Plugins/Runtime/Metasound/Source/MetasoundEditor/Private/MetasoundEditorModule.h
--- a/Engine/Plugins/Runtime/Metasound/Source/MetasoundEditor/Private/MetasoundEditorModule.h
+++ b/Engine/Plugins/Runtime/Metasound/Source/MetasoundEditor/Private/MetasoundEditorModule.h
@@ -1,12 +1,38 @@
 // Copyright Epic Games, Inc. All Rights Reserved.
-
 #pragma once
 
 #include "CoreMinimal.h"
+#include "EdGraph/EdGraphNode.h"
+#include "EdGraph/EdGraphPin.h"
 #include "Modules/ModuleInterface.h"
+#include "MetasoundDataReference.h"
+#include "Templates/Function.h"
 
 DECLARE_LOG_CATEGORY_EXTERN(LogMetasoundEditor, Log, All);
 
-class METASOUNDEDITOR_API IMetasoundEditorModule : public IModuleInterface
+
+namespace Metasound
 {
-};
+	namespace Editor
+	{
+		struct FEditorDataType
+		{
+			FEdGraphPinType PinType;
+			FDataTypeRegistryInfo RegistryInfo;
+
+			FEditorDataType(FName InGraphDataType, const FDataTypeRegistryInfo& InRegistryInfo)
+				: PinType(InGraphDataType, NAME_None, nullptr, EPinContainerType::None, false, FEdGraphTerminalType())
+				, RegistryInfo(InRegistryInfo)
+			{
+			}
+		};
+
+		class METASOUNDEDITOR_API IMetasoundEditorModule : public IModuleInterface
+		{
+		public:
+			virtual void RegisterDataType(FName InPinCategoryName, const FDataTypeRegistryInfo& InRegistryInfo) = 0;
+			virtual const FEditorDataType& FindDataType(FName InDataTypeName) const = 0;
+			virtual void IterateDataTypes(TUniqueFunction<void(const FEditorDataType&)> InDataTypeFunction) const = 0;
+		};
+	} // namespace Editor
+} // namespace Metasound
diff --git a/Engine/Plugins/Runtime/Metasound/Source/MetasoundEditor/Private/MetasoundEditorTabFactory.h b/Engine/Plugins/Runtime/Metasound/Source/MetasoundEditor/Private/MetasoundEditorTabFactory.h
new file mode 100644
--- /dev/null
+++ b/Engine/Plugins/Runtime/Metasound/Source/MetasoundEditor/Private/MetasoundEditorTabFactory.h
@@ -0,0 +1,67 @@
+// Copyright Epic Games, Inc. All Rights Reserved.
+#pragma once
+
+#include "CoreMinimal.h"
+#include "EditorStyleSet.h"
+#include "Framework/Docking/TabManager.h"
+#include "Internationalization/Internationalization.h"
+#include "Widgets/Docking/SDockTab.h"
+#include "Templates/SharedPointer.h"
+
+#define LOCTEXT_NAMESPACE "MetasoundEditor"
+
+namespace Metasound
+{
+	namespace Editor
+	{
+		namespace TabFactory
+		{
+			namespace Names
+			{
+				const FName GraphCanvas = "MetasoundEditor_GraphCanvas";
+				const FName Palette = "MetasoundEditor_Palette";
+				const FName Properties = "MetasoundEditor_Properties";
+			}
+
+			TSharedRef<SDockTab> CreateGraphCanvasTab(TSharedPtr<SGraphEditor> GraphEditor, const FSpawnTabArgs& Args)
+			{
+				check(Args.GetTabId() == Names::GraphCanvas);
+
+				TSharedRef<SDockTab> SpawnedTab = SNew(SDockTab)
+					.Label(LOCTEXT("MetasoundGraphCanvasTitle", "Viewport"));
+
+				if (GraphEditor.IsValid())
+				{
+					SpawnedTab->SetContent(GraphEditor.ToSharedRef());
+				}
+
+				return SpawnedTab;
+			}
+
+			TSharedRef<SDockTab> CreatePaletteTab(TSharedPtr<SMetasoundPalette> Palette, const FSpawnTabArgs& Args)
+			{
+				check(Args.GetTabId() == Names::Palette);
+
+				return SNew(SDockTab)
+					.Icon(FEditorStyle::GetBrush("Kismet.Tabs.Palette"))
+					.Label(LOCTEXT("MetasoundPaletteTitle", "Palette"))
+					[
+						Palette.ToSharedRef()
+					];
+			}
+
+			TSharedRef<SDockTab> CreatePropertiesTab(TSharedPtr<IDetailsView> PropertiesView, const FSpawnTabArgs& Args)
+			{
+				check(Args.GetTabId() == Names::Properties);
+
+				return SNew(SDockTab)
+					.Icon(FEditorStyle::GetBrush("LevelEditor.Tabs.Details"))
+					.Label(LOCTEXT("MetasoundDetailsTitle", "Details"))
+					[
+						PropertiesView.ToSharedRef()
+					];
+			}
+		}
+	} // namespace Editor
+} // namespace Metasound
+#undef LOCTEXT_NAMESPACE // MetasoundEditor
diff --git a/Engine/Plugins/Runtime/Metasound/Source/MetasoundEditor/Private/MetasoundEditorUtilities.cpp b/Engine/Plugins/Runtime/Metasound/Source/MetasoundEditor/Private/MetasoundEditorUtilities.cpp
deleted file mode 100644
diff --git a/Engine/Plugins/Runtime/Metasound/Source/MetasoundEditor/Private/MetasoundFactory.cpp b/Engine/Plugins/Runtime/Metasound/Source/MetasoundEditor/Private/MetasoundFactory.cpp
--- a/Engine/Plugins/Runtime/Metasound/Source/MetasoundEditor/Private/MetasoundFactory.cpp
+++ b/Engine/Plugins/Runtime/Metasound/Source/MetasoundEditor/Private/MetasoundFactory.cpp
@@ -1,7 +1,7 @@
 // Copyright Epic Games, Inc. All Rights Reserved.
-
 #include "MetasoundFactory.h"
 
+#include "Kismet/KismetSystemLibrary.h"
 #include "Metasound.h"
 #include "MetasoundEditorGraphNode.h"
 
@@ -18,6 +18,13 @@ UMetasoundFactory::UMetasoundFactory(const FObjectInitializer& ObjectInitializer
 
 UObject* UMetasoundFactory::FactoryCreateNew(UClass* InClass, UObject* InParent, FName Name, EObjectFlags Flags, UObject* InContext, FFeedbackContext* InFeedbackContext)
 {
-	UMetasound* MetaSound = NewObject<UMetasound>(InParent, Name, Flags);
-	return MetaSound;
+	UMetasound* Metasound = NewObject<UMetasound>(InParent, Name, Flags);
+
+	FMetasoundClassMetadata Metadata;
+	Metadata.NodeName = Metasound->GetName();
+	Metadata.NodeType = EMetasoundClassType::MetasoundGraph;
+	Metadata.AuthorName = FText::FromString(UKismetSystemLibrary::GetPlatformUserName());
+	Metasound->SetMetadata(Metadata);
+
+	return Metasound;
 }
diff --git a/Engine/Plugins/Runtime/Metasound/Source/MetasoundEditor/Public/IMetasoundEditor.h b/Engine/Plugins/Runtime/Metasound/Source/MetasoundEditor/Public/IMetasoundEditor.h
--- a/Engine/Plugins/Runtime/Metasound/Source/MetasoundEditor/Public/IMetasoundEditor.h
+++ b/Engine/Plugins/Runtime/Metasound/Source/MetasoundEditor/Public/IMetasoundEditor.h
@@ -13,8 +13,6 @@
 class IMetasoundEditor : public FAssetEditorToolkit
 {
 	virtual UMetasound* GetMetasound() const = 0;
-	virtual void SetSelection(TArray<UObject*> SelectedObjects) = 0;
+	virtual void SetSelection(const TArray<UObject*>& SelectedObjects) = 0;
 	virtual bool GetBoundsForSelectedNodes(FSlateRect& Rect, float Padding) = 0;
-	virtual int32 GetNumberOfSelectedNodes() const = 0;
-	virtual TSet<UObject*> GetSelectedNodes() const = 0;
 };
diff --git a/Engine/Plugins/Runtime/Metasound/Source/MetasoundEditor/Public/MetasoundEditorCommands.h b/Engine/Plugins/Runtime/Metasound/Source/MetasoundEditor/Public/MetasoundEditorCommands.h
--- a/Engine/Plugins/Runtime/Metasound/Source/MetasoundEditor/Public/MetasoundEditorCommands.h
+++ b/Engine/Plugins/Runtime/Metasound/Source/MetasoundEditor/Public/MetasoundEditorCommands.h
@@ -6,36 +6,45 @@
 #include "Framework/Commands/Commands.h"
 #include "EditorStyleSet.h"
 
-class FMetasoundEditorCommands : public TCommands<FMetasoundEditorCommands>
+namespace Metasound
 {
-public:
-	/** Constructor */
-	FMetasoundEditorCommands() 
-		: TCommands<FMetasoundEditorCommands>("MetasoundEditor", NSLOCTEXT("Contexts", "MetasoundEditor", "Metasound Graph Editor"), NAME_None, "MetasoundStyle")
+	namespace Editor
 	{
-	}
-	
-	/** Plays the Metasound */
-	TSharedPtr<FUICommandInfo> Play;
-	
-	/** Stops the currently playing Metasound */
-	TSharedPtr<FUICommandInfo> Stop;
+		class FEditorCommands : public TCommands<FEditorCommands>
+		{
+		public:
+			/** Constructor */
+			FEditorCommands()
+				: TCommands<FEditorCommands>("MetasoundEditor", NSLOCTEXT("Contexts", "MetasoundEditor", "Metasound Graph Editor"), NAME_None, "MetasoundStyle")
+			{
+			}
 
-	/** Plays stops the currently playing Metasound */
-	TSharedPtr<FUICommandInfo> TogglePlayback;
+			/** Plays the Metasound */
+			TSharedPtr<FUICommandInfo> Play;
 
-	/** Selects the SoundWave in the content browser */
-	TSharedPtr<FUICommandInfo> BrowserSync;
+			/** Stops the currently playing Metasound */
+			TSharedPtr<FUICommandInfo> Stop;
 
-	/** Breaks the node input/output link */
-	TSharedPtr<FUICommandInfo> BreakLink;
+			/** Compiles the Metasound */
+			TSharedPtr<FUICommandInfo> Compile;
 
-	/** Adds an input to the node */
-	TSharedPtr<FUICommandInfo> AddInput;
+			/** Plays stops the currently playing Metasound */
+			TSharedPtr<FUICommandInfo> TogglePlayback;
 
-	/** Removes an input from the node */
-	TSharedPtr<FUICommandInfo> DeleteInput;
+			/** Selects the Metasound in the content browser. If referencing Metasound nodes are selected, selects referenced assets instead. */
+			TSharedPtr<FUICommandInfo> BrowserSync;
 
-	/** Initialize commands */
-	virtual void RegisterCommands() override;
-};
+			/** Breaks the node input/output link */
+			TSharedPtr<FUICommandInfo> BreakLink;
+
+			/** Adds an input to the node */
+			TSharedPtr<FUICommandInfo> AddInput;
+
+			/** Removes an input from the node */
+			TSharedPtr<FUICommandInfo> DeleteInput;
+
+			/** Initialize commands */
+			virtual void RegisterCommands() override;
+		};
+	} // namespace Editor
+} // namespace Metasound
\ No newline at end of file
diff --git a/Engine/Plugins/Runtime/Metasound/Source/MetasoundEditor/Public/MetasoundEditorGraph.h b/Engine/Plugins/Runtime/Metasound/Source/MetasoundEditor/Public/MetasoundEditorGraph.h
--- a/Engine/Plugins/Runtime/Metasound/Source/MetasoundEditor/Public/MetasoundEditorGraph.h
+++ b/Engine/Plugins/Runtime/Metasound/Source/MetasoundEditor/Public/MetasoundEditorGraph.h
@@ -1,20 +1,30 @@
 // Copyright Epic Games, Inc. All Rights Reserved.
-
 #pragma once
 
 #include "CoreMinimal.h"
 #include "EdGraph/EdGraph.h"
 #include "Metasound.h"
+#include "MetasoundEditor.h"
 #include "UObject/ObjectMacros.h"
 
 #include "MetasoundEditorGraph.generated.h"
 
+// Forward Declarations
+struct FMetasoundDocument;
+
+
 UCLASS(MinimalAPI)
 class UMetasoundEditorGraph : public UEdGraph
 {
 	GENERATED_BODY()
 
 public:
+	UMetasound* GetMetasound() const;
+	UMetasound& GetMetasoundChecked() const;
+
+private:
 	UPROPERTY()
-	UMetasound* Metasound;
+	UMetasound* ParentMetasound;
+
+	friend class Metasound::Editor::FEditor;
 };
diff --git a/Engine/Plugins/Runtime/Metasound/Source/MetasoundEditor/Public/MetasoundEditorGraphNode.h b/Engine/Plugins/Runtime/Metasound/Source/MetasoundEditor/Public/MetasoundEditorGraphNode.h
--- a/Engine/Plugins/Runtime/Metasound/Source/MetasoundEditor/Public/MetasoundEditorGraphNode.h
+++ b/Engine/Plugins/Runtime/Metasound/Source/MetasoundEditor/Public/MetasoundEditorGraphNode.h
@@ -4,7 +4,9 @@
 
 #include "CoreMinimal.h"
 #include "UObject/ObjectMacros.h"
-#include "MetasoundEditorGraphNode_Base.h"
+#include "MetasoundFrontend.h"
+#include "MetasoundFrontendDataLayout.h"
+#include "MetasoundFrontendRegistries.h"
 
 #include "MetasoundEditorGraphNode.generated.h"
 
@@ -14,34 +16,33 @@ class USoundNode;
 
 
 UCLASS(MinimalAPI)
-class UMetasoundEditorGraphNode : public UMetasoundEditorGraphNode_Base
+class UMetasoundEditorGraphNode : public UEdGraphNode
 {
 	GENERATED_UCLASS_BODY()
 
+public:
 	/** Fix up the node's owner after being copied */
 	METASOUNDEDITOR_API void PostCopyNode();
+
 	/** Create a new input pin for this node */
 	METASOUNDEDITOR_API void CreateInputPin();
-	/** Add an input pin to this node and recompile the MetaSoundEditor */
-	METASOUNDEDITOR_API void AddInputPin();
-	/** Remove a specific input pin from this node and recompile the MetaSoundEditor */
-	METASOUNDEDITOR_API void RemoveInputPin(UEdGraphPin* InGraphPin);
+
 	/** Estimate the width of this Node from the length of its title */
 	METASOUNDEDITOR_API int32 EstimateNodeWidth() const;
-	/** Checks whether an input can be added to this node */
-	METASOUNDEDITOR_API bool CanAddInputPin() const;
-
 
-	// UMetasoundEditorGraphNode_Base interface
-	virtual void CreateInputPins() override;
-	// End of UMetasoundEditorGraphNode_Base interface
+	METASOUNDEDITOR_API void IteratePins(TUniqueFunction<void(UEdGraphPin* /* Pin */, int32 /* Index */)> Func, EEdGraphPinDirection InPinDirection = EGPD_MAX);
 
 	// UEdGraphNode interface
+	virtual void AllocateDefaultPins() override;
+	virtual void AutowireNewNode(UEdGraphPin* FromPin) override;
+	virtual bool CanCreateUnderSpecifiedSchema(const UEdGraphSchema* Schema) const override;
 	virtual FText GetNodeTitle(ENodeTitleType::Type TitleType) const override;
-	virtual void PrepareForCopying() override;
+	virtual FString GetDocumentationExcerptName() const override;
+	virtual FString GetDocumentationLink() const override;
 	virtual void GetNodeContextMenuActions(class UToolMenu* Menu, class UGraphNodeContextMenuContext* Context) const override;
 	virtual FText GetTooltipText() const override;
-	virtual FString GetDocumentationExcerptName() const override;
+	virtual void PrepareForCopying() override;
+	virtual void ReconstructNode() override;
 	// End of UEdGraphNode interface
 
 	// UObject interface
@@ -49,4 +50,11 @@ class UMetasoundEditorGraphNode : public UMetasoundEditorGraphNode_Base
 	virtual void PostEditImport() override;
 	virtual void PostDuplicate(bool bDuplicateForPIE) override;
 	// End of UObject interface
+
+	Metasound::Frontend::FNodeHandle GetNodeHandle() const;
+	void SetNodeID(uint32 InNodeID);
+
+protected:
+	UPROPERTY()
+	uint32 NodeID = INDEX_NONE;
 };
diff --git a/Engine/Plugins/Runtime/Metasound/Source/MetasoundEditor/Public/MetasoundEditorGraphNode_Base.h b/Engine/Plugins/Runtime/Metasound/Source/MetasoundEditor/Public/MetasoundEditorGraphNode_Base.h
deleted file mode 100644
diff --git a/Engine/Plugins/Runtime/Metasound/Source/MetasoundEditor/Public/MetasoundEditorGraphNode_Root.h b/Engine/Plugins/Runtime/Metasound/Source/MetasoundEditor/Public/MetasoundEditorGraphNode_Root.h
deleted file mode 100644
diff --git a/Engine/Plugins/Runtime/Metasound/Source/MetasoundEditor/Public/MetasoundEditorGraphSchema.h b/Engine/Plugins/Runtime/Metasound/Source/MetasoundEditor/Public/MetasoundEditorGraphSchema.h
--- a/Engine/Plugins/Runtime/Metasound/Source/MetasoundEditor/Public/MetasoundEditorGraphSchema.h
+++ b/Engine/Plugins/Runtime/Metasound/Source/MetasoundEditor/Public/MetasoundEditorGraphSchema.h
@@ -3,15 +3,121 @@
 #pragma once
 
 #include "AssetData.h"
+#include "ConnectionDrawingPolicy.h"
 #include "CoreMinimal.h"
 #include "EdGraph/EdGraphSchema.h"
+#include "EdGraphUtilities.h"
+#include "MetasoundFrontend.h"
 #include "UObject/ObjectMacros.h"
 
 #include "MetasoundEditorGraphSchema.generated.h"
 
+
+// Forward Declarations
 class UEdGraph;
-class UMetasound;
 class UEdGraphNode;
+class UMetasound;
+class UMetasoundEditorGraphNode;
+
+
+namespace Metasound
+{
+	namespace Editor
+	{
+		struct FGraphConnectionDrawingPolicyFactory : public FGraphPanelPinConnectionFactory
+		{
+		public:
+			virtual ~FGraphConnectionDrawingPolicyFactory() = default;
+
+			// FGraphPanelPinConnectionFactory
+			virtual class FConnectionDrawingPolicy* CreateConnectionPolicy(
+				const UEdGraphSchema* Schema,
+				int32 InBackLayerID,
+				int32 InFrontLayerID,
+				float ZoomFactor,
+				const FSlateRect& InClippingRect,
+				FSlateWindowElementList& InDrawElements,
+				UEdGraph* InGraphObj) const override;
+			// ~FGraphPanelPinConnectionFactory
+		};
+
+		// This class draws the connections for an UEdGraph using a SoundCue schema
+		class FGraphConnectionDrawingPolicy : public FConnectionDrawingPolicy
+		{
+		protected:
+			// Times for one execution pair within the current graph
+			struct FTimePair
+			{
+				double PredExecTime;
+				double ThisExecTime;
+
+				FTimePair()
+					: PredExecTime(0.0)
+					, ThisExecTime(0.0)
+				{
+				}
+			};
+
+			// Map of pairings
+			using FExecPairingMap = TMap<UEdGraphNode*, FTimePair>;
+
+			// Map of nodes that preceded before a given node in the execution sequence (one entry for each pairing)
+			TMap<UEdGraphNode*, FExecPairingMap> PredecessorNodes;
+
+			UEdGraph* GraphObj = nullptr;
+
+			float ActiveWireThickness;
+			float InactiveWireThickness;
+
+		public:
+			FGraphConnectionDrawingPolicy(int32 InBackLayerID, int32 InFrontLayerID, float ZoomFactor, const FSlateRect& InClippingRect, FSlateWindowElementList& InDrawElements, UEdGraph* InGraphObj);
+
+			// FConnectionDrawingPolicy interface
+			virtual void DetermineWiringStyle(UEdGraphPin* OutputPin, UEdGraphPin* InputPin, /*inout*/ FConnectionParams& Params) override;
+			// End of FConnectionDrawingPolicy interface
+		};
+	} // namespace Editor
+} // namespace Metasound
+
+/** Action to add an input to the graph */
+USTRUCT()
+struct METASOUNDEDITOR_API FMetasoundGraphSchemaAction_NewInput : public FEdGraphSchemaAction
+{
+	GENERATED_USTRUCT_BODY();
+
+	UPROPERTY()
+	FName NodeTypeName;
+
+	FMetasoundGraphSchemaAction_NewInput()
+		: FEdGraphSchemaAction()
+	{}
+
+	FMetasoundGraphSchemaAction_NewInput(FText InNodeCategory, FText InDisplayName, FName InTypeName, FText InToolTip, const int32 InGrouping);
+
+	//~ Begin FEdGraphSchemaAction Interface
+	virtual UEdGraphNode* PerformAction(UEdGraph* ParentGraph, UEdGraphPin* FromPin, const FVector2D Location, bool bSelectNewNode = true) override;
+	//~ End FEdGraphSchemaAction Interface
+};
+
+/** Action to add an output to the graph */
+USTRUCT()
+struct METASOUNDEDITOR_API FMetasoundGraphSchemaAction_NewOutput : public FEdGraphSchemaAction
+{
+	GENERATED_USTRUCT_BODY();
+
+	UPROPERTY()
+	FName NodeTypeName;
+
+	FMetasoundGraphSchemaAction_NewOutput()
+		: FEdGraphSchemaAction()
+	{}
+
+	FMetasoundGraphSchemaAction_NewOutput(FText InNodeCategory, FText InDisplayName, FName InTypeName, FText InToolTip, const int32 InGrouping);
+
+	//~ Begin FEdGraphSchemaAction Interface
+	virtual UEdGraphNode* PerformAction(UEdGraph* ParentGraph, UEdGraphPin* FromPin, const FVector2D Location, bool bSelectNewNode = true) override;
+	//~ End FEdGraphSchemaAction Interface
+};
 
 /** Action to add a node to the graph */
 USTRUCT()
@@ -19,18 +125,15 @@ struct METASOUNDEDITOR_API FMetasoundGraphSchemaAction_NewNode : public FEdGraph
 {
 	GENERATED_USTRUCT_BODY();
 
-	/** Class of node we want to create */
-	UPROPERTY()
-	UClass* MetasoundNodeClass;
+	/** ClassInfo of node to create */
+	Metasound::Frontend::FNodeClassInfo NodeClassInfo;
 
 	FMetasoundGraphSchemaAction_NewNode() 
 		: FEdGraphSchemaAction()
-		, MetasoundNodeClass(nullptr)
 	{}
 
 	FMetasoundGraphSchemaAction_NewNode(FText InNodeCategory, FText InMenuDesc, FText InToolTip, const int32 InGrouping)
 		: FEdGraphSchemaAction(MoveTemp(InNodeCategory), MoveTemp(InMenuDesc), MoveTemp(InToolTip), InGrouping)
-		, MetasoundNodeClass(nullptr)
 	{}
 
 	//~ Begin FEdGraphSchemaAction Interface
@@ -39,7 +142,7 @@ struct METASOUNDEDITOR_API FMetasoundGraphSchemaAction_NewNode : public FEdGraph
 
 private:
 	/** Connects new node to output of selected nodes */
-	void ConnectToSelectedNodes(UMetasound* NewNodeclass, UEdGraph* ParentGraph) const;
+	void ConnectToSelectedNodes(UMetasoundEditorGraphNode* NewGraphNode, UEdGraph* ParentGraph) const;
 };
 
 /** Action to add nodes to the graph based on selected objects*/
diff --git a/Engine/Plugins/Runtime/Metasound/Source/MetasoundEditor/Public/MetasoundEditorUtilities.h b/Engine/Plugins/Runtime/Metasound/Source/MetasoundEditor/Public/MetasoundEditorUtilities.h
deleted file mode 100644
diff --git a/Engine/Plugins/Runtime/Metasound/Source/MetasoundEngine/Public/Metasound.h b/Engine/Plugins/Runtime/Metasound/Source/MetasoundEngine/Public/Metasound.h
--- a/Engine/Plugins/Runtime/Metasound/Source/MetasoundEngine/Public/Metasound.h
+++ b/Engine/Plugins/Runtime/Metasound/Source/MetasoundEngine/Public/Metasound.h
@@ -14,6 +14,24 @@
 #include "Metasound.generated.h"
 
 
+USTRUCT()
+struct METASOUNDENGINE_API FEdMetasoundInput
+{
+	GENERATED_BODY()
+
+	UPROPERTY()
+	UEdGraphNode* InputNode = nullptr;
+};
+
+USTRUCT()
+struct METASOUNDENGINE_API FEdMetasoundOutput
+{
+	GENERATED_BODY()
+
+	UPROPERTY()
+	UEdGraphNode* OutputNode = nullptr;
+};
+
 /**
  * This asset type is used for Metasound assets that can only be used as nodes in other Metasound graphs.
  * Because of this, they can have any inputs or outputs they need.
@@ -24,12 +42,18 @@ class METASOUNDENGINE_API UMetasound : public UObject, public FMetasoundAssetBas
 	GENERATED_BODY()
 
 protected:
-	UPROPERTY()
+	UPROPERTY(EditAnywhere, Category = Hidden)
 	FMetasoundDocument RootMetasoundDocument;
 
 #if WITH_EDITORONLY_DATA
 	UPROPERTY()
 	UEdGraph* Graph;
+
+	UPROPERTY()
+	TArray<FEdMetasoundInput> Inputs;
+
+	UPROPERTY()
+	TArray<FEdMetasoundOutput> Outputs;
 #endif // WITH_EDITORONLY_DATA
 
 public:
diff --git a/Engine/Plugins/Runtime/Metasound/Source/MetasoundFrontend/MetasoundFrontend.Build.cs b/Engine/Plugins/Runtime/Metasound/Source/MetasoundFrontend/MetasoundFrontend.Build.cs
--- a/Engine/Plugins/Runtime/Metasound/Source/MetasoundFrontend/MetasoundFrontend.Build.cs
+++ b/Engine/Plugins/Runtime/Metasound/Source/MetasoundFrontend/MetasoundFrontend.Build.cs
@@ -7,21 +7,26 @@ namespace UnrealBuildTool.Rules
 {
 	public class MetasoundFrontend : ModuleRules
 	{
-        public MetasoundFrontend(ReadOnlyTargetRules Target) : base(Target)
+		public MetasoundFrontend(ReadOnlyTargetRules Target) : base(Target)
 		{
-            PublicDependencyModuleNames.AddRange(
-				new string[] {
-                    "Core",
+			OptimizeCode = CodeOptimization.Never;
+
+			PublicDependencyModuleNames.AddRange
+			(
+				new string[]
+				{
+					"Core",
+					"Serialization",
 					"SignalProcessing"
-                }
-            );
+				}
+			);
 
-			PrivateDependencyModuleNames.AddRange(
+			PrivateDependencyModuleNames.AddRange
+			(
 				new string[]
 				{
 					"MetasoundGraphCore",
-					"CoreUObject",
-					"Serialization"
+					"CoreUObject"
 				}
 			);
 
diff --git a/Engine/Plugins/Runtime/Metasound/Source/MetasoundFrontend/Private/MetasoundFrontend.cpp b/Engine/Plugins/Runtime/Metasound/Source/MetasoundFrontend/Private/MetasoundFrontend.cpp
--- a/Engine/Plugins/Runtime/Metasound/Source/MetasoundFrontend/Private/MetasoundFrontend.cpp
+++ b/Engine/Plugins/Runtime/Metasound/Source/MetasoundFrontend/Private/MetasoundFrontend.cpp
@@ -626,10 +626,23 @@ namespace Metasound
 			return InHandle.ConnectWithConverterNode(*this, InNodeClassName);
 		}
 
+		TDescriptionPtr<FMetasoundClassDescription> FNodeHandle::GetNodeClassDescriptionForNodeHandle(const FHandleInitParams& InitParams, EMetasoundClassType InNodeClassType)
+		{
+			if (InNodeClassType != EMetasoundClassType::Input && InNodeClassType != EMetasoundClassType::Output)
+			{
+				return TDescriptionPtr<FMetasoundClassDescription>(InitParams.InAccessPoint, Path::GetDependencyPath(InitParams.InClassName));
+			}
+			else
+			{
+				// input nodes and output nodes don't have class descriptions.
+				return TDescriptionPtr<FMetasoundClassDescription>(nullptr, FDescPath());
+			}
+		}
+
 		FNodeHandle::FNodeHandle(FHandleInitParams::EPrivateToken PrivateToken, const FHandleInitParams& InParams, EMetasoundClassType InNodeClassType)
 			: ITransactable(MetasoundUndoRollLimitCvar, InParams.InOwningAsset)
 			, NodePtr(InParams.InAccessPoint, InParams.InPath)
-			, NodeClass(InParams.InAccessPoint, Path::GetDependencyPath(InParams.InClassName))
+			, NodeClass(GetNodeClassDescriptionForNodeHandle(InParams, InNodeClassType))
 			, NodeClassType(InNodeClassType)
 		{
 			if (InParams.InAccessPoint.IsValid())
@@ -810,9 +823,35 @@ namespace Metasound
 				return DefaultClassName;
 			}
 
+			if (NodeClassType == EMetasoundClassType::Input)
+			{
+				static const FString InputClassName = TEXT("Input");
+				return InputClassName;
+			}
+			else if (NodeClassType == EMetasoundClassType::Output)
+			{
+				static const FString OutputClassName = TEXT("Output");
+				return OutputClassName;
+			}
+
 			return NodeClass->Metadata.NodeName;
 		}
 
+		FNodeClassInfo FNodeHandle::GetClassInfo() const
+		{
+			FNodeClassInfo ClassInfo;
+
+			if (IsValid())
+			{
+				ClassInfo.NodeName = NodeClass->Metadata.NodeName;
+				ClassInfo.NodeType = NodeClass->Metadata.NodeType;
+				ClassInfo.LookupKey.NodeHash = NodeClass->ExternalNodeClassLookupInfo.ExternalNodeClassHash;
+				ClassInfo.LookupKey.NodeName = NodeClass->ExternalNodeClassLookupInfo.ExternalNodeClassName;
+			}
+
+			return ClassInfo;
+		}
+
 		void FNodeHandle::GetContainedGraph(FGraphHandle& OutGraph)
 		{
 			if (!IsValid())
@@ -1099,6 +1138,44 @@ namespace Metasound
 			return OutArray;
 		}
 
+		bool FGraphHandle::ContainsOutputNodeWithName(const FString& InName) const
+		{
+			if (!IsValid())
+			{
+				return false;
+			}
+
+			const TArray<FMetasoundNodeDescription>& NodeDescriptions = GraphPtr->Nodes;
+			for (const FMetasoundNodeDescription& NodeDescription : NodeDescriptions)
+			{
+				if (NodeDescription.ObjectTypeOfNode == EMetasoundClassType::Output && NodeDescription.Name == InName)
+				{
+					return true;
+				}
+			}
+
+			return false;
+		}
+
+		bool FGraphHandle::ContainsInputNodeWithName(const FString& InName) const
+		{
+			if (!IsValid())
+			{
+				return false;
+			}
+
+			const TArray<FMetasoundNodeDescription>& NodeDescriptions = GraphPtr->Nodes;
+			for (const FMetasoundNodeDescription& NodeDescription : NodeDescriptions)
+			{
+				if (NodeDescription.ObjectTypeOfNode == EMetasoundClassType::Input && NodeDescription.Name == InName)
+				{
+					return true;
+				}
+			}
+
+			return false;
+		}
+
 		FNodeHandle FGraphHandle::GetOutputNodeWithName(const FString& InName)
 		{
 			if (!IsValid())
@@ -1162,7 +1239,6 @@ namespace Metasound
 				return FNodeHandle::InvalidHandle();
 			}
 
-			
 			TArray<FMetasoundInputDescription>& Inputs = GraphsClassDeclaration->Inputs;
 
 			// Sanity check that this input has a unique name.
@@ -1237,7 +1313,7 @@ namespace Metasound
 			return true;
 		}
 
-		FNodeHandle FGraphHandle::AddNewOutput(const FMetasoundOutputDescription& OutDescription)
+		FNodeHandle FGraphHandle::AddNewOutput(const FMetasoundOutputDescription& InDescription)
 		{
 			if (!IsValid())
 			{
@@ -1259,18 +1335,18 @@ namespace Metasound
 			// Sanity check that this input has a unique name.
 			for (const FMetasoundOutputDescription& Output : Outputs)
 			{
-				if (!ensureAlwaysMsgf(Output.Name != OutDescription.Name, TEXT("Tried to add a new output with a name that already exists!")))
+				if (!ensureAlwaysMsgf(Output.Name != InDescription.Name, TEXT("Tried to add a new output with a name that already exists!")))
 				{
 					return FNodeHandle::InvalidHandle();
 				}
 			}
 
 			// Add the output to this node's class description.
-			Outputs.Add(OutDescription);
+			Outputs.Add(InDescription);
 
 			// Add a node for this output to the graph description.
 			FMetasoundNodeDescription NewNodeDescription;
-			NewNodeDescription.Name = OutDescription.Name;
+			NewNodeDescription.Name = InDescription.Name;
 			NewNodeDescription.UniqueID = NewUniqueId;
 			NewNodeDescription.ObjectTypeOfNode = EMetasoundClassType::Output;
 
diff --git a/Engine/Plugins/Runtime/Metasound/Source/MetasoundFrontend/Private/MetasoundFrontendBaseClasses.cpp b/Engine/Plugins/Runtime/Metasound/Source/MetasoundFrontend/Private/MetasoundFrontendBaseClasses.cpp
--- a/Engine/Plugins/Runtime/Metasound/Source/MetasoundFrontend/Private/MetasoundFrontendBaseClasses.cpp
+++ b/Engine/Plugins/Runtime/Metasound/Source/MetasoundFrontend/Private/MetasoundFrontendBaseClasses.cpp
@@ -468,7 +468,7 @@ namespace Metasound
 				CurrentStep = GoToNext(CurrentPath, CurrentStep);
 			}
 
-			if (ensureAlwaysMsgf(CurrentPath.Num() == 0, TEXT("Couldn't resolve part of the path.")))
+			if (ensureAlwaysMsgf(CurrentStep.Type == Path::EDescType::Class && CurrentPath.Num() == 0, TEXT("Couldn't resolve part of the path.")))
 			{
 				return CurrentStep.DescriptionStructPtr.Get<FMetasoundClassDescription*>();
 			}
diff --git a/Engine/Plugins/Runtime/Metasound/Source/MetasoundFrontend/Public/MetasoundFrontend.h b/Engine/Plugins/Runtime/Metasound/Source/MetasoundFrontend/Public/MetasoundFrontend.h
--- a/Engine/Plugins/Runtime/Metasound/Source/MetasoundFrontend/Public/MetasoundFrontend.h
+++ b/Engine/Plugins/Runtime/Metasound/Source/MetasoundFrontend/Public/MetasoundFrontend.h
@@ -35,10 +35,13 @@ namespace Metasound
 	 */
 	namespace Frontend
 	{
+		// Forward Declarations
+		class FInputHandle;
+
 		// Struct with the basics of a node class' information,
 		// used to look up that node from our node browser functions,
 		// and also used in FGraphHandle::AddNewNode.
-		struct FNodeClassInfo
+		struct METASOUNDFRONTEND_API FNodeClassInfo
 		{
 			// The descriptive name of this node class.
 			FString NodeName;
@@ -51,32 +54,31 @@ namespace Metasound
 		};
 
 		// Get all available nodes of any type.
-		TArray<FNodeClassInfo> GetAllAvailableNodeClasses();
+		METASOUNDFRONTEND_API TArray<FNodeClassInfo> GetAllAvailableNodeClasses();
 
 		// Get all possible nodes whose name begins with a specific namespace.
-		TArray<FNodeClassInfo> GetAllNodeClassesInNamespace(const FString& InNamespace);
+		METASOUNDFRONTEND_API TArray<FNodeClassInfo> GetAllNodeClassesInNamespace(const FString& InNamespace);
 
 		// Similar to GetAllNodeClassesInNamespace, except searches all nodes for a match with a substring.
-		TArray<FNodeClassInfo> GetAllNodesWhoseNameContains(const FString& InSubstring);
+		METASOUNDFRONTEND_API TArray<FNodeClassInfo> GetAllNodesWhoseNameContains(const FString& InSubstring);
 
 		// Searches for any node types that can output a specific type.
-		TArray<FNodeClassInfo> GetAllNodesWithAnOutputOfType(const FName& InType);
+		METASOUNDFRONTEND_API TArray<FNodeClassInfo> GetAllNodesWithAnOutputOfType(const FName& InType);
 
-		TArray<FNodeClassInfo> GetAllNodesWithAnInputOfType(const FName& InType);
+		METASOUNDFRONTEND_API TArray<FNodeClassInfo> GetAllNodesWithAnInputOfType(const FName& InType);
 
 		// gets all metadata (name, description, author, what to say if it's missing) for a given node.
-		FMetasoundClassMetadata GenerateMetadataForNode(const FNodeClassInfo& InInfo);
-
+		METASOUNDFRONTEND_API FMetasoundClassMetadata GenerateMetadataForNode(const FNodeClassInfo& InInfo);
 
 		// Generates a new FMetasoundClassDescription for a given node class. Only used by classes that manipulate Metasound Description data directly.
-		FMetasoundClassDescription GenerateClassDescriptionForNode(const FNodeClassInfo& InInfo);
+		METASOUNDFRONTEND_API FMetasoundClassDescription GenerateClassDescriptionForNode(const FNodeClassInfo& InInfo);
 
 		// Returns a list of all available data types.
-		TArray<FName> GetAllAvailableDataTypes();
+		METASOUNDFRONTEND_API TArray<FName> GetAllAvailableDataTypes();
 
 		// outputs the traits for a given data type.
 		// returns false if InDataType couldn't be found.
-		bool GetTraitsForDataType(FName InDataType, FDataTypeRegistryInfo& OutInfo);
+		METASOUNDFRONTEND_API bool GetTraitsForDataType(FName InDataType, FDataTypeRegistryInfo& OutInfo);
 
 		// This will open a json document at the given absolute path and convert it into a 
 		// metasound document struct. 
@@ -85,7 +87,7 @@ namespace Metasound
 
 		// Struct that indicates whether an input and an output can be connected,
 		// and whether an intermediate node is necessary to connect the two.
-		struct FConnectability
+		struct METASOUNDFRONTEND_API FConnectability
 		{
 			enum class EConnectable
 			{
@@ -102,9 +104,7 @@ namespace Metasound
 			TArray<FNodeClassInfo> PossibleConverterNodeClasses;
 		};
 
-		class FInputHandle;
-
-		struct FHandleInitParams
+		struct METASOUNDFRONTEND_API FHandleInitParams
 		{
 			TWeakPtr<FDescriptionAccessPoint> InAccessPoint;
 
@@ -247,6 +247,8 @@ namespace Metasound
 			FInputHandle GetInputWithName(const FString& InName);
 			FOutputHandle GetOutputWithName(const FString& InName);
 
+			FNodeClassInfo GetClassInfo() const;
+
 			EMetasoundClassType GetNodeType() const;
 			const FString& GetNodeClassName() const;
 
@@ -260,6 +262,8 @@ namespace Metasound
 
 		private:
 
+			static TDescriptionPtr<FMetasoundClassDescription> GetNodeClassDescriptionForNodeHandle(const FHandleInitParams& InitParams, EMetasoundClassType InNodeClassType);
+
 			TDescriptionPtr<FMetasoundNodeDescription> NodePtr;
 			TDescriptionPtr<FMetasoundClassDescription> NodeClass;
 
@@ -297,13 +301,16 @@ namespace Metasound
 			TArray<FNodeHandle> GetOutputNodes();
 			TArray<FNodeHandle> GetInputNodes();
 
+			bool ContainsOutputNodeWithName(const FString& InName) const;
+			bool ContainsInputNodeWithName(const FString& InName) const;
+
 			FNodeHandle GetOutputNodeWithName(const FString& InName);
 			FNodeHandle GetInputNodeWithName(const FString& InName);
 
 			FNodeHandle AddNewInput(const FMetasoundInputDescription& InDescription);
 			bool RemoveInput(const FString& InputName);
 
-			FNodeHandle AddNewOutput(const FMetasoundOutputDescription& OutDescription);
+			FNodeHandle AddNewOutput(const FMetasoundOutputDescription& InDescription);
 			bool RemoveOutput(const FString& OutputName);
 
 			// This can be used to determine what kind of property editor we should use for the data type of a given input.
diff --git a/Engine/Plugins/Runtime/Metasound/Source/MetasoundFrontend/Public/MetasoundFrontendBaseClasses.h b/Engine/Plugins/Runtime/Metasound/Source/MetasoundFrontend/Public/MetasoundFrontendBaseClasses.h
--- a/Engine/Plugins/Runtime/Metasound/Source/MetasoundFrontend/Public/MetasoundFrontendBaseClasses.h
+++ b/Engine/Plugins/Runtime/Metasound/Source/MetasoundFrontend/Public/MetasoundFrontendBaseClasses.h
@@ -630,6 +630,7 @@ namespace Metasound
 
 			TWeakPtr<FDescriptionAccessPoint> AccessPoint;
 			FDescPath PathFromRoot;
+			bool bAccessFailure = false;
 
 		public:
 			TDescriptionPtr() = delete;
@@ -638,13 +639,21 @@ namespace Metasound
 				: AccessPoint(InAccessPoint)
 				, PathFromRoot(InPathFromRoot)
 			{
+				// Test accessor
+				bAccessFailure = Get() == nullptr;
 			}
 
+			TDescriptionPtr(TDescriptionPtr&& InDescriptionPtr) = default;
+			TDescriptionPtr(const TDescriptionPtr& InDescriptionPtr) = default;
+
+			TDescriptionPtr& operator =(TDescriptionPtr&& InDescriptionPtr) = default;
+			TDescriptionPtr& operator =(const TDescriptionPtr & InDescriptionPtr) = default;
+
 			// @returns false if the description this accessor was referencing has been destroyed,
 			// or if this was created with bad arguments.
 			bool IsValid() const
 			{
-				return AccessPoint.IsValid() && PathFromRoot.IsValid();
+				return AccessPoint.IsValid() && PathFromRoot.IsValid() && !bAccessFailure;
 			}
 
 			FDescPath GetPath() const
diff --git a/Engine/Plugins/Runtime/Metasound/Source/MetasoundFrontend/Public/MetasoundFrontendDataLayout.h b/Engine/Plugins/Runtime/Metasound/Source/MetasoundFrontend/Public/MetasoundFrontendDataLayout.h
--- a/Engine/Plugins/Runtime/Metasound/Source/MetasoundFrontend/Public/MetasoundFrontendDataLayout.h
+++ b/Engine/Plugins/Runtime/Metasound/Source/MetasoundFrontend/Public/MetasoundFrontendDataLayout.h
@@ -1,11 +1,12 @@
 // Copyright Epic Games, Inc. All Rights Reserved.
-
 #pragma once
 
 #include "CoreMinimal.h"
+#include "Internationalization/Text.h"
 #include "MetasoundGraph.h"
 #include "MetasoundFrontendDataLayout.generated.h"
 
+
 UENUM()
 enum class EMetasoundClassType : uint8
 {
@@ -43,23 +44,21 @@ struct FMetasoundClassMetadata
 	FString NodeName;
 
 	// This will always be set to EMetasoundObjectType::Metasound.
-	UPROPERTY()
+	UPROPERTY(VisibleAnywhere, Category = General)
 	EMetasoundClassType NodeType;
 
-	// Optional longform description of what this Metasound does.
-	// Can be displayed in a tooltip.
-	UPROPERTY()
-	FString MetasoundDescription;
+	// Optional longform description of what this Metasound does. Can be displayed in a tooltip.
+	UPROPERTY(EditAnywhere, Category = General, meta = (DisplayName = "Description"))
+	FText MetasoundDescription;
 
-	// If another Metasound depends on this metasoundand doesn't package it, 
-	// this prompt will show if we can't find this Metasound. 
-	// It should describe how to get this metasound.
+	// Prompt that will will show if Metasound is missing as a dependency listed in another Metasound.
+	// Can optionally provide hints as to how to obtain this Metasound (ex. documentation URL, Plugin name).
 	UPROPERTY()
-	FString PromptIfMissing;
+	FText PromptIfMissing;
 
 	// Original author of this Metasound Object.
-	UPROPERTY()
-	FString AuthorName;
+	UPROPERTY(EditAnywhere, Category = General, meta = (DisplayName = "Author"))
+	FText AuthorName;
 };
 
 
@@ -222,7 +221,7 @@ struct FMetasoundClassDescription
 	UPROPERTY()
 	uint32 UniqueID;
 
-	UPROPERTY()
+	UPROPERTY(EditAnywhere, Category = Hidden)
 	FMetasoundClassMetadata Metadata;
 
 	UPROPERTY()
@@ -281,7 +280,7 @@ struct FMetasoundDocument
 
 	// The highest level Metasound Class in this document.
 	// Typically of type MetasoundGraph, contains graph of the node classes listed in Dependencies.
-	UPROPERTY()
+	UPROPERTY(EditAnywhere, Category = Hidden)
 	FMetasoundClassDescription RootClass;
 
 	// this is the list of every dependency required by the RootClass, as well as nested dependencies.
diff --git a/Engine/Plugins/Runtime/Metasound/Source/MetasoundGraphCore/Private/MetasoundGraph.cpp b/Engine/Plugins/Runtime/Metasound/Source/MetasoundGraphCore/Private/MetasoundGraph.cpp
--- a/Engine/Plugins/Runtime/Metasound/Source/MetasoundGraphCore/Private/MetasoundGraph.cpp
+++ b/Engine/Plugins/Runtime/Metasound/Source/MetasoundGraphCore/Private/MetasoundGraph.cpp
@@ -1,24 +1,19 @@
 // Copyright Epic Games, Inc. All Rights Reserved.
-
 #include "MetasoundGraph.h"
 
+#include "Internationalization/Text.h"
+
+
 namespace Metasound
 {
-	const FName FGraph::ClassName = FName(TEXT("Graph"));
-
 	FGraph::FGraph(const FString& InDescription)
 		: Description(InDescription)
 	{
 	}
 
-	FGraph::~FGraph()
+	const FText& FGraph::GetDescription() const
 	{
-	}
-
-	const FString& FGraph::GetDescription() const
-	{
-		static FString ClassDescription(TEXT("Unused"));
-		return ClassDescription;
+		return FText::GetEmpty();
 	}
 
 	const FString& FGraph::GetInstanceName() const
@@ -28,20 +23,18 @@ namespace Metasound
 
 	const FName& FGraph::GetClassName() const
 	{
+		static const FName ClassName("Graph");
 		return ClassName;
 	}
 
-
-	const FString& FGraph::GetAuthorName() const
+	const FText& FGraph::GetAuthorName() const
 	{
-		static FString Author(TEXT("Unused"));
-		return Author;
+		return PluginAuthor;
 	}
 
-	const FString& FGraph::GetPromptIfMissing() const
+	const FText& FGraph::GetPromptIfMissing() const
 	{
-		static FString Prompt(TEXT("Unused"));
-		return Prompt;
+		return PluginNodeMissingPrompt;
 	}
 
 	const FInputDataVertexCollection& FGraph::GetInputDataVertices() const
diff --git a/Engine/Plugins/Runtime/Metasound/Source/MetasoundGraphCore/Public/MetasoundADSRNode.h b/Engine/Plugins/Runtime/Metasound/Source/MetasoundGraphCore/Public/MetasoundADSRNode.h
--- a/Engine/Plugins/Runtime/Metasound/Source/MetasoundGraphCore/Public/MetasoundADSRNode.h
+++ b/Engine/Plugins/Runtime/Metasound/Source/MetasoundGraphCore/Public/MetasoundADSRNode.h
@@ -1,13 +1,14 @@
 // Copyright Epic Games, Inc. All Rights Reserved.
-
 #pragma once
 
 #include "CoreMinimal.h"
+#include "Internationalization/Text.h"
 #include "MetasoundNode.h"
 #include "MetasoundOperatorInterface.h"
 #include "MetasoundBuilderInterface.h"
 #include "MetasoundDataReferenceCollection.h"
 
+
 namespace Metasound
 {
 	/** FADSRNode
@@ -16,24 +17,22 @@ namespace Metasound
 	 */
 	class METASOUNDGRAPHCORE_API FADSRNode : public FNode
 	{
-
-			// The operator factory for this node.
-			class FOperatorFactory : public IOperatorFactory
-			{
-				public:
-					virtual TUniquePtr<IOperator> CreateOperator(const INode& InNode, const FOperatorSettings& InOperatorSettings, const FDataReferenceCollection& InInputDataReferences, TArray<TUniquePtr<IOperatorBuildError>>& OutErrors) override;
-
-					template<typename DataType>
-					void SetReadableRefIfInCollection(const FString& InName, const FDataReferenceCollection& InCollection, TDataReadReference<DataType>& DataRef)
+		// The operator factory for this node.
+		class FOperatorFactory : public IOperatorFactory
+		{
+			public:
+				virtual TUniquePtr<IOperator> CreateOperator(const INode& InNode, const FOperatorSettings& InOperatorSettings, const FDataReferenceCollection& InInputDataReferences, TArray<TUniquePtr<IOperatorBuildError>>& OutErrors) override;
+
+				template<typename DataType>
+				void SetReadableRefIfInCollection(const FString& InName, const FDataReferenceCollection& InCollection, TDataReadReference<DataType>& DataRef)
+				{
+					// TODO: add a helper function to FDataReferenceCollection to do a SetDataIfContains.
+					if (InCollection.ContainsDataReadReference<DataType>(InName))
 					{
-						// TODO: add a helper function to FDataReferenceCollection to do a SetDataIfContains.
-						if (InCollection.ContainsDataReadReference<DataType>(InName))
-						{
-							DataRef = InCollection.GetDataReadReference<DataType>(InName);
-						}
+						DataRef = InCollection.GetDataReadReference<DataType>(InName);
 					}
-
-			};
+				}
+		};
 
 		public:
 			/** Class name for FADSRNode */
@@ -68,29 +67,26 @@ namespace Metasound
 			/** Returns the type name of the FADSRNode */
 			virtual const FName& GetClassName() const override;
 
-			virtual const FString& GetDescription() const override
+			virtual const FText& GetDescription() const override
 			{
-				static FString StaticDescription = TEXT("This node, when bopped, emits an envelope that can be applied to .");
+				static const FText StaticDescription = NSLOCTEXT("MetasoundGraphCore", "Metasound_ADSRNodeDescription", "Emits an ADSR (Attack, decay, sustain, & release) envelope when bopped.");
 				return StaticDescription;
 			}
 
-			virtual const FString& GetAuthorName() const override
+			virtual const FText& GetAuthorName() const override
 			{
-				static FString Author = TEXT("Epic Games");
-				return Author;
+				return PluginAuthor;
 			}
 
-			virtual const FString& GetPromptIfMissing() const override
+			virtual const FText& GetPromptIfMissing() const override
 			{
-				static FString Prompt = TEXT("Make sure that the Metasound plugin is loaded.");
-				return Prompt;
+				return PluginNodeMissingPrompt;
 			}
 
 			/** Return a factory for building a metasound operator. */
 			virtual IOperatorFactory& GetDefaultOperatorFactory() override;
 
 		private:
-
 			float DefaultAttackMs;
 			float DefaultDecayMs;
 			float DefaultSustainMs;
@@ -98,6 +94,4 @@ namespace Metasound
 
 			FOperatorFactory Factory;
 	};
-
-
-}
+} // namespace Metasound
diff --git a/Engine/Plugins/Runtime/Metasound/Source/MetasoundGraphCore/Public/MetasoundAudioMultiplyNode.h b/Engine/Plugins/Runtime/Metasound/Source/MetasoundGraphCore/Public/MetasoundAudioMultiplyNode.h
--- a/Engine/Plugins/Runtime/Metasound/Source/MetasoundGraphCore/Public/MetasoundAudioMultiplyNode.h
+++ b/Engine/Plugins/Runtime/Metasound/Source/MetasoundGraphCore/Public/MetasoundAudioMultiplyNode.h
@@ -1,37 +1,37 @@
 // Copyright Epic Games, Inc. All Rights Reserved.
-
 #pragma once
 
 #include "CoreMinimal.h"
+#include "Internationalization/Text.h"
 #include "MetasoundNode.h"
 #include "MetasoundOperatorInterface.h"
 #include "MetasoundBuilderInterface.h"
 #include "MetasoundDataReferenceCollection.h"
 
+
 namespace Metasound
 {
 	class METASOUNDGRAPHCORE_API FAudioMultiplyNode : public FNode
 	{
-
-			class FOperatorFactory : public IOperatorFactory
-			{
-				public:
-					virtual TUniquePtr<IOperator> CreateOperator(const INode& InNode, const FOperatorSettings& InOperatorSettings, const FDataReferenceCollection& InInputDataReferences, TArray<TUniquePtr<IOperatorBuildError>>& OutErrors) override;
-
-					template<typename ParamType>
-					bool SetReadableRefIfInCollection(const FString& InParamName, const FDataReferenceCollection& InCollection, TDataReadReference<ParamType>& ParamRef)
+		class FOperatorFactory : public IOperatorFactory
+		{
+			public:
+				virtual TUniquePtr<IOperator> CreateOperator(const INode& InNode, const FOperatorSettings& InOperatorSettings, const FDataReferenceCollection& InInputDataReferences, TArray<TUniquePtr<IOperatorBuildError>>& OutErrors) override;
+
+				template<typename ParamType>
+				bool SetReadableRefIfInCollection(const FString& InParamName, const FDataReferenceCollection& InCollection, TDataReadReference<ParamType>& ParamRef)
+				{
+					// TODO: add a helper function to FDataReferenceCollection to do a SetParamIfContains.
+					if (InCollection.ContainsDataReadReference<ParamType>(InParamName))
 					{
-						// TODO: add a helper function to FDataReferenceCollection to do a SetParamIfContains.
-						if (InCollection.ContainsDataReadReference<ParamType>(InParamName))
-						{
-							ParamRef = InCollection.GetDataReadReference<ParamType>(InParamName);
-							return true;
-						}
-
-						return false;
+						ParamRef = InCollection.GetDataReadReference<ParamType>(InParamName);
+						return true;
 					}
 
-			};
+					return false;
+				}
+
+		};
 
 		public:
 			static const FName ClassName;
@@ -47,28 +47,23 @@ namespace Metasound
 
 			virtual IOperatorFactory& GetDefaultOperatorFactory() override;
 
-			virtual const FString& GetDescription() const override
+			virtual const FText& GetDescription() const override
 			{
-				static FString StaticDescription = TEXT("This node multiplies to audio signals together. This is useful for amplitude modulation and other applications.");
+				static const FText StaticDescription = NSLOCTEXT("MetasoundGraphCore", "Metasound_AudioMultipleNodeDescription", "Multiplies two audio signals together (useful for amplitude modulation).");
 				return StaticDescription;
 			}
 
-			virtual const FString& GetAuthorName() const override
+			virtual const FText& GetAuthorName() const override
 			{
-				static FString Author = TEXT("Epic Games");
-				return Author;
+				return PluginAuthor;
 			}
 
-			virtual const FString& GetPromptIfMissing() const override
+			virtual const FText& GetPromptIfMissing() const override
 			{
-				static FString Prompt = TEXT("Make sure that the Metasound plugin is loaded.");
-				return Prompt;
+				return PluginNodeMissingPrompt;
 			}
 
 		private:
-
 			FOperatorFactory Factory;
 	};
-
-	
-}
+} // namespace Metasound
diff --git a/Engine/Plugins/Runtime/Metasound/Source/MetasoundGraphCore/Public/MetasoundGraph.h b/Engine/Plugins/Runtime/Metasound/Source/MetasoundGraphCore/Public/MetasoundGraph.h
--- a/Engine/Plugins/Runtime/Metasound/Source/MetasoundGraphCore/Public/MetasoundGraph.h
+++ b/Engine/Plugins/Runtime/Metasound/Source/MetasoundGraphCore/Public/MetasoundGraph.h
@@ -1,10 +1,13 @@
 // Copyright Epic Games, Inc. All Rights Reserved.
-
 #pragma once
 
 #include "CoreMinimal.h"
 #include "MetasoundNode.h"
 
+// Forward Declarations
+class FText;
+
+
 namespace Metasound
 {
 	// TODO: consider making nodes TSharedPtr<INode>
@@ -15,10 +18,8 @@ namespace Metasound
 	class METASOUNDGRAPHCORE_API FGraph : public IGraph
 	{
 		public:
-			static const FName ClassName;
-
 			FGraph(const FString& InInstanceName);
-			virtual ~FGraph();
+			virtual ~FGraph() = default;
 
 			/** Return the name of this specific instance of the node class. */
 			virtual const FString& GetInstanceName() const override;
@@ -27,12 +28,12 @@ namespace Metasound
 			virtual const FName& GetClassName() const override;
 
 			/** Return a longer text description describing how this node is used. */
-			virtual const FString& GetDescription() const override;
+			virtual const FText& GetDescription() const override;
 
 			/** Return the original author of this node class. */
-			virtual const FString& GetAuthorName() const override;
+			virtual const FText& GetAuthorName() const override;
 
-			virtual const FString& GetPromptIfMissing() const override;
+			virtual const FText& GetPromptIfMissing() const override;
 
 			/** Return a collection of input parameter descriptions for this graph. */
 			virtual const FInputDataVertexCollection& GetInputDataVertices() const override;
diff --git a/Engine/Plugins/Runtime/Metasound/Source/MetasoundGraphCore/Public/MetasoundInputNode.h b/Engine/Plugins/Runtime/Metasound/Source/MetasoundGraphCore/Public/MetasoundInputNode.h
--- a/Engine/Plugins/Runtime/Metasound/Source/MetasoundGraphCore/Public/MetasoundInputNode.h
+++ b/Engine/Plugins/Runtime/Metasound/Source/MetasoundGraphCore/Public/MetasoundInputNode.h
@@ -1,101 +1,103 @@
 // Copyright Epic Games, Inc. All Rights Reserved.
-
 #pragma once
 
 #include "CoreMinimal.h"
+#include "Internationalization/Text.h"
 #include "MetasoundOperatorInterface.h"
 #include "MetasoundNodeInterface.h"
 
+#define LOCTEXT_NAMESPACE "MetasoundGraphCore"
+
+
 namespace Metasound
 {
 	template<typename DataType>
 	class TInputNode : public INode
 	{
-			static_assert(TDataReferenceTypeInfo<DataType>::bIsValidSpecialization, "Please use DECLARE_METASOUND_DATA_REFERENCE_TYPES with this class before trying to create an input node with it.");
-
-			class FInputOperator : public IOperator
-			{
-				public:
-					using FDataWriteReference = TDataWriteReference<DataType>;
-
-					FInputOperator(const FString& InDataReferenceName, FDataWriteReference InDataReference)
-					{
-						Inputs.AddDataWriteReference<DataType>(InDataReferenceName, InDataReference);
-						Outputs.AddDataReadReference<DataType>(InDataReferenceName, InDataReference);
-					}
-
-					virtual ~FInputOperator() {}
-
-					virtual const FDataReferenceCollection& GetInputs() const override
-					{
-						return Inputs;
-					}
-
-					virtual const FDataReferenceCollection& GetOutputs() const override
-					{
-						return Outputs;
-					}
-
-					virtual FExecuteFunction GetExecuteFunction() override
-					{
-						return nullptr;
-					}
-
-				private:
-					FDataReferenceCollection Outputs;
-					FDataReferenceCollection Inputs;
-			};
-
-			// TODO: Would be nice if we didn't utilize the copy constructor by default here. But it's a nice catch-all
-			// for our use case.  Could store constructor parameters of the paramtype, but run the risk of holding rvalues
-			// or pointers to objects which are no longer valid. 
-
-			// TODO: instead of this copy factory, register a DataType factory lambda that uses a FDataTypeLiteralParam
-			class FCopyOperatorFactory : public IOperatorFactory
-			{
-				static_assert(TIsConstructible<DataType, const DataType&>::Value, "Data Type must be copy constructible!");
+		static_assert(TDataReferenceTypeInfo<DataType>::bIsValidSpecialization, "Please use DECLARE_METASOUND_DATA_REFERENCE_TYPES with this class before trying to create an input node with it.");
+
+		class FInputOperator : public IOperator
+		{
+			public:
+				using FDataWriteReference = TDataWriteReference<DataType>;
+
+				FInputOperator(const FString& InDataReferenceName, FDataWriteReference InDataReference)
+				{
+					Inputs.AddDataWriteReference<DataType>(InDataReferenceName, InDataReference);
+					Outputs.AddDataReadReference<DataType>(InDataReferenceName, InDataReference);
+				}
+
+				virtual ~FInputOperator() {}
+
+				virtual const FDataReferenceCollection& GetInputs() const override
+				{
+					return Inputs;
+				}
+
+				virtual const FDataReferenceCollection& GetOutputs() const override
+				{
+					return Outputs;
+				}
+
+				virtual FExecuteFunction GetExecuteFunction() override
+				{
+					return nullptr;
+				}
+
+			private:
+				FDataReferenceCollection Outputs;
+				FDataReferenceCollection Inputs;
+		};
+
+		// TODO: Would be nice if we didn't utilize the copy constructor by default here. But it's a nice catch-all
+		// for our use case.  Could store constructor parameters of the paramtype, but run the risk of holding rvalues
+		// or pointers to objects which are no longer valid. 
+
+		// TODO: instead of this copy factory, register a DataType factory lambda that uses a FDataTypeLiteralParam
+		class FCopyOperatorFactory : public IOperatorFactory
+		{
+			static_assert(TIsConstructible<DataType, const DataType&>::Value, "Data Type must be copy constructible!");
 				
-				public:
-					using FDataWriteReference = TDataWriteReference<DataType>;
-					using FInputNodeType = TInputNode<DataType>;
-
-					template< typename... ArgTypes >
-					FCopyOperatorFactory(ArgTypes&&... Args)
-					:	Data(Forward<ArgTypes>(Args)...)
-					{
-					}
-
-					/*
-					template<typename = typename TEnableIf< TIsConstructible<DataType>::Value >::Type >
-					FCopyOperatorFactory()
-					:	Data()
-					{
-					}
-					*/
-
-					FCopyOperatorFactory(FDataTypeLiteralParam InitParam, const FOperatorSettings& InSettings)
-						: Data(InitParam.ParseTo<DataType>(InSettings))
-					{
-					}
-
-					virtual TUniquePtr<IOperator> CreateOperator(const INode& InNode, const FOperatorSettings& InOperatorSettings, const FDataReferenceCollection& InInputDataReferences, TArray<TUniquePtr<IOperatorBuildError>>& OutErrors) override
-					{
-						// Use copy constructor to create a new parameter reference.
-						FDataWriteReference DataRef = FDataWriteReference::CreateNew(Data);
-						// TODO: Write special version of this for audio types since they will need to be constructed based upon the operator settings. 
-						// TODO: Need to ponder how inputs are initialized, but template specialization might do the trick. 
-
-						const FInputNodeType& InputNode = static_cast<const FInputNodeType&>(InNode);
-
-						return MakeUnique<FInputOperator>(InputNode.GetVertexName(), DataRef);
-					}
-
-				private:
-					DataType Data;
-			};
+			public:
+				using FDataWriteReference = TDataWriteReference<DataType>;
+				using FInputNodeType = TInputNode<DataType>;
+
+				template< typename... ArgTypes >
+				FCopyOperatorFactory(ArgTypes&&... Args)
+				:	Data(Forward<ArgTypes>(Args)...)
+				{
+				}
+
+				/*
+				template<typename = typename TEnableIf< TIsConstructible<DataType>::Value >::Type >
+				FCopyOperatorFactory()
+				:	Data()
+				{
+				}
+				*/
+
+				FCopyOperatorFactory(FDataTypeLiteralParam InitParam, const FOperatorSettings& InSettings)
+					: Data(InitParam.ParseTo<DataType>(InSettings))
+				{
+				}
+
+				virtual TUniquePtr<IOperator> CreateOperator(const INode& InNode, const FOperatorSettings& InOperatorSettings, const FDataReferenceCollection& InInputDataReferences, TArray<TUniquePtr<IOperatorBuildError>>& OutErrors) override
+				{
+					// Use copy constructor to create a new parameter reference.
+					FDataWriteReference DataRef = FDataWriteReference::CreateNew(Data);
+					// TODO: Write special version of this for audio types since they will need to be constructed based upon the operator settings. 
+					// TODO: Need to ponder how inputs are initialized, but template specialization might do the trick. 
+
+					const FInputNodeType& InputNode = static_cast<const FInputNodeType&>(InNode);
+
+					return MakeUnique<FInputOperator>(InputNode.GetVertexName(), DataRef);
+				}
+
+			private:
+				DataType Data;
+		};
 
 		public:
-			
 			// Constructors that need to get working on linux:
 			/*
 			template<typename = typename TEnableIf< TIsConstructible<DataType>::Value >::Type>
@@ -136,7 +138,6 @@ namespace Metasound
 				Inputs.Add(InputVertexKey, InputVertex);
 
 			}
-			
 
 			explicit TInputNode(const FString& InNodeDescription, const FString& InVertexName, FDataTypeLiteralParam InParam, const FOperatorSettings& InSettings)
 				: NodeDescription(InNodeDescription)
@@ -163,26 +164,24 @@ namespace Metasound
 			{
 				// TODO: although this is ok with MSVC's lax template instantiation, every other compiler will complain about TDataReferenceTypeInfo.
 				//static const FName ClassName = FName(FString(TEXT("Input_")) + FString(TDataReferenceTypeInfo<DataType>::TypeName));
-				static const FName ClassName = FName(TEXT("InputNode"));
+				static const FName ClassName("InputNode");
 				return ClassName;
 			}
 
-			virtual const FString& GetDescription() const override
+			virtual const FText& GetDescription() const override
 			{
-				static FString Description = TEXT("An Input into the current Metasound graph.");
+				static const FText Description = LOCTEXT("Metasound_InputNodeDescription", "Input into the parent Metasound graph.");
 				return Description;
 			}
 
-			virtual const FString& GetAuthorName() const override
+			virtual const FText& GetAuthorName() const override
 			{
-				static FString Author = TEXT("Epic Games");
-				return Author;
+				return PluginAuthor;
 			}
 
-			virtual const FString& GetPromptIfMissing() const override
+			virtual const FText& GetPromptIfMissing() const override
 			{
-				static FString Prompt = TEXT("Make sure that the Metasound plugin is loaded.");
-				return Prompt;
+				return PluginNodeMissingPrompt;
 			}
 
 			const FString& GetVertexName() const
@@ -205,7 +204,6 @@ namespace Metasound
 				return Factory;
 			}
 
-
 		private:
 			FString NodeDescription;
 			FString VertexName;
@@ -215,4 +213,5 @@ namespace Metasound
 			FInputDataVertexCollection Inputs;
 			FOutputDataVertexCollection Outputs;
 	};
-}
+} // namespace Metasound
+#undef LOCTEXT_NAMESPACE //MetasoundOutputNode
diff --git a/Engine/Plugins/Runtime/Metasound/Source/MetasoundGraphCore/Public/MetasoundNodeInterface.h b/Engine/Plugins/Runtime/Metasound/Source/MetasoundGraphCore/Public/MetasoundNodeInterface.h
--- a/Engine/Plugins/Runtime/Metasound/Source/MetasoundGraphCore/Public/MetasoundNodeInterface.h
+++ b/Engine/Plugins/Runtime/Metasound/Source/MetasoundGraphCore/Public/MetasoundNodeInterface.h
@@ -3,11 +3,15 @@
 #pragma once
 
 #include "CoreMinimal.h"
+#include "Internationalization/Text.h"
 #include "MetasoundOperatorInterface.h"
 #include "MetasoundDataReference.h"
 
 namespace Metasound
 {
+	static const FText PluginAuthor = NSLOCTEXT("MetasoundGraphCore", "Metasound_DefaultAuthor", "Epic Games, Inc.");
+	static const FText PluginNodeMissingPrompt = NSLOCTEXT("MetasoundGraphCore", "Metasound_DefaultMissingPrompt", "Make sure that the Metasound plugin is loaded.");
+
 	/** FDataVertex
 	 *
 	 *  This describes a data vertex on a Node.
@@ -148,16 +152,16 @@ namespace Metasound
 			virtual const FName& GetClassName() const = 0;
 
 			/** Return a longer text description describing how this node is used. */
-			virtual const FString& GetDescription() const = 0;
+			virtual const FText& GetDescription() const = 0;
 
 			/** Return the original author of this node class. */
-			virtual const FString& GetAuthorName() const = 0;
+			virtual const FText& GetAuthorName() const = 0;
 
 			/** 
 			 *  Return an optional prompt on how users can get the plugin this node is in,
 			 *  if they have found a metasound that uses this node but don't have this plugin downloaded or enabled.
 			 */
-			virtual const FString& GetPromptIfMissing() const = 0;
+			virtual const FText& GetPromptIfMissing() const = 0;
 
 			/** Return a collection of input parameter descriptions for this node. */
 			virtual const FInputDataVertexCollection& GetInputDataVertices() const = 0;
diff --git a/Engine/Plugins/Runtime/Metasound/Source/MetasoundGraphCore/Public/MetasoundOscNode.h b/Engine/Plugins/Runtime/Metasound/Source/MetasoundGraphCore/Public/MetasoundOscNode.h
--- a/Engine/Plugins/Runtime/Metasound/Source/MetasoundGraphCore/Public/MetasoundOscNode.h
+++ b/Engine/Plugins/Runtime/Metasound/Source/MetasoundGraphCore/Public/MetasoundOscNode.h
@@ -33,26 +33,23 @@ namespace Metasound
 
 			virtual IOperatorFactory& GetDefaultOperatorFactory() override;
 
-			virtual const FString& GetDescription() const override
+			virtual const FText& GetDescription() const override
 			{
-				static FString StaticDescription = TEXT("This node emits an audio signal of a sinusoid.");
+				static const FText StaticDescription = NSLOCTEXT("MetasoundGraphCore", "Metasound_OscNodeDescription", "Emits an audio signal of a sinusoid.");
 				return StaticDescription;
 			}
 
-			virtual const FString& GetAuthorName() const override
+			virtual const FText& GetAuthorName() const override
 			{
-				static FString Author = TEXT("Epic Games");
-				return Author;
+				return PluginAuthor;
 			}
 
-			virtual const FString& GetPromptIfMissing() const override
+			virtual const FText& GetPromptIfMissing() const override
 			{
-				static FString Prompt = TEXT("Make sure that the Metasound plugin is loaded.");
-				return Prompt;
+				return PluginNodeMissingPrompt;
 			}
 
 		private:
-
 			float DefaultFrequency;
 			FOperatorFactory Factory;
 	};
diff --git a/Engine/Plugins/Runtime/Metasound/Source/MetasoundGraphCore/Public/MetasoundOutputNode.h b/Engine/Plugins/Runtime/Metasound/Source/MetasoundGraphCore/Public/MetasoundOutputNode.h
--- a/Engine/Plugins/Runtime/Metasound/Source/MetasoundGraphCore/Public/MetasoundOutputNode.h
+++ b/Engine/Plugins/Runtime/Metasound/Source/MetasoundGraphCore/Public/MetasoundOutputNode.h
@@ -1,79 +1,77 @@
 // Copyright Epic Games, Inc. All Rights Reserved.
-
 #pragma once
 
 #include "CoreMinimal.h"
-#include "MetasoundOperatorInterface.h"
 #include "MetasoundBuilderInterface.h"
 #include "MetasoundNodeInterface.h"
+#include "MetasoundOperatorInterface.h"
+
+#define LOCTEXT_NAMESPACE "MetasoundGraphCore"
 
-#define LOCTEXT_NAMESPACE "MetasoundOutputNode"
 
 namespace Metasound
 {
 	template<typename DataType>
 	class TOutputNode : public INode
 	{
-			static_assert(TDataReferenceTypeInfo<DataType>::bIsValidSpecialization, "Please use DECLARE_METASOUND_DATA_REFERENCE_TYPES with this class before trying to create an output node with it.");
+		static_assert(TDataReferenceTypeInfo<DataType>::bIsValidSpecialization, "Please use DECLARE_METASOUND_DATA_REFERENCE_TYPES with this class before trying to create an output node with it.");
 			
-			class FOutputOperator : public IOperator
-			{
-				public:
-					using FDataReadReference = TDataReadReference<DataType>;
-
-					FOutputOperator(const FString& InDataReferenceName, FDataReadReference InDataReference)
-					{
-						Outputs.AddDataReadReference<DataType>(InDataReferenceName, InDataReference);
-					}
-
-					virtual ~FOutputOperator() {}
-
-					virtual const FDataReferenceCollection& GetInputs() const override
-					{
-						return Inputs;
-					}
-
-					virtual const FDataReferenceCollection& GetOutputs() const override
-					{
-						return Outputs;
-					}
-
-					virtual FExecuteFunction GetExecuteFunction() override
-					{
-						return nullptr;
-					}
-
-				private:
-					FDataReferenceCollection Inputs;
-					FDataReferenceCollection Outputs;
-			};
-
-			class FOutputOperatorFactory : public IOperatorFactory
-			{
-				public:
-					FOutputOperatorFactory(const FString& InDataReferenceName)
-					:	DataReferenceName(InDataReferenceName)
+		class FOutputOperator : public IOperator
+		{
+			public:
+				using FDataReadReference = TDataReadReference<DataType>;
+
+				FOutputOperator(const FString& InDataReferenceName, FDataReadReference InDataReference)
+				{
+					Outputs.AddDataReadReference<DataType>(InDataReferenceName, InDataReference);
+				}
+
+				virtual ~FOutputOperator() {}
+
+				virtual const FDataReferenceCollection& GetInputs() const override
+				{
+					return Inputs;
+				}
+
+				virtual const FDataReferenceCollection& GetOutputs() const override
+				{
+					return Outputs;
+				}
+
+				virtual FExecuteFunction GetExecuteFunction() override
+				{
+					return nullptr;
+				}
+
+			private:
+				FDataReferenceCollection Inputs;
+				FDataReferenceCollection Outputs;
+		};
+
+		class FOutputOperatorFactory : public IOperatorFactory
+		{
+			public:
+				FOutputOperatorFactory(const FString& InDataReferenceName)
+				:	DataReferenceName(InDataReferenceName)
+				{
+				}
+
+				virtual TUniquePtr<IOperator> CreateOperator(const INode& InNode, const FOperatorSettings& InOperatorSettings, const FDataReferenceCollection& InInputDataReferences, TArray<TUniquePtr<IOperatorBuildError>>& OutErrors) override
+				{
+					if (!InInputDataReferences.ContainsDataReadReference<DataType>(DataReferenceName))
 					{
+						// TODO: Add build error.
+						return TUniquePtr<IOperator>(nullptr);
 					}
 
-					virtual TUniquePtr<IOperator> CreateOperator(const INode& InNode, const FOperatorSettings& InOperatorSettings, const FDataReferenceCollection& InInputDataReferences, TArray<TUniquePtr<IOperatorBuildError>>& OutErrors) override
-					{
-						if (!InInputDataReferences.ContainsDataReadReference<DataType>(DataReferenceName))
-						{
-							// TODO: Add build error.
-							return TUniquePtr<IOperator>(nullptr);
-						}
-
-						return MakeUnique<FOutputOperator>(DataReferenceName, InInputDataReferences.GetDataReadReference<DataType>(DataReferenceName));
-					}
+					return MakeUnique<FOutputOperator>(DataReferenceName, InInputDataReferences.GetDataReadReference<DataType>(DataReferenceName));
+				}
 
-				private:
-					FString DataReferenceName;
-			};
+			private:
+				FString DataReferenceName;
+		};
 
 		public:
-
-
 			TOutputNode(const FString& InNodeDescription, const FString& InVertexName)
 			:	NodeDescription(InNodeDescription)
 			,	VertexName(InVertexName)
@@ -95,8 +93,7 @@ namespace Metasound
 				// TODO: Any special formatting for these node type names?
 				// TODO: although this is ok with MSVC's lax template instantiation, every other compiler will complain about TDataReferenceTypeInfo.
 				//static const FName ClassName = FName(FString(TEXT("Input_")) + FString(TDataReferenceTypeInfo<DataType>::TypeName));
-				static const FName ClassName = FName(TEXT("OutputNode"));
-
+				static const FName ClassName("OutputNode");
 				return ClassName;
 			}
 
@@ -105,22 +102,20 @@ namespace Metasound
 				return NodeDescription;
 			}
 
-			virtual const FString& GetDescription() const override
+			virtual const FText& GetDescription() const override
 			{
-				static FString Description = TEXT("This node can be used as an output for a given type.");
+				static const FText Description = LOCTEXT("Metasound_InputNodeDescription", "Output from the parent Metasound graph.");
 				return Description;
 			}
 
-			virtual const FString& GetAuthorName() const override
+			virtual const FText& GetAuthorName() const override
 			{
-				static FString Author = TEXT("Epic Games");
-				return Author;
+				return PluginAuthor;
 			}
 
-			virtual const FString& GetPromptIfMissing() const override
+			virtual const FText& GetPromptIfMissing() const override
 			{
-				static FString Prompt = TEXT("Make sure that the Metasound plugin is loaded.");
-				return Prompt;
+				return PluginNodeMissingPrompt;
 			}
 
 			virtual const FInputDataVertexCollection& GetInputDataVertices() const override
@@ -149,5 +144,4 @@ namespace Metasound
 			FOutputDataVertexCollection Outputs;
 	};
 }
-
-#undef LOCTEXT_NAMESPACE //MetasoundOutputNode
+#undef LOCTEXT_NAMESPACE // MetasoundOutputNode
diff --git a/Engine/Plugins/Runtime/Metasound/Source/MetasoundGraphCore/Public/MetasoundPeriodicBopNode.h b/Engine/Plugins/Runtime/Metasound/Source/MetasoundGraphCore/Public/MetasoundPeriodicBopNode.h
--- a/Engine/Plugins/Runtime/Metasound/Source/MetasoundGraphCore/Public/MetasoundPeriodicBopNode.h
+++ b/Engine/Plugins/Runtime/Metasound/Source/MetasoundGraphCore/Public/MetasoundPeriodicBopNode.h
@@ -14,11 +14,10 @@ namespace Metasound
 {
 	class METASOUNDGRAPHCORE_API FPeriodicBopNode : public FNode
 	{
-
-			class FOperatorFactory : public IOperatorFactory
-			{
-				virtual TUniquePtr<IOperator> CreateOperator(const INode& InNode, const FOperatorSettings& InOperatorSettings, const FDataReferenceCollection& InInputDataReferences, TArray<TUniquePtr<IOperatorBuildError>>& OutErrors) override;
-			};
+		class FOperatorFactory : public IOperatorFactory
+		{
+			virtual TUniquePtr<IOperator> CreateOperator(const INode& InNode, const FOperatorSettings& InOperatorSettings, const FDataReferenceCollection& InInputDataReferences, TArray<TUniquePtr<IOperatorBuildError>>& OutErrors) override;
+		};
 
 		public:
 			static const FName ClassName;
@@ -34,25 +33,22 @@ namespace Metasound
 
 			virtual IOperatorFactory& GetDefaultOperatorFactory() override;
 
-			virtual const FString& GetDescription() const override
+			virtual const FText& GetDescription() const override
 			{
-				static FString StaticDescription = TEXT("This node emits a bop periodically, based on the duration given.");
+				static const FText StaticDescription = NSLOCTEXT("MetasoundGraphCore", "Metasound_PeriodicBopNodeDescription", "Emits a bop periodically based on the period duration given.");
 				return StaticDescription;
 			}
 
-			virtual const FString& GetAuthorName() const override
+			virtual const FText& GetAuthorName() const override
 			{
-				static FString Author = TEXT("Epic Games");
-				return Author;
+				return PluginAuthor;
 			}
 
-			virtual const FString& GetPromptIfMissing() const override
+			virtual const FText& GetPromptIfMissing() const override
 			{
-				static FString Prompt = TEXT("Make sure that the Metasound plugin is loaded.");
-				return Prompt;
+				return PluginNodeMissingPrompt;
 			}
 
-
 		private:
 			float DefaultPeriod;
 			FOperatorFactory Factory;
