Operator builder errors and unit test

diff --git a/Engine/Plugins/Runtime/Metasound/Source/MetasoundGraphCore/Private/MetasoundBuildError.cpp b/Engine/Plugins/Runtime/Metasound/Source/MetasoundGraphCore/Private/MetasoundBuildError.cpp
new file mode 100644
--- /dev/null
+++ b/Engine/Plugins/Runtime/Metasound/Source/MetasoundGraphCore/Private/MetasoundBuildError.cpp
@@ -0,0 +1,145 @@
+// Copyright Epic Games, Inc. All Rights Reserved.
+
+#include "MetasoundBuildError.h"
+
+#include "CoreMinimal.h"
+
+#include "MetasoundBuilderInterface.h"
+#include "MetasoundNodeInterface.h"
+#include "MetasoundOperatorInterface.h"
+
+#define LOCTEXT_NAMESPACE "MetasoundBuildError"
+
+namespace Metasound
+{
+	FBuildErrorBase::FBuildErrorBase(const FName& InErrorType, const FText& InErrorDescription)
+	:	ErrorType(InErrorType)
+	,	ErrorDescription(InErrorDescription)
+	{
+	}
+
+	/** Returns the type of error. */
+	const FName& FBuildErrorBase::GetErrorType() const
+	{
+		return ErrorType;
+	}
+
+	/** Returns a human readable error description. */
+	const FText& FBuildErrorBase::GetErrorDescription() const
+	{
+		return ErrorDescription;
+	}
+
+	/** Returns an array of Nodes associated with the error. */
+	const TArray<const INode*>& FBuildErrorBase::GetNodes() const 
+	{
+		return Nodes;
+	}
+
+	/** Returns an array of edges associated with the error. */
+	const TArray<FDataEdge>& FBuildErrorBase::GetEdges() const
+	{
+		return Edges;
+	}
+
+	void FBuildErrorBase::AddEdge(const FDataEdge& InEdge)
+	{
+		Edges.Add(InEdge);
+	}
+
+	void FBuildErrorBase::AddEdges(TArrayView<const FDataEdge> InEdges)
+	{
+		Edges.Append(InEdges.GetData(), InEdges.Num());
+	}
+
+	void FBuildErrorBase::AddNode(const INode& InNode)
+	{
+		Nodes.Add(&InNode);
+	}
+
+	void FBuildErrorBase::AddNodes(TArrayView<INode const * const> InNodes)
+	{
+		Nodes.Append(InNodes.GetData(), InNodes.Num());
+	}
+
+	const FName FDanglingEdgeError::ErrorType = FName(TEXT("MetasoundDanglingEdge"));
+
+	FDanglingEdgeError::FDanglingEdgeError(const FDataEdge& InEdge)
+	:	FBuildErrorBase(ErrorType, LOCTEXT("DanglingEdgeError", "Edge is not connected"))
+	{
+		AddEdge(InEdge);
+	}
+
+	const FName FGraphCycleError::ErrorType = FName(TEXT("MetasoundGraphCycleError"));
+
+	FGraphCycleError::FGraphCycleError(const TArray<INode*>& InNodes, const TArray<FDataEdge>& InEdges)
+	:	FBuildErrorBase(ErrorType, LOCTEXT("GraphCycleError", "Graph contains cycles."))
+	{
+		AddNodes(InNodes);
+		AddEdges(InEdges);
+	}
+
+	const FName FInternalError::ErrorType = FName(TEXT("MetasoundInternalError"));
+
+	FInternalError::FInternalError(const FString& InFileName, int32 InLineNumber)
+	:	FBuildErrorBase(ErrorType, LOCTEXT("InternalError", "Internal error."))
+	,	FileName(InFileName)
+	,	LineNumber(InLineNumber)
+	{
+	}
+
+	const FString& FInternalError::GetFileName() const
+	{
+		return FileName;
+	}
+
+	int32 FInternalError::GetLineNumber() const
+	{
+		return LineNumber;
+	}
+
+	const FName FMissingInputDataReferenceError::ErrorType = FName(TEXT("MetasoundMissingInputDataReferenceError"));
+	
+	FMissingInputDataReferenceError::FMissingInputDataReferenceError(const FInputDataDestination& InInputDataDestination)
+	:	FBuildErrorBase(ErrorType, LOCTEXT("MissingInputDataReferenceError", "Missing input data reference."))
+	,	InputDataDestination(InInputDataDestination)
+	{
+		if (nullptr != InInputDataDestination.Node)
+		{
+			AddNode(*InInputDataDestination.Node);
+		}
+	}
+
+	const FInputDataDestination& FMissingInputDataReferenceError::GetInputDataDestination() const
+	{
+		return InputDataDestination;
+	}
+
+
+	const FName FMissingOutputDataReferenceError::ErrorType = FName(TEXT("MetasoundMissingOutputDataReferenceError"));
+	
+	FMissingOutputDataReferenceError::FMissingOutputDataReferenceError(const FOutputDataSource& InOutputDataSource)
+	:	FBuildErrorBase(ErrorType, LOCTEXT("MissingOutputDataReferenceError", "Missing output data reference."))
+	,	OutputDataSource(InOutputDataSource)
+	{
+		if (nullptr != InOutputDataSource.Node)
+		{
+			AddNode(*InOutputDataSource.Node);
+		}
+	}
+
+	const FOutputDataSource& FMissingOutputDataReferenceError::GetOutputDataSource() const
+	{
+		return OutputDataSource;
+	}
+
+	const FName FInvalidConnectionDataTypeError::ErrorType = FName(TEXT("MetasoundInvalidConnectionDataTypeError"));
+
+	FInvalidConnectionDataTypeError::FInvalidConnectionDataTypeError(const FDataEdge& InEdge)
+	:	FBuildErrorBase(ErrorType, LOCTEXT("InvalidConnectionDataTypeError", "Source and destination data types do not match."))
+	{
+		AddEdge(InEdge);
+	}
+}
+
+#undef LOCTEXT_NAMESPACE
diff --git a/Engine/Plugins/Runtime/Metasound/Source/MetasoundGraphCore/Private/MetasoundOperatorBuilder.cpp b/Engine/Plugins/Runtime/Metasound/Source/MetasoundGraphCore/Private/MetasoundOperatorBuilder.cpp
--- a/Engine/Plugins/Runtime/Metasound/Source/MetasoundGraphCore/Private/MetasoundOperatorBuilder.cpp
+++ b/Engine/Plugins/Runtime/Metasound/Source/MetasoundGraphCore/Private/MetasoundOperatorBuilder.cpp
@@ -5,8 +5,34 @@
 #include "MetasoundOperatorInterface.h"
 #include "MetasoundGraphOperator.h"
 
+#include "MetasoundBuildError.h"
+
 namespace Metasound
 {
+	namespace OperatorBuilderPrivate
+	{
+		using FBuildErrorPointer = TUniquePtr<IOperatorBuildError>;
+
+		template<typename ErrorType, typename... ArgTypes>
+		void AddBuildError(TArray<FBuildErrorPointer>& OutErrors, ArgTypes&&... Args)
+		{
+			OutErrors.Add(MakeUnique<ErrorType>(Forward<ArgTypes>(Args)...));
+		}
+
+		void GetEdgesBetweenNodes(const TArray<INode*>& InNodes, const TArray<FDataEdge>& InEdges, TArray<FDataEdge>& OutEdges)
+		{
+			OutEdges = InEdges.FilterByPredicate(
+				[&](const FDataEdge& InEdge) 
+				{
+					bool bEdgeInNodes = InNodes.Contains(InEdge.To.Node);
+					bEdgeInNodes &= InNodes.Contains(InEdge.From.Node);
+
+					return bEdgeInNodes;
+				}
+			);
+		}
+	}
+
 	FOperatorBuilder::FOperatorBuilder(const FOperatorSettings& InSettings)
 	:	OperatorSettings(InSettings)
 	{
@@ -28,17 +54,14 @@ namespace Metasound
 		bSuccess = GroupInputEdges(InGraph.GetDataEdges(), NodeInputEdges, OutErrors);
 		if (!bSuccess)
 		{
-			// TODO: add error to error stack.
 			return TUniquePtr<IOperator>(nullptr);
 		}
 
 		// TODO: Validate edges.
 
-		// TODO: need to include inputs/outputs vertices here. 
 		bSuccess = TopologicalSort(InGraph.GetDataEdges(), SortedNodes, OutErrors);
 		if (!bSuccess)
 		{
-			// TODO: add error to error stack.
 			return TUniquePtr<IOperator>(nullptr);
 		}
 
@@ -50,7 +73,6 @@ namespace Metasound
 		bSuccess = CreateOperators(SortedNodes, NodeInputEdges, Operators, NodeData, OutErrors);
 		if (!bSuccess)
 		{
-			// TODO: add error to error stack.
 			return TUniquePtr<IOperator>(nullptr);
 		}
 
@@ -71,9 +93,10 @@ namespace Metasound
 
 	bool FOperatorBuilder::TopologicalSort(const TArray<FDataEdge>& InEdges, TArray<INode*>& OutNodes, TArray<FBuildErrorPtr>& OutErrors) const
 	{
+		using namespace OperatorBuilderPrivate;
 		using FNodePair = TPair<INode*, INode*>;
 		using FNodeMultiMap = TMultiMap<INode*, INode*>;
-		// TODO: Implement depth first topological sort to imporove cache coherency.
+		// TODO: Implement depth first topological sort to improve cache coherency.
 		
 		TArray<INode*> UniqueNodes;
 		FNodeMultiMap Dependencies;
@@ -82,12 +105,17 @@ namespace Metasound
 		{
 			if (nullptr == Edge.To.Node)
 			{
-				// TODO: add build error for unconnected edge.
+				// Edges are required to be connected to a destination node.
+				AddBuildError<FDanglingEdgeError>(OutErrors, Edge);
+
 				continue;
 			}
+
 			if (nullptr == Edge.From.Node)
 			{
-				// TODO: add build error for unconnected edge.
+				// Edges are required to be connected to a source node.
+				AddBuildError<FDanglingEdgeError>(OutErrors, Edge);
+
 				continue;
 			}
 
@@ -106,7 +134,14 @@ namespace Metasound
 
 			if (0 == IndependentNodes.Num())
 			{
-				// TODO: add build error. likely a cycle in the graph.
+				// add build error. likely a cycle in the graph.
+				
+				// Get edges in the cylce
+				TArray<FDataEdge> CycleEdges;
+				GetEdgesBetweenNodes(UniqueNodes, InEdges, CycleEdges);
+
+				AddBuildError<FGraphCycleError>(OutErrors, UniqueNodes, CycleEdges);
+
 				return false;
 			}
 
@@ -137,6 +172,8 @@ namespace Metasound
 
 	bool FOperatorBuilder::GatherInputDataReferences(const INode* InNode, const FNodeEdgeMultiMap& InEdgeMap, const FNodeDataReferenceMap& InDataReferenceMap, FDataReferenceCollection& OutCollection, TArray<FBuildErrorPtr>& OutErrors) const
 	{
+		using namespace OperatorBuilderPrivate;
+
 		// Find all input edges associated with this node. 
 		TArray<const FDataEdge*> Edges;
 		InEdgeMap.MultiFind(InNode, Edges);
@@ -146,9 +183,10 @@ namespace Metasound
 		{
 			if (!InDataReferenceMap.Contains(Edge->From.Node))
 			{
-				// TODO: Add to build error.
-				// TODO: This is likely due to a failed topological sort and is more of an internal error
+				// This is likely due to a failed topological sort and is more of an internal error
 				// than a user error.
+				AddBuildError<FInternalError>(OutErrors, TEXT(__FILE__), __LINE__);
+
 				return false;
 			}
 
@@ -157,18 +195,21 @@ namespace Metasound
 
 			if (!FromDataReferenceCollection.ContainsDataReadReference(Edge->From.Vertex.VertexName, Edge->From.Vertex.DataReferenceTypeName))
 			{
-				// TODO: This is likely a node programming error where the edges reported by the INode interface
+				// This is likely a node programming error where the edges reported by the INode interface
 				// did not match the readable parameter refs created by the operators outputs. Or, the edge description is invalid.
-				// TODO: add build error to stack
+
 				// TODO: consider checking that outputs of node operators match output descriptions of nodes. 
 				// TODO: consider checking the edges are supported by nodes existing parameter descriptions.
+				AddBuildError<FMissingOutputDataReferenceError>(OutErrors, Edge->From);
+
 				return false;
 			}
 
 			if (Edge->From.Vertex.DataReferenceTypeName != Edge->To.Vertex.DataReferenceTypeName)
 			{
-				// TODO: add build error.
 				// TODO: may want to have this check higher up in the stack. 
+				AddBuildError<FInvalidConnectionDataTypeError>(OutErrors, *Edge);
+
 				return false;
 			}
 
@@ -178,7 +219,7 @@ namespace Metasound
 
 			if (!bSuccess)
 			{
-				// TODO: add build error.
+				AddBuildError<FMissingOutputDataReferenceError>(OutErrors, Edge->From);
 				return false;
 			}
 		}
@@ -195,9 +236,6 @@ namespace Metasound
 			// Gather the input parameters for this IOperator from the output parameters of already created IOperators. 
 			if (!GatherInputDataReferences(Node, InNodeInputEdges, OutDataReferences, InputCollection, OutErrors))
 			{
-				// TODO Is this a separate error, or the stack of the same error. Perhaps this shouldn't add an error
-				// because it's already sending the error downstream. 
-				// TODO: add to build error
 				return false;
 			}
 
@@ -207,9 +245,6 @@ namespace Metasound
 
 			if (!Operator.IsValid())
 			{
-				// TODO: add vailed operator creation erorr. 
-				// TODO Is this a separate error, or the stack of the same error. Perhaps this shouldn't add an error
-				// because it's already sending the error downstream. 
 				return false;
 			}
 
@@ -225,17 +260,18 @@ namespace Metasound
 
 	bool FOperatorBuilder::GatherGraphDataReferences(const IGraph& InGraph, FNodeDataReferenceMap& InNodeDataReferences, FDataReferenceCollection& OutGraphInputs, FDataReferenceCollection& OutGraphOutputs, TArray<FBuildErrorPtr>& OutErrors) const
 	{
+		using namespace OperatorBuilderPrivate;
 		using FDestinationElement = FInputDataDestinationCollection::ElementType;
 		using FSourceElement = FOutputDataSourceCollection::ElementType;
-		// Gather inputs
 
+		// Gather inputs
 		for (const FDestinationElement& Element : InGraph.GetInputDataDestinations())
 		{
 			const FInputDataDestination& InputDestination = Element.Value;
 
 			if (!InNodeDataReferences.Contains(InputDestination.Node))
 			{
-				// TODO: build error.  input node does not exist.
+				AddBuildError<FMissingInputDataReferenceError>(OutErrors, InputDestination);
 				return false;
 			}
 
@@ -243,7 +279,7 @@ namespace Metasound
 
 			if (!Collection.ContainsDataWriteReference(InputDestination.Vertex.VertexName, InputDestination.Vertex.DataReferenceTypeName))
 			{
-				// TODO: build error.  parameter ref does not exist
+				AddBuildError<FMissingInputDataReferenceError>(OutErrors, InputDestination);
 				return false;
 			}
 
@@ -251,8 +287,7 @@ namespace Metasound
 
 			if (!bSuccess)
 			{
-				// TODO: add build error
-				// Shouldn't be hitting this error because of earlier "ContainsWritable..." call.
+				AddBuildError<FMissingInputDataReferenceError>(OutErrors, InputDestination);
 				return false;
 			}
 		}
@@ -264,7 +299,7 @@ namespace Metasound
 
 			if (!InNodeDataReferences.Contains(OutputSource.Node))
 			{
-				// TODO: build error.  input node does not exist.
+				AddBuildError<FMissingOutputDataReferenceError>(OutErrors, OutputSource);
 				return false;
 			}
 
@@ -272,7 +307,7 @@ namespace Metasound
 
 			if (!Collection.ContainsDataReadReference(OutputSource.Vertex.VertexName, OutputSource.Vertex.DataReferenceTypeName))
 			{
-				// TODO: build error.  parameter ref does not exist
+				AddBuildError<FMissingOutputDataReferenceError>(OutErrors, OutputSource);
 				return false;
 			}
 
@@ -280,8 +315,7 @@ namespace Metasound
 
 			if (!bSuccess)
 			{
-				// TODO: add build error
-				// Shouldn't be hitting this error because of earlier "ContainsWritable..." call.
+				AddBuildError<FMissingOutputDataReferenceError>(OutErrors, OutputSource);
 				return false;
 			}
 		}
@@ -298,7 +332,6 @@ namespace Metasound
 
 		if (!bSuccess)
 		{
-			// TODO: add build error.
 			return TUniquePtr<IOperator>(nullptr);
 		}
 
diff --git a/Engine/Plugins/Runtime/Metasound/Source/MetasoundGraphCore/Public/MetasoundBuildError.h b/Engine/Plugins/Runtime/Metasound/Source/MetasoundGraphCore/Public/MetasoundBuildError.h
new file mode 100644
--- /dev/null
+++ b/Engine/Plugins/Runtime/Metasound/Source/MetasoundGraphCore/Public/MetasoundBuildError.h
@@ -0,0 +1,167 @@
+// Copyright Epic Games, Inc. All Rights Reserved.
+
+#pragma once
+
+#include "CoreMinimal.h"
+
+#include "MetasoundBuilderInterface.h"
+#include "MetasoundNodeInterface.h"
+#include "MetasoundOperatorInterface.h"
+
+namespace Metasound
+{
+	/** FBuildErrorBase
+	 *
+	 * A general build error which contains a error type and human readable 
+	 * error description.
+	 */
+	class METASOUNDGRAPHCORE_API FBuildErrorBase : public IOperatorBuildError
+	{
+		public:
+
+			FBuildErrorBase(const FName& InErrorType, const FText& InErrorDescription);
+
+			virtual ~FBuildErrorBase() = default;
+
+			/** Returns the type of error. */
+			virtual const FName& GetErrorType() const override;
+
+			/** Returns a human readable error description. */
+			virtual const FText& GetErrorDescription() const override;
+
+			/** Returns an array of edges associated with the error. */
+			virtual const TArray<FDataEdge>& GetEdges() const override;
+
+			/** Returns an array of Nodes associated with the error. */
+			virtual const TArray<const INode*>& GetNodes() const override;
+
+		protected:
+
+			// Add edges to be associated with error.
+			void AddEdge(const FDataEdge& InEdge);
+			void AddEdges(TArrayView<const FDataEdge> InEdges);
+
+			// Add nodes to be associated with error.
+			void AddNode(const INode& InNode);
+			void AddNodes(TArrayView<INode const* const> InNodes);
+
+		private:
+
+			FName ErrorType;
+			FText ErrorDescription;
+
+			TArray<const INode*> Nodes;
+			TArray<FDataEdge> Edges;
+	};
+
+	/** FDanglingEdgeError
+	 *
+	 * Caused by FDataEdges pointing to invalid nodes.
+	 */
+	class METASOUNDGRAPHCORE_API FDanglingEdgeError : public FBuildErrorBase
+	{
+		public:
+			static const FName ErrorType;
+
+			FDanglingEdgeError(const FDataEdge& InEdge);
+
+			virtual ~FDanglingEdgeError() = default;
+
+		private:
+	};
+
+	/** FGraphCycleError
+	 *
+	 * Caused by circular paths in graph.
+	 */
+	class METASOUNDGRAPHCORE_API FGraphCycleError : public FBuildErrorBase
+	{
+		public:
+			static const FName ErrorType;
+
+			FGraphCycleError(const TArray<INode*>& InNodes, const TArray<FDataEdge>& InEdges);
+
+			virtual ~FGraphCycleError() = default;
+
+		private:
+	};
+
+	/** FInternalError
+	 *
+	 * Caused by internal state or logic errors. 
+	 */
+	class METASOUNDGRAPHCORE_API FInternalError : public FBuildErrorBase
+	{
+		public:
+			static const FName ErrorType;
+
+			FInternalError(const FString& InFileName, int32 InLineNumber);
+
+			virtual ~FInternalError() = default;
+
+			const FString& GetFileName() const;
+			int32 GetLineNumber() const;
+
+		private:
+			FString FileName;
+			int32 LineNumber;
+	};
+
+	/** FMissingInputDataReferenceError
+	 *
+	 * Caused by IOperators not exposing expected IDataReferences in their input
+	 * FDataReferenceCollection
+	 */
+	class METASOUNDGRAPHCORE_API FMissingInputDataReferenceError : public FBuildErrorBase
+	{
+		public:
+
+			static const FName ErrorType;
+			
+			FMissingInputDataReferenceError(const FInputDataDestination& InInputDataDestination);
+
+			virtual ~FMissingInputDataReferenceError() = default;
+
+			const FInputDataDestination& GetInputDataDestination() const;
+
+		private:
+			FInputDataDestination InputDataDestination;
+	};
+
+	/** FMissingOutputDataReferenceError
+	 *
+	 * Caused by IOperators not exposing expected IDataReferences in their output
+	 * FDataReferenceCollection
+	 */
+	class METASOUNDGRAPHCORE_API FMissingOutputDataReferenceError : public FBuildErrorBase
+	{
+		public:
+
+			static const FName ErrorType;
+			
+			FMissingOutputDataReferenceError(const FOutputDataSource& InOutputDataSource);
+
+			virtual ~FMissingOutputDataReferenceError() = default;
+
+			const FOutputDataSource& GetOutputDataSource() const;
+
+		private:
+			FOutputDataSource OutputDataSource;
+	};
+
+	/** FInvalidConnectionDataTypeError
+	 *
+	 * Caused when edges describe a connection between vertices with different
+	 * data types.
+	 */
+	class METASOUNDGRAPHCORE_API FInvalidConnectionDataTypeError  : public FBuildErrorBase
+	{
+		public:
+
+			static const FName ErrorType;
+
+			FInvalidConnectionDataTypeError(const FDataEdge& InEdge);
+
+			virtual ~FInvalidConnectionDataTypeError() = default;
+	};
+}
diff --git a/Engine/Plugins/Runtime/Metasound/Source/MetasoundGraphCore/Public/MetasoundBuilderInterface.h b/Engine/Plugins/Runtime/Metasound/Source/MetasoundGraphCore/Public/MetasoundBuilderInterface.h
--- a/Engine/Plugins/Runtime/Metasound/Source/MetasoundGraphCore/Public/MetasoundBuilderInterface.h
+++ b/Engine/Plugins/Runtime/Metasound/Source/MetasoundGraphCore/Public/MetasoundBuilderInterface.h
@@ -15,33 +15,22 @@ namespace Metasound
 	class IOperatorBuildError
 	{
 		public:
-			virtual ~IOperatorBuildError() {}
+			virtual ~IOperatorBuildError() = default;
 
-			/** Returns the INode associated with the error. */
-			virtual const INode* GetNode() const = 0;
-
-			/** Returns teh FDataEdge associated with the error. */
-			virtual const FDataEdge& GetEdge() const = 0;
-
-			/** Returns the FInputDataVertex associated with the error. */
-			virtual const FInputDataVertex& GetInputDataVertex() const = 0;
-
-			/** Returns the FOutputDataVertex associated with the error. */
-			virtual const FOutputDataVertex& GetOutputDataVertex() const = 0;
-
-			/** Returns the parameter type name associated with the error. */
-			virtual const FString& GetDataReferenceTypeName() const = 0;
-
-			/** Returns the parameter name associated with the error. */
-			virtual const FString& GetDataReferenceName() const = 0;
-
-			/** Returns teh type of error. */
-			virtual const FString& GetErrorType() const = 0; // TODO: FName?
+			/** Returns the type of error. */
+			virtual const FName& GetErrorType() const = 0;
 
 			/** Returns a human readable error description. */
 			virtual const FText& GetErrorDescription() const = 0;
+
+			/** Returns an array of Nodes associated with the error. */
+			virtual const TArray<const INode*>& GetNodes() const = 0;
+
+			/** Returns an array of edges associated with the error. */
+			virtual const TArray<FDataEdge>& GetEdges() const = 0;
 	};
 
+
 	/** IOperatorFactory
 	 *
 	 * IOperatorFactory defines an interface for building an IOperator from an INode.  In practice,
@@ -51,7 +40,7 @@ namespace Metasound
 	class IOperatorFactory
 	{
 		public:
-			virtual ~IOperatorFactory() {}
+			virtual ~IOperatorFactory() = default;
 
 			/** Create a new IOperator.
 			 *
@@ -76,7 +65,7 @@ namespace Metasound
 			/** A TUniquePtr of an IOperatorBuildError */
 			using FBuildErrorPtr = TUniquePtr<IOperatorBuildError>;
 
-			virtual ~IOperatorBuilder() {}
+			virtual ~IOperatorBuilder() = default;
 
 			/** Build a graph operator from a graph. 
 			 *
diff --git a/Engine/Plugins/Runtime/Metasound/Source/MetasoundGraphCore/Public/MetasoundNodeInterface.h b/Engine/Plugins/Runtime/Metasound/Source/MetasoundGraphCore/Public/MetasoundNodeInterface.h
--- a/Engine/Plugins/Runtime/Metasound/Source/MetasoundGraphCore/Public/MetasoundNodeInterface.h
+++ b/Engine/Plugins/Runtime/Metasound/Source/MetasoundGraphCore/Public/MetasoundNodeInterface.h
@@ -136,6 +136,26 @@ namespace Metasound
 		}
 	};
 
+	typedef FInputDataVertexCollection::ElementType FInputDataVertexElement;
+	typedef FOutputDataVertexCollection::ElementType FOutputDataVertexElement;
+
+	template<typename DataType>
+	FORCEINLINE FInputDataVertexElement MakeInputDataVertexElement(const FString& InName, const FText& InDescription)
+	{
+		FInputDataVertex Vertex = MakeInputDataVertex<DataType>(InName, InDescription);
+		return FInputDataVertexElement(MakeDataVertexKey(Vertex), Vertex);
+	}
+
+	template<typename DataType>
+	FORCEINLINE FOutputDataVertexElement MakeOutputDataVertexElement(const FString& InName, const FText& InDescription)
+	{
+		FOutputDataVertex Vertex = MakeOutputDataVertex<DataType>(InName, InDescription);
+		return FOutputDataVertexElement(MakeDataVertexKey(Vertex), Vertex);
+	}
+
+
+	//typedef uint32 FNodeId;
+
 	/** INodeBase
 	 * 
 	 * Interface for all nodes that can describe their name, type, inputs and outputs.
@@ -145,6 +165,9 @@ namespace Metasound
 		public:
 			virtual ~INodeBase() {}
 
+			// TODO: talk to ethan about nodeids. Maybe this should be handled by frontend. 
+			//FNodeId GetNodeId() const = 0;
+
 			/** Return the name of this specific instance of the node class. */
 			virtual const FString& GetInstanceName() const = 0;
 
