Another IWYU fix for script compile
#rnx

diff --git a/Engine/Plugins/Runtime/Metasound/Source/MetasoundEditor/Private/MetasoundEditorGraphBuilder.cpp b/Engine/Plugins/Runtime/Metasound/Source/MetasoundEditor/Private/MetasoundEditorGraphBuilder.cpp
--- a/Engine/Plugins/Runtime/Metasound/Source/MetasoundEditor/Private/MetasoundEditorGraphBuilder.cpp
+++ b/Engine/Plugins/Runtime/Metasound/Source/MetasoundEditor/Private/MetasoundEditorGraphBuilder.cpp
@@ -2,6 +2,7 @@
 #include "MetasoundEditorGraphBuilder.h"
 
 #include "EdGraph/EdGraph.h"
+#include "EdGraph/EdGraphNode.h"
 #include "EdGraph/EdGraphPin.h"
 #include "Metasound.h"
 #include "MetasoundEditorGraph.h"
diff --git a/Engine/Plugins/Runtime/Metasound/Source/MetasoundEditor/Private/MetasoundEditorGraphBuilder.h b/Engine/Plugins/Runtime/Metasound/Source/MetasoundEditor/Private/MetasoundEditorGraphBuilder.h
--- a/Engine/Plugins/Runtime/Metasound/Source/MetasoundEditor/Private/MetasoundEditorGraphBuilder.h
+++ b/Engine/Plugins/Runtime/Metasound/Source/MetasoundEditor/Private/MetasoundEditorGraphBuilder.h
@@ -4,6 +4,7 @@
 #include "CoreMinimal.h"
 
 // Forward Declarations
+class UEdGraphNode;
 class UMetasound;
 class UMetasoundEditorGraphNode;
 
