Checkpoint on more input/output work
- Expose I/O arrays to UI (TODO: fix rename. We currently have two places we're caching input/output names and no easy way to snatch I/O metadata from I/O node.  Could we add function to reference name from description instead of caching on node?)
- Stub in json import exposition to UI
- Fix crash with removing inputs/outputs
- Various clean-up

diff --git a/Engine/Plugins/Runtime/Metasound/Source/MetasoundEditor/Private/MetasoundDetailCustomization.cpp b/Engine/Plugins/Runtime/Metasound/Source/MetasoundEditor/Private/MetasoundDetailCustomization.cpp
--- a/Engine/Plugins/Runtime/Metasound/Source/MetasoundEditor/Private/MetasoundDetailCustomization.cpp
+++ b/Engine/Plugins/Runtime/Metasound/Source/MetasoundEditor/Private/MetasoundDetailCustomization.cpp
@@ -4,6 +4,7 @@
 #include "CoreMinimal.h"
 #include "DetailCategoryBuilder.h"
 #include "DetailLayoutBuilder.h"
+#include "IDetailGroup.h"
 #include "Metasound.h"
 #include "PropertyEditorDelegates.h"
 #include "PropertyHandle.h"
@@ -11,18 +12,32 @@
 #include "Templates/SharedPointer.h"
 
 
+#define LOCTEXT_NAMESPACE "MetasoundEditor"
+
 void FMetasoundDetailCustomization::CustomizeDetails(IDetailLayoutBuilder& DetailLayout)
 {
-	IDetailCategoryBuilder& Category = DetailLayout.EditCategory("General");
+	// General Category
+	IDetailCategoryBuilder& GeneralCategoryBuilder = DetailLayout.EditCategory("General");
 
 	TSharedPtr<IPropertyHandle> AuthorHandle = DetailLayout.GetProperty("RootMetasoundDocument.RootClass.Metadata.AuthorName");
 	TSharedPtr<IPropertyHandle> DescHandle = DetailLayout.GetProperty("RootMetasoundDocument.RootClass.Metadata.MetasoundDescription");
 	TSharedPtr<IPropertyHandle> NodeTypeHandle = DetailLayout.GetProperty("RootMetasoundDocument.RootClass.Metadata.NodeType");
 
-	Category.AddProperty(NodeTypeHandle);
-	Category.AddProperty(AuthorHandle);
-	Category.AddProperty(DescHandle);
+	GeneralCategoryBuilder.AddProperty(NodeTypeHandle);
+	GeneralCategoryBuilder.AddProperty(AuthorHandle);
+	GeneralCategoryBuilder.AddProperty(DescHandle);
+
+	IDetailCategoryBuilder& ParametersCategoryBuilder = DetailLayout.EditCategory("Parameters");
+
+	IDetailGroup& InputDetailGroup = ParametersCategoryBuilder.AddGroup("Inputs", LOCTEXT("MetasoundDetailsGroupInputs", "Inputs"), false /* bForAdvanced */, true /* bStartExpanded */);
+	TSharedPtr<IPropertyHandle> InputsHandle = DetailLayout.GetProperty("RootMetasoundDocument.RootClass.Inputs");
+	InputDetailGroup.AddPropertyRow(InputsHandle.ToSharedRef());
+
+	IDetailGroup& OutputDetailGroup = ParametersCategoryBuilder.AddGroup("Inputs", LOCTEXT("MetasoundDetailsGroupOutputs", "Outputs"), false /* bForAdvanced */, true /* bStartExpanded */);
+	TSharedPtr<IPropertyHandle> OutputsHandle = DetailLayout.GetProperty("RootMetasoundDocument.RootClass.Outputs");
+	OutputDetailGroup.AddPropertyRow(OutputsHandle.ToSharedRef());
 
 	// Hack to hide parent structs for nested metadata properties
 	DetailLayout.HideCategory("Hidden");
 }
+#undef LOCTEXT_NAMESPACE
\ No newline at end of file
diff --git a/Engine/Plugins/Runtime/Metasound/Source/MetasoundEditor/Private/MetasoundEditor.cpp b/Engine/Plugins/Runtime/Metasound/Source/MetasoundEditor/Private/MetasoundEditor.cpp
--- a/Engine/Plugins/Runtime/Metasound/Source/MetasoundEditor/Private/MetasoundEditor.cpp
+++ b/Engine/Plugins/Runtime/Metasound/Source/MetasoundEditor/Private/MetasoundEditor.cpp
@@ -15,7 +15,6 @@
 #include "HAL/PlatformApplicationMisc.h"
 #include "IDetailsView.h"
 #include "Kismet2/BlueprintEditorUtils.h"
-#include "MetasoundAssetSerializer.h"
 #include "MetasoundEditorCommands.h"
 #include "MetasoundEditorGraph.h"
 #include "MetasoundEditorGraphBuilder.h"
@@ -251,16 +250,26 @@ namespace Metasound
 				GetToolkitCommands(),
 				FToolBarExtensionDelegate::CreateLambda([this](FToolBarBuilder& ToolbarBuilder)
 				{
-					ToolbarBuilder.BeginSection("Compile");
+					ToolbarBuilder.BeginSection("Utilities");
 					{
 						ToolbarBuilder.AddToolBarButton
 						(
-							FEditorCommands::Get().Compile,
+							FEditorCommands::Get().Import,
 							NAME_None,
 							TAttribute<FText>(),
 							TAttribute<FText>(),
-							TAttribute<FSlateIcon>::Create([this]() { return GetCompileStatusImage(); }),
-							"CompileMetasound"
+							TAttribute<FSlateIcon>::Create([this]() { return GetImportStatusImage(); }),
+							"ImportMetasound"
+						);
+
+						ToolbarBuilder.AddToolBarButton
+						(
+							FEditorCommands::Get().Export,
+							NAME_None,
+							TAttribute<FText>(),
+							TAttribute<FText>(),
+							TAttribute<FSlateIcon>::Create([this]() { return GetExportStatusImage(); }),
+							"ExportMetasound"
 						);
 					}
 					ToolbarBuilder.EndSection();
@@ -277,12 +286,19 @@ namespace Metasound
 			AddToolbarExtender(ToolbarExtender);
 		}
 
-		FSlateIcon FEditor::GetCompileStatusImage() const
+		FSlateIcon FEditor::GetImportStatusImage() const
+		{
+			// TODO: Get/register import image
+			const FName IconName = "MetasoundEditor.Export";
+			return FSlateIcon("MetasoundStyle", IconName);
+		}
+
+		FSlateIcon FEditor::GetExportStatusImage() const
 		{
-			FName IconName = "MetasoundEditor.Compile";
+			FName IconName = "MetasoundEditor.Export";
 			if (!bPassedValidation)
 			{
-				IconName = "MetasoundEditor.CompileError";
+				IconName = "MetasoundEditor.ExportError";
 			}
 
 			return FSlateIcon("MetasoundStyle", IconName);
@@ -301,8 +317,12 @@ namespace Metasound
 				FExecuteAction::CreateSP(this, &FEditor::Stop));
 
 			ToolkitCommands->MapAction(
-				Commands.Compile,
-				FExecuteAction::CreateSP(this, &FEditor::Compile));
+				Commands.Import,
+				FExecuteAction::CreateSP(this, &FEditor::Import));
+
+			ToolkitCommands->MapAction(
+				Commands.Export,
+				FExecuteAction::CreateSP(this, &FEditor::Export));
 
 			ToolkitCommands->MapAction(
 				Commands.TogglePlayback,
@@ -317,12 +337,23 @@ namespace Metasound
 				FExecuteAction::CreateSP(this, &FEditor::RedoGraphAction));
 		}
 
-		void FEditor::Compile()
+		void FEditor::Import()
+		{
+			if (Metasound)
+			{
+				// TODO: Prompt OFD and provide path from user
+				const FString Path = FPaths::ProjectIntermediateDir() / TEXT("Metasounds") + FPaths::ChangeExtension(Metasound->GetPathName(), FMetasoundAssetBase::FileExtension);
+				Metasound->ImportFromJSON(Path);
+			}
+		}
+
+		void FEditor::Export()
 		{
 			if (Metasound)
 			{
-				const UMetasoundEditorGraph* Graph = CastChecked<UMetasoundEditorGraph>(Metasound->GetGraph());
-				bPassedValidation = Metasound::Engine::AssetSerializer::MetasoundToJSON(*Metasound);
+				// TODO: Prompt OFD and provide path from user
+				const FString Path = FPaths::ProjectIntermediateDir() / TEXT("Metasounds") + FPaths::ChangeExtension(Metasound->GetPathName(), FMetasoundAssetBase::FileExtension);
+				Metasound->ExportToJSON(Path);
 			}
 		}
 
@@ -738,8 +769,6 @@ namespace Metasound
 				Node->CreateNewGuid();
 			}
 
-			// TODO: Compile metasound here???
-
 			// Update UI
 			MetasoundGraphEditor->NotifyGraphChanged();
 
diff --git a/Engine/Plugins/Runtime/Metasound/Source/MetasoundEditor/Private/MetasoundEditor.h b/Engine/Plugins/Runtime/Metasound/Source/MetasoundEditor/Private/MetasoundEditor.h
--- a/Engine/Plugins/Runtime/Metasound/Source/MetasoundEditor/Private/MetasoundEditor.h
+++ b/Engine/Plugins/Runtime/Metasound/Source/MetasoundEditor/Private/MetasoundEditor.h
@@ -136,10 +136,13 @@ namespace Metasound
 			/** Binds new graph commands to delegates */
 			void BindGraphCommands();
 
-			FSlateIcon GetCompileStatusImage() const;
+			FSlateIcon GetImportStatusImage() const;
+
+			FSlateIcon GetExportStatusImage() const;
 
 			/** Toolbar command methods */
-			void Compile();
+			void Import();
+			void Export();
 			void Play();
 			void PlayNode();
 			void Stop();
diff --git a/Engine/Plugins/Runtime/Metasound/Source/MetasoundEditor/Private/MetasoundEditorCommands.cpp b/Engine/Plugins/Runtime/Metasound/Source/MetasoundEditor/Private/MetasoundEditorCommands.cpp
--- a/Engine/Plugins/Runtime/Metasound/Source/MetasoundEditor/Private/MetasoundEditorCommands.cpp
+++ b/Engine/Plugins/Runtime/Metasound/Source/MetasoundEditor/Private/MetasoundEditorCommands.cpp
@@ -13,7 +13,8 @@ namespace Metasound
 			UI_COMMAND(Stop, "Stop", "Stops Metasound (If currently playing)", EUserInterfaceActionType::Button, FInputChord());
 			UI_COMMAND(TogglePlayback, "Toggle Playback", "Plays the Metasound or stops the currently playing Metasound", EUserInterfaceActionType::Button, FInputChord(EKeys::SpaceBar));
 
-			UI_COMMAND(Compile, "Compile", "Compiles Metasound", EUserInterfaceActionType::Button, FInputChord());
+			UI_COMMAND(Import, "Import", "Imports Metasound from Json", EUserInterfaceActionType::Button, FInputChord());
+			UI_COMMAND(Export, "Export", "Exports Metasound to Json", EUserInterfaceActionType::Button, FInputChord());
 
 			UI_COMMAND(BrowserSync, "Browse", "Selects the Metasound in the content browser. If referencing Metasound nodes are selected, selects referenced assets instead.", EUserInterfaceActionType::Button, FInputChord());
 			UI_COMMAND(AddInput, "Add Input", "Adds an input to the node", EUserInterfaceActionType::Button, FInputChord());
diff --git a/Engine/Plugins/Runtime/Metasound/Source/MetasoundEditor/Private/MetasoundEditorGraphBuilder.cpp b/Engine/Plugins/Runtime/Metasound/Source/MetasoundEditor/Private/MetasoundEditorGraphBuilder.cpp
--- a/Engine/Plugins/Runtime/Metasound/Source/MetasoundEditor/Private/MetasoundEditorGraphBuilder.cpp
+++ b/Engine/Plugins/Runtime/Metasound/Source/MetasoundEditor/Private/MetasoundEditorGraphBuilder.cpp
@@ -93,7 +93,26 @@ namespace Metasound
 			Frontend::FGraphHandle GraphHandle = Graph->GetMetasoundChecked().GetRootGraphHandle();
 			if (GraphHandle.IsValid() && NodeHandle.IsValid())
 			{
-				GraphHandle.RemoveNode(NodeHandle, true /* bEvenIfInputOrOutputNode */);
+				switch (NodeHandle.GetNodeType())
+				{
+					case EMetasoundClassType::Input:
+					{
+						GraphHandle.RemoveInput(NodeHandle.GetNodeName());
+					}
+					break;
+
+					case EMetasoundClassType::Output:
+					{
+						GraphHandle.RemoveOutput(NodeHandle.GetNodeName());
+					}
+					break;
+
+					case EMetasoundClassType::External:
+					{
+						GraphHandle.RemoveNode(NodeHandle);
+					}
+					break;
+				}
 			}
 
 			InNode.PostEditChange();
diff --git a/Engine/Plugins/Runtime/Metasound/Source/MetasoundEditor/Private/MetasoundEditorGraphNode.cpp b/Engine/Plugins/Runtime/Metasound/Source/MetasoundEditor/Private/MetasoundEditorGraphNode.cpp
--- a/Engine/Plugins/Runtime/Metasound/Source/MetasoundEditor/Private/MetasoundEditorGraphNode.cpp
+++ b/Engine/Plugins/Runtime/Metasound/Source/MetasoundEditor/Private/MetasoundEditorGraphNode.cpp
@@ -35,13 +35,13 @@ void UMetasoundEditorGraphNode::PostLoad()
 			// Makes sure pin has a name for lookup purposes but user will never see it
 			if (Pin->Direction == EGPD_Input)
 			{
-				Pin->PinName = CreateUniquePinName(TEXT("Input"));
+				Pin->PinName = CreateUniquePinName("Input");
 			}
 			else
 			{
-				Pin->PinName = CreateUniquePinName(TEXT("Output"));
+				Pin->PinName = CreateUniquePinName("Output");
 			}
-			Pin->PinFriendlyName = FText::FromString(TEXT(" "));
+			Pin->PinFriendlyName = FText::GetEmpty();
 		}
 	}
 }
@@ -60,24 +60,38 @@ void UMetasoundEditorGraphNode::CreateInputPin()
 
 int32 UMetasoundEditorGraphNode::EstimateNodeWidth() const
 {
-	const int32 EstimatedCharWidth = 6;
-	FString NodeTitle = GetNodeTitle(ENodeTitleType::FullTitle).ToString();
-	UFont* Font = GetDefault<UEditorEngine>()->EditorFont;
-	int32 Result = NodeTitle.Len()*EstimatedCharWidth;
-
-	if (Font)
+	const FString NodeTitle = GetNodeTitle(ENodeTitleType::FullTitle).ToString();
+	if (const UFont* Font = GetDefault<UEditorEngine>()->EditorFont)
+	{
+		return Font->GetStringSize(*NodeTitle);
+	}
+	else
 	{
-		Result = Font->GetStringSize(*NodeTitle);
+		static const int32 EstimatedCharWidth = 6;
+		return NodeTitle.Len() * EstimatedCharWidth;
 	}
+}
 
-	return Result;
+UMetasound& UMetasoundEditorGraphNode::GetMetasoundChecked()
+{
+	UMetasoundEditorGraph* EdGraph = CastChecked<UMetasoundEditorGraph>(GetGraph());
+	return EdGraph->GetMetasoundChecked();
 }
 
-Metasound::Frontend::FNodeHandle UMetasoundEditorGraphNode::GetNodeHandle() const
+const UMetasound& UMetasoundEditorGraphNode::GetMetasoundChecked() const
 {
 	UMetasoundEditorGraph* EdGraph = CastChecked<UMetasoundEditorGraph>(GetGraph());
-	UMetasound& Metasound = EdGraph->GetMetasoundChecked();
-	return Metasound.GetRootGraphHandle().GetNodeWithId(NodeID);
+	return EdGraph->GetMetasoundChecked();
+}
+
+Metasound::Frontend::FGraphHandle UMetasoundEditorGraphNode::GetRootGraphHandle() const
+{
+	return GetMetasoundChecked().GetRootGraphHandle();
+}
+
+Metasound::Frontend::FNodeHandle UMetasoundEditorGraphNode::GetNodeHandle() const
+{
+	return GetRootGraphHandle().GetNodeWithId(NodeID);
 }
 
 void UMetasoundEditorGraphNode::IteratePins(TUniqueFunction<void(UEdGraphPin* /* Pin */, int32 /* Index */)> InFunc, EEdGraphPinDirection InPinDirection)
@@ -165,8 +179,31 @@ void UMetasoundEditorGraphNode::SetNodeID(uint32 InNodeID)
 
 FText UMetasoundEditorGraphNode::GetNodeTitle(ENodeTitleType::Type TitleType) const
 {
-	const Metasound::Frontend::FNodeHandle NodeHandle = GetNodeHandle();
-	return FText::FromString(NodeHandle.GetNodeClassName());
+	using namespace Metasound::Frontend;
+
+	FNodeHandle NodeHandle = GetNodeHandle();
+	switch (NodeHandle.GetNodeType())
+	{
+		case EMetasoundClassType::Input:
+		{
+			return FText::Format(LOCTEXT("MetasoundGraphNode_TitleFormat", "Get Input {0}"), FText::FromString(NodeHandle.GetNodeName()));
+		}
+		break;
+
+		case EMetasoundClassType::Output:
+		{
+			return FText::Format(LOCTEXT("MetasoundGraphNode_TitleFormat", "Set Output {0}"), FText::FromString(NodeHandle.GetNodeName()));
+		}
+		break;
+
+		default:
+		case EMetasoundClassType::External:
+		case EMetasoundClassType::MetasoundGraph:
+		{
+			return FText::FromString(NodeHandle.GetNodeClassName());
+		}
+		break;
+	}
 }
 
 void UMetasoundEditorGraphNode::PrepareForCopying()
@@ -241,11 +278,32 @@ void UMetasoundEditorGraphNode::GetNodeContextMenuActions(UToolMenu* Menu, UGrap
 
 FText UMetasoundEditorGraphNode::GetTooltipText() const
 {
-	FText Tooltip = FText::GetEmpty();
+	using namespace Metasound::Frontend;
 
-	// TODO: Retrieve from UMetasound
+	// TODO: Add author to ToolTip
+	FNodeHandle NodeHandle = GetNodeHandle();
+	switch (NodeHandle.GetNodeType())
+	{
+		case EMetasoundClassType::Input:
+		{
+			return GetMetasoundChecked().GetInputToolTip(NodeHandle.GetNodeName());
+		}
+		break;
+
+		case EMetasoundClassType::Output:
+		{
+			return GetMetasoundChecked().GetOutputToolTip(NodeHandle.GetNodeName());
+		}
+		break;
 
-	return Tooltip;
+		default:
+		case EMetasoundClassType::External:
+		case EMetasoundClassType::MetasoundGraph:
+		{
+			return GenerateClassDescriptionForNode(NodeHandle.GetClassInfo()).Metadata.MetasoundDescription;
+		}
+		break;
+	}
 }
 
 FString UMetasoundEditorGraphNode::GetDocumentationExcerptName() const
diff --git a/Engine/Plugins/Runtime/Metasound/Source/MetasoundEditor/Private/MetasoundEditorGraphSchema.cpp b/Engine/Plugins/Runtime/Metasound/Source/MetasoundEditor/Private/MetasoundEditorGraphSchema.cpp
--- a/Engine/Plugins/Runtime/Metasound/Source/MetasoundEditor/Private/MetasoundEditorGraphSchema.cpp
+++ b/Engine/Plugins/Runtime/Metasound/Source/MetasoundEditor/Private/MetasoundEditorGraphSchema.cpp
@@ -32,6 +32,31 @@ namespace Metasound
 {
 	namespace Editor
 	{
+		FString GetDataTypeDisplayName(const FName& InDataTypeName)
+		{
+			FString CategoryString = InDataTypeName.ToString();
+			int32 Index = 0;
+			CategoryString.FindLastChar(':', Index);
+
+			return CategoryString.RightChop(Index + 1);
+		}
+
+		TArray<FString> GetDataTypeNameCategories(const FName& InDataTypeName)
+		{
+			FString CategoryString = InDataTypeName.ToString();
+
+			TArray<FString> Categories;
+			CategoryString.ParseIntoArray(Categories, TEXT(":"));
+
+			if (Categories.Num() > 0)
+			{
+				// Remove name
+				Categories.RemoveAt(Categories.Num() - 1);
+			}
+
+			return Categories;
+		}
+
 		TSharedPtr<FEditor> GetEditorForGraph(const UMetasound& Metasound)
 		{
 			TSharedPtr<IToolkit> FoundAssetEditor = FToolkitManager::Get().FindEditorForAsset(CastChecked<const UObject>(&Metasound));
@@ -171,15 +196,15 @@ UEdGraphNode* FMetasoundGraphSchemaAction_NewInput::PerformAction(UEdGraph* Pare
 	UMetasound& ParentMetasound = CastChecked<UMetasoundEditorGraph>(ParentGraph)->GetMetasoundChecked();
 	ParentMetasound.Modify();
 
-	FString NameBase = FString::Printf(TEXT("%s"), *NodeTypeName.ToString());
+	FString NameBase = GetDataTypeDisplayName(NodeTypeName);
 
 	FGraphHandle GraphHandle = ParentMetasound.GetRootGraphHandle();
-	FString NewNodeName = NameBase;
 
-	int32 i = 0;
-	while (GraphHandle.ContainsOutputNodeWithName(NewNodeName))
+	int32 i = 1;
+	FString NewNodeName = NameBase + FString::Printf(TEXT("_%02d"), i);
+	while (GraphHandle.ContainsInputNodeWithName(NewNodeName))
 	{
-		NewNodeName = NameBase + FString::Printf(TEXT("_%u"), ++i);
+		NewNodeName = NameBase + FString::Printf(TEXT("_%02d"), ++i);
 	}
 
 	UEdGraphNode* NewGraphNode = FGraphBuilder::AddInput(ParentMetasound, Location, NewNodeName, NodeTypeName, FText::GetEmpty());
@@ -204,15 +229,15 @@ UEdGraphNode* FMetasoundGraphSchemaAction_NewOutput::PerformAction(UEdGraph* Par
 	UMetasound& ParentMetasound = CastChecked<UMetasoundEditorGraph>(ParentGraph)->GetMetasoundChecked();
 	ParentMetasound.Modify();
 
-	FString NameBase = FString::Printf(TEXT("%s"), *NodeTypeName.ToString());
+	FString NameBase = GetDataTypeDisplayName(NodeTypeName);
 
 	FGraphHandle GraphHandle = ParentMetasound.GetRootGraphHandle();
-	FString NewNodeName = NameBase;
 
-	int32 i = 0;
+	int32 i = 1;
+	FString NewNodeName = NameBase + FString::Printf(TEXT("_%02d"), i);
 	while (GraphHandle.ContainsOutputNodeWithName(NewNodeName))
 	{
-		NewNodeName = NameBase + FString::Printf(TEXT("_%u"), ++i);
+		NewNodeName = NameBase + FString::Printf(TEXT("_%02d"), ++i);
 	}
 
 	UEdGraphNode* NewGraphNode = FGraphBuilder::AddOutput(ParentMetasound, Location, NewNodeName, NodeTypeName, FText::GetEmpty());
@@ -457,22 +482,18 @@ void UMetasoundEditorGraphSchema::GetAllMetasoundActions(FGraphActionMenuBuilder
 	EditorModule.IterateDataTypes([InputMenuName, OutputMenuName, InMenuBuilder = &ActionMenuBuilder](const FEditorDataType& DataType)
 	{
 		const FName DataTypeName = DataType.RegistryInfo.DataTypeName;
+		const FText DataTypeDisplayName = FText::FromString(GetDataTypeDisplayName(DataTypeName));
 		const FText DataTypeTextName = FText::FromName(DataTypeName);
 		const FText MenuJoinFormat = LOCTEXT("MetasoundFormatNodeSubCategory", "{0}|{1}");
 
-		FString SubMenu = DataTypeName.ToString();
-		int32 Index = 0;
-		SubMenu.FindLastChar(':', Index);
-
-		const FText DisplayName = FText::FromString(SubMenu.RightChop(Index + 1));
+		const TArray<FString> Categories = GetDataTypeNameCategories(DataTypeName);
+		const FText CategoriesText = FText::FromString(FString::Join(Categories, TEXT("|")));
 
-		SubMenu.RemoveAt(Index, SubMenu.Len() - Index);
-		SubMenu.ReplaceCharInline(':', '|');
 
 		TSharedPtr<FMetasoundGraphSchemaAction_NewInput> AddInputNodeAction = MakeShared<FMetasoundGraphSchemaAction_NewInput>
 		(
-			FText::Format(MenuJoinFormat, InputMenuName, FText::FromString(SubMenu)),
-			DisplayName,
+			FText::Format(MenuJoinFormat, InputMenuName, CategoriesText),
+			DataTypeDisplayName,
 			DataTypeName,
 			FText::Format(LOCTEXT("MetasoundTooltipAddInputFormat", "Adds an input of type {0} to the Metasound"), DataTypeTextName),
 			0
@@ -481,8 +502,8 @@ void UMetasoundEditorGraphSchema::GetAllMetasoundActions(FGraphActionMenuBuilder
 
 		TSharedPtr<FMetasoundGraphSchemaAction_NewOutput> AddOutputNodeAction = MakeShared<FMetasoundGraphSchemaAction_NewOutput>
 		(
-			FText::Format(MenuJoinFormat, OutputMenuName, FText::FromString(SubMenu)),
-			DisplayName,
+			FText::Format(MenuJoinFormat, OutputMenuName, CategoriesText),
+			DataTypeDisplayName,
 			DataTypeName,
 			FText::Format(LOCTEXT("MetasoundTooltipAddOutputFormat", "Adds an output of type {0} to the Metasound"), DataTypeTextName),
 			0
diff --git a/Engine/Plugins/Runtime/Metasound/Source/MetasoundEditor/Public/MetasoundEditorCommands.h b/Engine/Plugins/Runtime/Metasound/Source/MetasoundEditor/Public/MetasoundEditorCommands.h
--- a/Engine/Plugins/Runtime/Metasound/Source/MetasoundEditor/Public/MetasoundEditorCommands.h
+++ b/Engine/Plugins/Runtime/Metasound/Source/MetasoundEditor/Public/MetasoundEditorCommands.h
@@ -25,8 +25,11 @@ namespace Metasound
 			/** Stops the currently playing Metasound */
 			TSharedPtr<FUICommandInfo> Stop;
 
-			/** Compiles the Metasound */
-			TSharedPtr<FUICommandInfo> Compile;
+			/** Imports the Metasound from Json */
+			TSharedPtr<FUICommandInfo> Import;
+
+			/** Exports the Metasound */
+			TSharedPtr<FUICommandInfo> Export;
 
 			/** Plays stops the currently playing Metasound */
 			TSharedPtr<FUICommandInfo> TogglePlayback;
diff --git a/Engine/Plugins/Runtime/Metasound/Source/MetasoundEditor/Public/MetasoundEditorGraphNode.h b/Engine/Plugins/Runtime/Metasound/Source/MetasoundEditor/Public/MetasoundEditorGraphNode.h
--- a/Engine/Plugins/Runtime/Metasound/Source/MetasoundEditor/Public/MetasoundEditorGraphNode.h
+++ b/Engine/Plugins/Runtime/Metasound/Source/MetasoundEditor/Public/MetasoundEditorGraphNode.h
@@ -13,6 +13,7 @@
 
 // Forward Declarations
 class UEdGraphPin;
+class UMetasound;
 class USoundNode;
 
 
@@ -48,10 +49,15 @@ public:
 
 	// UObject interface
 	virtual void PostLoad() override;
+
 	virtual void PostEditImport() override;
 	virtual void PostDuplicate(bool bDuplicateForPIE) override;
 	// End of UObject interface
 
+	UMetasound& GetMetasoundChecked();
+	const UMetasound& GetMetasoundChecked() const;
+
+	Metasound::Frontend::FGraphHandle GetRootGraphHandle() const;
 	Metasound::Frontend::FNodeHandle GetNodeHandle() const;
 	void SetNodeID(uint32 InNodeID);
 
diff --git a/Engine/Plugins/Runtime/Metasound/Source/MetasoundEngine/Private/Metasound.cpp b/Engine/Plugins/Runtime/Metasound/Source/MetasoundEngine/Private/Metasound.cpp
--- a/Engine/Plugins/Runtime/Metasound/Source/MetasoundEngine/Private/Metasound.cpp
+++ b/Engine/Plugins/Runtime/Metasound/Source/MetasoundEngine/Private/Metasound.cpp
@@ -1,14 +1,15 @@
 // Copyright Epic Games, Inc. All Rights Reserved.
 #include "Metasound.h"
 
-#include "Backends/JsonStructSerializerBackend.h"
 #include "CoreMinimal.h"
 #include "HAL/FileManager.h"
 #include "StructSerializer.h"
+#include "UObject/UnrealType.h"
+
 
 UMetasound::UMetasound(const FObjectInitializer& ObjectInitializer)
 	: Super(ObjectInitializer)
-	, FMetasoundAssetBase()
+	, FMetasoundAssetBase(RootMetasoundDocument)
 {
 }
 
@@ -35,6 +36,32 @@ const UEdGraph& UMetasound::GetGraphChecked() const
 	return *Graph;
 }
 
+const FText& UMetasound::GetInputToolTip(FString InputName) const
+{
+	for (const FMetasoundInputDescription& Desc : RootMetasoundDocument.RootClass.Inputs)
+	{
+		if (Desc.Name == InputName)
+		{
+			return Desc.ToolTip;
+		}
+	}
+
+	return FText::GetEmpty();
+}
+
+const FText& UMetasound::GetOutputToolTip(FString OutputName) const
+{
+	for (const FMetasoundOutputDescription& Desc : RootMetasoundDocument.RootClass.Outputs)
+	{
+		if (Desc.Name == OutputName)
+		{
+			return Desc.ToolTip;
+		}
+	}
+
+	return FText::GetEmpty();
+}
+
 void UMetasound::SetGraph(UEdGraph* InGraph)
 {
 	Graph = InGraph;
diff --git a/Engine/Plugins/Runtime/Metasound/Source/MetasoundEngine/Public/Metasound.h b/Engine/Plugins/Runtime/Metasound/Source/MetasoundEngine/Public/Metasound.h
--- a/Engine/Plugins/Runtime/Metasound/Source/MetasoundEngine/Public/Metasound.h
+++ b/Engine/Plugins/Runtime/Metasound/Source/MetasoundEngine/Public/Metasound.h
@@ -13,24 +13,10 @@
 
 #include "Metasound.generated.h"
 
+// Forward Declarations
+class FEditPropertyChain;
+struct FPropertyChangedEvent;
 
-USTRUCT()
-struct METASOUNDENGINE_API FEdMetasoundInput
-{
-	GENERATED_BODY()
-
-	UPROPERTY()
-	UEdGraphNode* InputNode = nullptr;
-};
-
-USTRUCT()
-struct METASOUNDENGINE_API FEdMetasoundOutput
-{
-	GENERATED_BODY()
-
-	UPROPERTY()
-	UEdGraphNode* OutputNode = nullptr;
-};
 
 /**
  * This asset type is used for Metasound assets that can only be used as nodes in other Metasound graphs.
@@ -48,12 +34,6 @@ protected:
 #if WITH_EDITORONLY_DATA
 	UPROPERTY()
 	UEdGraph* Graph;
-
-	UPROPERTY()
-	TArray<FEdMetasoundInput> Inputs;
-
-	UPROPERTY()
-	TArray<FEdMetasoundOutput> Outputs;
 #endif // WITH_EDITORONLY_DATA
 
 public:
@@ -73,6 +53,10 @@ public:
 	// Metasound UObject for editor serialization purposes.
 	// @param Editor graph associated with UMetasound.
 	void SetGraph(UEdGraph* InGraph);
+
+	const FText& GetInputToolTip(FString InputName) const;
+	const FText& GetOutputToolTip(FString OutputName) const;
+
 #endif // WITH_EDITORONLY_DATA
 
 	FMetasoundDocument& GetDocument() override
@@ -80,9 +64,16 @@ public:
 		return RootMetasoundDocument;
 	}
 
-	UObject* GetOwningAsset() override
+	const FMetasoundDocument& GetDocument() const override
+	{
+		return RootMetasoundDocument;
+	}
+
+	UObject* GetOwningAsset() const override
 	{
-		return this;
+		// Hack to allow handles to manipulate while providing
+		// ability to inspect via a handle from the const version of GetRootGraphHandle()
+		return const_cast<UObject*>(CastChecked<const UObject>(this));
 	}
 
 	// Updates the Metasound's metadata (name, author, etc).
diff --git a/Engine/Plugins/Runtime/Metasound/Source/MetasoundEngine/Public/MetasoundAssetSerializer.h b/Engine/Plugins/Runtime/Metasound/Source/MetasoundEngine/Public/MetasoundAssetSerializer.h
deleted file mode 100644
diff --git a/Engine/Plugins/Runtime/Metasound/Source/MetasoundFrontend/Private/MetasoundAssetBase.cpp b/Engine/Plugins/Runtime/Metasound/Source/MetasoundFrontend/Private/MetasoundAssetBase.cpp
--- a/Engine/Plugins/Runtime/Metasound/Source/MetasoundFrontend/Private/MetasoundAssetBase.cpp
+++ b/Engine/Plugins/Runtime/Metasound/Source/MetasoundFrontend/Private/MetasoundAssetBase.cpp
@@ -8,6 +8,17 @@
 #include "StructSerializer.h"
 
 
+const FString FMetasoundAssetBase::FileExtension(TEXT(".metasound"));
+
+FMetasoundAssetBase::FMetasoundAssetBase()
+{
+}
+
+FMetasoundAssetBase::FMetasoundAssetBase(FMetasoundDocument& InDocument)
+	: AccessPoint(MakeShared<Metasound::Frontend::FDescriptionAccessPoint>(InDocument))
+{
+}
+
 void FMetasoundAssetBase::SetMetadata(FMetasoundClassMetadata& InMetadata)
 {
 	InMetadata.NodeType = EMetasoundClassType::MetasoundGraph;
@@ -19,17 +30,12 @@ FMetasoundClassMetadata FMetasoundAssetBase::GetMetadata()
 	return GetDocument().RootClass.Metadata;
 }
 
-TWeakPtr<Metasound::Frontend::FDescriptionAccessPoint> FMetasoundAssetBase::GetGraphAccessPoint()
+TWeakPtr<Metasound::Frontend::FDescriptionAccessPoint> FMetasoundAssetBase::GetGraphAccessPoint() const
 {
-	if (!AccessPoint.IsValid())
-	{
-		AccessPoint = MakeShared<Metasound::Frontend::FDescriptionAccessPoint>(GetDocument());
-	}
-
 	return TWeakPtr<Metasound::Frontend::FDescriptionAccessPoint>(AccessPoint);
 }
 
-Metasound::Frontend::FGraphHandle FMetasoundAssetBase::GetRootGraphHandle()
+Metasound::Frontend::FGraphHandle FMetasoundAssetBase::GetRootGraphHandle() const
 {
 	using FDescPath = Metasound::Frontend::FDescPath;
 	using FGraphHandle = Metasound::Frontend::FGraphHandle;
@@ -69,7 +75,12 @@ TArray<Metasound::Frontend::FGraphHandle> FMetasoundAssetBase::GetAllSubgraphHan
 	return OutArray;
 }
 
-bool FMetasoundAssetBase::ExportToJSON(const FString& InAbsolutePath)
+bool FMetasoundAssetBase::ImportFromJSON(const FString& InAbsolutePath)
+{
+	return Metasound::Frontend::ImportJSONToMetasound(InAbsolutePath, GetDocument());
+}
+
+bool FMetasoundAssetBase::ExportToJSON(const FString& InAbsolutePath) const
 {
 	if (TUniquePtr<FArchive> FileWriter = TUniquePtr<FArchive>(IFileManager::Get().CreateFileWriter(*InAbsolutePath)))
 	{
diff --git a/Engine/Plugins/Runtime/Metasound/Source/MetasoundFrontend/Private/MetasoundFrontend.cpp b/Engine/Plugins/Runtime/Metasound/Source/MetasoundFrontend/Private/MetasoundFrontend.cpp
--- a/Engine/Plugins/Runtime/Metasound/Source/MetasoundFrontend/Private/MetasoundFrontend.cpp
+++ b/Engine/Plugins/Runtime/Metasound/Source/MetasoundFrontend/Private/MetasoundFrontend.cpp
@@ -812,6 +812,11 @@ namespace Metasound
 				return EMetasoundClassType::Invalid;
 			}
 
+			if (NodeClassType == EMetasoundClassType::Input || NodeClassType == EMetasoundClassType::Output)
+			{
+				return NodeClassType;
+			}
+
 			return NodeClass->Metadata.NodeType;
 		}
 
@@ -880,6 +885,18 @@ namespace Metasound
 			return NodePtr->UniqueID;
 		}
 
+
+		const FString& FNodeHandle::GetNodeName() const
+		{
+			if (!IsValid())
+			{
+				static const FString DefaultName = "InvalidNodeHandle";
+				return DefaultName;
+			}
+
+			return NodePtr->Name;
+		}
+
 		uint32 FNodeHandle::GetNodeID(const FDescPath& InNodePath)
 		{
 			if (!ensureAlwaysMsgf(InNodePath.Path.Num() != 0, TEXT("Tried to get a node ID from an empty path.")))
@@ -1070,7 +1087,7 @@ namespace Metasound
 			return OutArray;
 		}
 
-		FNodeHandle FGraphHandle::GetNodeWithId(uint32 InNodeId)
+		FNodeHandle FGraphHandle::GetNodeWithId(uint32 InNodeId) const
 		{
 			if (!IsValid())
 			{
@@ -1304,7 +1321,7 @@ namespace Metasound
 			}
 
 			// Finally, remove the node, and remove the input.
-			if (!ensureAlwaysMsgf(RemoveNode(InputNode, true), TEXT("Call to RemoveNode failed.")))
+			if (!ensureAlwaysMsgf(RemoveNodeInternal(InputNode), TEXT("Call to RemoveNodeInternal failed.")))
 			{
 				return false;
 			}
@@ -1386,18 +1403,18 @@ namespace Metasound
 			// find the corresponding node handle to delete.
 			FNodeHandle OutputNode = GetOutputNodeWithName(OutputName);
 
-			// If we found the input declared in the class description but couldn't find it in the graph,
-			// something has gone terribly wrong. Remove the input from the description, but still ensure.
-			if (!ensureAlwaysMsgf(OutputNode.IsValid(), TEXT(R"(Couldn't find an input node with name %s, even though we found the input listed as a dependency.
+			// If we found the output declared in the class description but couldn't find it in the graph,
+			// something has gone terribly wrong. Remove the output from the description, but still ensure.
+			if (!ensureAlwaysMsgf(OutputNode.IsValid(), TEXT(R"(Couldn't find an output node with name %s, even though we found the output listed as a dependency.
 				This indicates the underlying FMetasoundClassDescription is corrupted.
-				Removing the Input in the class dependency to resolve...)"), *OutputName))
+				Removing the Output in the class dependency to resolve...)"), *OutputName))
 			{
 				Outputs.RemoveAt(IndexOfOutputToRemove);
 				return true;
 			}
 
-			// Finally, remove the node, and remove the input.
-			if (!ensureAlwaysMsgf(RemoveNode(OutputNode, true), TEXT("Call to RemoveNode failed.")))
+			// Finally, remove the node, and remove the output.
+			if (!ensureAlwaysMsgf(RemoveNodeInternal(OutputNode), TEXT("Call to RemoveNodeInternal failed.")))
 			{
 				return false;
 			}
@@ -1576,7 +1593,18 @@ namespace Metasound
 			return FNodeHandle(FHandleInitParams::PrivateToken, InitParams, NewNodeDescription.ObjectTypeOfNode);
 		}
 
-		bool FGraphHandle::RemoveNode(const FNodeHandle& InNode, bool bEvenIfInputOrOutputNode /*= false*/)
+		bool FGraphHandle::RemoveNode(const FNodeHandle& InNode)
+		{
+			if (!ensureAlwaysMsgf(InNode.GetNodeType() != EMetasoundClassType::Input && InNode.GetNodeType() != EMetasoundClassType::Output,
+				TEXT("Inputs and outputs must be removed explicitly using 'RemoveInput' or 'RemoveOutput').")))
+			{
+				return false;
+			}
+
+			return RemoveNodeInternal(InNode);
+		}
+
+		bool FGraphHandle::RemoveNodeInternal(const FNodeHandle& InNode)
 		{
 			if (!IsValid())
 			{
@@ -1620,10 +1648,16 @@ namespace Metasound
 				return false;
 			}
 
+			if (InNode.GetNodeType() == EMetasoundClassType::Input || InNode.GetNodeType() == EMetasoundClassType::Output)
+			{
+				Nodes.RemoveAt(IndexOfNodeToRemove);
+				return true;
+			}
+
 			// This should never hit based on the logic above.
-			if (!ensureAlwaysMsgf(NodesOfClass > 0, TEXT("Found node with matching ID (%u) but mismatched class (%s). Likely means that the underlying class description was corrupted."), 
-				*InNode.GetNodeClassName(),
-				InNode.GetNodeType()))
+			if (!ensureAlwaysMsgf(NodesOfClass > 0, TEXT("Found node with matching ID (%u) but mismatched class (%s). Likely means that the underlying class description was corrupted."),
+				InNode.GetNodeID(),
+				*InNode.GetNodeClassName()))
 			{
 				return false;
 			}
@@ -1637,7 +1671,7 @@ namespace Metasound
 
 				uint32 UniqueIDForThisDependency = INDEX_NONE;
 				TArray<FMetasoundClassDescription>& DependencyClasses = OwningDocument->Dependencies;
-				
+
 				// scan the owning document's depenency classes for a dependency with this name.
 				int32 IndexOfDependencyInDocument = -1;
 				for (int32 Index = 0; Index < DependencyClasses.Num(); Index++)
@@ -1650,7 +1684,7 @@ namespace Metasound
 					}
 				}
 
-				if (!ensureAlwaysMsgf(IndexOfDependencyToRemove >= 0, 
+				if (!ensureAlwaysMsgf(IndexOfDependencyToRemove >= 0,
 					TEXT("Couldn't find node class %s in the list of dependencies for this document, but found it in the nodes list. This likely means that the underlying class description is corrupted.)"), *NodeClassName))
 				{
 					return false;
@@ -1668,7 +1702,7 @@ namespace Metasound
 				}
 
 				if (ensureAlwaysMsgf(IndexOfDependencyToRemove > 0, TEXT(R"(Couldn't find node class %s in the list of dependencies for this graph, but found it in the nodes list.
-					This likely means that the underlying class description is corrupted.)"), *NodeClassName))
+				This likely means that the underlying class description is corrupted.)"), *NodeClassName))
 				{
 					DependencyIDs.RemoveAt(IndexOfDependencyToRemove);
 				}
diff --git a/Engine/Plugins/Runtime/Metasound/Source/MetasoundFrontend/Public/MetasoundAssetBase.h b/Engine/Plugins/Runtime/Metasound/Source/MetasoundFrontend/Public/MetasoundAssetBase.h
--- a/Engine/Plugins/Runtime/Metasound/Source/MetasoundFrontend/Public/MetasoundAssetBase.h
+++ b/Engine/Plugins/Runtime/Metasound/Source/MetasoundFrontend/Public/MetasoundAssetBase.h
@@ -12,6 +12,11 @@
 class METASOUNDFRONTEND_API FMetasoundAssetBase
 {
 public:
+	static const FString FileExtension;
+
+	FMetasoundAssetBase();
+	FMetasoundAssetBase(FMetasoundDocument& InDocument);
+
 	virtual ~FMetasoundAssetBase() = default;
 
 	// Sets/overwrites the root class metadata
@@ -20,11 +25,14 @@ public:
 	// Returns the root class metadata
 	FMetasoundClassMetadata GetMetadata();
 
+	// Imports the asset from a JSON file at provided path
+	bool ImportFromJSON(const FString& InAbsolutePath);
+
 	// Exports the asset to JSON file at provided path
-	bool ExportToJSON(const FString& InAbsolutePath);
+	bool ExportToJSON(const FString& InAbsolutePath) const;
 
 	// Returns handle for the root metasound graph of this asset.
-	Metasound::Frontend::FGraphHandle GetRootGraphHandle();
+	Metasound::Frontend::FGraphHandle GetRootGraphHandle() const;
 
 	// Returns all handles for subgraphs referenced
 	TArray<Metasound::Frontend::FGraphHandle> GetAllSubgraphHandles();
@@ -39,15 +47,17 @@ protected:
 	// Returns document object responsible for serializing asset
 	virtual FMetasoundDocument& GetDocument() = 0;
 
+	// Returns document object responsible for serializing asset
+	virtual const FMetasoundDocument& GetDocument() const = 0;
+
 	// Returns the owning asset responsible for transactions applied to metasound
-	virtual UObject* GetOwningAsset() = 0;
+	virtual UObject* GetOwningAsset() const = 0;
 
 	// Returns a weak pointer that can be used to build a TDescriptionPtr
 	// for direct editing of the FMetasoundClassDescription tree.
 	// For advance use only, and requires knowledge of Metasound::Frontend::FDescPath syntax.
 	// For most use cases, use GetRootGraphHandle() instead.
-	TWeakPtr<Metasound::Frontend::FDescriptionAccessPoint> GetGraphAccessPoint();
+	TWeakPtr<Metasound::Frontend::FDescriptionAccessPoint> GetGraphAccessPoint() const;
 
-private:
 	TSharedPtr<Metasound::Frontend::FDescriptionAccessPoint> AccessPoint;
 };
diff --git a/Engine/Plugins/Runtime/Metasound/Source/MetasoundFrontend/Public/MetasoundFrontend.h b/Engine/Plugins/Runtime/Metasound/Source/MetasoundFrontend/Public/MetasoundFrontend.h
--- a/Engine/Plugins/Runtime/Metasound/Source/MetasoundFrontend/Public/MetasoundFrontend.h
+++ b/Engine/Plugins/Runtime/Metasound/Source/MetasoundFrontend/Public/MetasoundFrontend.h
@@ -83,7 +83,7 @@ namespace Metasound
 		// This will open a json document at the given absolute path and convert it into a 
 		// metasound document struct. 
 		// @returns false if the file couldn't be found or parsed into a document.
-		bool ImportJSONToMetasound(const FString& InPath, FMetasoundDocument& OutMetasoundDocument);
+		METASOUNDFRONTEND_API bool ImportJSONToMetasound(const FString& InPath, FMetasoundDocument& OutMetasoundDocument);
 
 		// Struct that indicates whether an input and an output can be connected,
 		// and whether an intermediate node is necessary to connect the two.
@@ -260,6 +260,8 @@ namespace Metasound
 			uint32 GetNodeID() const;
 			static uint32 GetNodeID(const FDescPath& InNodePath);
 
+			const FString& GetNodeName() const;
+
 		private:
 
 			static TDescriptionPtr<FMetasoundClassDescription> GetNodeClassDescriptionForNodeHandle(const FHandleInitParams& InitParams, EMetasoundClassType InNodeClassType);
@@ -297,7 +299,7 @@ namespace Metasound
 			bool IsValid() const;
 
 			TArray<FNodeHandle> GetAllNodes();
-			FNodeHandle GetNodeWithId(uint32 InNodeId);
+			FNodeHandle GetNodeWithId(uint32 InNodeId) const;
 			TArray<FNodeHandle> GetOutputNodes();
 			TArray<FNodeHandle> GetInputNodes();
 
@@ -332,7 +334,7 @@ namespace Metasound
 
 			// Remove the node corresponding to this node handle.
 			// On success, invalidates the received node handle.
-			bool RemoveNode(const FNodeHandle& InNode, bool bEvenIfInputOrOutputNode = false);
+			bool RemoveNode(const FNodeHandle& InNode);
 
 			// Returns the metadata for the current graph, including the name, description and author.
 			FMetasoundClassMetadata GetGraphMetadata();
@@ -374,6 +376,12 @@ namespace Metasound
 			// the outermost document for all dependencies for the document this graph lives in.
 			TDescriptionPtr<FMetasoundDocument> OwningDocument;
 
+			// Remove the node corresponding to this node handle.
+			// On success, invalidates the received node handle.
+			// Can remove inputs and outputs, but does not remove
+			// from the input/output arrays.
+			bool RemoveNodeInternal(const FNodeHandle& InNode);
+
 			// Scans all existing node ids to guarantee a new unique ID.
 			uint32 FindNewUniqueNodeId();
 
diff --git a/Engine/Plugins/Runtime/Metasound/Source/MetasoundFrontend/Public/MetasoundFrontendDataLayout.h b/Engine/Plugins/Runtime/Metasound/Source/MetasoundFrontend/Public/MetasoundFrontendDataLayout.h
--- a/Engine/Plugins/Runtime/Metasound/Source/MetasoundFrontend/Public/MetasoundFrontendDataLayout.h
+++ b/Engine/Plugins/Runtime/Metasound/Source/MetasoundFrontend/Public/MetasoundFrontendDataLayout.h
@@ -92,15 +92,15 @@ struct FMetasoundInputDescription
 	GENERATED_BODY()
 
 	// The descriptive name of this input.
-	UPROPERTY()
+	UPROPERTY(EditAnywhere, Category = Parameters)
 	FString Name;
 
 	// The type of this input. This type matches with a type declared via the DECLARE_METASOUND_DATA_REFERENCE_TYPES macro.
-	UPROPERTY()
+	UPROPERTY(VisibleAnywhere, Category = Parameters)
 	FName TypeName;
 
 	// Optional description text about this input.
-	UPROPERTY()
+	UPROPERTY(EditAnywhere, Category = Parameters)
 	FText ToolTip;
 
 	// The optional literal value, if we have one.
@@ -114,15 +114,15 @@ struct FMetasoundOutputDescription
 	GENERATED_BODY()
 
 	// The descriptive name of this output.
-	UPROPERTY()
+	UPROPERTY(EditAnywhere, Category = Parameters)
 	FString Name;
 
 	// The type of this output. This type matches with a type declared via the DECLARE_METASOUND_DATA_REFERENCE_TYPES macro.
-	UPROPERTY()
+	UPROPERTY(VisibleAnywhere, Category = Parameters)
 	FName TypeName;
 
 	// Optional description text about this output.
-	UPROPERTY()
+	UPROPERTY(EditAnywhere, Category = Parameters)
 	FText ToolTip;
 };
 
@@ -224,10 +224,10 @@ struct FMetasoundClassDescription
 	UPROPERTY(EditAnywhere, Category = Hidden)
 	FMetasoundClassMetadata Metadata;
 
-	UPROPERTY()
+	UPROPERTY(EditAnywhere, Category = Hidden)
 	TArray<FMetasoundInputDescription> Inputs;
 
-	UPROPERTY()
+	UPROPERTY(EditAnywhere, Category = Hidden)
 	TArray<FMetasoundOutputDescription> Outputs;
 
 	// If this object is itself of the Metasound type, here we list the unique IDs of other objects it depends on to be built.
