Fix icon paths
#rnx

diff --git a/Engine/Plugins/Runtime/Metasound/Source/MetasoundEditor/Private/MetasoundEditor.cpp b/Engine/Plugins/Runtime/Metasound/Source/MetasoundEditor/Private/MetasoundEditor.cpp
--- a/Engine/Plugins/Runtime/Metasound/Source/MetasoundEditor/Private/MetasoundEditor.cpp
+++ b/Engine/Plugins/Runtime/Metasound/Source/MetasoundEditor/Private/MetasoundEditor.cpp
@@ -288,8 +288,7 @@ namespace Metasound
 
 		FSlateIcon FEditor::GetImportStatusImage() const
 		{
-			// TODO: Get/register import image
-			const FName IconName = "MetasoundEditor.Export";
+			const FName IconName = "MetasoundEditor.Import";
 			return FSlateIcon("MetasoundStyle", IconName);
 		}
 
diff --git a/Engine/Plugins/Runtime/Metasound/Source/MetasoundEditor/Private/MetasoundEditorModule.cpp b/Engine/Plugins/Runtime/Metasound/Source/MetasoundEditor/Private/MetasoundEditorModule.cpp
--- a/Engine/Plugins/Runtime/Metasound/Source/MetasoundEditor/Private/MetasoundEditorModule.cpp
+++ b/Engine/Plugins/Runtime/Metasound/Source/MetasoundEditor/Private/MetasoundEditorModule.cpp
@@ -64,10 +64,12 @@ namespace Metasound
 					Set("MetasoundEditor.Play.Small", new FSlateImageBrush(RootToContentDir(TEXT("Icons/icon_SCueEd_PlayCue_40x.png")), Icon20x20));
 					Set("MetasoundEditor.Stop", new FSlateImageBrush(RootToContentDir(TEXT("Icons/icon_SCueEd_Stop_40x.png")), Icon40x40));
 					Set("MetasoundEditor.Stop.Small", new FSlateImageBrush(RootToContentDir(TEXT("Icons/icon_SCueEd_Stop_40x.png")), Icon20x20));
-					Set("MetasoundEditor.Compile", new FSlateImageBrush(RootToContentDir(TEXT("/Old/Kismet2/compile_40px.png")), Icon40x40));
-					Set("MetasoundEditor.Compile.Small", new FSlateImageBrush(RootToContentDir(TEXT("/Old/Kismet2/compile_40px.png")), Icon20x20));
-					Set("MetasoundEditor.CompileError", new FSlateImageBrush(RootToContentDir(TEXT("/Old/Kismet2/CompileStatus_Fail.png")), Icon40x40));
-					Set("MetasoundEditor.CompileError.Small", new FSlateImageBrush(RootToContentDir(TEXT("/Old/Kismet2/CompileStatus_Broken_Small.png")), Icon20x20));
+					Set("MetasoundEditor.Import", new FSlateImageBrush(RootToContentDir(TEXT("/Old/Kismet2/compile_40px.png")), Icon40x40));
+					Set("MetasoundEditor.Import.Small", new FSlateImageBrush(RootToContentDir(TEXT("/Old/Kismet2/compile_40px.png")), Icon20x20));
+					Set("MetasoundEditor.Export", new FSlateImageBrush(RootToContentDir(TEXT("/Old/Kismet2/compile_40px.png")), Icon40x40));
+					Set("MetasoundEditor.Export.Small", new FSlateImageBrush(RootToContentDir(TEXT("/Old/Kismet2/compile_40px.png")), Icon20x20));
+					Set("MetasoundEditor.ExportError", new FSlateImageBrush(RootToContentDir(TEXT("/Old/Kismet2/CompileStatus_Fail.png")), Icon40x40));
+					Set("MetasoundEditor.ExportError.Small", new FSlateImageBrush(RootToContentDir(TEXT("/Old/Kismet2/CompileStatus_Broken_Small.png")), Icon20x20));
 				}
 
 				FSlateStyleRegistry::RegisterSlateStyle(*this);
