Minor feedback from earlier CL to nix the 'Get' and 'Set' on I/O nodes
#rnx

diff --git a/Engine/Plugins/Runtime/Metasound/Source/MetasoundEditor/Private/MetasoundEditorGraphNode.cpp b/Engine/Plugins/Runtime/Metasound/Source/MetasoundEditor/Private/MetasoundEditorGraphNode.cpp
--- a/Engine/Plugins/Runtime/Metasound/Source/MetasoundEditor/Private/MetasoundEditorGraphNode.cpp
+++ b/Engine/Plugins/Runtime/Metasound/Source/MetasoundEditor/Private/MetasoundEditorGraphNode.cpp
@@ -186,13 +186,13 @@ FText UMetasoundEditorGraphNode::GetNodeTitle(ENodeTitleType::Type TitleType) co
 	{
 		case EMetasoundClassType::Input:
 		{
-			return FText::Format(LOCTEXT("MetasoundGraphNode_TitleFormat", "Get Input {0}"), FText::FromString(NodeHandle.GetNodeName()));
+			return FText::Format(LOCTEXT("MetasoundGraphNode_TitleFormat", "Input {0}"), FText::FromString(NodeHandle.GetNodeName()));
 		}
 		break;
 
 		case EMetasoundClassType::Output:
 		{
-			return FText::Format(LOCTEXT("MetasoundGraphNode_TitleFormat", "Set Output {0}"), FText::FromString(NodeHandle.GetNodeName()));
+			return FText::Format(LOCTEXT("MetasoundGraphNode_TitleFormat", "Output {0}"), FText::FromString(NodeHandle.GetNodeName()));
 		}
 		break;
 
