Fix static analysis bug
#rnx

diff --git a/Engine/Plugins/Runtime/Metasound/Source/MetasoundEditor/Private/MetasoundEditorGraphSchema.cpp b/Engine/Plugins/Runtime/Metasound/Source/MetasoundEditor/Private/MetasoundEditorGraphSchema.cpp
--- a/Engine/Plugins/Runtime/Metasound/Source/MetasoundEditor/Private/MetasoundEditorGraphSchema.cpp
+++ b/Engine/Plugins/Runtime/Metasound/Source/MetasoundEditor/Private/MetasoundEditorGraphSchema.cpp
@@ -371,7 +371,7 @@ const FPinConnectionResponse UMetasoundEditorGraphSchema::CanCreateConnection(co
 	const UEdGraphPin* InputPin = nullptr;
 	const UEdGraphPin* OutputPin = nullptr;
 
-	if (!CategorizePinsByDirection(PinA, PinB, /*out*/ InputPin, /*out*/ OutputPin))
+	if (!CategorizePinsByDirection(PinA, PinB, InputPin, OutputPin))
 	{
 		return FPinConnectionResponse(CONNECT_RESPONSE_DISALLOW, LOCTEXT("ConnectionIncompatible", "Directions are not compatible"));
 	}
@@ -381,7 +381,7 @@ const FPinConnectionResponse UMetasoundEditorGraphSchema::CanCreateConnection(co
 		return FPinConnectionResponse(CONNECT_RESPONSE_DISALLOW, LOCTEXT("ConnectionLoop", "Connection would cause loop"));
 	}
 
-	if (InputPin->PinType.PinCategory != InputPin->PinType.PinCategory)
+	if (InputPin->PinType.PinCategory != OutputPin->PinType.PinCategory)
 	{
 		return FPinConnectionResponse(CONNECT_RESPONSE_DISALLOW, LOCTEXT("ConnectionTypeIncorrect", "Connection pin types do not match"));
 	}
