Implement UMetasoundSource, Archetype registry, and implement Archetype support in the metasound editor.

diff --git a/Engine/Plugins/Runtime/Metasound/Source/MetasoundEditor/Private/MetasoundAssetTypeActions.cpp b/Engine/Plugins/Runtime/Metasound/Source/MetasoundEditor/Private/MetasoundAssetTypeActions.cpp
--- a/Engine/Plugins/Runtime/Metasound/Source/MetasoundEditor/Private/MetasoundAssetTypeActions.cpp
+++ b/Engine/Plugins/Runtime/Metasound/Source/MetasoundEditor/Private/MetasoundAssetTypeActions.cpp
@@ -3,6 +3,7 @@
 #include "MetasoundAssetTypeActions.h"
 
 #include "Metasound.h"
+#include "MetasoundSource.h"
 #include "MetasoundEditor.h"
 
 
@@ -40,6 +41,32 @@ namespace Metasound
 				}
 			}
 		}
+
+		UClass* FAssetTypeActions_MetasoundSource::GetSupportedClass() const
+		{
+			return UMetasoundSource::StaticClass();
+		}
+
+		void FAssetTypeActions_MetasoundSource::OpenAssetEditor(const TArray<UObject*>& InObjects, TSharedPtr<IToolkitHost> ToolkitHost)
+		{
+			const EToolkitMode::Type Mode = ToolkitHost.IsValid() ? EToolkitMode::WorldCentric : EToolkitMode::Standalone;
+
+			for (UObject* Object : InObjects)
+			{
+				if (UMetasoundSource* Metasound = Cast<UMetasoundSource>(Object))
+				{
+					TSharedRef<FEditor> NewEditor = MakeShared<FEditor>();
+					NewEditor->InitMetasoundEditor(Mode, ToolkitHost, Metasound);
+				}
+			}
+		}
+
+		const TArray<FText>& FAssetTypeActions_MetasoundSource::GetSubMenus() const
+		{
+			return AssetTypeActionSubMenu;
+		}
 	} // namespace Editor
 } // namespace Metasound
+
 #undef LOCTEXT_NAMESPACE
+
diff --git a/Engine/Plugins/Runtime/Metasound/Source/MetasoundEditor/Private/MetasoundAssetTypeActions.h b/Engine/Plugins/Runtime/Metasound/Source/MetasoundEditor/Private/MetasoundAssetTypeActions.h
--- a/Engine/Plugins/Runtime/Metasound/Source/MetasoundEditor/Private/MetasoundAssetTypeActions.h
+++ b/Engine/Plugins/Runtime/Metasound/Source/MetasoundEditor/Private/MetasoundAssetTypeActions.h
@@ -24,5 +24,18 @@ namespace Metasound
 
 			virtual const TArray<FText>& GetSubMenus() const override;
 		};
+
+		class FAssetTypeActions_MetasoundSource : public FAssetTypeActions_Base
+		{
+		public:
+			// IAssetTypeActions Implementation
+			virtual FText GetName() const override { return NSLOCTEXT("AssetTypeActions", "AssetTypeActions_MetasoundSource", "MetasoundSource"); }
+			virtual FColor GetTypeColor() const override { return FColor(220, 160, 0); }
+			virtual UClass* GetSupportedClass() const override;
+			virtual uint32 GetCategories() override { return EAssetTypeCategories::Sounds; }
+			virtual void OpenAssetEditor(const TArray<UObject*>& InObjects, TSharedPtr<IToolkitHost> ToolkitHost) override;
+
+			virtual const TArray<FText>& GetSubMenus() const override;
+		};
 	} // namespace Editor
 } // namespace Metasound
diff --git a/Engine/Plugins/Runtime/Metasound/Source/MetasoundEditor/Private/MetasoundEditor.cpp b/Engine/Plugins/Runtime/Metasound/Source/MetasoundEditor/Private/MetasoundEditor.cpp
--- a/Engine/Plugins/Runtime/Metasound/Source/MetasoundEditor/Private/MetasoundEditor.cpp
+++ b/Engine/Plugins/Runtime/Metasound/Source/MetasoundEditor/Private/MetasoundEditor.cpp
@@ -93,18 +93,22 @@ namespace Metasound
 
 		void FEditor::InitMetasoundEditor(const EToolkitMode::Type Mode, const TSharedPtr<IToolkitHost>& InitToolkitHost, UObject* ObjectToEdit)
 		{
-
-			Metasound = CastChecked<UMetasound>(ObjectToEdit);
+			checkf(Frontend::IsObjectAMetasoundArchetype(ObjectToEdit), TEXT("Object passed in was not registered as a valid metasound archetype!"));
+			
+			Metasound = ObjectToEdit;
+			FMetasoundAssetBase* MetasoundAsset = Frontend::GetObjectAsAssetBase(ObjectToEdit);
+			check(MetasoundAsset);
 
 			// Support undo/redo
-			Metasound->SetFlags(RF_Transactional);
+			ObjectToEdit->SetFlags(RF_Transactional);
 
-			if (!Metasound->GetGraph())
+			if (!MetasoundAsset->GetGraph())
 			{
 				UMetasoundEditorGraph* Graph = NewObject<UMetasoundEditorGraph>(Metasound);
 				Graph->ParentMetasound = Metasound;
 				Graph->Schema = UMetasoundEditorGraphSchema::StaticClass();
-				Metasound->SetGraph(Graph);
+				MetasoundAsset->SetGraph(Graph);
+				FGraphBuilder::RebuildGraph(*Metasound);
 			}
 
 			GEditor->RegisterForUndo(this);
@@ -163,7 +167,7 @@ namespace Metasound
 			RegenerateMenusAndToolbars();
 		}
 
-		UMetasound* FEditor::GetMetasound() const
+		UObject* FEditor::GetMetasoundObject() const
 		{
 			return Metasound;
 		}
@@ -338,44 +342,70 @@ namespace Metasound
 
 		void FEditor::Import()
 		{
-			if (Metasound)
+			FMetasoundAssetBase* MetasoundAsset = Frontend::GetObjectAsAssetBase(Metasound);
+			if (MetasoundAsset)
 			{
 				// TODO: Prompt OFD and provide path from user
-				const FString Path = FPaths::ProjectIntermediateDir() / TEXT("Metasounds") + FPaths::ChangeExtension(Metasound->GetPathName(), FMetasoundAssetBase::FileExtension);
-				Metasound->ImportFromJSON(Path);
+				const FString InputPath = FPaths::ProjectIntermediateDir() / TEXT("Metasounds") + FPaths::ChangeExtension(Metasound->GetPathName(), FMetasoundAssetBase::FileExtension);
+				
+				// TODO: use the same directory as the currently open metasound
+				const FString OutputPath = FString("/Game/ImportedMetasound/GeneratedMetasound");
+
+				FMetasoundDocument MetasoundDoc;
+
+				if (Frontend::ImportJSONToMetasound(InputPath, MetasoundDoc))
+				{
+					Frontend::GetObjectForDocument(MetasoundDoc, OutputPath);
+				}
+				else
+				{
+					UE_LOG(LogTemp, Warning, TEXT("Couldn't import Metasound at path: %s"), *InputPath);
+				}
 			}
 		}
 
 		void FEditor::Export()
 		{
-			if (Metasound)
+			FMetasoundAssetBase* InMetasoundAsset = Frontend::GetObjectAsAssetBase(Metasound);
+			check(InMetasoundAsset);
+
+			if (!InMetasoundAsset)
 			{
-				// TODO: Prompt OFD and provide path from user
-				const FString Path = FPaths::ProjectIntermediateDir() / TEXT("Metasounds") + FPaths::ChangeExtension(Metasound->GetPathName(), FMetasoundAssetBase::FileExtension);
-				Metasound->ExportToJSON(Path);
+				return;
 			}
+
+			static const FString MetasoundExtension(TEXT(".metasound"));
+
+			// TODO: We could just make this an object.
+			const FString Path = FPaths::ProjectSavedDir() / TEXT("Metasounds") + FPaths::ChangeExtension(Metasound->GetPathName(), MetasoundExtension);
+			InMetasoundAsset->ExportToJSON(Path);
 		}
 
 		void FEditor::Play()
 		{
 			// 	TODO: Implement play
-			// 	check(GEditor);
-			// 	GEditor->PlayPreviewSound(Metasound);
-			// 
-			// 	MetasoundGraphEditor->RegisterActiveTimer(0.0f, 
-			// 		FWidgetActiveTimerDelegate::CreateLambda([](double InCurrentTime, float InDeltaTime)
-			// 		{
-			// 			UAudioComponent* PreviewComp = GEditor->GetPreviewAudioComponent();
-			// 			if (PreviewComp && PreviewComp->IsPlaying())
-			// 			{
-			// 				return EActiveTimerReturnType::Continue;
-			// 			}
-			// 			else
-			// 			{
-			// 				return EActiveTimerReturnType::Stop;
-			// 			}
-			// 		})
-			// 	);
+			check(GEditor);
+
+			// TODO: toggle the Play button if Metasound is a USoundBase
+			if (USoundBase* MetasoundToPlay = Cast<USoundBase>(Metasound))
+			{
+				GEditor->PlayPreviewSound(MetasoundToPlay);
+
+				MetasoundGraphEditor->RegisterActiveTimer(0.0f,
+					FWidgetActiveTimerDelegate::CreateLambda([](double InCurrentTime, float InDeltaTime)
+						{
+							UAudioComponent* PreviewComp = GEditor->GetPreviewAudioComponent();
+							if (PreviewComp && PreviewComp->IsPlaying())
+							{
+								return EActiveTimerReturnType::Continue;
+							}
+							else
+							{
+								return EActiveTimerReturnType::Stop;
+							}
+						})
+				);
+			}
 		}
 
 		void FEditor::PlayNode()
@@ -543,11 +573,14 @@ namespace Metasound
 			InEvents.OnTextCommitted = FOnNodeTextCommitted::CreateSP(this, &FEditor::OnNodeTitleCommitted);
 			InEvents.OnNodeDoubleClicked = FSingleNodeEvent::CreateSP(this, &FEditor::PlaySingleNode);
 
+			FMetasoundAssetBase* MetasoundAsset = Frontend::GetObjectAsAssetBase(Metasound);
+			check(MetasoundAsset);
+
 			return SNew(SGraphEditor)
 				.AdditionalCommands(GraphEditorCommands)
 				.IsEditable(true)
 				.Appearance(AppearanceInfo)
-				.GraphToEdit(Metasound->GetGraph())
+				.GraphToEdit(MetasoundAsset->GetGraph())
 				.GraphEvents(InEvents)
 				.AutoExpandActionMenu(true)
 				.ShowGraphStateOverlay(false);
@@ -563,7 +596,7 @@ namespace Metasound
 				{
 					if (Cast<UMetasoundEditorGraphNode>(*SetIt))
 					{
-						Selection.Add(GetMetasound());
+						Selection.Add(GetMetasoundObject());
 					}
 					else
 					{
@@ -573,7 +606,7 @@ namespace Metasound
 			}
 			else
 			{
-				Selection.Add(GetMetasound());
+				Selection.Add(GetMetasoundObject());
 			}
 
 			SetSelection(Selection);
@@ -601,11 +634,16 @@ namespace Metasound
 			const FGraphPanelSelectionSet SelectedNodes = MetasoundGraphEditor->GetSelectedNodes();
 			MetasoundGraphEditor->ClearSelectionSet();
 
-			FGraphHandle GraphHandle = Graph->ParentMetasound->GetRootGraphHandle();
-			for (FGraphPanelSelectionSet::TConstIterator NodeIt(SelectedNodes); NodeIt; ++NodeIt)
+			UObject* ParentMetasoundObject = Graph->ParentMetasound;
+
+			if (FMetasoundAssetBase* MetasoundAsset = GetObjectAsAssetBase(ParentMetasoundObject))
 			{
-				UMetasoundEditorGraphNode* Node = CastChecked<UMetasoundEditorGraphNode>(*NodeIt);
-				FGraphBuilder::DeleteNode(*Node, false /* bInRecordTransaction */);
+				FGraphHandle GraphHandle = MetasoundAsset->GetRootGraphHandle();
+				for (FGraphPanelSelectionSet::TConstIterator NodeIt(SelectedNodes); NodeIt; ++NodeIt)
+				{
+					UMetasoundEditorGraphNode* Node = CastChecked<UMetasoundEditorGraphNode>(*NodeIt);
+					FGraphBuilder::DeleteNode(*Node, false /* bInRecordTransaction */);
+				}
 			}
 		}
 
@@ -708,7 +746,14 @@ namespace Metasound
 
 		void FEditor::PasteNodesAtLocation(const FVector2D& Location)
 		{
-			UEdGraph* Graph = Metasound->GetGraph();
+			FMetasoundAssetBase* MetasoundAsset = Frontend::GetObjectAsAssetBase(Metasound);
+
+			if (!ensureAlways(MetasoundAsset))
+			{
+				return;
+			}
+
+			UEdGraph* Graph = MetasoundAsset->GetGraph();
 			if (!Graph)
 			{
 				return;
diff --git a/Engine/Plugins/Runtime/Metasound/Source/MetasoundEditor/Private/MetasoundEditor.h b/Engine/Plugins/Runtime/Metasound/Source/MetasoundEditor/Private/MetasoundEditor.h
--- a/Engine/Plugins/Runtime/Metasound/Source/MetasoundEditor/Private/MetasoundEditor.h
+++ b/Engine/Plugins/Runtime/Metasound/Source/MetasoundEditor/Private/MetasoundEditor.h
@@ -43,7 +43,7 @@ namespace Metasound
 			void InitMetasoundEditor(const EToolkitMode::Type Mode, const TSharedPtr<IToolkitHost>& InitToolkitHost, UObject* ObjectToEdit);
 
 			/** IMetasoundEditor interface */
-			virtual UMetasound* GetMetasound() const override;
+			virtual UObject* GetMetasoundObject() const override;
 			virtual void SetSelection(const TArray<UObject*>& SelectedObjects) override;
 			virtual bool GetBoundsForSelectedNodes(FSlateRect& Rect, float Padding) override;
 
@@ -194,7 +194,7 @@ namespace Metasound
 			TSharedPtr<FUICommandList> GraphEditorCommands;
 
 			/** The Metasound asset being edited */
-			UMetasound* Metasound = nullptr;
+			UObject* Metasound = nullptr;
 
 			/** Whether or not metasound being edited is valid */
 			bool bPassedValidation = true;
diff --git a/Engine/Plugins/Runtime/Metasound/Source/MetasoundEditor/Private/MetasoundEditorGraph.cpp b/Engine/Plugins/Runtime/Metasound/Source/MetasoundEditor/Private/MetasoundEditorGraph.cpp
--- a/Engine/Plugins/Runtime/Metasound/Source/MetasoundEditor/Private/MetasoundEditorGraph.cpp
+++ b/Engine/Plugins/Runtime/Metasound/Source/MetasoundEditor/Private/MetasoundEditorGraph.cpp
@@ -4,12 +4,12 @@
 #include "EdGraph/EdGraphNode.h"
 
 
-UMetasound* UMetasoundEditorGraph::GetMetasound() const
+UObject* UMetasoundEditorGraph::GetMetasound() const
 {
 	return ParentMetasound;
 }
 
-UMetasound& UMetasoundEditorGraph::GetMetasoundChecked() const
+UObject& UMetasoundEditorGraph::GetMetasoundChecked() const
 {
 	check(ParentMetasound);
 	return *ParentMetasound;
diff --git a/Engine/Plugins/Runtime/Metasound/Source/MetasoundEditor/Private/MetasoundEditorGraphBuilder.cpp b/Engine/Plugins/Runtime/Metasound/Source/MetasoundEditor/Private/MetasoundEditorGraphBuilder.cpp
--- a/Engine/Plugins/Runtime/Metasound/Source/MetasoundEditor/Private/MetasoundEditorGraphBuilder.cpp
+++ b/Engine/Plugins/Runtime/Metasound/Source/MetasoundEditor/Private/MetasoundEditorGraphBuilder.cpp
@@ -26,13 +26,16 @@ namespace Metasound
 		const FName FGraphBuilder::PinPrimitiveInteger = "Int";
 		const FName FGraphBuilder::PinPrimitiveString = "String";
 
-		UEdGraphNode* FGraphBuilder::AddNode(UMetasound& InMetasound, const FVector2D& Location, Frontend::FNodeHandle& InNodeHandle, bool bInSelectNewNode)
+		UEdGraphNode* FGraphBuilder::AddNode(UObject& InMetasound, const FVector2D& Location, Frontend::FNodeHandle& InNodeHandle, bool bInSelectNewNode)
 		{
 			const FScopedTransaction Transaction(LOCTEXT("AddMetasoundGraphNode", "Add Metasound Node"));
 
 			const FString EdNodeName = InNodeHandle.GetNodeClassName() + TEXT("_") + FString::FromInt(InNodeHandle.GetNodeID());
 
-			UEdGraph& Graph = InMetasound.GetGraphChecked();
+			FMetasoundAssetBase* MetasoundAsset = Frontend::GetObjectAsAssetBase(&InMetasound);
+			check(MetasoundAsset);
+
+			UEdGraph& Graph = MetasoundAsset->GetGraphChecked();
 			FGraphNodeCreator<UMetasoundEditorGraphNode> NodeCreator(Graph);
 			UMetasoundEditorGraphNode* NewGraphNode = NodeCreator.CreateNode(bInSelectNewNode);
 			NodeCreator.Finalize();
@@ -50,31 +53,40 @@ namespace Metasound
 			return NewGraphNode;
 		}
 
-		UEdGraphNode* FGraphBuilder::AddNode(UMetasound& InMetasound, const FVector2D& Location, const Frontend::FNodeClassInfo& InClassInfo, bool bInSelectNewNode)
+		UEdGraphNode* FGraphBuilder::AddNode(UObject& InMetasound, const FVector2D& Location, const Frontend::FNodeClassInfo& InClassInfo, bool bInSelectNewNode)
 		{
-			Frontend::FNodeHandle NodeHandle = InMetasound.GetRootGraphHandle().AddNewNode(InClassInfo);
+			FMetasoundAssetBase* MetasoundAsset = Frontend::GetObjectAsAssetBase(&InMetasound);
+			check(MetasoundAsset);
+
+			Frontend::FNodeHandle NodeHandle = MetasoundAsset->GetRootGraphHandle().AddNewNode(InClassInfo);
 			return AddNode(InMetasound, Location, NodeHandle, bInSelectNewNode);
 		}
 
-		UEdGraphNode* FGraphBuilder::AddInput(UMetasound& InMetasound, const FVector2D& Location, const FString& InName, const FName InTypeName, const FText& InToolTip, bool bInSelectNewNode)
+		UEdGraphNode* FGraphBuilder::AddInput(UObject& InMetasound, const FVector2D& Location, const FString& InName, const FName InTypeName, const FText& InToolTip, bool bInSelectNewNode)
 		{
 			FMetasoundInputDescription Description;
 			Description.Name = InName;
 			Description.TypeName = InTypeName;
 			Description.ToolTip = InToolTip;
 
-			Frontend::FNodeHandle NodeHandle = InMetasound.GetRootGraphHandle().AddNewInput(Description);
+			FMetasoundAssetBase* MetasoundAsset = Frontend::GetObjectAsAssetBase(&InMetasound);
+			check(MetasoundAsset);
+
+			Frontend::FNodeHandle NodeHandle = MetasoundAsset->GetRootGraphHandle().AddNewInput(Description);
 			return AddNode(InMetasound, Location, NodeHandle, bInSelectNewNode);
 		}
 
-		UEdGraphNode* FGraphBuilder::AddOutput(UMetasound& InMetasound, const FVector2D& Location, const FString& InName, const FName InTypeName, const FText& InToolTip, bool bInSelectNewNode)
+		UEdGraphNode* FGraphBuilder::AddOutput(UObject& InMetasound, const FVector2D& Location, const FString& InName, const FName InTypeName, const FText& InToolTip, bool bInSelectNewNode)
 		{
 			FMetasoundOutputDescription Description;
 			Description.Name = InName;
 			Description.TypeName = InTypeName;
 			Description.ToolTip = InToolTip;
 
-			Frontend::FNodeHandle NodeHandle = InMetasound.GetRootGraphHandle().AddNewOutput(Description);
+			FMetasoundAssetBase* MetasoundAsset = Frontend::GetObjectAsAssetBase(&InMetasound);
+			check(MetasoundAsset);
+
+			Frontend::FNodeHandle NodeHandle = MetasoundAsset->GetRootGraphHandle().AddNewOutput(Description);
 			return AddNode(InMetasound, Location, NodeHandle, bInSelectNewNode);
 		}
 
@@ -90,7 +102,11 @@ namespace Metasound
 			}
 
 			const Frontend::FNodeHandle& NodeHandle = InNode.GetNodeHandle();
-			Frontend::FGraphHandle GraphHandle = Graph->GetMetasoundChecked().GetRootGraphHandle();
+
+			FMetasoundAssetBase* MetasoundAsset = Frontend::GetObjectAsAssetBase(Graph->GetMetasound());
+			check(MetasoundAsset);
+
+			Frontend::FGraphHandle GraphHandle = MetasoundAsset->GetRootGraphHandle();
 			if (GraphHandle.IsValid() && NodeHandle.IsValid())
 			{
 				switch (NodeHandle.GetNodeType())
@@ -119,12 +135,15 @@ namespace Metasound
 			InNode.MarkPackageDirty();
 		}
 
-		void FGraphBuilder::RebuildGraph(UMetasound& InMetasound)
+		void FGraphBuilder::RebuildGraph(UObject& InMetasound)
 		{
 			using namespace Frontend;
 
-			FGraphHandle GraphHandle = InMetasound.GetRootGraphHandle();
-			UMetasoundEditorGraph* Graph = CastChecked<UMetasoundEditorGraph>(InMetasound.GetGraph());
+			FMetasoundAssetBase* MetasoundAsset = GetObjectAsAssetBase(&InMetasound);
+			check(MetasoundAsset);
+
+			FGraphHandle GraphHandle = MetasoundAsset->GetRootGraphHandle();
+			UMetasoundEditorGraph* Graph = CastChecked<UMetasoundEditorGraph>(MetasoundAsset->GetGraph());
 
 			Graph->Nodes.Reset();
 
diff --git a/Engine/Plugins/Runtime/Metasound/Source/MetasoundEditor/Private/MetasoundEditorGraphBuilder.h b/Engine/Plugins/Runtime/Metasound/Source/MetasoundEditor/Private/MetasoundEditorGraphBuilder.h
--- a/Engine/Plugins/Runtime/Metasound/Source/MetasoundEditor/Private/MetasoundEditorGraphBuilder.h
+++ b/Engine/Plugins/Runtime/Metasound/Source/MetasoundEditor/Private/MetasoundEditorGraphBuilder.h
@@ -28,18 +28,18 @@ namespace Metasound
 			static const FName PinPrimitiveString;
 
 			// Adds a node to the editor graph that corresponds to the provided node handle.
-			static UEdGraphNode* AddNode(UMetasound& InMetasound, const FVector2D& Location, Frontend::FNodeHandle& InNodeHandle, bool bInSelectNewNode = true);
+			static UEdGraphNode* AddNode(UObject& InMetasound, const FVector2D& Location, Frontend::FNodeHandle& InNodeHandle, bool bInSelectNewNode = true);
 
 			// Adds a node with the given class info to both the editor and document graphs
-			static UEdGraphNode* AddNode(UMetasound& InMetasound, const FVector2D& Location, const Frontend::FNodeClassInfo& InClassInfo, bool bInSelectNewNode = true);
+			static UEdGraphNode* AddNode(UObject& InMetasound, const FVector2D& Location, const Frontend::FNodeClassInfo& InClassInfo, bool bInSelectNewNode = true);
 
-			static UEdGraphNode* AddInput(UMetasound& InMetasound, const FVector2D& Location, const FString& InName, const FName InTypeName, const FText& InToolTip, bool bInSelectNewNode = true);
+			static UEdGraphNode* AddInput(UObject& InMetasound, const FVector2D& Location, const FString& InName, const FName InTypeName, const FText& InToolTip, bool bInSelectNewNode = true);
 
-			static UEdGraphNode* AddOutput(UMetasound& InMetasound, const FVector2D& Location, const FString& InName, const FName InTypeName, const FText& InToolTip, bool bInSelectNewNode = true);
+			static UEdGraphNode* AddOutput(UObject& InMetasound, const FVector2D& Location, const FString& InName, const FName InTypeName, const FText& InToolTip, bool bInSelectNewNode = true);
 
 			static void DeleteNode(UMetasoundEditorGraphNode& InNode, bool bInRecordTransaction = true);
 
-			static void RebuildGraph(UMetasound& InMetasound);
+			static void RebuildGraph(UObject& InMetasound);
 
 			static void RebuildNodePins(UMetasoundEditorGraphNode& InGraphNode, Frontend::FNodeHandle& InNodeHandle, bool bInRecordTransaction = true);
 		};
diff --git a/Engine/Plugins/Runtime/Metasound/Source/MetasoundEditor/Private/MetasoundEditorGraphNode.cpp b/Engine/Plugins/Runtime/Metasound/Source/MetasoundEditor/Private/MetasoundEditorGraphNode.cpp
--- a/Engine/Plugins/Runtime/Metasound/Source/MetasoundEditor/Private/MetasoundEditorGraphNode.cpp
+++ b/Engine/Plugins/Runtime/Metasound/Source/MetasoundEditor/Private/MetasoundEditorGraphNode.cpp
@@ -8,6 +8,7 @@
 #include "Framework/Commands/GenericCommands.h"
 #include "GraphEditorActions.h"
 #include "Metasound.h"
+#include "MetasoundFrontend.h"
 #include "MetasoundEditorGraph.h"
 #include "MetasoundEditorGraphBuilder.h"
 #include "MetasoundEditorGraphSchema.h"
@@ -72,13 +73,13 @@ int32 UMetasoundEditorGraphNode::EstimateNodeWidth() const
 	}
 }
 
-UMetasound& UMetasoundEditorGraphNode::GetMetasoundChecked()
+UObject& UMetasoundEditorGraphNode::GetMetasoundChecked()
 {
 	UMetasoundEditorGraph* EdGraph = CastChecked<UMetasoundEditorGraph>(GetGraph());
 	return EdGraph->GetMetasoundChecked();
 }
 
-const UMetasound& UMetasoundEditorGraphNode::GetMetasoundChecked() const
+const UObject& UMetasoundEditorGraphNode::GetMetasoundChecked() const
 {
 	UMetasoundEditorGraph* EdGraph = CastChecked<UMetasoundEditorGraph>(GetGraph());
 	return EdGraph->GetMetasoundChecked();
@@ -86,7 +87,10 @@ const UMetasound& UMetasoundEditorGraphNode::GetMetasoundChecked() const
 
 Metasound::Frontend::FGraphHandle UMetasoundEditorGraphNode::GetRootGraphHandle() const
 {
-	return GetMetasoundChecked().GetRootGraphHandle();
+	const FMetasoundAssetBase* MetasoundAsset = Metasound::Frontend::GetObjectAsAssetBase(&GetMetasoundChecked());
+	check(MetasoundAsset);
+
+	return MetasoundAsset->GetRootGraphHandle();
 }
 
 Metasound::Frontend::FNodeHandle UMetasoundEditorGraphNode::GetNodeHandle() const
@@ -282,17 +286,21 @@ FText UMetasoundEditorGraphNode::GetTooltipText() const
 
 	// TODO: Add author to ToolTip
 	FNodeHandle NodeHandle = GetNodeHandle();
+
+	const FMetasoundAssetBase* MetasoundAsset = GetObjectAsAssetBase(&GetMetasoundChecked());
+	check(MetasoundAsset);
+
 	switch (NodeHandle.GetNodeType())
 	{
 		case EMetasoundClassType::Input:
 		{
-			return GetMetasoundChecked().GetInputToolTip(NodeHandle.GetNodeName());
+			return MetasoundAsset->GetInputToolTip(NodeHandle.GetNodeName());
 		}
 		break;
 
 		case EMetasoundClassType::Output:
 		{
-			return GetMetasoundChecked().GetOutputToolTip(NodeHandle.GetNodeName());
+			return MetasoundAsset->GetOutputToolTip(NodeHandle.GetNodeName());
 		}
 		break;
 
diff --git a/Engine/Plugins/Runtime/Metasound/Source/MetasoundEditor/Private/MetasoundEditorGraphSchema.cpp b/Engine/Plugins/Runtime/Metasound/Source/MetasoundEditor/Private/MetasoundEditorGraphSchema.cpp
--- a/Engine/Plugins/Runtime/Metasound/Source/MetasoundEditor/Private/MetasoundEditorGraphSchema.cpp
+++ b/Engine/Plugins/Runtime/Metasound/Source/MetasoundEditor/Private/MetasoundEditorGraphSchema.cpp
@@ -57,7 +57,7 @@ namespace Metasound
 			return Categories;
 		}
 
-		TSharedPtr<FEditor> GetEditorForGraph(const UMetasound& Metasound)
+		TSharedPtr<FEditor> GetEditorForGraph(const UObject& Metasound)
 		{
 			TSharedPtr<IToolkit> FoundAssetEditor = FToolkitManager::Get().FindEditorForAsset(CastChecked<const UObject>(&Metasound));
 			return StaticCastSharedPtr<FEditor, IToolkit>(FoundAssetEditor);
@@ -170,7 +170,7 @@ UEdGraphNode* FMetasoundGraphSchemaAction_NewNode::PerformAction(UEdGraph* Paren
 
 	ParentGraph->Modify();
 
-	UMetasound& ParentMetasound = CastChecked<UMetasoundEditorGraph>(ParentGraph)->GetMetasoundChecked();
+	UObject& ParentMetasound = CastChecked<UMetasoundEditorGraph>(ParentGraph)->GetMetasoundChecked();
 	ParentMetasound.Modify();
 
 	UEdGraphNode* NewGraphNode = FGraphBuilder::AddNode(ParentMetasound, Location, NodeClassInfo);
@@ -193,12 +193,15 @@ UEdGraphNode* FMetasoundGraphSchemaAction_NewInput::PerformAction(UEdGraph* Pare
 
 	ParentGraph->Modify();
 
-	UMetasound& ParentMetasound = CastChecked<UMetasoundEditorGraph>(ParentGraph)->GetMetasoundChecked();
+	UObject& ParentMetasound = CastChecked<UMetasoundEditorGraph>(ParentGraph)->GetMetasoundChecked();
 	ParentMetasound.Modify();
 
 	FString NameBase = GetDataTypeDisplayName(NodeTypeName);
 
-	FGraphHandle GraphHandle = ParentMetasound.GetRootGraphHandle();
+	FMetasoundAssetBase* MetasoundAsset = GetObjectAsAssetBase(&ParentMetasound);
+	check(MetasoundAsset);
+
+	FGraphHandle GraphHandle = MetasoundAsset->GetRootGraphHandle();
 
 	int32 i = 1;
 	FString NewNodeName = NameBase + FString::Printf(TEXT("_%02d"), i);
@@ -226,12 +229,15 @@ UEdGraphNode* FMetasoundGraphSchemaAction_NewOutput::PerformAction(UEdGraph* Par
 
 	ParentGraph->Modify();
 
-	UMetasound& ParentMetasound = CastChecked<UMetasoundEditorGraph>(ParentGraph)->GetMetasoundChecked();
+	UObject& ParentMetasound = CastChecked<UMetasoundEditorGraph>(ParentGraph)->GetMetasoundChecked();
 	ParentMetasound.Modify();
 
 	FString NameBase = GetDataTypeDisplayName(NodeTypeName);
 
-	FGraphHandle GraphHandle = ParentMetasound.GetRootGraphHandle();
+	FMetasoundAssetBase* MetasoundAsset = GetObjectAsAssetBase(&ParentMetasound);
+	check(MetasoundAsset);
+
+	FGraphHandle GraphHandle = MetasoundAsset->GetRootGraphHandle();
 
 	int32 i = 1;
 	FString NewNodeName = NameBase + FString::Printf(TEXT("_%02d"), i);
diff --git a/Engine/Plugins/Runtime/Metasound/Source/MetasoundEditor/Private/MetasoundEditorModule.cpp b/Engine/Plugins/Runtime/Metasound/Source/MetasoundEditor/Private/MetasoundEditorModule.cpp
--- a/Engine/Plugins/Runtime/Metasound/Source/MetasoundEditor/Private/MetasoundEditorModule.cpp
+++ b/Engine/Plugins/Runtime/Metasound/Source/MetasoundEditor/Private/MetasoundEditorModule.cpp
@@ -11,6 +11,7 @@
 #include "EditorStyleSet.h"
 #include "IDetailCustomization.h"
 #include "Metasound.h"
+#include "MetasoundSource.h"
 #include "MetasoundAssetTypeActions.h"
 #include "MetasoundDetailCustomization.h"
 #include "MetasoundEditorGraphBuilder.h"
@@ -150,19 +151,24 @@ namespace Metasound
 				// Register Metasound asset type actions
 				IAssetTools& AssetTools = FModuleManager::LoadModuleChecked<FAssetToolsModule>(AssetToolName).Get();
 				AddAssetAction<FAssetTypeActions_Metasound>(AssetTools, AssetActions);
+				AddAssetAction<FAssetTypeActions_MetasoundSource>(AssetTools, AssetActions);
 
 				FPropertyEditorModule& PropertyModule = FModuleManager::LoadModuleChecked<FPropertyEditorModule>("PropertyEditor");
+				
 				PropertyModule.RegisterCustomClassLayout(
 					UMetasound::StaticClass()->GetFName(),
 					FOnGetDetailCustomizationInstance::CreateLambda([]() { return MakeShared<FMetasoundDetailCustomization>(); }));
 
+				PropertyModule.RegisterCustomClassLayout(
+					UMetasoundSource::StaticClass()->GetFName(),
+					FOnGetDetailCustomizationInstance::CreateLambda([]() { return MakeShared<FMetasoundDetailCustomization>(); }));
+
 				StyleSet = MakeShared<FSlateStyle>();
 
 				RegisterCoreDataTypes();
 
 				GraphConnectionFactory = MakeShared<FGraphConnectionDrawingPolicyFactory>();
 				FEdGraphUtilities::RegisterVisualPinConnectionFactory(GraphConnectionFactory);
-
 			}
 
 			virtual void ShutdownModule() override
diff --git a/Engine/Plugins/Runtime/Metasound/Source/MetasoundEditor/Private/MetasoundFactory.cpp b/Engine/Plugins/Runtime/Metasound/Source/MetasoundEditor/Private/MetasoundFactory.cpp
--- a/Engine/Plugins/Runtime/Metasound/Source/MetasoundEditor/Private/MetasoundFactory.cpp
+++ b/Engine/Plugins/Runtime/Metasound/Source/MetasoundEditor/Private/MetasoundFactory.cpp
@@ -3,8 +3,9 @@
 
 #include "Kismet/KismetSystemLibrary.h"
 #include "Metasound.h"
+#include "MetasoundSource.h"
 #include "MetasoundEditorGraphNode.h"
-
+#include "MetasoundEditorGraphBuilder.h"
 
 UMetasoundFactory::UMetasoundFactory(const FObjectInitializer& ObjectInitializer)
 	: Super(ObjectInitializer)
@@ -18,13 +19,37 @@ UMetasoundFactory::UMetasoundFactory(const FObjectInitializer& ObjectInitializer
 
 UObject* UMetasoundFactory::FactoryCreateNew(UClass* InClass, UObject* InParent, FName Name, EObjectFlags Flags, UObject* InContext, FFeedbackContext* InFeedbackContext)
 {
-	UMetasound* Metasound = NewObject<UMetasound>(InParent, Name, Flags);
+	UMetasound* NewMetasound = NewObject<UMetasound>(InParent, Name, Flags);
+	FMetasoundClassMetadata Metadata;
+	Metadata.NodeName = NewMetasound->GetName();
+	Metadata.NodeType = EMetasoundClassType::MetasoundGraph;
+	Metadata.AuthorName = FText::FromString(UKismetSystemLibrary::GetPlatformUserName());
+	NewMetasound->SetMetadata(Metadata);
+
+	NewMetasound->ConformDocumentToArchetype();
+
+	return NewMetasound;
+}
+
+UMetasoundSourceFactory::UMetasoundSourceFactory(const FObjectInitializer& ObjectInitializer)
+	: Super(ObjectInitializer)
+{
+	SupportedClass = UMetasoundSource::StaticClass();
+
+	bCreateNew = true;
+	bEditorImport = false;
+	bEditAfterNew = true;
+}
 
+UObject* UMetasoundSourceFactory::FactoryCreateNew(UClass* InClass, UObject* InParent, FName Name, EObjectFlags Flags, UObject* InContext, FFeedbackContext* InFeedbackContext)
+{
+	UMetasoundSource* MetasoundSource = NewObject<UMetasoundSource>(InParent, Name, Flags);
 	FMetasoundClassMetadata Metadata;
-	Metadata.NodeName = Metasound->GetName();
+	Metadata.NodeName = MetasoundSource->GetName();
 	Metadata.NodeType = EMetasoundClassType::MetasoundGraph;
 	Metadata.AuthorName = FText::FromString(UKismetSystemLibrary::GetPlatformUserName());
-	Metasound->SetMetadata(Metadata);
+	MetasoundSource->SetMetadata(Metadata);
+	MetasoundSource->ConformDocumentToArchetype();
 
-	return Metasound;
+	return MetasoundSource;
 }
diff --git a/Engine/Plugins/Runtime/Metasound/Source/MetasoundEditor/Public/IMetasoundEditor.h b/Engine/Plugins/Runtime/Metasound/Source/MetasoundEditor/Public/IMetasoundEditor.h
--- a/Engine/Plugins/Runtime/Metasound/Source/MetasoundEditor/Public/IMetasoundEditor.h
+++ b/Engine/Plugins/Runtime/Metasound/Source/MetasoundEditor/Public/IMetasoundEditor.h
@@ -12,7 +12,7 @@
 
 class IMetasoundEditor : public FAssetEditorToolkit
 {
-	virtual UMetasound* GetMetasound() const = 0;
+	virtual UObject* GetMetasoundObject() const = 0;
 	virtual void SetSelection(const TArray<UObject*>& SelectedObjects) = 0;
 	virtual bool GetBoundsForSelectedNodes(FSlateRect& Rect, float Padding) = 0;
 };
diff --git a/Engine/Plugins/Runtime/Metasound/Source/MetasoundEditor/Public/MetasoundEditorGraph.h b/Engine/Plugins/Runtime/Metasound/Source/MetasoundEditor/Public/MetasoundEditorGraph.h
--- a/Engine/Plugins/Runtime/Metasound/Source/MetasoundEditor/Public/MetasoundEditorGraph.h
+++ b/Engine/Plugins/Runtime/Metasound/Source/MetasoundEditor/Public/MetasoundEditorGraph.h
@@ -19,12 +19,12 @@ class UMetasoundEditorGraph : public UEdGraph
 	GENERATED_BODY()
 
 public:
-	UMetasound* GetMetasound() const;
-	UMetasound& GetMetasoundChecked() const;
+	UObject* GetMetasound() const;
+	UObject& GetMetasoundChecked() const;
 
 private:
 	UPROPERTY()
-	UMetasound* ParentMetasound;
+	UObject* ParentMetasound;
 
 	friend class Metasound::Editor::FEditor;
 };
diff --git a/Engine/Plugins/Runtime/Metasound/Source/MetasoundEditor/Public/MetasoundEditorGraphNode.h b/Engine/Plugins/Runtime/Metasound/Source/MetasoundEditor/Public/MetasoundEditorGraphNode.h
--- a/Engine/Plugins/Runtime/Metasound/Source/MetasoundEditor/Public/MetasoundEditorGraphNode.h
+++ b/Engine/Plugins/Runtime/Metasound/Source/MetasoundEditor/Public/MetasoundEditorGraphNode.h
@@ -54,8 +54,8 @@ public:
 	virtual void PostDuplicate(bool bDuplicateForPIE) override;
 	// End of UObject interface
 
-	UMetasound& GetMetasoundChecked();
-	const UMetasound& GetMetasoundChecked() const;
+	UObject& GetMetasoundChecked();
+	const UObject& GetMetasoundChecked() const;
 
 	Metasound::Frontend::FGraphHandle GetRootGraphHandle() const;
 	Metasound::Frontend::FNodeHandle GetNodeHandle() const;
diff --git a/Engine/Plugins/Runtime/Metasound/Source/MetasoundEditor/Public/MetasoundFactory.h b/Engine/Plugins/Runtime/Metasound/Source/MetasoundEditor/Public/MetasoundFactory.h
--- a/Engine/Plugins/Runtime/Metasound/Source/MetasoundEditor/Public/MetasoundFactory.h
+++ b/Engine/Plugins/Runtime/Metasound/Source/MetasoundEditor/Public/MetasoundFactory.h
@@ -17,3 +17,13 @@ class UMetasoundFactory : public UFactory
 	virtual UObject* FactoryCreateNew(UClass* InClass, UObject* InParent, FName Name, EObjectFlags Flags, UObject* InContext, FFeedbackContext* InFeedbackContext) override;
 	//~ Begin UFactory Interface
 };
+
+UCLASS(hidecategories = Object, MinimalAPI)
+class UMetasoundSourceFactory : public UFactory
+{
+	GENERATED_UCLASS_BODY()
+
+	//~ Begin UFactory Interface
+	virtual UObject* FactoryCreateNew(UClass* InClass, UObject* InParent, FName Name, EObjectFlags Flags, UObject* InContext, FFeedbackContext* InFeedbackContext) override;
+	//~ Begin UFactory Interface
+};
diff --git a/Engine/Plugins/Runtime/Metasound/Source/MetasoundEngine/MetasoundEngine.Build.cs b/Engine/Plugins/Runtime/Metasound/Source/MetasoundEngine/MetasoundEngine.Build.cs
--- a/Engine/Plugins/Runtime/Metasound/Source/MetasoundEngine/MetasoundEngine.Build.cs
+++ b/Engine/Plugins/Runtime/Metasound/Source/MetasoundEngine/MetasoundEngine.Build.cs
@@ -14,7 +14,8 @@ namespace UnrealBuildTool.Rules
 				new string[]
 				{
 					"Core",
-					"MetasoundGraphCore"
+					"MetasoundGraphCore",
+					"MetasoundGenerator"
 				}
 			);
 
diff --git a/Engine/Plugins/Runtime/Metasound/Source/MetasoundEngine/Private/Metasound.cpp b/Engine/Plugins/Runtime/Metasound/Source/MetasoundEngine/Private/Metasound.cpp
--- a/Engine/Plugins/Runtime/Metasound/Source/MetasoundEngine/Private/Metasound.cpp
+++ b/Engine/Plugins/Runtime/Metasound/Source/MetasoundEngine/Private/Metasound.cpp
@@ -6,6 +6,7 @@
 #include "StructSerializer.h"
 #include "UObject/UnrealType.h"
 
+#include "MetasoundArchetypeRegistration.h"
 
 UMetasound::UMetasound(const FObjectInitializer& ObjectInitializer)
 	: Super(ObjectInitializer)
@@ -36,32 +37,6 @@ const UEdGraph& UMetasound::GetGraphChecked() const
 	return *Graph;
 }
 
-const FText& UMetasound::GetInputToolTip(FString InputName) const
-{
-	for (const FMetasoundInputDescription& Desc : RootMetasoundDocument.RootClass.Inputs)
-	{
-		if (Desc.Name == InputName)
-		{
-			return Desc.ToolTip;
-		}
-	}
-
-	return FText::GetEmpty();
-}
-
-const FText& UMetasound::GetOutputToolTip(FString OutputName) const
-{
-	for (const FMetasoundOutputDescription& Desc : RootMetasoundDocument.RootClass.Outputs)
-	{
-		if (Desc.Name == OutputName)
-		{
-			return Desc.ToolTip;
-		}
-	}
-
-	return FText::GetEmpty();
-}
-
 void UMetasound::SetGraph(UEdGraph* InGraph)
 {
 	Graph = InGraph;
@@ -76,9 +51,17 @@ void UMetasound::SetMetadata(FMetasoundClassMetadata& InMetadata)
 	MarkPackageDirty();
 }
 
-// delete this asset's current metasound document,
-// and replace it with InClassDescription.
-void UMetasound::SetMetasoundDocument(const FMetasoundDocument& InDocument)
+void UMetasound::PostLoad()
+{
+	ConformDocumentToArchetype();
+}
+
+FMetasoundArchetype UMetasound::GetArchetype() const
 {
-	RootMetasoundDocument = InDocument;
+	FMetasoundArchetype Archetype;
+	static FName MetasoundArchetypeName = TEXT("Generic Metasound");
+
+	Archetype.ArchetypeName = MetasoundArchetypeName;
+
+	return Archetype;
 }
diff --git a/Engine/Plugins/Runtime/Metasound/Source/MetasoundEngine/Private/MetasoundEngineModule.cpp b/Engine/Plugins/Runtime/Metasound/Source/MetasoundEngine/Private/MetasoundEngineModule.cpp
--- a/Engine/Plugins/Runtime/Metasound/Source/MetasoundEngine/Private/MetasoundEngineModule.cpp
+++ b/Engine/Plugins/Runtime/Metasound/Source/MetasoundEngine/Private/MetasoundEngineModule.cpp
@@ -2,6 +2,7 @@
 #include "MetasoundEngineModule.h"
 
 #include "Metasound.h"
+#include "MetasoundSource.h"
 #include "Modules/ModuleManager.h"
 
 
@@ -11,6 +12,9 @@ class FMetasoundEngineModule : public IMetasoundEngineModule
 {
 	virtual void StartupModule() override
 	{
+		Metasound::Frontend::RegisterArchetype<UMetasoundSource>();
+		Metasound::Frontend::RegisterArchetype<UMetasound>();
+
 		UE_LOG(LogMetasoundEngine, Log, TEXT("Metasound Engine Initialized"));
 	}
 };
diff --git a/Engine/Plugins/Runtime/Metasound/Source/MetasoundEngine/Private/MetasoundSource.cpp b/Engine/Plugins/Runtime/Metasound/Source/MetasoundEngine/Private/MetasoundSource.cpp
new file mode 100644
--- /dev/null
+++ b/Engine/Plugins/Runtime/Metasound/Source/MetasoundEngine/Private/MetasoundSource.cpp
@@ -0,0 +1,147 @@
+// Copyright Epic Games, Inc. All Rights Reserved.
+
+#include "MetasoundSource.h"
+
+#include "CoreMinimal.h"
+#include "Internationalization/Text.h"
+#include "AssetRegistryModule.h"
+
+#include "MetasoundAssetBase.h"
+#include "MetasoundGenerator.h"
+#include "MetasoundAudioFormats.h"
+#include "MetasoundPrimitives.h"
+#include "MetasoundDataReferenceTypes.h"
+
+#include "MetasoundArchetypeRegistration.h"
+
+#if WITH_EDITORONLY_DATA
+#include "EdGraph/EdGraph.h"
+#endif // WITH_EDITORONLY_DATA
+
+#define LOCTEXT_NAMESPACE "MetasoundSource"
+
+static FName MetasoundSourceArchetypeName = FName(TEXT("Metasound Source"));
+
+UMetasoundSource::UMetasoundSource(const FObjectInitializer& ObjectInitializer)
+	: Super(ObjectInitializer)
+	, FMetasoundAssetBase(RootMetasoundDocument)
+{
+	// TODO: 
+	NumChannels = 1;
+	Duration = INDEFINITELY_LOOPING_DURATION;
+	bLooping = true;
+	
+	// todo: ensure that we have a method so that the audio engine can be authoritative over the sample rate the UMetasoundSource runs at.
+	SampleRate = 48000.0f;
+
+}
+
+void UMetasoundSource::SetMetadata(FMetasoundClassMetadata& InMetadata)
+{
+	RootMetasoundDocument.RootClass.Metadata = InMetadata;
+}
+
+bool UMetasoundSource::IsPlayable() const
+{
+	// todo: cache off whether this metasound is buildable to an operator.
+	return true;
+}
+
+bool UMetasoundSource::SupportsSubtitles() const
+{
+	return Super::SupportsSubtitles();
+}
+
+float UMetasoundSource::GetDuration()
+{
+	// eh? this is kind of a weird field anyways.
+	return Super::GetDuration();
+}
+
+ISoundGeneratorPtr UMetasoundSource::CreateSoundGenerator(const FSoundGeneratorInitParams& InParams)
+{
+	NumChannels = 1;
+	Duration = INDEFINITELY_LOOPING_DURATION;
+	bLooping = true;
+	SampleRate = InParams.SampleRate;
+
+	Metasound::FOperatorSettings InSettings(InParams.SampleRate, InParams.NumFramesPerCallback);
+
+	TArray<Metasound::IOperatorBuilder::FBuildErrorPtr> BuildErrors;
+
+	Metasound::Frontend::FGraphHandle RootGraph = GetRootGraphHandle();
+	ensureAlways(RootGraph.IsValid());
+
+	TUniquePtr<Metasound::IOperator> Operator = RootGraph.BuildOperator(InSettings, BuildErrors);
+	if (ensureAlways(Operator.IsValid()))
+	{
+		Metasound::FDataReferenceCollection Outputs = Operator->GetOutputs();
+
+		Metasound::FMetasoundGeneratorInitParams InitParams =
+		{
+			MoveTemp(Operator),
+			Outputs.GetDataReadReferenceOrConstruct<Metasound::FAudioBuffer>(GetAudioOutputName(), InParams.NumFramesPerCallback),
+			Outputs.GetDataReadReferenceOrConstruct<Metasound::FBop>(GetIsFinishedOutputName(), false)
+		};
+
+		return ISoundGeneratorPtr(new Metasound::FMetasoundGenerator(MoveTemp(InitParams)));
+	}
+	else
+	{
+		return ISoundGeneratorPtr(nullptr);
+	}
+}
+
+void UMetasoundSource::PostLoad()
+{
+	Super::PostLoad();
+
+	ConformDocumentToArchetype();
+}
+
+const FString& UMetasoundSource::GetOnPlayInputName()
+{
+	static FString BopInputName = FString(TEXT("On Play"));
+	return BopInputName;
+}
+
+const FString& UMetasoundSource::GetAudioOutputName()
+{
+	static FString AudioOutputName = FString(TEXT("Generated Audio"));
+	return AudioOutputName;
+}
+
+const FString& UMetasoundSource::GetIsFinishedOutputName()
+{
+	static FString OnFinishedOutputName = FString(TEXT("On Finished"));
+	return OnFinishedOutputName;
+}
+
+FMetasoundArchetype UMetasoundSource::GetArchetype() const
+{
+	FMetasoundArchetype Archetype;
+	Archetype.ArchetypeName = MetasoundSourceArchetypeName;
+
+	FMetasoundInputDescription OnPlayBop;
+	OnPlayBop.Name = GetOnPlayInputName();
+	OnPlayBop.TypeName = Metasound::Frontend::GetDataTypeName<Metasound::FBop>();
+	OnPlayBop.ToolTip = LOCTEXT("OnPlayBopToolTip", "This input is bopped when this source is first played.");
+
+	Archetype.RequiredInputs.Add(OnPlayBop);
+
+	FMetasoundOutputDescription GeneratedAudio;
+	GeneratedAudio.Name = GetAudioOutputName();
+	GeneratedAudio.TypeName = Metasound::Frontend::GetDataTypeName<Metasound::FAudioBuffer>();
+	GeneratedAudio.ToolTip = LOCTEXT("GeneratedAudioToolTip", "This should be the resulting output audio from this source.");
+	Archetype.RequiredOutputs.Add(GeneratedAudio);
+
+	FMetasoundOutputDescription OnFinished;
+	OnFinished.Name = GetIsFinishedOutputName();
+	OnFinished.TypeName = Metasound::Frontend::GetDataTypeName<Metasound::FBop>();
+	OnFinished.ToolTip = LOCTEXT("OnFinishedToolTip", "This output can be bopped to stop the sound.");
+	Archetype.RequiredOutputs.Add(OnFinished);
+
+	return Archetype;
+}
+
+#undef LOCTEXT_NAMESPACE
diff --git a/Engine/Plugins/Runtime/Metasound/Source/MetasoundEngine/Public/Metasound.h b/Engine/Plugins/Runtime/Metasound/Source/MetasoundEngine/Public/Metasound.h
--- a/Engine/Plugins/Runtime/Metasound/Source/MetasoundEngine/Public/Metasound.h
+++ b/Engine/Plugins/Runtime/Metasound/Source/MetasoundEngine/Public/Metasound.h
@@ -44,18 +44,15 @@ public:
 	// Returns the graph associated with this Metasound. Graph is required to be referenced on
 	// Metasound UObject for editor serialization purposes.
 	// @return Editor graph associated with UMetasound.
-	UEdGraph* GetGraph();
-	const UEdGraph* GetGraph() const;
-	UEdGraph& GetGraphChecked();
-	const UEdGraph& GetGraphChecked() const;
+	virtual UEdGraph* GetGraph() override;
+	virtual const UEdGraph* GetGraph() const override;
+	virtual UEdGraph& GetGraphChecked() override;
+	virtual const UEdGraph& GetGraphChecked() const override;
 
 	// Sets the graph associated with this Metasound. Graph is required to be referenced on
 	// Metasound UObject for editor serialization purposes.
 	// @param Editor graph associated with UMetasound.
-	void SetGraph(UEdGraph* InGraph);
-
-	const FText& GetInputToolTip(FString InputName) const;
-	const FText& GetOutputToolTip(FString OutputName) const;
+	virtual void SetGraph(UEdGraph* InGraph) override;
 
 #endif // WITH_EDITORONLY_DATA
 
@@ -80,6 +77,9 @@ public:
 	// @param InMetadata Metadata containing corrections to the class metadata.
 	void SetMetadata(FMetasoundClassMetadata& InMetadata) override;
 
-	// Deletes Metasound's current metasound document, and replaces it with InClassDescription.
-	void SetMetasoundDocument(const FMetasoundDocument& InDocument);
+	// Get the most up to date archetype for metasound.
+	virtual FMetasoundArchetype GetArchetype() const override;
+
+	void PostLoad() override;
+
 };
\ No newline at end of file
diff --git a/Engine/Plugins/Runtime/Metasound/Source/MetasoundEngine/Public/MetasoundArchetypeRegistration.h b/Engine/Plugins/Runtime/Metasound/Source/MetasoundEngine/Public/MetasoundArchetypeRegistration.h
new file mode 100644
--- /dev/null
+++ b/Engine/Plugins/Runtime/Metasound/Source/MetasoundEngine/Public/MetasoundArchetypeRegistration.h
@@ -0,0 +1,83 @@
+// Copyright Epic Games, Inc. All Rights Reserved.
+#pragma once
+
+#include "CoreMinimal.h"
+#include "MetasoundFrontend.h"
+#include "MetasoundAssetBase.h"
+
+#if WITH_EDITOR
+#include "AssetRegistryModule.h"
+#endif
+
+namespace Metasound
+{
+	namespace Frontend
+	{
+		// Use this function to register a specific metasound archetype with the frontend.
+		// Metasound Archetypes are lists of required inputs and outputs.
+		// For example, a Metasound Source has a required audio output,
+		// and a Metasound Source Effect has a required audio input and audio output.
+		template <typename UClassType>
+		bool RegisterArchetype()
+		{
+			static_assert(std::is_base_of<FMetasoundAssetBase, UClassType>::value, "Metasound Archetypes have to inherit from FMetasoundAssetBase (see: UMetasound, UMetasoundSource)");
+
+			// Autogenerate our getter and cast lambdas here:
+			FMetasoundArchetypeRegistryParams_Internal InternalParams =
+			{
+				GetDefault<UClassType>()->GetArchetype(),
+				UClassType::StaticClass(),
+				// SafeCast:
+				[](UObject* InObject) -> FMetasoundAssetBase*
+				{
+					check(InObject);
+					return static_cast<FMetasoundAssetBase*>(CastChecked<UClassType>(InObject));
+				},
+				// SafeConstCast:
+				[](const UObject* InObject) -> const FMetasoundAssetBase*
+				{
+					check(InObject);
+					return static_cast<const FMetasoundAssetBase*>(CastChecked<const UClassType>(InObject));
+				},
+				//Object Getter:
+				[](const FMetasoundDocument& InDocument, const FString& InPath) -> UObject*
+				{
+					UPackage* PackageToSaveTo = nullptr;
+					if (GIsEditor)
+					{
+						FText InvalidPathReason;
+						bool const bValidPackageName = FPackageName::IsValidLongPackageName(InPath, false, &InvalidPathReason);
+
+						if (!ensureAlwaysMsgf(bValidPackageName, TEXT("Tried to generate a Metasound UObject with an invalid package path/name Falling back to transient package, which means we won't be able to save this asset.")))
+						{
+							PackageToSaveTo = GetTransientPackage();
+						}
+						else
+						{
+							PackageToSaveTo = CreatePackage(nullptr, *InPath);
+						}
+					}
+					else
+					{
+						PackageToSaveTo = GetTransientPackage();
+					}
+
+					UClassType* NewMetasoundObject = NewObject<UClassType>(PackageToSaveTo, *InDocument.RootClass.Metadata.NodeName);
+					NewMetasoundObject->SetDocument(InDocument);
+
+#if WITH_EDITOR
+					AsyncTask(ENamedThreads::GameThread, [NewMetasoundObject]() {
+						FAssetRegistryModule::AssetCreated(NewMetasoundObject);
+						NewMetasoundObject->MarkPackageDirty();
+						// todo: how do you get the package for a uobject and save it? I forget
+					});
+#endif
+
+					return NewMetasoundObject;
+				}
+			};
+
+			return RegisterArchetype_Internal(MoveTemp(InternalParams));
+		}
+	}
+}
\ No newline at end of file
diff --git a/Engine/Plugins/Runtime/Metasound/Source/MetasoundEngine/Public/MetasoundSource.h b/Engine/Plugins/Runtime/Metasound/Source/MetasoundEngine/Public/MetasoundSource.h
new file mode 100644
--- /dev/null
+++ b/Engine/Plugins/Runtime/Metasound/Source/MetasoundEngine/Public/MetasoundSource.h
@@ -0,0 +1,111 @@
+// Copyright Epic Games, Inc. All Rights Reserved.
+#pragma once
+
+#include "CoreMinimal.h"
+#include "MetasoundFrontend.h"
+#include "MetasoundFrontendDataLayout.h"
+#include "MetasoundAssetBase.h"
+
+#include "Sound/SoundWaveProcedural.h"
+
+
+#if WITH_EDITORONLY_DATA
+#include "EdGraph/EdGraph.h"
+#endif // WITH_EDITORONLY_DATA
+
+#include "MetasoundSource.generated.h"
+
+
+/**
+ * This asset type is used for Metasound assets that can only be used as nodes in other Metasound graphs.
+ * Because of this, they can have any inputs or outputs they need.
+ */
+UCLASS(hidecategories = object, BlueprintType)
+class METASOUNDENGINE_API UMetasoundSource : public USoundWaveProcedural, public FMetasoundAssetBase
+{
+	GENERATED_BODY()
+
+protected:
+	UPROPERTY(EditAnywhere, Category = Hidden)
+	FMetasoundDocument RootMetasoundDocument;
+
+#if WITH_EDITORONLY_DATA
+	UPROPERTY()
+	UEdGraph* Graph;
+#endif // WITH_EDITORONLY_DATA
+
+public:
+	UMetasoundSource(const FObjectInitializer& ObjectInitializer);
+
+#if WITH_EDITORONLY_DATA
+
+	// Returns the graph associated with this Metasound. Graph is required to be referenced on
+	// Metasound UObject for editor serialization purposes.
+	// @return Editor graph associated with UMetasoundSource.
+	virtual UEdGraph* GetGraph() override
+	{
+		return Graph;
+	}
+
+	virtual const UEdGraph* GetGraph() const override
+	{
+		return Graph;
+	}
+
+	virtual UEdGraph& GetGraphChecked() override
+	{
+		check(Graph);
+		return *Graph;
+	}
+
+	virtual const UEdGraph& GetGraphChecked() const override
+	{
+		check(Graph);
+		return *Graph;
+	}
+
+	// Sets the graph associated with this Metasound. Graph is required to be referenced on
+	// Metasound UObject for editor serialization purposes.
+	// @param Editor graph associated with UMetasoundSource.
+	virtual void SetGraph(UEdGraph* InGraph) override
+	{
+		Graph = InGraph;
+	}
+#endif // WITH_EDITORONLY_DATA
+
+	FMetasoundDocument& GetDocument() override
+	{
+		return RootMetasoundDocument;
+	}
+
+	const FMetasoundDocument& GetDocument() const override
+	{
+		return RootMetasoundDocument;
+	}
+
+	UObject* GetOwningAsset() const override
+	{
+		// Hack to allow handles to manipulate while providing
+		// ability to inspect via a handle from the const version of GetRootGraphHandle()
+		return const_cast<UObject*>(CastChecked<const UObject>(this));
+	}
+
+	// Updates the Metasound's metadata (name, author, etc).
+	// @param InMetadata Metadata containing corrections to the class metadata.
+	void SetMetadata(FMetasoundClassMetadata& InMetadata) override;
+
+	bool IsPlayable() const override;
+	bool SupportsSubtitles() const override;
+	float GetDuration() override;
+	ISoundGeneratorPtr CreateSoundGenerator(const FSoundGeneratorInitParams& InParams) override;
+	void PostLoad() override;
+
+	// Get the most up to date archetype for metasound sources.
+	virtual FMetasoundArchetype GetArchetype() const override;
+
+private:
+
+	static const FString& GetOnPlayInputName();
+	static const FString& GetAudioOutputName();
+	static const FString& GetIsFinishedOutputName();
+};
\ No newline at end of file
diff --git a/Engine/Plugins/Runtime/Metasound/Source/MetasoundFrontend/Private/MetasoundAssetBase.cpp b/Engine/Plugins/Runtime/Metasound/Source/MetasoundFrontend/Private/MetasoundAssetBase.cpp
--- a/Engine/Plugins/Runtime/Metasound/Source/MetasoundFrontend/Private/MetasoundAssetBase.cpp
+++ b/Engine/Plugins/Runtime/Metasound/Source/MetasoundFrontend/Private/MetasoundAssetBase.cpp
@@ -2,7 +2,7 @@
 
 #include "MetasoundAssetBase.h"
 
-#include "Backends/JsonStructSerializerBackend.h"
+#include "MetasoundJsonBackend.h"
 #include "HAL/FileManager.h"
 #include "IStructSerializerBackend.h"
 #include "StructSerializer.h"
@@ -25,11 +25,57 @@ void FMetasoundAssetBase::SetMetadata(FMetasoundClassMetadata& InMetadata)
 	GetDocument().RootClass.Metadata = InMetadata;
 }
 
+void FMetasoundAssetBase::SetDocument(const FMetasoundDocument& InDocument)
+{
+	FMetasoundDocument& Document = GetDocument();
+	Document = InDocument;
+
+	Document.Archetype = GetArchetype();
+	GetRootGraphHandle().FixDocumentToMatchArchetype();
+}
+
+void FMetasoundAssetBase::ConformDocumentToArchetype()
+{
+	FMetasoundDocument& Document = GetDocument();
+	Document.Archetype = GetArchetype();
+	GetRootGraphHandle().FixDocumentToMatchArchetype();
+}
+
 FMetasoundClassMetadata FMetasoundAssetBase::GetMetadata()
 {
 	return GetDocument().RootClass.Metadata;
 }
 
+const FText& FMetasoundAssetBase::GetInputToolTip(FString InputName) const
+{
+	const FMetasoundDocument& RootMetasoundDocument = GetDocument();
+
+	for (const FMetasoundInputDescription& Desc : RootMetasoundDocument.RootClass.Inputs)
+	{
+		if (Desc.Name == InputName)
+		{
+			return Desc.ToolTip;
+		}
+	}
+
+	return FText::GetEmpty();
+}
+
+const FText& FMetasoundAssetBase::GetOutputToolTip(FString OutputName) const
+{
+	const FMetasoundDocument& RootMetasoundDocument = GetDocument();
+
+	for (const FMetasoundOutputDescription& Desc : RootMetasoundDocument.RootClass.Outputs)
+	{
+		if (Desc.Name == OutputName)
+		{
+			return Desc.ToolTip;
+		}
+	}
+
+	return FText::GetEmpty();
+}
+
 TWeakPtr<Metasound::Frontend::FDescriptionAccessPoint> FMetasoundAssetBase::GetGraphAccessPoint() const
 {
 	return TWeakPtr<Metasound::Frontend::FDescriptionAccessPoint>(AccessPoint);
@@ -84,7 +130,7 @@ bool FMetasoundAssetBase::ExportToJSON(const FString& InAbsolutePath) const
 {
 	if (TUniquePtr<FArchive> FileWriter = TUniquePtr<FArchive>(IFileManager::Get().CreateFileWriter(*InAbsolutePath)))
 	{
-		FJsonStructSerializerBackend Backend(*FileWriter, EStructSerializerBackendFlags::Default);
+		Metasound::TJsonStructSerializerBackend<Metasound::DefaultCharType> Backend(*FileWriter, EStructSerializerBackendFlags::Default);
 		FStructSerializer::Serialize<FMetasoundDocument>(GetDocument(), Backend);
 		FileWriter->Close();
 
diff --git a/Engine/Plugins/Runtime/Metasound/Source/MetasoundFrontend/Private/MetasoundFrontend.cpp b/Engine/Plugins/Runtime/Metasound/Source/MetasoundFrontend/Private/MetasoundFrontend.cpp
--- a/Engine/Plugins/Runtime/Metasound/Source/MetasoundFrontend/Private/MetasoundFrontend.cpp
+++ b/Engine/Plugins/Runtime/Metasound/Source/MetasoundFrontend/Private/MetasoundFrontend.cpp
@@ -2,8 +2,7 @@
 
 #include "MetasoundFrontend.h"
 
-#include "Backends/JsonStructSerializerBackend.h"
-#include "Backends/JsonStructDeserializerBackend.h"
+#include "MetasoundJsonBackend.h"
 #include "StructSerializer.h"
 #include "StructDeserializer.h"
 #include "HAL/FileManager.h"
@@ -13,7 +12,6 @@
 #include "MetasoundAudioFormats.h"
 #include "MetasoundOperatorBuilder.h"
 
-
 static int32 MetasoundUndoRollLimitCvar = 128;
 FAutoConsoleVariableRef CVarMetasoundUndoRollLimit(
 	TEXT("au.Metasound.Frontend.UndoRollLimit"),
@@ -220,7 +218,7 @@ namespace Metasound
 		{
 			if (TUniquePtr<FArchive> FileReader = TUniquePtr<FArchive>(IFileManager::Get().CreateFileReader(*InPath)))
 			{
-				FJsonStructDeserializerBackend Backend(*FileReader);
+				TJsonStructDeserializerBackend<DefaultCharType> Backend(*FileReader);
 				bool DeserializeResult = FStructDeserializer::Deserialize(OutMetasoundDocument, Backend);
 
 				FileReader->Close();
@@ -230,6 +228,171 @@ namespace Metasound
 			return false;
 		}
 
+		// Struct used for any pertinent archetype data.
+		struct FArchetypeRegistryElement
+		{
+			FMetasoundArchetype Archetype;
+			UClass* ArchetypeUClass;
+
+			// Constructor used to generate an instance of the UObject version of this archetype from scratch.
+			TUniqueFunction<UObject* (const FMetasoundDocument&, const FString&)> ObjectConstructor;
+
+			// template-generated lambdas used to safely sidecast to FMetasoundBase*.
+			TUniqueFunction<FMetasoundAssetBase* (UObject*)> SafeCast;
+			TUniqueFunction<const FMetasoundAssetBase* (const UObject*)> SafeConstCast;
+		};
+
+		static TMap<FName, FArchetypeRegistryElement> ArchetypeRegistry;
+
+		bool RegisterArchetype_Internal(FMetasoundArchetypeRegistryParams_Internal&& InParams)
+		{
+			FName ArchetypeName = InParams.ArchetypeDescription.ArchetypeName;
+
+			if (!ArchetypeRegistry.Contains(ArchetypeName))
+			{
+				FArchetypeRegistryElement RegistryElement = 
+				{ 
+					InParams.ArchetypeDescription,
+					InParams.ArchetypeUClass,
+					MoveTemp(InParams.ObjectGetter),
+					MoveTemp(InParams.SafeCast),
+					MoveTemp(InParams.SafeConstCast)
+				};
+
+				ArchetypeRegistry.Add(ArchetypeName, MoveTemp(RegistryElement));
+
+				return true;
+			}
+			else
+			{
+				return false;
+			}
+		}
+
+		TArray<FName> GetAllRegisteredArchetypes()
+		{
+			TArray<FName> OutArchetypes;
+			for (auto& RegistryTuple : ArchetypeRegistry)
+			{
+				OutArchetypes.Add(RegistryTuple.Key);
+			}
+
+			return OutArchetypes;
+		}
+
+		UObject* GetObjectForDocument(const FMetasoundDocument& InDocument, const FString& InPath)
+		{
+			FName ArchetypeName = InDocument.Archetype.ArchetypeName;
+			if (ArchetypeRegistry.Contains(ArchetypeName))
+			{
+				return ArchetypeRegistry[ArchetypeName].ObjectConstructor(InDocument, InPath);
+			}
+			else
+			{
+				return nullptr;
+			}
+		}
+
+		bool IsObjectAMetasoundArchetype(const UObject* InObject)
+		{
+			UClass* ObjectClass = InObject->GetClass();
+
+			for (auto& RegistryTuple : ArchetypeRegistry)
+			{
+				if (ObjectClass == RegistryTuple.Value.ArchetypeUClass)
+				{
+					return true;
+				}
+			}
+
+			return false;
+		}
+
+		FMetasoundAssetBase* GetObjectAsAssetBase(UObject* InObject)
+		{
+			UClass* ObjectClass = InObject->GetClass();
+
+			for (auto& RegistryTuple : ArchetypeRegistry)
+			{
+				if (ObjectClass == RegistryTuple.Value.ArchetypeUClass)
+				{
+					return RegistryTuple.Value.SafeCast(InObject);
+				}
+			}
+
+			return nullptr;
+		}
+
+		const FMetasoundAssetBase* GetObjectAsAssetBase(const UObject* InObject)
+		{
+			UClass* ObjectClass = InObject->GetClass();
+
+			for (auto& RegistryTuple : ArchetypeRegistry)
+			{
+				if (ObjectClass == RegistryTuple.Value.ArchetypeUClass)
+				{
+					return RegistryTuple.Value.SafeConstCast(InObject);
+				}
+			}
+
+			return nullptr;
+		}
+
+		void FGraphHandle::FixDocumentToMatchArchetype()
+		{
+			// TODO: Also check if this is the root class.
+			if (!IsValid())
+			{
+				return;
+			}
+
+			FMetasoundDocument& DocumentToFixUp = *OwningDocument;
+
+			// Add any missing inputs from the Required Inputs list:
+			const TArray<FMetasoundInputDescription>& RequiredInputs = DocumentToFixUp.Archetype.RequiredInputs;
+			TArray<FMetasoundInputDescription>& CurrentInputs = DocumentToFixUp.RootClass.Inputs;
+			for (const FMetasoundInputDescription& RequiredInput : RequiredInputs)
+			{
+				bool bFoundRequiredInput = false;
+				for (const FMetasoundInputDescription& Input : CurrentInputs)
+				{
+					if (Input.TypeName == RequiredInput.TypeName && Input.Name == RequiredInput.Name)
+					{
+						bFoundRequiredInput = true;
+						break;
+					}
+				}
+
+				if (!bFoundRequiredInput)
+				{
+					//CurrentInputs.Add(RequiredInput);
+					AddNewInput(RequiredInput);
+				}
+			}
+
+			// Add any missing outputs from the Required Outputs list:
+			const TArray<FMetasoundOutputDescription>& RequiredOutputs = DocumentToFixUp.Archetype.RequiredOutputs;
+			TArray<FMetasoundOutputDescription>& CurrentOutputs = DocumentToFixUp.RootClass.Outputs;
+			for (const FMetasoundOutputDescription& RequiredOutput : RequiredOutputs)
+			{
+				bool bFoundRequiredOutput = false;
+				for (const FMetasoundOutputDescription& Output : CurrentOutputs)
+				{
+					if (Output.TypeName == RequiredOutput.TypeName && Output.Name == RequiredOutput.Name)
+					{
+						bFoundRequiredOutput = true;
+						break;
+					}
+				}
+
+				if (!bFoundRequiredOutput)
+				{
+					//CurrentOutputs.Add(RequiredOutput);
+					AddNewOutput(RequiredOutput);
+				}
+			}
+		}
+
 		FInputHandle::FInputHandle(FHandleInitParams::EPrivateToken PrivateToken, const FHandleInitParams& InParams, const FString& InputName)
 			: ITransactable(MetasoundUndoRollLimitCvar, InParams.InOwningAsset)
 			, NodePtr(InParams.InAccessPoint, InParams.InPath)
@@ -848,10 +1011,28 @@ namespace Metasound
 
 			if (IsValid())
 			{
-				ClassInfo.NodeName = NodeClass->Metadata.NodeName;
-				ClassInfo.NodeType = NodeClass->Metadata.NodeType;
-				ClassInfo.LookupKey.NodeHash = NodeClass->ExternalNodeClassLookupInfo.ExternalNodeClassHash;
-				ClassInfo.LookupKey.NodeName = NodeClass->ExternalNodeClassLookupInfo.ExternalNodeClassName;
+
+				if (NodeClassType == EMetasoundClassType::Input)
+				{
+					ClassInfo.NodeName = TEXT("Input");
+					ClassInfo.NodeType = NodeClassType;
+					ClassInfo.LookupKey.NodeHash = 0;
+					ClassInfo.LookupKey.NodeName = FName();
+				}
+				else if (NodeClassType == EMetasoundClassType::Output)
+				{
+					ClassInfo.NodeName = TEXT("Output");
+					ClassInfo.NodeType = NodeClassType;
+					ClassInfo.LookupKey.NodeHash = 0;
+					ClassInfo.LookupKey.NodeName = FName();
+				}
+				else
+				{
+					ClassInfo.NodeName = NodeClass->Metadata.NodeName;
+					ClassInfo.NodeType = NodeClass->Metadata.NodeType;
+					ClassInfo.LookupKey.NodeHash = NodeClass->ExternalNodeClassLookupInfo.ExternalNodeClassHash;
+					ClassInfo.LookupKey.NodeName = NodeClass->ExternalNodeClassLookupInfo.ExternalNodeClassName;
+				}
 			}
 
 			return ClassInfo;
@@ -1269,6 +1450,7 @@ namespace Metasound
 
 			// Add the input to this node's class description.
 			Inputs.Add(InDescription);
+			ClearLiteralForInput(InDescription.Name);
 
 			FMetasoundNodeDescription NewNodeDescription;
 			NewNodeDescription.Name = InDescription.Name;
@@ -1279,6 +1461,7 @@ namespace Metasound
 
 			FDescPath NodePath = GraphPtr.GetPath()[Path::EFromGraph::ToNodes][NewNodeDescription.UniqueID];
 			FHandleInitParams InitParams = { GraphPtr.GetAccessPoint(), NodePath, NewNodeDescription.Name, OwningAsset };
+
 			// todo: add special enum for input and output nodes
 			return FNodeHandle(FHandleInitParams::PrivateToken, InitParams, NewNodeDescription.ObjectTypeOfNode);
 		}
@@ -1754,10 +1937,13 @@ namespace Metasound
 				return false;
 			}
 
+			
 			if (TUniquePtr<FArchive> FileWriter = TUniquePtr<FArchive>(IFileManager::Get().CreateFileWriter(*InAbsolutePath)))
 			{
-				FJsonStructSerializerBackend Backend(*FileWriter, EStructSerializerBackendFlags::Default);
+				
+				TJsonStructSerializerBackend<DefaultCharType> Backend(*FileWriter, EStructSerializerBackendFlags::Default);
 				FStructSerializer::Serialize<FMetasoundClassDescription>(GraphsClassDeclaration.GetChecked(), Backend);
+		
 				FileWriter->Close();
 
 				return true;
diff --git a/Engine/Plugins/Runtime/Metasound/Source/MetasoundFrontend/Private/MetasoundJsonBackend.h b/Engine/Plugins/Runtime/Metasound/Source/MetasoundFrontend/Private/MetasoundJsonBackend.h
new file mode 100644
--- /dev/null
+++ b/Engine/Plugins/Runtime/Metasound/Source/MetasoundFrontend/Private/MetasoundJsonBackend.h
@@ -0,0 +1,716 @@
+// Copyright Epic Games, Inc. All Rights Reserved.
+
+#pragma once
+
+#include "CoreMinimal.h"
+
+#include "HAL/Platform.h"
+
+#include "IStructSerializerBackend.h"
+#include "IStructDeserializerBackend.h"
+#include "Serialization/JsonWriter.h"
+#include "Serialization/JsonReader.h"
+
+#include "UObject/Class.h"
+#include "UObject/UnrealType.h"
+#include "UObject/EnumProperty.h"
+#include "UObject/TextProperty.h"
+
+// This file is largely copied from JsonStructSerializerBackend / JsonStructDeserializerBackend.
+// The primary difference is that it exposes template args for which character encoding to use,
+// as well as our printing policy.
+
+namespace Metasound
+{
+	typedef ANSICHAR DefaultCharType;
+
+	struct StructDeserializerBackendUtilities
+	{
+		/**
+		* Clears the value of the given property.
+		*
+		* @param Property The property to clear.
+		* @param Outer The property that contains the property to be cleared, if any.
+		* @param Data A pointer to the memory holding the property's data.
+		* @param ArrayIndex The index of the element to clear (if the property is an array).
+		* @return true on success, false otherwise.
+		* @see SetPropertyValue
+		*/
+		static bool ClearPropertyValue(FProperty* Property, FProperty* Outer, void* Data, int32 ArrayIndex)
+		{
+			FArrayProperty* ArrayProperty = CastField<FArrayProperty>(Outer);
+
+			if (ArrayProperty != nullptr)
+			{
+				if (ArrayProperty->Inner != Property)
+				{
+					return false;
+				}
+
+				FScriptArrayHelper ArrayHelper(ArrayProperty, ArrayProperty->template ContainerPtrToValuePtr<void>(Data));
+				ArrayIndex = ArrayHelper.AddValue();
+			}
+
+			Property->ClearValue_InContainer(Data, ArrayIndex);
+
+			return true;
+		}
+
+
+		/**
+		* Gets a pointer to object of the given property.
+		*
+		* @param Property The property to get.
+		* @param Outer The property that contains the property to be get, if any.
+		* @param Data A pointer to the memory holding the property's data.
+		* @param ArrayIndex The index of the element to set (if the property is an array).
+		* @return A pointer to the object represented by the property, null otherwise..
+		* @see ClearPropertyValue
+		*/
+		static void* GetPropertyValuePtr(FProperty* Property, FProperty* Outer, void* Data, int32 ArrayIndex)
+		{
+			check(Property);
+
+			if (FArrayProperty* ArrayProperty = CastField<FArrayProperty>(Outer))
+			{
+				if (ArrayProperty->Inner != Property)
+				{
+					return nullptr;
+				}
+
+				FScriptArrayHelper ArrayHelper(ArrayProperty, ArrayProperty->template ContainerPtrToValuePtr<void>(Data));
+				int32 Index = ArrayHelper.AddValue();
+
+				return ArrayHelper.GetRawPtr(Index);
+			}
+
+			if (ArrayIndex >= Property->ArrayDim)
+			{
+				return nullptr;
+			}
+
+			return Property->template ContainerPtrToValuePtr<void>(Data, ArrayIndex);
+		}
+
+		/**
+		* Sets the value of the given property.
+		*
+		* @param Property The property to set.
+		* @param Outer The property that contains the property to be set, if any.
+		* @param Data A pointer to the memory holding the property's data.
+		* @param ArrayIndex The index of the element to set (if the property is an array).
+		* @return true on success, false otherwise.
+		* @see ClearPropertyValue
+		*/
+		template<typename PropertyType, typename ValueType>
+		static bool SetPropertyValue(PropertyType* Property, FProperty* Outer, void* Data, int32 ArrayIndex, const ValueType& Value)
+		{
+			if (void* Ptr = GetPropertyValuePtr(Property, Outer, Data, ArrayIndex))
+			{
+				Property->SetPropertyValue(Ptr, Value);
+				return true;
+			}
+
+			return false;
+		}
+	};
+
+	namespace JsonStructSerializerBackend
+	{
+		// Writes a property value to the serialization output.
+		template<typename JsonWriterType, typename ValueType>
+		static void WritePropertyValue(const TSharedRef<JsonWriterType> JsonWriter, const FStructSerializerState& State, const ValueType& Value)
+		{
+			if ((State.ValueProperty == nullptr) ||
+				(State.ValueProperty->ArrayDim > 1) ||
+				State.ValueProperty->GetOwner<FArrayProperty>() ||
+				State.ValueProperty->GetOwner<FSetProperty>())
+			{
+				JsonWriter->WriteValue(Value);
+			}
+			else if (State.KeyProperty != nullptr)
+			{
+				FString KeyString;
+				State.KeyProperty->ExportTextItem(KeyString, State.KeyData, nullptr, nullptr, PPF_None);
+				JsonWriter->WriteValue(KeyString, Value);
+			}
+			else
+			{
+				JsonWriter->WriteValue(State.ValueProperty->GetName(), Value);
+			}
+		}
+
+		// Writes a null value to the serialization output.
+		static void WriteNull(const TSharedRef<TJsonWriter<UCS2CHAR>> JsonWriter, const FStructSerializerState& State)
+		{
+			if ((State.ValueProperty == nullptr) ||
+				(State.ValueProperty->ArrayDim > 1) ||
+				State.ValueProperty->GetOwner<FArrayProperty>() ||
+				State.ValueProperty->GetOwner<FSetProperty>())
+			{
+				JsonWriter->WriteNull();
+			}
+			else if (State.KeyProperty != nullptr)
+			{
+				FString KeyString;
+				State.KeyProperty->ExportTextItem(KeyString, State.KeyData, nullptr, nullptr, PPF_None);
+				JsonWriter->WriteNull(KeyString);
+			}
+			else
+			{
+				JsonWriter->WriteNull(State.ValueProperty->GetName());
+			}
+		}
+	}
+
+	/**
+	 * Implements a writer for UStruct serialization using Json.
+	 *
+	 * Optionally, CharType and PrettyPrintPolicy can be subsituted using template arguments.
+	 */
+	template <class CharType, class PrintPolicy = TPrettyJsonPrintPolicy<CharType>>
+	class TJsonStructSerializerBackend
+		: public IStructSerializerBackend
+	{
+	public:
+
+		/**
+		 * Creates and initializes a new legacy instance.
+		 * @note Deprecated, use the two-parameter constructor with EStructSerializerBackendFlags::Legacy if you need backwards compatibility with code compiled prior to 4.22.
+		 *
+		 * @param InArchive The archive to serialize into.
+		 */
+		UE_DEPRECATED(4.22, "Use the two-parameter constructor with EStructSerializerBackendFlags::Legacy only if you need backwards compatibility with code compiled prior to 4.22; otherwise use EStructSerializerBackendFlags::Default.")
+			TJsonStructSerializerBackend(FArchive& InArchive)
+			: JsonWriter(TJsonWriter<CharType, PrintPolicy>::Create(&InArchive))
+			, Flags(EStructSerializerBackendFlags::Legacy)
+		{ }
+
+		/**
+		 * Creates and initializes a new instance with the given flags.
+		 *
+		 * @param InArchive The archive to serialize into.
+		 * @param InFlags The flags that control the serialization behavior (typically EStructSerializerBackendFlags::Default).
+		 */
+		TJsonStructSerializerBackend(FArchive& InArchive, const EStructSerializerBackendFlags InFlags)
+			: JsonWriter(TJsonWriter<CharType, PrintPolicy>::Create(&InArchive))
+			, Flags(InFlags)
+		{ }
+
+		~TJsonStructSerializerBackend() = default;
+
+	public:
+
+		// IStructSerializerBackend interface
+
+		virtual void BeginArray(const FStructSerializerState& State) override
+		{
+			if (State.ValueProperty->GetOwner<FArrayProperty>())
+			{
+				JsonWriter->WriteArrayStart();
+			}
+			else if (State.KeyProperty != nullptr)
+			{
+				FString KeyString;
+				State.KeyProperty->ExportTextItem(KeyString, State.KeyData, nullptr, nullptr, PPF_None);
+				JsonWriter->WriteArrayStart(KeyString);
+			}
+			else
+			{
+				JsonWriter->WriteArrayStart(State.ValueProperty->GetName());
+			}
+		}
+
+		virtual void BeginStructure(const FStructSerializerState& State) override
+		{
+			if (State.ValueProperty != nullptr)
+			{
+				if (State.ValueProperty->GetOwner<FArrayProperty>() || State.ValueProperty->GetOwner<FSetProperty>())
+				{
+					JsonWriter->WriteObjectStart();
+				}
+				else if (State.KeyProperty != nullptr)
+				{
+					FString KeyString;
+					State.KeyProperty->ExportTextItem(KeyString, State.KeyData, nullptr, nullptr, PPF_None);
+					JsonWriter->WriteObjectStart(KeyString);
+				}
+				else
+				{
+					JsonWriter->WriteObjectStart(State.ValueProperty->GetName());
+				}
+			}
+			else
+			{
+				JsonWriter->WriteObjectStart();
+			}
+		}
+
+		virtual void EndArray(const FStructSerializerState& State) override
+		{
+			JsonWriter->WriteArrayEnd();
+		}
+
+		virtual void EndStructure(const FStructSerializerState& State) override
+		{
+			JsonWriter->WriteObjectEnd();
+		}
+
+		virtual void WriteComment(const FString& Comment) override
+		{
+			// Json does not support comments
+		}
+
+		virtual void WriteProperty(const FStructSerializerState& State, int32 ArrayIndex = 0) override
+		{
+			using namespace JsonStructSerializerBackend;
+
+			// booleans
+			if (State.FieldType == FBoolProperty::StaticClass())
+			{
+				WritePropertyValue(JsonWriter, State, CastFieldChecked<FBoolProperty>(State.ValueProperty)->GetPropertyValue_InContainer(State.ValueData, ArrayIndex));
+			}
+
+			// unsigned bytes & enumerations
+			else if (State.FieldType == FEnumProperty::StaticClass())
+			{
+				FEnumProperty* EnumProperty = CastFieldChecked<FEnumProperty>(State.ValueProperty);
+
+				WritePropertyValue(JsonWriter, State, EnumProperty->GetEnum()->GetNameStringByValue(EnumProperty->GetUnderlyingProperty()->GetSignedIntPropertyValue(EnumProperty->ContainerPtrToValuePtr<void>(State.ValueData, ArrayIndex))));
+			}
+			else if (State.FieldType == FByteProperty::StaticClass())
+			{
+				FByteProperty* ByteProperty = CastFieldChecked<FByteProperty>(State.ValueProperty);
+
+				if (ByteProperty->IsEnum())
+				{
+					WritePropertyValue(JsonWriter, State, ByteProperty->Enum->GetNameStringByValue(ByteProperty->GetPropertyValue_InContainer(State.ValueData, ArrayIndex)));
+				}
+				else
+				{
+					WritePropertyValue(JsonWriter, State, (double)ByteProperty->GetPropertyValue_InContainer(State.ValueData, ArrayIndex));
+				}
+			}
+
+			// floating point numbers
+			else if (State.FieldType == FDoubleProperty::StaticClass())
+			{
+				WritePropertyValue(JsonWriter, State, CastFieldChecked<FDoubleProperty>(State.ValueProperty)->GetPropertyValue_InContainer(State.ValueData, ArrayIndex));
+			}
+			else if (State.FieldType == FFloatProperty::StaticClass())
+			{
+				WritePropertyValue(JsonWriter, State, CastFieldChecked<FFloatProperty>(State.ValueProperty)->GetPropertyValue_InContainer(State.ValueData, ArrayIndex));
+			}
+
+			// signed integers
+			else if (State.FieldType == FIntProperty::StaticClass())
+			{
+				WritePropertyValue(JsonWriter, State, (double)CastFieldChecked<FIntProperty>(State.ValueProperty)->GetPropertyValue_InContainer(State.ValueData, ArrayIndex));
+			}
+			else if (State.FieldType == FInt8Property::StaticClass())
+			{
+				WritePropertyValue(JsonWriter, State, (double)CastFieldChecked<FInt8Property>(State.ValueProperty)->GetPropertyValue_InContainer(State.ValueData, ArrayIndex));
+			}
+			else if (State.FieldType == FInt16Property::StaticClass())
+			{
+				WritePropertyValue(JsonWriter, State, (double)CastFieldChecked<FInt16Property>(State.ValueProperty)->GetPropertyValue_InContainer(State.ValueData, ArrayIndex));
+			}
+			else if (State.FieldType == FInt64Property::StaticClass())
+			{
+				WritePropertyValue(JsonWriter, State, (double)CastFieldChecked<FInt64Property>(State.ValueProperty)->GetPropertyValue_InContainer(State.ValueData, ArrayIndex));
+			}
+
+			// unsigned integers
+			else if (State.FieldType == FUInt16Property::StaticClass())
+			{
+				WritePropertyValue(JsonWriter, State, (double)CastFieldChecked<FUInt16Property>(State.ValueProperty)->GetPropertyValue_InContainer(State.ValueData, ArrayIndex));
+			}
+			else if (State.FieldType == FUInt32Property::StaticClass())
+			{
+				WritePropertyValue(JsonWriter, State, (double)CastFieldChecked<FUInt32Property>(State.ValueProperty)->GetPropertyValue_InContainer(State.ValueData, ArrayIndex));
+			}
+			else if (State.FieldType == FUInt64Property::StaticClass())
+			{
+				WritePropertyValue(JsonWriter, State, (double)CastFieldChecked<FUInt64Property>(State.ValueProperty)->GetPropertyValue_InContainer(State.ValueData, ArrayIndex));
+			}
+
+			// names, strings & text
+			else if (State.FieldType == FNameProperty::StaticClass())
+			{
+				WritePropertyValue(JsonWriter, State, CastFieldChecked<FNameProperty>(State.ValueProperty)->GetPropertyValue_InContainer(State.ValueData, ArrayIndex).ToString());
+			}
+			else if (State.FieldType == FStrProperty::StaticClass())
+			{
+				WritePropertyValue(JsonWriter, State, CastFieldChecked<FStrProperty>(State.ValueProperty)->GetPropertyValue_InContainer(State.ValueData, ArrayIndex));
+			}
+			else if (State.FieldType == FTextProperty::StaticClass())
+			{
+				const FText& TextValue = CastFieldChecked<FTextProperty>(State.ValueProperty)->GetPropertyValue_InContainer(State.ValueData, ArrayIndex);
+				if (EnumHasAnyFlags(Flags, EStructSerializerBackendFlags::WriteTextAsComplexString))
+				{
+					FString TextValueString;
+					FTextStringHelper::WriteToBuffer(TextValueString, TextValue);
+					WritePropertyValue(JsonWriter, State, TextValueString);
+				}
+				else
+				{
+					WritePropertyValue(JsonWriter, State, TextValue.ToString());
+				}
+			}
+
+			// classes & objects
+			else if (State.FieldType == FClassProperty::StaticClass())
+			{
+				UObject* const& Value = CastFieldChecked<FClassProperty>(State.ValueProperty)->GetPropertyValue_InContainer(State.ValueData, ArrayIndex);
+				WritePropertyValue(JsonWriter, State, Value ? Value->GetPathName() : FString());
+			}
+			else if (State.FieldType == FSoftClassProperty::StaticClass())
+			{
+				FSoftObjectPtr const& Value = CastFieldChecked<FSoftClassProperty>(State.ValueProperty)->GetPropertyValue_InContainer(State.ValueData, ArrayIndex);
+				WritePropertyValue(JsonWriter, State, Value.IsValid() ? Value->GetPathName() : FString());
+			}
+			else if (State.FieldType == FObjectProperty::StaticClass())
+			{
+				UObject* const& Value = CastFieldChecked<FObjectProperty>(State.ValueProperty)->GetPropertyValue_InContainer(State.ValueData, ArrayIndex);
+				WritePropertyValue(JsonWriter, State, Value ? Value->GetPathName() : FString());
+			}
+			else if (State.FieldType == FWeakObjectProperty::StaticClass())
+			{
+				FWeakObjectPtr const& Value = CastFieldChecked<FWeakObjectProperty>(State.ValueProperty)->GetPropertyValue_InContainer(State.ValueData, ArrayIndex);
+				WritePropertyValue(JsonWriter, State, Value.IsValid() ? Value.Get()->GetPathName() : FString());
+			}
+			else if (State.FieldType == FSoftObjectProperty::StaticClass())
+			{
+				FSoftObjectPtr const& Value = CastFieldChecked<FSoftObjectProperty>(State.ValueProperty)->GetPropertyValue_InContainer(State.ValueData, ArrayIndex);
+				WritePropertyValue(JsonWriter, State, Value.ToString());
+			}
+
+			// unsupported property type
+			else
+			{
+				UE_LOG(LogSerialization, Verbose, TEXT("FJsonStructSerializerBackend: Property %s cannot be serialized, because its type (%s) is not supported"), *State.ValueProperty->GetFName().ToString(), *State.ValueType->GetFName().ToString());
+			}
+		}
+
+	protected:
+
+		// Allow access to the internal JsonWriter to subclasses
+		TSharedRef<TJsonWriter<CharType, PrintPolicy>>& GetWriter()
+		{
+			return JsonWriter;
+		}
+
+	private:
+
+		/** Holds the Json writer used for the actual serialization. */
+		TSharedRef<TJsonWriter<CharType, PrintPolicy>> JsonWriter;
+
+		/** Flags controlling the serialization behavior. */
+		EStructSerializerBackendFlags Flags;
+	};
+
+
+	/**
+	* Implements a reader for UStruct deserialization using Json.
+	*
+	* Optionally, CharType that we're decoding from can be substituted using the CharType template argument.
+	 */
+	template <class CharType>
+	class TJsonStructDeserializerBackend
+		: public IStructDeserializerBackend
+	{
+	public:
+
+		/**
+		 * Creates and initializes a new instance.
+		 *
+		 * @param Archive The archive to deserialize from.
+		 */
+		TJsonStructDeserializerBackend(FArchive& Archive)
+			: JsonReader(TJsonReader<CharType>::Create(&Archive))
+		{ }
+
+		~TJsonStructDeserializerBackend() = default;
+
+	public:
+
+		// IStructDeserializerBackend interface
+
+		virtual const FString& GetCurrentPropertyName() const override
+		{
+			return JsonReader->GetIdentifier();
+		}
+
+		virtual FString GetDebugString() const override
+		{
+			return FString::Printf(TEXT("Line: %u, Ch: %u"), JsonReader->GetLineNumber(), JsonReader->GetCharacterNumber());
+		}
+
+		virtual const FString& GetLastErrorMessage() const override
+		{
+			return JsonReader->GetErrorMessage();
+		}
+
+		virtual bool GetNextToken(EStructDeserializerBackendTokens& OutToken) override
+		{
+			if (!JsonReader->ReadNext(LastNotation))
+			{
+				return false;
+			}
+
+			switch (LastNotation)
+			{
+			case EJsonNotation::ArrayEnd:
+				OutToken = EStructDeserializerBackendTokens::ArrayEnd;
+				break;
+
+			case EJsonNotation::ArrayStart:
+				OutToken = EStructDeserializerBackendTokens::ArrayStart;
+				break;
+
+			case EJsonNotation::Boolean:
+			case EJsonNotation::Null:
+			case EJsonNotation::Number:
+			case EJsonNotation::String:
+			{
+				OutToken = EStructDeserializerBackendTokens::Property;
+			}
+			break;
+
+			case EJsonNotation::Error:
+				OutToken = EStructDeserializerBackendTokens::Error;
+				break;
+
+			case EJsonNotation::ObjectEnd:
+				OutToken = EStructDeserializerBackendTokens::StructureEnd;
+				break;
+
+			case EJsonNotation::ObjectStart:
+				OutToken = EStructDeserializerBackendTokens::StructureStart;
+				break;
+
+			default:
+				OutToken = EStructDeserializerBackendTokens::None;
+			}
+
+			return true;
+		}
+
+		virtual bool ReadProperty(FProperty* Property, FProperty* Outer, void* Data, int32 ArrayIndex) override
+		{
+			switch (LastNotation)
+			{
+				// boolean values
+			case EJsonNotation::Boolean:
+			{
+				bool BoolValue = JsonReader->GetValueAsBoolean();
+
+				if (FBoolProperty* BoolProperty = CastField<FBoolProperty>(Property))
+				{
+					return StructDeserializerBackendUtilities::SetPropertyValue(BoolProperty, Outer, Data, ArrayIndex, BoolValue);
+				}
+
+				const FCoreTexts& CoreTexts = FCoreTexts::Get();
+
+				UE_LOG(LogSerialization, Verbose, TEXT("Boolean field %s with value '%s' is not supported in FProperty type %s (%s)"), *Property->GetFName().ToString(), BoolValue ? *(CoreTexts.True.ToString()) : *(CoreTexts.False.ToString()), *Property->GetClass()->GetName(), *GetDebugString());
+
+				return false;
+			}
+			break;
+
+			// numeric values
+			case EJsonNotation::Number:
+			{
+				double NumericValue = JsonReader->GetValueAsNumber();
+
+				if (FByteProperty* ByteProperty = CastField<FByteProperty>(Property))
+				{
+					return StructDeserializerBackendUtilities::SetPropertyValue(ByteProperty, Outer, Data, ArrayIndex, (int8)NumericValue);
+				}
+
+				if (FDoubleProperty* DoubleProperty = CastField<FDoubleProperty>(Property))
+				{
+					return StructDeserializerBackendUtilities::SetPropertyValue(DoubleProperty, Outer, Data, ArrayIndex, (double)NumericValue);
+				}
+
+				if (FFloatProperty* FloatProperty = CastField<FFloatProperty>(Property))
+				{
+					return StructDeserializerBackendUtilities::SetPropertyValue(FloatProperty, Outer, Data, ArrayIndex, (float)NumericValue);
+				}
+
+				if (FIntProperty* IntProperty = CastField<FIntProperty>(Property))
+				{
+					return StructDeserializerBackendUtilities::SetPropertyValue(IntProperty, Outer, Data, ArrayIndex, (int32)NumericValue);
+				}
+
+				if (FUInt32Property* UInt32Property = CastField<FUInt32Property>(Property))
+				{
+					return StructDeserializerBackendUtilities::SetPropertyValue(UInt32Property, Outer, Data, ArrayIndex, (uint32)NumericValue);
+				}
+
+				if (FInt16Property* Int16Property = CastField<FInt16Property>(Property))
+				{
+					return StructDeserializerBackendUtilities::SetPropertyValue(Int16Property, Outer, Data, ArrayIndex, (int16)NumericValue);
+				}
+
+				if (FUInt16Property* FInt16Property = CastField<FUInt16Property>(Property))
+				{
+					return StructDeserializerBackendUtilities::SetPropertyValue(FInt16Property, Outer, Data, ArrayIndex, (uint16)NumericValue);
+				}
+
+				if (FInt64Property* Int64Property = CastField<FInt64Property>(Property))
+				{
+					return StructDeserializerBackendUtilities::SetPropertyValue(Int64Property, Outer, Data, ArrayIndex, (int64)NumericValue);
+				}
+
+				if (FUInt64Property* FInt64Property = CastField<FUInt64Property>(Property))
+				{
+					return StructDeserializerBackendUtilities::SetPropertyValue(FInt64Property, Outer, Data, ArrayIndex, (uint64)NumericValue);
+				}
+
+				if (FInt8Property* Int8Property = CastField<FInt8Property>(Property))
+				{
+					return StructDeserializerBackendUtilities::SetPropertyValue(Int8Property, Outer, Data, ArrayIndex, (int8)NumericValue);
+				}
+
+				UE_LOG(LogSerialization, Verbose, TEXT("Numeric field %s with value '%f' is not supported in FProperty type %s (%s)"), *Property->GetFName().ToString(), NumericValue, *Property->GetClass()->GetName(), *GetDebugString());
+
+				return false;
+			}
+			break;
+
+			// null values
+			case EJsonNotation::Null:
+				return StructDeserializerBackendUtilities::ClearPropertyValue(Property, Outer, Data, ArrayIndex);
+
+				// strings, names, enumerations & object/class reference
+			case EJsonNotation::String:
+			{
+				const FString& StringValue = JsonReader->GetValueAsString();
+
+				if (FStrProperty* StrProperty = CastField<FStrProperty>(Property))
+				{
+					return StructDeserializerBackendUtilities::SetPropertyValue(StrProperty, Outer, Data, ArrayIndex, StringValue);
+				}
+
+				if (FNameProperty* NameProperty = CastField<FNameProperty>(Property))
+				{
+					return StructDeserializerBackendUtilities::SetPropertyValue(NameProperty, Outer, Data, ArrayIndex, FName(*StringValue));
+				}
+
+				if (FTextProperty* TextProperty = CastField<FTextProperty>(Property))
+				{
+					FText TextValue;
+					if (!FTextStringHelper::ReadFromBuffer(*StringValue, TextValue))
+					{
+						TextValue = FText::FromString(StringValue);
+					}
+					return StructDeserializerBackendUtilities::SetPropertyValue(TextProperty, Outer, Data, ArrayIndex, TextValue);
+				}
+
+				if (FByteProperty* ByteProperty = CastField<FByteProperty>(Property))
+				{
+					if (!ByteProperty->Enum)
+					{
+						return false;
+					}
+
+					int32 Value = ByteProperty->Enum->GetValueByName(*StringValue);
+					if (Value == INDEX_NONE)
+					{
+						return false;
+					}
+
+					return StructDeserializerBackendUtilities::SetPropertyValue(ByteProperty, Outer, Data, ArrayIndex, (uint8)Value);
+				}
+
+				if (FEnumProperty* EnumProperty = CastField<FEnumProperty>(Property))
+				{
+					int64 Value = EnumProperty->GetEnum()->GetValueByName(*StringValue);
+					if (Value == INDEX_NONE)
+					{
+						return false;
+					}
+
+					if (void* ElementPtr = StructDeserializerBackendUtilities::GetPropertyValuePtr(EnumProperty, Outer, Data, ArrayIndex))
+					{
+						EnumProperty->GetUnderlyingProperty()->SetIntPropertyValue(ElementPtr, Value);
+						return true;
+					}
+
+					return false;
+				}
+
+				if (FClassProperty* ClassProperty = CastField<FClassProperty>(Property))
+				{
+					return StructDeserializerBackendUtilities::SetPropertyValue(ClassProperty, Outer, Data, ArrayIndex, LoadObject<UClass>(nullptr, *StringValue, nullptr, LOAD_NoWarn));
+				}
+
+				if (FSoftClassProperty* SoftClassProperty = CastField<FSoftClassProperty>(Property))
+				{
+					return StructDeserializerBackendUtilities::SetPropertyValue(SoftClassProperty, Outer, Data, ArrayIndex, FSoftObjectPtr(LoadObject<UClass>(nullptr, *StringValue, nullptr, LOAD_NoWarn)));
+				}
+
+				if (FObjectProperty* ObjectProperty = CastField<FObjectProperty>(Property))
+				{
+					return StructDeserializerBackendUtilities::SetPropertyValue(ObjectProperty, Outer, Data, ArrayIndex, StaticFindObject(ObjectProperty->PropertyClass, nullptr, *StringValue));
+				}
+
+				if (FWeakObjectProperty* WeakObjectProperty = CastField<FWeakObjectProperty>(Property))
+				{
+					return StructDeserializerBackendUtilities::SetPropertyValue(WeakObjectProperty, Outer, Data, ArrayIndex, FWeakObjectPtr(StaticFindObject(WeakObjectProperty->PropertyClass, nullptr, *StringValue)));
+				}
+
+				if (FSoftObjectProperty* SoftObjectProperty = CastField<FSoftObjectProperty>(Property))
+				{
+					return StructDeserializerBackendUtilities::SetPropertyValue(SoftObjectProperty, Outer, Data, ArrayIndex, FSoftObjectPtr(FSoftObjectPath(StringValue)));
+				}
+
+				UE_LOG(LogSerialization, Verbose, TEXT("String field %s with value '%s' is not supported in FProperty type %s (%s)"), *Property->GetFName().ToString(), *StringValue, *Property->GetClass()->GetName(), *GetDebugString());
+
+				return false;
+			}
+			break;
+			}
+
+			return true;
+		}
+
+		virtual void SkipArray() override
+		{
+			JsonReader->SkipArray();
+		}
+
+		virtual void SkipStructure() override
+		{
+			JsonReader->SkipObject();
+		}
+
+	protected:
+		FString& GetLastIdentifier()
+		{
+			return LastIdentifier;
+		}
+
+		EJsonNotation GetLastNotation()
+		{
+			return LastNotation;
+		}
+
+		TSharedRef<TJsonReader<CharType>>& GetReader()
+		{
+			return JsonReader;
+		}
+
+	private:
+
+		/** Holds the name of the last read Json identifier. */
+		FString LastIdentifier;
+
+		/** Holds the last read Json notation. */
+		EJsonNotation LastNotation;
+
+		/** Holds the Json reader used for the actual reading of the archive. */
+		TSharedRef<TJsonReader<CharType>> JsonReader;
+	};
+}
diff --git a/Engine/Plugins/Runtime/Metasound/Source/MetasoundFrontend/Public/MetasoundAssetBase.h b/Engine/Plugins/Runtime/Metasound/Source/MetasoundFrontend/Public/MetasoundAssetBase.h
--- a/Engine/Plugins/Runtime/Metasound/Source/MetasoundFrontend/Public/MetasoundAssetBase.h
+++ b/Engine/Plugins/Runtime/Metasound/Source/MetasoundFrontend/Public/MetasoundAssetBase.h
@@ -8,6 +8,7 @@
 #include "MetasoundGraph.h"
 #include "UObject/WeakObjectPtrTemplates.h"
 
+class UEdGraph;
 
 class METASOUNDFRONTEND_API FMetasoundAssetBase
 {
@@ -19,12 +20,34 @@ public:
 
 	virtual ~FMetasoundAssetBase() = default;
 
+#if WITH_EDITORONLY_DATA
+
+	// Returns the graph associated with this Metasound. Graph is required to be referenced on
+	// Metasound UObject for editor serialization purposes.
+	// @return Editor graph associated with this metasound uobject.
+	virtual UEdGraph* GetGraph() = 0;
+	virtual const UEdGraph* GetGraph() const = 0;
+	virtual UEdGraph& GetGraphChecked() = 0;
+	virtual const UEdGraph& GetGraphChecked() const = 0;
+
+	// Sets the graph associated with this Metasound. Graph is required to be referenced on
+	// Metasound UObject for editor serialization purposes.
+	// @param Editor graph associated with this metasound object.
+	virtual void SetGraph(UEdGraph* InGraph) = 0;
+#endif // WITH_EDITORONLY_DATA
+
 	// Sets/overwrites the root class metadata
 	virtual void SetMetadata(FMetasoundClassMetadata& InMetadata);
 
+	// Returns  a description of the required inputs and outputs for this metasound UClass.
+	virtual FMetasoundArchetype GetArchetype() const = 0;
+
 	// Returns the root class metadata
 	FMetasoundClassMetadata GetMetadata();
 
+	const FText& GetInputToolTip(FString InputName) const;
+	const FText& GetOutputToolTip(FString OutputName) const;
+
 	// Imports the asset from a JSON file at provided path
 	bool ImportFromJSON(const FString& InAbsolutePath);
 
@@ -37,6 +60,13 @@ public:
 	// Returns all handles for subgraphs referenced
 	TArray<Metasound::Frontend::FGraphHandle> GetAllSubgraphHandles();
 
+	// Overwrites the entire document and fixes it up based on any 
+	// inputs or outputs in the archetype that are missing from the graph.
+	void SetDocument(const FMetasoundDocument& InDocument);
+
+	// This must be called on UObject::PostLoad, as well as in this asset's UFactory, to fix up the root document based on the most recent version of the archetype.
+	void ConformDocumentToArchetype();
+
 protected:
 	// Returns private token allowing implementing asset class to use graph/node handle system
 	static Metasound::Frontend::FHandleInitParams::EPrivateToken GetPrivateToken()
diff --git a/Engine/Plugins/Runtime/Metasound/Source/MetasoundFrontend/Public/MetasoundFrontend.h b/Engine/Plugins/Runtime/Metasound/Source/MetasoundFrontend/Public/MetasoundFrontend.h
--- a/Engine/Plugins/Runtime/Metasound/Source/MetasoundFrontend/Public/MetasoundFrontend.h
+++ b/Engine/Plugins/Runtime/Metasound/Source/MetasoundFrontend/Public/MetasoundFrontend.h
@@ -12,7 +12,6 @@
 // Forward Declarations
 class FMetasoundAssetBase;
 
-
 namespace Metasound
 {	
 	/**
@@ -73,6 +72,13 @@ namespace Metasound
 		// Generates a new FMetasoundClassDescription for a given node class. Only used by classes that manipulate Metasound Description data directly.
 		METASOUNDFRONTEND_API FMetasoundClassDescription GenerateClassDescriptionForNode(const FNodeClassInfo& InInfo);
 
+		template<typename DataType>
+		FName GetDataTypeName()
+		{
+			static FName DataTypeName = FName(TDataReferenceTypeInfo<DataType>::TypeName);
+			return DataTypeName;
+		}
+
 		// Returns a list of all available data types.
 		METASOUNDFRONTEND_API TArray<FName> GetAllAvailableDataTypes();
 
@@ -85,6 +91,48 @@ namespace Metasound
 		// @returns false if the file couldn't be found or parsed into a document.
 		METASOUNDFRONTEND_API bool ImportJSONToMetasound(const FString& InPath, FMetasoundDocument& OutMetasoundDocument);
 
+		struct METASOUNDFRONTEND_API FMetasoundArchetypeRegistryParams_Internal
+		{
+			FMetasoundArchetype ArchetypeDescription;
+
+			// The UClass associated with this specific archetype.
+			UClass* ArchetypeUClass;
+
+			// template-generated lambdas used to safely sidecast to FMetasoundBase*.
+			TUniqueFunction<FMetasoundAssetBase* (UObject*)> SafeCast;
+			TUniqueFunction<const FMetasoundAssetBase* (const UObject*)> SafeConstCast;
+
+			// This function should construct a new UObject of the given archetype's type
+			// given a metasound document with a matching archetype.
+			// The first argument is the document to use.
+			// The second argument is the path relative to the content directory to save the soundwave to.
+			TUniqueFunction<UObject* (const FMetasoundDocument&, const FString&)> ObjectGetter;
+		};
+
+		// ಠ╭╮ಠ         ಠ╭╮ಠ           ಠ╭╮ಠ
+		// Please don't use this function.
+		// See RegisterArchetype<Class> instead, in MetasoundArchetypeRegistration.h
+		METASOUNDFRONTEND_API bool RegisterArchetype_Internal(FMetasoundArchetypeRegistryParams_Internal&& InParams);
+
+		METASOUNDFRONTEND_API TArray<FName> GetAllRegisteredArchetypes();
+
+		// Returns a new UObject, whose class corresponds to the archetype in the document.
+		// @param InDocument a fully generated metasound document, typically retrieved from ImportJSONToMetasound().
+		// @param InPath, path in content directory to save the generated UAsset to (ex. "/game/MyDir/MyMetasoundAsset".
+		//                if InPath is invalid, we won't save to an asset.
+		// @returns nullptr if we couldn't find the archetype.
+		METASOUNDFRONTEND_API UObject* GetObjectForDocument(const FMetasoundDocument& InDocument, const FString& InPath);
+
+		// This returns true if the object is listed as one of our metasound archetypes.
+		// If it is, the object can be safely static cast to FMetasoundAssetBase*.
+		METASOUNDFRONTEND_API bool IsObjectAMetasoundArchetype(const UObject* InObject);
+
+		// These functions are used to safely sidecast between a UObject of a given metasound archetype
+		// and FMetasoundAssetBase*.
+		// @returns nullptr if the object wasn't a registered archetype.
+		METASOUNDFRONTEND_API FMetasoundAssetBase* GetObjectAsAssetBase(UObject* InObject);
+		METASOUNDFRONTEND_API const FMetasoundAssetBase* GetObjectAsAssetBase(const UObject* InObject);
+
 		// Struct that indicates whether an input and an output can be connected,
 		// and whether an intermediate node is necessary to connect the two.
 		struct METASOUNDFRONTEND_API FConnectability
@@ -325,6 +373,8 @@ namespace Metasound
 			bool SetInputToLiteral(const FString& InInputName, int32 InValue);
 			bool SetInputToLiteral(const FString& InInputName, float InValue);
 			bool SetInputToLiteral(const FString& InInputName, const FString& InValue);
+			bool SetInputToLiteral(const FString& InInputName, UObject* InValue);
+			bool SetInputToLiteral(const FString& InInputName, TArray<UObject*> InValue);
 
 			// This can be used to clear the current literal for a given input.
 			// @returns false if the input name couldn't be found.
@@ -365,6 +415,9 @@ namespace Metasound
 			// @todo: make an API for doing this async.
 			TUniquePtr<IOperator> BuildOperator(const FOperatorSettings& InSettings, TArray<IOperatorBuilder::FBuildErrorPtr>& OutBuildErrors) const;
 
+			// This function will ensure that the document's root class has the inputs and outputs required by the archetype.
+			void FixDocumentToMatchArchetype();
+
 		private:
 
 			// The graph struct itself.
diff --git a/Engine/Plugins/Runtime/Metasound/Source/MetasoundFrontend/Public/MetasoundFrontendDataLayout.h b/Engine/Plugins/Runtime/Metasound/Source/MetasoundFrontend/Public/MetasoundFrontendDataLayout.h
--- a/Engine/Plugins/Runtime/Metasound/Source/MetasoundFrontend/Public/MetasoundFrontendDataLayout.h
+++ b/Engine/Plugins/Runtime/Metasound/Source/MetasoundFrontend/Public/MetasoundFrontendDataLayout.h
@@ -69,20 +69,22 @@ struct FMetasoundLiteralDescription
 {
 	GENERATED_BODY()
 
+	// HACK: Once we make a property inspector for these we should hide them.
+
 	// The actual type of this literal.
-	UPROPERTY()
+	UPROPERTY(EditAnywhere, Category = Parameters)
 	EMetasoundLiteralType LiteralType;
 
-	UPROPERTY()
+	UPROPERTY(EditAnywhere, Category = Parameters)
 	bool AsBool;
 
-	UPROPERTY();
+	UPROPERTY(EditAnywhere, Category = Parameters)
 	int32 AsInteger;
 
-	UPROPERTY()
+	UPROPERTY(EditAnywhere, Category = Parameters)
 	float AsFloat;
 
-	UPROPERTY()
+	UPROPERTY(EditAnywhere, Category = Parameters)
 	FString AsString;
 };
 
@@ -104,7 +106,8 @@ struct FMetasoundInputDescription
 	FText ToolTip;
 
 	// The optional literal value, if we have one.
-	UPROPERTY()
+	// NOTE: in the future we'll have a specific details customization for this.
+	UPROPERTY(EditAnywhere, Category = Parameters)
 	FMetasoundLiteralDescription LiteralValue;
 };
 
diff --git a/Engine/Plugins/Runtime/Metasound/Source/MetasoundGenerator/Private/MetasoundGenerator.cpp b/Engine/Plugins/Runtime/Metasound/Source/MetasoundGenerator/Private/MetasoundGenerator.cpp
--- a/Engine/Plugins/Runtime/Metasound/Source/MetasoundGenerator/Private/MetasoundGenerator.cpp
+++ b/Engine/Plugins/Runtime/Metasound/Source/MetasoundGenerator/Private/MetasoundGenerator.cpp
@@ -16,24 +16,33 @@
 
 namespace Metasound
 {
-	FMetasoundGenerator::FMetasoundGenerator(FOperatorUniquePtr InGraphOperator, const FMultichannelAudioFormatReadRef& InGraphOutputAudioRef)
-	:	NumChannels(0)
-	,	NumFramesPerExecute(0)
-	,	NumSamplesPerExecute(0)
-	,	GraphOutputAudioRef(InGraphOutputAudioRef) // Copy audio ref here to avoid unneeded construction of unused graph output audio object.
+	
+
+	FMetasoundGenerator::FMetasoundGenerator(FMetasoundGeneratorInitParams&& InParams)
+		: NumChannels(0)
+		, NumFramesPerExecute(0)
+		, NumSamplesPerExecute(0)
+		, GraphOutputAudioRef(InParams.GraphOutputAudioRef) // Copy audio ref here to avoid unneeded construction of unused graph output audio object.
+		, OnFinishedBopRef(InParams.BopOnFinishRef)
 	{
-		SetGraph(MoveTemp(InGraphOperator), InGraphOutputAudioRef);
+		SetGraph(MoveTemp(InParams));
 	}
 
 	FMetasoundGenerator::~FMetasoundGenerator()
 	{
 	}
 
-	bool FMetasoundGenerator::UpdateGraphOperator(FOperatorUniquePtr InGraphOperator, const FMultichannelAudioFormatReadRef& InGraphOutputAudioRef)
+	bool FMetasoundGenerator::UpdateGraphOperator(FMetasoundGeneratorInitParams&& InParams)
 	{
-		if (InGraphOutputAudioRef->GetNumChannels() == NumChannels)
+		// multichannel version:
+		//const int32 GraphNumChannels = InParams.GraphOutputAudioRef->GetNumChannels();
+
+		// single channel version:
+		const int32 GraphNumChannels = 1;
+
+		if (GraphNumChannels == NumChannels)
 		{
-			SetGraph(MoveTemp(InGraphOperator), InGraphOutputAudioRef);
+			SetGraph(MoveTemp(InParams));
 
 			return true;
 		}
@@ -41,7 +50,7 @@ namespace Metasound
 		return false;
 	}
 
-	void FMetasoundGenerator::SetGraph(FOperatorUniquePtr InGraphOperator, const FMultichannelAudioFormatReadRef& InGraphOutputAudioRef)
+	void FMetasoundGenerator::SetGraph(FMetasoundGeneratorInitParams&& InParams)
 	{
 		InterleavedAudioBuffer.Reset();
 		NumFramesPerExecute = 0;
@@ -49,17 +58,26 @@ namespace Metasound
 
 		// The graph operator and graph audio output contain all the values needed
 		// by the sound generator.
-		RootExecuter.SetOperator(MoveTemp(InGraphOperator));
-		GraphOutputAudioRef = InGraphOutputAudioRef;
+		RootExecuter.SetOperator(MoveTemp(InParams.GraphOperator));
+		GraphOutputAudioRef = InParams.GraphOutputAudioRef;
+		OnFinishedBopRef = InParams.BopOnFinishRef;
 
 		// Query the graph output to get the number of output audio channels.
-		NumChannels = GraphOutputAudioRef->GetNumChannels();
+		// Multichannel version:
+		//NumChannels = GraphOutputAudioRef->GetNumChannels();
+
+		// Single channel version:
+		NumChannels = 1;
 
 		if (NumChannels > 0)
 		{
 			// All buffers have same number of frames, so only need to query
 			// first buffer to know number of frames.
-			NumFramesPerExecute = GraphOutputAudioRef->GetBuffers()[0]->Num();
+			
+			// Multichannel version:
+			//NumFramesPerExecute = GraphOutputAudioRef->GetBuffers()[0]->Num();
+			// single channel version:
+			NumFramesPerExecute = GraphOutputAudioRef->Num();
 
 			NumSamplesPerExecute = NumFramesPerExecute * NumChannels;
 
@@ -73,7 +91,11 @@ namespace Metasound
 
 	int32 FMetasoundGenerator::GetNumChannels() const
 	{
-		return GraphOutputAudioRef->GetNumChannels();
+		// Multichannel version:
+		// return GraphOutputAudioRef->GetNumChannels();
+
+		// Single channel version:
+		return 1;
 	}
 
 	int32 FMetasoundGenerator::OnGenerateAudio(float* OutAudio, int32 NumSamplesRemaining)
@@ -125,6 +147,16 @@ namespace Metasound
 	}
 
 
+	int32 FMetasoundGenerator::GetDesiredNumSamplesToRenderPerCallback() const
+	{
+		return NumFramesPerExecute;
+	}
+
+	bool FMetasoundGenerator::IsFinished() const
+	{
+		return *OnFinishedBopRef;
+	}
+
 	int32 FMetasoundGenerator::FillWithBuffer(const Audio::AlignedFloatBuffer& InBuffer, float* OutAudio, int32 MaxNumOutputSamples)
 	{
 		int32 InNum = InBuffer.Num();
@@ -148,6 +180,9 @@ namespace Metasound
 
 	void FMetasoundGenerator::InterleaveGeneratedAudio()
 	{
+		/**
+		// TODO: Convert MetasoundSource / MetasoundGenerator to take Multichannel Audio
+		// once we've got converter nodes up and running.
 		const FMultichannelAudioFormat& DeinterleavedAudio = *GraphOutputAudioRef;
 
 		// Prepare output buffer
@@ -175,6 +210,26 @@ namespace Metasound
 				OutputPtr += NumChannels;
 			}
 		}
+
+		*/
+
+		// Single channel version (for now):
+		const FAudioBuffer& InputAudioBuffer = *GraphOutputAudioRef;
+
+		for (int32 ChannelIndex = 0; ChannelIndex < NumChannels; ChannelIndex++)
+		{
+			const FAudioBuffer& InputBuffer = InputAudioBuffer;
+
+			const float* InputPtr = InputBuffer.GetData();
+			float* OutputPtr = &InterleavedAudioBuffer.GetData()[ChannelIndex];
+
+			// Assign values to output for single channel.
+			for (int32 FrameIndex = 0; FrameIndex < NumFramesPerExecute; FrameIndex++)
+			{
+				*OutputPtr = InputPtr[FrameIndex];
+				OutputPtr += NumChannels;
+			}
+		}
 	}
 } 
 
diff --git a/Engine/Plugins/Runtime/Metasound/Source/MetasoundGenerator/Private/SynthComponentMetasoundExample.cpp b/Engine/Plugins/Runtime/Metasound/Source/MetasoundGenerator/Private/SynthComponentMetasoundExample.cpp
--- a/Engine/Plugins/Runtime/Metasound/Source/MetasoundGenerator/Private/SynthComponentMetasoundExample.cpp
+++ b/Engine/Plugins/Runtime/Metasound/Source/MetasoundGenerator/Private/SynthComponentMetasoundExample.cpp
@@ -151,10 +151,23 @@ void USynthComponentMetasoundExample::SetGraphOperator(FOperatorUniquePtr InGrap
 	if (InGraphOperator.IsValid())
 	{
 		FDataReferenceCollection Outputs = InGraphOperator->GetOutputs();
-		FMultichannelAudioFormatReadRef Audio = Outputs.GetDataReadReferenceOrConstruct<FMultichannelAudioFormat>(InOutputAudioName, NumChannels, 256);
-		NumChannels = Audio->GetNumChannels();
+		FAudioBufferReadRef Audio = Outputs.GetDataReadReferenceOrConstruct<FAudioBuffer>(InOutputAudioName, 256);
+		FBopReadRef OnFinished = FBopReadRef::CreateNew(); // unused in this example
+		
+		// Multichannel version:
+		//NumChannels = Audio->GetNumChannels();
 
-		Generator = MakeShared<Metasound::FMetasoundGenerator, ESPMode::ThreadSafe>(MoveTemp(InGraphOperator), Audio);
+		// Single channel version:
+		NumChannels = 1;
+
+		FMetasoundGeneratorInitParams InitParams =
+		{
+			MoveTemp(InGraphOperator),
+			Audio,
+			OnFinished
+		};
+
+		Generator = MakeShared<Metasound::FMetasoundGenerator, ESPMode::ThreadSafe>(MoveTemp(InitParams));
 	}
 	else
 	{
@@ -173,9 +186,24 @@ bool USynthComponentMetasoundExample::PushGraphOperator(FOperatorUniquePtr InGra
 	}
 
 	FDataReferenceCollection Outputs = InGraphOperator->GetOutputs();
-	FMultichannelAudioFormatReadRef Audio = Outputs.GetDataReadReferenceOrConstruct<FMultichannelAudioFormat>(InOutputAudioName, NumChannels, 256);
+	FAudioBufferReadRef AudioReadRef = Outputs.GetDataReadReferenceOrConstruct<FAudioBuffer>(InOutputAudioName, 256);
+	
+	// On Finished isn't used in this example.
+	FBopReadRef OnFinishedRef = FBopReadRef::CreateNew();
+
+	FMetasoundGeneratorInitParams GeneratorInitParams =
+	{
+		MoveTemp(InGraphOperator),
+		AudioReadRef,
+		OnFinishedRef
+	};
 
-	if (Generator->UpdateGraphOperator(MoveTemp(InGraphOperator), Audio))
+	
+	GeneratorInitParams.GraphOutputAudioRef = Outputs.GetDataReadReferenceOrConstruct<FAudioBuffer>(InOutputAudioName, 256);
+	GeneratorInitParams.GraphOperator = MoveTemp(InGraphOperator);
+
+
+	if (Generator->UpdateGraphOperator(MoveTemp(GeneratorInitParams)))
 	{
 		return true;
 	}
@@ -183,7 +211,7 @@ bool USynthComponentMetasoundExample::PushGraphOperator(FOperatorUniquePtr InGra
 	return false;
 }
 
-ISoundGeneratorPtr USynthComponentMetasoundExample::CreateSoundGenerator(int32 InSampleRate, int32 InNumChannels)
+ISoundGeneratorPtr USynthComponentMetasoundExample::CreateSoundGenerator(const FSoundGeneratorInitParams& InParams)
 {
 	using namespace Metasound;
 
@@ -191,7 +219,7 @@ ISoundGeneratorPtr USynthComponentMetasoundExample::CreateSoundGenerator(int32 I
 
 	TUniquePtr<FSynthComponentMetasoundExampleGraph> ExampleGraph = CreateSynthComponentMetasoundExampleGraph();
 
-	FOperatorSettings Settings(InSampleRate, 256);
+	FOperatorSettings Settings(InParams.SampleRate, InParams.NumFramesPerCallback);
 	FOperatorBuilder Builder(Settings);
 
 	TArray<IOperatorBuilder::FBuildErrorPtr> Errors;
@@ -202,6 +230,3 @@ ISoundGeneratorPtr USynthComponentMetasoundExample::CreateSoundGenerator(int32 I
 	return Generator;
 }
 
-
-
-
diff --git a/Engine/Plugins/Runtime/Metasound/Source/MetasoundGenerator/Public/MetasoundGenerator.h b/Engine/Plugins/Runtime/Metasound/Source/MetasoundGenerator/Public/MetasoundGenerator.h
--- a/Engine/Plugins/Runtime/Metasound/Source/MetasoundGenerator/Public/MetasoundGenerator.h
+++ b/Engine/Plugins/Runtime/Metasound/Source/MetasoundGenerator/Public/MetasoundGenerator.h
@@ -8,10 +8,18 @@
 #include "MetasoundExecutableOperator.h"
 #include "MetasoundAudioFormats.h"
 #include "MetasoundPrimitives.h"
+#include "MetasoundDataReferenceTypes.h"
 #include "Sound/SoundGenerator.h"
 
 namespace Metasound
 {
+	struct FMetasoundGeneratorInitParams
+	{
+		TUniquePtr<Metasound::IOperator> GraphOperator;
+		FAudioBufferReadRef GraphOutputAudioRef;
+		FBopReadRef BopOnFinishRef;
+	};
+
 	/** FMetasoundGenerator generates audio from a given metasound IOperator
 	 * which produces a multichannel audio output.
 	 */
@@ -19,14 +27,14 @@ namespace Metasound
 	{
 	public:
 		using FOperatorUniquePtr = TUniquePtr<Metasound::IOperator>;
-		using FMultichannelAudioFormatReadRef = Metasound::FMultichannelAudioFormatReadRef;
+		using FAudioBufferReadRef = Metasound::FAudioBufferReadRef;
 
 		/** Create the generator with a graph operator and an output audio reference.
 		 *
 		 * @param InGraphOperator - Unique pointer to the IOperator which executes the entire graph.
 		 * @param InGraphOutputAudioRef - Read reference to the audio buffer filled by the InGraphOperator.
 		 */
-		FMetasoundGenerator(FOperatorUniquePtr InGraphOperator, const FMultichannelAudioFormatReadRef& InGraphOutputAudioRef);
+		FMetasoundGenerator(FMetasoundGeneratorInitParams&& InParams);
 
 		virtual ~FMetasoundGenerator();
 
@@ -86,19 +94,21 @@ namespace Metasound
 		 *
 		 * @return Returns true if the graph operator was successfully swapped. False otherwise.
 		 */
-		bool UpdateGraphOperator(FOperatorUniquePtr InGraphOperator, const FMultichannelAudioFormatReadRef& InGraphOutputAudioRef);
+		bool UpdateGraphOperator(FMetasoundGeneratorInitParams&& InParams);
 
 		/** Return the number of audio channels. */
 		int32 GetNumChannels() const;
 
 		//~ Begin FSoundGenerator
 		virtual int32 OnGenerateAudio(float* OutAudio, int32 NumSamples) override;
+		int32 GetDesiredNumSamplesToRenderPerCallback() const override;
+		bool IsFinished() const override;
 		//~ End FSoundGenerator
 
 	private:
 
 		// Internal set graph after checking compatibility.
-		void SetGraph(FOperatorUniquePtr InGraphOperator, const FMultichannelAudioFormatReadRef& InGraphOutputAudioRef);
+		void SetGraph(FMetasoundGeneratorInitParams&& InParams);
 
 		// Fill OutAudio with data in InBuffer, up to maximum number of samples.
 		// Returns the number of samples used.
@@ -113,7 +123,10 @@ namespace Metasound
 		int32 NumFramesPerExecute;
 		int32 NumSamplesPerExecute;
 
-		FMultichannelAudioFormatReadRef GraphOutputAudioRef;
+		FAudioBufferReadRef GraphOutputAudioRef;
+
+		// This will be bopped when this metasound is finished playing.
+		FBopReadRef OnFinishedBopRef;
 
 		Audio::AlignedFloatBuffer InterleavedAudioBuffer;
 
@@ -121,7 +134,3 @@ namespace Metasound
 	};
 }
 
-
-
-
-
diff --git a/Engine/Plugins/Runtime/Metasound/Source/MetasoundGenerator/Public/SynthComponentMetasoundExample.h b/Engine/Plugins/Runtime/Metasound/Source/MetasoundGenerator/Public/SynthComponentMetasoundExample.h
--- a/Engine/Plugins/Runtime/Metasound/Source/MetasoundGenerator/Public/SynthComponentMetasoundExample.h
+++ b/Engine/Plugins/Runtime/Metasound/Source/MetasoundGenerator/Public/SynthComponentMetasoundExample.h
@@ -24,7 +24,7 @@ class METASOUNDGENERATOR_API USynthComponentMetasoundExample : public USynthComp
 
 public:
 	using FOperatorUniquePtr = TUniquePtr<Metasound::IOperator>;
-	using FMultichannelAudioFormatReadRef = Metasound::FMultichannelAudioFormatReadRef;
+	using FAudioBufferReadRef = Metasound::FAudioBufferReadRef;
 
     USynthComponentMetasoundExample(const FObjectInitializer& ObjInitializer);
 
@@ -48,20 +48,20 @@ public:
 	 * the next time CreateSoundGenerator(...) is called.
 	 *
 	 * @param InGraphOperator - Operator which executes entire MetasoundGraph.
-	 * @param InOutputAudioName - Name of the FMultichannelAudioReadRef data ref in the graph output.
+	 * @param InOutputAudioName - Name of the FAudioBufferReadRef data ref in the graph output.
 	 */
 	void SetGraphOperator(FOperatorUniquePtr InGraphOperator, const FString& InOutputAudioName);
 
 	/** Attempt to push a new graph operator to the existing metasound generator.
 	 *
 	 * @param InGraphOperator - Operator which executes entire MetasoundGraph.
-	 * @param InOutputAudioName - Name of the FMultichannelAudioReadRef data ref in the graph output.
+	 * @param InOutputAudioName - Name of the FAudioBufferReadRef data ref in the graph output.
 	 *
 	 * @return True on success, false on failure.
 	 */
 	bool PushGraphOperator(FOperatorUniquePtr InGraphOperator, const FString& InOutputAudioName);
 
-    virtual ISoundGeneratorPtr CreateSoundGenerator(int32 InSampleRate, int32 InNumChannels) override;
+    virtual ISoundGeneratorPtr CreateSoundGenerator(const FSoundGeneratorInitParams& InParams) override;
 
 #if WITH_EDITOR
 	virtual void PostEditChangeProperty(struct FPropertyChangedEvent & PropertyChangedEvent) override;
diff --git a/Engine/Plugins/Runtime/Metasound/Source/MetasoundGraphCore/Private/MetasoundGraphCoreModule.cpp b/Engine/Plugins/Runtime/Metasound/Source/MetasoundGraphCore/Private/MetasoundGraphCoreModule.cpp
--- a/Engine/Plugins/Runtime/Metasound/Source/MetasoundGraphCore/Private/MetasoundGraphCoreModule.cpp
+++ b/Engine/Plugins/Runtime/Metasound/Source/MetasoundGraphCore/Private/MetasoundGraphCoreModule.cpp
@@ -24,7 +24,7 @@ REGISTER_METASOUND_DATATYPE(double, "Primitive:Double", ::Metasound::ELiteralArg
 
 namespace Metasound
 {
-	REGISTER_METASOUND_DATATYPE(FBop, "Primitive:Bop")
+	REGISTER_METASOUND_DATATYPE(FBop, "Primitive:Bop", ::Metasound::ELiteralArgType::Boolean)
 	REGISTER_METASOUND_DATATYPE(FFloatTime, "Primitive:Time")
 	REGISTER_METASOUND_DATATYPE(FDoubleTime, "Primitive:Time:HighResolution")
 	REGISTER_METASOUND_DATATYPE(FSampleTime, "Primitive:Time:SampleResolution")
diff --git a/Engine/Plugins/Runtime/Metasound/Source/MetasoundGraphCore/Private/MetasoundOperatorBuilder.cpp b/Engine/Plugins/Runtime/Metasound/Source/MetasoundGraphCore/Private/MetasoundOperatorBuilder.cpp
--- a/Engine/Plugins/Runtime/Metasound/Source/MetasoundGraphCore/Private/MetasoundOperatorBuilder.cpp
+++ b/Engine/Plugins/Runtime/Metasound/Source/MetasoundGraphCore/Private/MetasoundOperatorBuilder.cpp
@@ -272,7 +272,8 @@ namespace Metasound
 			if (!InNodeDataReferences.Contains(InputDestination.Node))
 			{
 				AddBuildError<FMissingInputDataReferenceError>(OutErrors, InputDestination);
-				return false;
+				// This indicates an input node wasn't connected in some way to the output, and was culled earlier.
+				continue;
 			}
 
 			FDataReferenceCollection& Collection = InNodeDataReferences[InputDestination.Node].Inputs;
@@ -300,7 +301,7 @@ namespace Metasound
 			if (!InNodeDataReferences.Contains(OutputSource.Node))
 			{
 				AddBuildError<FMissingOutputDataReferenceError>(OutErrors, OutputSource);
-				return false;
+				continue;
 			}
 
 			FDataReferenceCollection& Collection = InNodeDataReferences[OutputSource.Node].Outputs;
diff --git a/Engine/Plugins/Runtime/Metasound/Source/MetasoundGraphCore/Public/MetasoundDataReferenceTypes.h b/Engine/Plugins/Runtime/Metasound/Source/MetasoundGraphCore/Public/MetasoundDataReferenceTypes.h
--- a/Engine/Plugins/Runtime/Metasound/Source/MetasoundGraphCore/Public/MetasoundDataReferenceTypes.h
+++ b/Engine/Plugins/Runtime/Metasound/Source/MetasoundGraphCore/Public/MetasoundDataReferenceTypes.h
@@ -11,6 +11,26 @@ namespace Metasound
 	class FBop
 	{
 		public:
+			// Constructor that, if called with true, bops the beginning of the callback.
+			FBop(bool bShouldBop)
+			{
+				if (bShouldBop)
+				{
+					BopFrame(0);
+				}
+			}
+
+			// Constructor that bops a specific audio frame within the callback.
+			FBop(uint32 InFrameToBop)
+			{
+				BopFrame(InFrameToBop);
+			}
+
+			// The default Bop Constructor 
+			FBop()
+				: FBop(false)
+			{}
+
 			void BopFrame(uint32 InFrameToBop)
 			{
 				// TODO: need to maintain a sorted array. 
@@ -46,9 +66,14 @@ namespace Metasound
 				return Frames[Index];
 			}
 
+			bool IsBopped() const
+			{
+				return Frames.Num() > 0;
+			}
+
 			operator bool() const 
 			{
-				return Frames.Num() >= 0;
+				return IsBopped();
 			}
 
 			void Reset()
diff --git a/Engine/Plugins/Runtime/Synthesis/Source/Synthesis/Classes/SynthComponents/SynthComponentMoto.h b/Engine/Plugins/Runtime/Synthesis/Source/Synthesis/Classes/SynthComponents/SynthComponentMoto.h
--- a/Engine/Plugins/Runtime/Synthesis/Source/Synthesis/Classes/SynthComponents/SynthComponentMoto.h
+++ b/Engine/Plugins/Runtime/Synthesis/Source/Synthesis/Classes/SynthComponents/SynthComponentMoto.h
@@ -38,7 +38,7 @@ public:
 	UFUNCTION(BlueprintCallable, Category = "MotoSynth")
 	bool IsEnabled() const;
 
-	virtual ISoundGeneratorPtr CreateSoundGenerator(int32 InSampleRate, int32 InNumChannels) override;
+	virtual ISoundGeneratorPtr CreateSoundGenerator(const FSoundGeneratorInitParams& InParams) override;
 
 private:
 	FVector2D RPMRange;
diff --git a/Engine/Plugins/Runtime/Synthesis/Source/Synthesis/Classes/SynthComponents/SynthComponentToneGenerator.h b/Engine/Plugins/Runtime/Synthesis/Source/Synthesis/Classes/SynthComponents/SynthComponentToneGenerator.h
--- a/Engine/Plugins/Runtime/Synthesis/Source/Synthesis/Classes/SynthComponents/SynthComponentToneGenerator.h
+++ b/Engine/Plugins/Runtime/Synthesis/Source/Synthesis/Classes/SynthComponents/SynthComponentToneGenerator.h
@@ -54,7 +54,7 @@ public:
 	UFUNCTION(BlueprintCallable, Category = "Tone Generator")
 	void SetVolume(float InVolume);
 
-	virtual ISoundGeneratorPtr CreateSoundGenerator(int32 InSampleRate, int32 InNumChannels) override;
+	virtual ISoundGeneratorPtr CreateSoundGenerator(const FSoundGeneratorInitParams& InParams) override;
 
 public:
 
diff --git a/Engine/Plugins/Runtime/Synthesis/Source/Synthesis/Private/SynthComponentMoto.cpp b/Engine/Plugins/Runtime/Synthesis/Source/Synthesis/Private/SynthComponentMoto.cpp
--- a/Engine/Plugins/Runtime/Synthesis/Source/Synthesis/Private/SynthComponentMoto.cpp
+++ b/Engine/Plugins/Runtime/Synthesis/Source/Synthesis/Private/SynthComponentMoto.cpp
@@ -57,7 +57,7 @@ void USynthComponentMoto::GetRPMRange(float& OutMinRPM, float& OutMaxRPM)
 	}
 }
 
-ISoundGeneratorPtr USynthComponentMoto::CreateSoundGenerator(int32 InSampleRate, int32 InNumChannels)
+ISoundGeneratorPtr USynthComponentMoto::CreateSoundGenerator(const FSoundGeneratorInitParams& InParams)
 {
 	if (!FMotoSynthEngine::IsMotoSynthEngineEnabled())
 	{
@@ -71,7 +71,7 @@ ISoundGeneratorPtr USynthComponentMoto::CreateSoundGenerator(int32 InSampleRate,
  
  		if (FMotoSynthEngine* MS = static_cast<FMotoSynthEngine*>(MotoSynthEngine.Get()))
  		{
- 			MS->Init(InSampleRate);
+ 			MS->Init(InParams.SampleRate);
  
  			FMotoSynthData AccelerationSourceData;
 			MotoSynthPreset->Settings.AccelerationSource->GetData(AccelerationSourceData);
diff --git a/Engine/Plugins/Runtime/Synthesis/Source/Synthesis/Private/SynthComponentToneGenerator.cpp b/Engine/Plugins/Runtime/Synthesis/Source/Synthesis/Private/SynthComponentToneGenerator.cpp
--- a/Engine/Plugins/Runtime/Synthesis/Source/Synthesis/Private/SynthComponentToneGenerator.cpp
+++ b/Engine/Plugins/Runtime/Synthesis/Source/Synthesis/Private/SynthComponentToneGenerator.cpp
@@ -79,7 +79,7 @@ void USynthComponentToneGenerator::SetVolume(float InVolume)
 	}
 }
 
-ISoundGeneratorPtr USynthComponentToneGenerator::CreateSoundGenerator(int32 InSampleRate, int32 InNumChannels)
+ISoundGeneratorPtr USynthComponentToneGenerator::CreateSoundGenerator(const FSoundGeneratorInitParams& InParams)
 {
-	return ToneGenerator = ISoundGeneratorPtr(new FToneGenerator(InSampleRate, InNumChannels, Frequency, Volume));
+	return ToneGenerator = ISoundGeneratorPtr(new FToneGenerator(InParams.SampleRate, InParams.NumChannels, Frequency, Volume));
 }
\ No newline at end of file
diff --git a/Engine/Source/Runtime/AudioMixer/Private/AudioMixerSourceBuffer.cpp b/Engine/Source/Runtime/AudioMixer/Private/AudioMixerSourceBuffer.cpp
--- a/Engine/Source/Runtime/AudioMixer/Private/AudioMixerSourceBuffer.cpp
+++ b/Engine/Source/Runtime/AudioMixer/Private/AudioMixerSourceBuffer.cpp
@@ -101,7 +101,11 @@ namespace Audio
 		// Retrieve a sound generator if this is a procedural sound wave
 		if (bProcedural)
 		{
-			SoundGenerator = InWave.CreateSoundGenerator(InSampleRate, NumChannels);
+			FSoundGeneratorInitParams InitParams;
+			InitParams.SampleRate = InSampleRate;
+			InitParams.NumChannels = NumChannels;
+			InitParams.NumFramesPerCallback = MONO_PCM_BUFFER_SAMPLES;
+			SoundGenerator = InWave.CreateSoundGenerator(InitParams);
 		}
 
 		const uint32 TotalSamples = MONO_PCM_BUFFER_SAMPLES * NumChannels;
@@ -335,10 +339,13 @@ namespace Audio
 			{
 				FProceduralAudioTaskData NewTaskData;
 
+				bool bIsFinished = false;
+
 				// Pass the generator instance to the async task
 				if (SoundGenerator.IsValid())
 				{
 					NewTaskData.SoundGenerator = SoundGenerator;
+					bIsFinished = SoundGenerator->IsFinished();
 				}
 				else
 				{
@@ -353,8 +360,7 @@ namespace Audio
 				check(!AsyncRealtimeAudioTask);
 				AsyncRealtimeAudioTask = CreateAudioTask(NewTaskData);
 
-				// Procedural sound waves never loop
-				return false;
+				return bIsFinished;
 			}
 
 			return true;
diff --git a/Engine/Source/Runtime/AudioMixer/Private/Components/SynthComponent.cpp b/Engine/Source/Runtime/AudioMixer/Private/Components/SynthComponent.cpp
--- a/Engine/Source/Runtime/AudioMixer/Private/Components/SynthComponent.cpp
+++ b/Engine/Source/Runtime/AudioMixer/Private/Components/SynthComponent.cpp
@@ -104,11 +104,11 @@ void USynthSound::OnEndGenerate()
 	}
 }
 
-ISoundGeneratorPtr USynthSound::CreateSoundGenerator(int32 InSampleRate, int32 InNumChannels)
+ISoundGeneratorPtr USynthSound::CreateSoundGenerator(const FSoundGeneratorInitParams& InParams)
 {
 	if (OwningSynthComponent)
 	{
-		return OwningSynthComponent->CreateSoundGeneratorInternal(SampleRate, NumChannels);
+		return OwningSynthComponent->CreateSoundGeneratorInternal(InParams);
 	}
 	return nullptr;
 }
@@ -550,7 +550,7 @@ void USynthComponent::SynthCommand(TFunction<void()> Command)
 	}
 }
 
-ISoundGeneratorPtr USynthComponent::CreateSoundGeneratorInternal(int32 InSampleRate, int32 InNumChannels)
+ISoundGeneratorPtr USynthComponent::CreateSoundGeneratorInternal(const FSoundGeneratorInitParams& InParams)
 {	
-	return SoundGenerator = CreateSoundGenerator(InSampleRate, InNumChannels);
+	return SoundGenerator = CreateSoundGenerator(InParams);
 }
diff --git a/Engine/Source/Runtime/AudioMixer/Public/Components/SynthComponent.h b/Engine/Source/Runtime/AudioMixer/Public/Components/SynthComponent.h
--- a/Engine/Source/Runtime/AudioMixer/Public/Components/SynthComponent.h
+++ b/Engine/Source/Runtime/AudioMixer/Public/Components/SynthComponent.h
@@ -51,7 +51,7 @@ class AUDIOMIXER_API USynthSound : public USoundWaveProcedural
 	virtual int32 OnGeneratePCMAudio(TArray<uint8>& OutAudio, int32 NumSamples) override;
 	virtual void OnEndGenerate() override;
 	virtual Audio::EAudioMixerStreamDataFormat::Type GetGeneratedPCMDataFormat() const override;
-	virtual ISoundGeneratorPtr CreateSoundGenerator(int32 InSampleRate, int32 InNumChannels) override;
+	virtual ISoundGeneratorPtr CreateSoundGenerator(const FSoundGeneratorInitParams& InParams) override;
 	/** End USoundWave */
 
 protected:
@@ -247,7 +247,7 @@ protected:
 
 	// Implemented by the synth component to create a generator object instead of generating audio directly on the synth component.
 	// This method prevents UObjects from having to exist in the audio render thread.
-	virtual ISoundGeneratorPtr CreateSoundGenerator(int32 InSampleRate, int32 InNumChannels) { return nullptr; }
+	virtual ISoundGeneratorPtr CreateSoundGenerator(const FSoundGeneratorInitParams& InParams) { return nullptr; }
 
 	// Called by procedural sound wave
 	// Returns the number of samples actually generated
@@ -264,7 +264,7 @@ protected:
 
 private:
 	// Creates the synth component's sound generator, calls into overridden client code to create the instance.
-	ISoundGeneratorPtr CreateSoundGeneratorInternal(int32 InSampleRate, int32 InNumChannels);
+	ISoundGeneratorPtr CreateSoundGeneratorInternal(const FSoundGeneratorInitParams& InParams);
 
 	UPROPERTY(Transient)
 	USynthSound* Synth;
diff --git a/Engine/Source/Runtime/Engine/Classes/Sound/SoundBase.h b/Engine/Source/Runtime/Engine/Classes/Sound/SoundBase.h
--- a/Engine/Source/Runtime/Engine/Classes/Sound/SoundBase.h
+++ b/Engine/Source/Runtime/Engine/Classes/Sound/SoundBase.h
@@ -260,7 +260,7 @@ public:
 	virtual bool HasCookedAmplitudeEnvelopeData() const { return false; }
 
 	/** Creates a sound generator instance from this sound base. Return true if this is being implemented by a subclass. Sound generators procedurally generate audio in the audio render thread. */
-	virtual ISoundGeneratorPtr CreateSoundGenerator(int32 InSampleRate, int32 InNumChannels) { return nullptr; }
+	virtual ISoundGeneratorPtr CreateSoundGenerator(const FSoundGeneratorInitParams& InParams) { return nullptr; }
 
 };
 
diff --git a/Engine/Source/Runtime/Engine/Classes/Sound/SoundGenerator.h b/Engine/Source/Runtime/Engine/Classes/Sound/SoundGenerator.h
--- a/Engine/Source/Runtime/Engine/Classes/Sound/SoundGenerator.h
+++ b/Engine/Source/Runtime/Engine/Classes/Sound/SoundGenerator.h
@@ -6,6 +6,21 @@
 #include "Containers/Queue.h"
 #include "Templates/Function.h"
 
+
+// Parameters used for constructing a new ISoundGenerator.
+struct FSoundGeneratorInitParams
+{
+	float SampleRate;
+	int32 NumChannels;
+	int32 NumFramesPerCallback;
+
+	FSoundGeneratorInitParams()
+		: SampleRate(0.0f)
+		, NumChannels(0)
+		, NumFramesPerCallback(0)
+	{}
+};
+
 class ENGINE_API ISoundGenerator
 {
 public:
@@ -24,6 +39,9 @@ public:
 	// Optional. Called on audio generator thread right when the generator ends generating.
 	virtual void OnEndGenerate() {}
 
+	// Optional. Can be overridden to end the sound when generating is finished.
+	virtual bool IsFinished() const { return false; }
+
 	// Retrieves the next buffer of audio from the generator, called from the audio mixer
 	int32 GetNextBuffer(float* OutAudio, int32 NumSamples, bool bRequireNumberSamples = false);
 
