- Add frontend-side connection support in MetasoundEditor
- Swap I/O detail customization groups for categories

diff --git a/Engine/Plugins/Runtime/Metasound/Source/MetasoundEditor/Private/MetasoundDetailCustomization.cpp b/Engine/Plugins/Runtime/Metasound/Source/MetasoundEditor/Private/MetasoundDetailCustomization.cpp
--- a/Engine/Plugins/Runtime/Metasound/Source/MetasoundEditor/Private/MetasoundDetailCustomization.cpp
+++ b/Engine/Plugins/Runtime/Metasound/Source/MetasoundEditor/Private/MetasoundDetailCustomization.cpp
@@ -27,15 +27,13 @@ void FMetasoundDetailCustomization::CustomizeDetails(IDetailLayoutBuilder& Detai
 	GeneralCategoryBuilder.AddProperty(AuthorHandle);
 	GeneralCategoryBuilder.AddProperty(DescHandle);
 
-	IDetailCategoryBuilder& ParametersCategoryBuilder = DetailLayout.EditCategory("Parameters");
-
-	IDetailGroup& InputDetailGroup = ParametersCategoryBuilder.AddGroup("Inputs", LOCTEXT("MetasoundDetailsGroupInputs", "Inputs"), false /* bForAdvanced */, true /* bStartExpanded */);
+	IDetailCategoryBuilder& InputsCategoryBuilder = DetailLayout.EditCategory("Inputs");
 	TSharedPtr<IPropertyHandle> InputsHandle = DetailLayout.GetProperty("RootMetasoundDocument.RootClass.Inputs");
-	InputDetailGroup.AddPropertyRow(InputsHandle.ToSharedRef());
+	InputsCategoryBuilder.AddProperty(InputsHandle);
 
-	IDetailGroup& OutputDetailGroup = ParametersCategoryBuilder.AddGroup("Inputs", LOCTEXT("MetasoundDetailsGroupOutputs", "Outputs"), false /* bForAdvanced */, true /* bStartExpanded */);
+	IDetailCategoryBuilder& OutputsCategoryBuilder = DetailLayout.EditCategory("Outputs");
 	TSharedPtr<IPropertyHandle> OutputsHandle = DetailLayout.GetProperty("RootMetasoundDocument.RootClass.Outputs");
-	OutputDetailGroup.AddPropertyRow(OutputsHandle.ToSharedRef());
+	OutputsCategoryBuilder.AddProperty(OutputsHandle);
 
 	// Hack to hide parent structs for nested metadata properties
 	DetailLayout.HideCategory("Hidden");
diff --git a/Engine/Plugins/Runtime/Metasound/Source/MetasoundEditor/Private/MetasoundEditorGraphSchema.cpp b/Engine/Plugins/Runtime/Metasound/Source/MetasoundEditor/Private/MetasoundEditorGraphSchema.cpp
--- a/Engine/Plugins/Runtime/Metasound/Source/MetasoundEditor/Private/MetasoundEditorGraphSchema.cpp
+++ b/Engine/Plugins/Runtime/Metasound/Source/MetasoundEditor/Private/MetasoundEditorGraphSchema.cpp
@@ -427,14 +427,38 @@ const FPinConnectionResponse UMetasoundEditorGraphSchema::CanCreateConnection(co
 
 bool UMetasoundEditorGraphSchema::TryCreateConnection(UEdGraphPin* PinA, UEdGraphPin* PinB) const
 {
+	UEdGraphPin* InputPin = nullptr;
+	UEdGraphPin* OutputPin = nullptr;
+	if (!CategorizePinsByDirection(PinA, PinB, InputPin, OutputPin))
+	{
+		return false;
+	}
+
 	const bool bModified = UEdGraphSchema::TryCreateConnection(PinA, PinB);
+	if (bModified)
+	{
+		UMetasoundEditorGraphNode* InputGraphNode = CastChecked<UMetasoundEditorGraphNode>(InputPin->GetOwningNode());
+		Metasound::Frontend::FNodeHandle InputNodeHandle = InputGraphNode->GetNodeHandle();
+		Metasound::Frontend::FInputHandle InputHandle = InputNodeHandle.GetInputWithName(InputPin->GetName());
+
+		UMetasoundEditorGraphNode* OutputGraphNode = CastChecked<UMetasoundEditorGraphNode>(OutputPin->GetOwningNode());
+		Metasound::Frontend::FNodeHandle OutputNodeHandle = OutputGraphNode->GetNodeHandle();
+		Metasound::Frontend::FOutputHandle OutputHandle = OutputNodeHandle.GetOutputWithName(OutputPin->GetName());
+
+		// TODO: Implement YesWithConverterNode with selected conversion option
+		if (!ensure(InputHandle.Connect(OutputHandle)))
+		{
+			InputPin->BreakLinkTo(PinB);
+			return false;
+		}
+	}
 	return bModified;
 }
 
 bool UMetasoundEditorGraphSchema::ShouldHidePinDefaultValue(UEdGraphPin* Pin) const
 {
 	// TODO: Determine if should be hidden from doc data
-	return true;
+	return false;
 }
 
 FLinearColor UMetasoundEditorGraphSchema::GetPinTypeColor(const FEdGraphPinType& PinType) const
@@ -445,18 +469,58 @@ FLinearColor UMetasoundEditorGraphSchema::GetPinTypeColor(const FEdGraphPinType&
 
 void UMetasoundEditorGraphSchema::BreakNodeLinks(UEdGraphNode& TargetNode) const
 {
+	using namespace Metasound::Frontend;
+
+	const FScopedTransaction Transaction(NSLOCTEXT("UnrealEd", "GraphEd_BreakPinLinks", "Break Pin Links"));
+
 	Super::BreakNodeLinks(TargetNode);
 
-	// TODO: Update MS document
+	FNodeHandle NodeHandle = CastChecked<UMetasoundEditorGraphNode>(&TargetNode)->GetNodeHandle();
+	const uint32 NodeID = NodeHandle.GetNodeID();
+
+	FGraphHandle GraphHandle = CastChecked<UMetasoundEditorGraphNode>(&TargetNode)->GetRootGraphHandle();
+	TArray<FNodeHandle> AllNodes = GraphHandle.GetAllNodes();
+	for (Metasound::Frontend::FNodeHandle& IterNode : AllNodes)
+	{
+		if (NodeID != IterNode.GetNodeID())
+		{
+			TArray<FInputHandle> Inputs = IterNode.GetAllInputs();
+			for (FInputHandle& Input : Inputs)
+			{
+				FOutputHandle Output = Input.GetCurrentlyConnectedOutput();
+				if (Output.GetOwningNodeID() == NodeID)
+				{
+					Input.Disconnect(Output);
+				}
+			}
+		}
+	}
 }
 
 void UMetasoundEditorGraphSchema::BreakPinLinks(UEdGraphPin& TargetPin, bool bSendsNodeNotifcation) const
 {
+	using namespace Metasound::Frontend;
+
 	const FScopedTransaction Transaction(NSLOCTEXT("UnrealEd", "GraphEd_BreakPinLinks", "Break Pin Links"));
 
-	Super::BreakPinLinks(TargetPin, bSendsNodeNotifcation);
+	if (TargetPin.Direction == EGPD_Input)
+	{
+		FNodeHandle NodeHandle = CastChecked<UMetasoundEditorGraphNode>(TargetPin.GetOwningNode())->GetNodeHandle();
+		FInputHandle InputHandle = NodeHandle.GetInputWithName(TargetPin.GetName());
+		InputHandle.Disconnect();
+	}
+	else
+	{
+		check(TargetPin.Direction == EGPD_Output);
+		for (UEdGraphPin* Pin : TargetPin.LinkedTo)
+		{
+			FNodeHandle NodeHandle = CastChecked<UMetasoundEditorGraphNode>(Pin->GetOwningNode())->GetNodeHandle();
+			FInputHandle InputHandle = NodeHandle.GetInputWithName(Pin->GetName());
+			InputHandle.Disconnect();
+		}
+	}
 
-	// TODO: Update MS document
+	Super::BreakPinLinks(TargetPin, bSendsNodeNotifcation);
 }
 
 void UMetasoundEditorGraphSchema::GetAssetsGraphHoverMessage(const TArray<FAssetData>& Assets, const UEdGraph* HoverGraph, FString& OutTooltipText, bool& bOutOkIcon) const
diff --git a/Engine/Plugins/Runtime/Metasound/Source/MetasoundFrontend/Private/MetasoundFrontend.cpp b/Engine/Plugins/Runtime/Metasound/Source/MetasoundFrontend/Private/MetasoundFrontend.cpp
--- a/Engine/Plugins/Runtime/Metasound/Source/MetasoundFrontend/Private/MetasoundFrontend.cpp
+++ b/Engine/Plugins/Runtime/Metasound/Source/MetasoundFrontend/Private/MetasoundFrontend.cpp
@@ -604,6 +604,58 @@ namespace Metasound
 			return true;
 		}
 
+		bool FInputHandle::Disconnect()
+		{
+			if (!IsValid())
+			{
+				return false;
+			}
+
+			TArray<FMetasoundNodeConnectionDescription>& Connections = NodePtr->InputConnections;
+			for (int32 i = 0; i < Connections.Num(); ++i)
+			{
+				if (Connections[i].InputName == GetInputName())
+				{
+					Connections.RemoveAtSwap(i);
+					return true;
+				}
+			}
+
+			return false;
+		}
+
+		bool FInputHandle::Disconnect(FOutputHandle& InHandle)
+		{
+			if (!IsValid() || !InHandle.IsValid())
+			{
+				return false;
+			}
+
+			if (!ensureAlwaysMsgf(InHandle.GetOutputType() == GetInputType(), TEXT("Tried to disconnect incompatible types!")))
+			{
+				return false;
+			}
+
+
+			FMetasoundNodeConnectionDescription* Connection = GetConnectionDescription();
+			if (!ensure(Connection))
+			{
+				return false;
+			}
+
+			const uint32 OutputNodeID = InHandle.GetOwningNodeID();
+			for (int32 i = 0; i < NodePtr->InputConnections.Num(); ++i)
+			{
+				if (NodePtr->InputConnections[i].NodeID == OutputNodeID)
+				{
+					NodePtr->InputConnections.RemoveAtSwap(i);
+					return true;
+				}
+			}
+
+			return false;
+		}
+
 		bool FInputHandle::ConnectWithConverterNode(FOutputHandle& InHandle, FString& InNodeClassName)
 		{
 			// (UEAU-473)
@@ -789,6 +841,16 @@ namespace Metasound
 			return InHandle.ConnectWithConverterNode(*this, InNodeClassName);
 		}
 
+		bool FOutputHandle::Disconnect(FInputHandle& InHandle)
+		{
+			if (!IsValid() || !InHandle.IsValid())
+			{
+				return false;
+			}
+
+			return InHandle.Disconnect(*this);
+		}
+
 		TDescriptionPtr<FMetasoundClassDescription> FNodeHandle::GetNodeClassDescriptionForNodeHandle(const FHandleInitParams& InitParams, EMetasoundClassType InNodeClassType)
 		{
 			if (InNodeClassType != EMetasoundClassType::Input && InNodeClassType != EMetasoundClassType::Output)
diff --git a/Engine/Plugins/Runtime/Metasound/Source/MetasoundFrontend/Public/MetasoundFrontend.h b/Engine/Plugins/Runtime/Metasound/Source/MetasoundFrontend/Public/MetasoundFrontend.h
--- a/Engine/Plugins/Runtime/Metasound/Source/MetasoundFrontend/Public/MetasoundFrontend.h
+++ b/Engine/Plugins/Runtime/Metasound/Source/MetasoundFrontend/Public/MetasoundFrontend.h
@@ -211,6 +211,7 @@ namespace Metasound
 			FConnectability CanConnectTo(const FInputHandle& InHandle) const;
 			bool Connect(FInputHandle& InHandle);
 			bool ConnectWithConverterNode(FInputHandle& InHandle, FString& InNodeClassName);
+			bool Disconnect(FInputHandle& InHandle);
 
 		private:
 			TDescriptionPtr<FMetasoundNodeDescription> NodePtr;
@@ -250,6 +251,8 @@ namespace Metasound
 			FConnectability CanConnectTo(const FOutputHandle& InHandle) const;
 			bool Connect(FOutputHandle& InHandle);
 			bool ConnectWithConverterNode(FOutputHandle& InHandle, FString& InNodeClassName);
+			bool Disconnect(FOutputHandle& InHandle);
+			bool Disconnect();
 
 		private:
 			TDescriptionPtr<FMetasoundNodeDescription> NodePtr;
