IWYU Fixes
#rnx

diff --git a/Engine/Plugins/Runtime/Metasound/Source/MetasoundEngine/Private/MetasoundEngineModule.cpp b/Engine/Plugins/Runtime/Metasound/Source/MetasoundEngine/Private/MetasoundEngineModule.cpp
--- a/Engine/Plugins/Runtime/Metasound/Source/MetasoundEngine/Private/MetasoundEngineModule.cpp
+++ b/Engine/Plugins/Runtime/Metasound/Source/MetasoundEngine/Private/MetasoundEngineModule.cpp
@@ -2,6 +2,7 @@
 #include "MetasoundEngineModule.h"
 
 #include "Metasound.h"
+#include "MetasoundArchetypeRegistration.h"
 #include "MetasoundSource.h"
 #include "Modules/ModuleManager.h"
 
diff --git a/Engine/Plugins/Runtime/Metasound/Source/MetasoundEngine/Public/MetasoundArchetypeRegistration.h b/Engine/Plugins/Runtime/Metasound/Source/MetasoundEngine/Public/MetasoundArchetypeRegistration.h
--- a/Engine/Plugins/Runtime/Metasound/Source/MetasoundEngine/Public/MetasoundArchetypeRegistration.h
+++ b/Engine/Plugins/Runtime/Metasound/Source/MetasoundEngine/Public/MetasoundArchetypeRegistration.h
@@ -2,12 +2,13 @@
 #pragma once
 
 #include "CoreMinimal.h"
-#include "MetasoundFrontend.h"
 #include "MetasoundAssetBase.h"
+#include "MetasoundFrontend.h"
 
 #if WITH_EDITOR
 #include "AssetRegistryModule.h"
-#endif
+#include "Async/TaskGraphInterfaces.h"
+#endif // WITH_EDITOR
 
 namespace Metasound
 {
@@ -66,7 +67,8 @@ namespace Metasound
 					NewMetasoundObject->SetDocument(InDocument);
 
 #if WITH_EDITOR
-					AsyncTask(ENamedThreads::GameThread, [NewMetasoundObject]() {
+					AsyncTask(ENamedThreads::GameThread, [NewMetasoundObject]()
+					{
 						FAssetRegistryModule::AssetCreated(NewMetasoundObject);
 						NewMetasoundObject->MarkPackageDirty();
 						// todo: how do you get the package for a uobject and save it? I forget
