Another IWYU fix
#rnx

diff --git a/Engine/Plugins/Runtime/Metasound/Source/MetasoundEngine/Public/MetasoundArchetypeRegistration.h b/Engine/Plugins/Runtime/Metasound/Source/MetasoundEngine/Public/MetasoundArchetypeRegistration.h
--- a/Engine/Plugins/Runtime/Metasound/Source/MetasoundEngine/Public/MetasoundArchetypeRegistration.h
+++ b/Engine/Plugins/Runtime/Metasound/Source/MetasoundEngine/Public/MetasoundArchetypeRegistration.h
@@ -7,6 +7,7 @@
 
 #if WITH_EDITOR
 #include "AssetRegistryModule.h"
+#include "Async/Async.h"
 #include "Async/TaskGraphInterfaces.h"
 #endif // WITH_EDITOR
 
