Update metasound operator settings to enforce SIMD alignment on frame counts

diff --git a/Engine/Plugins/Runtime/Metasound/Source/MetasoundGraphCore/Private/MetasoundADSRNode.cpp b/Engine/Plugins/Runtime/Metasound/Source/MetasoundGraphCore/Private/MetasoundADSRNode.cpp
--- a/Engine/Plugins/Runtime/Metasound/Source/MetasoundGraphCore/Private/MetasoundADSRNode.cpp
+++ b/Engine/Plugins/Runtime/Metasound/Source/MetasoundGraphCore/Private/MetasoundADSRNode.cpp
@@ -30,13 +30,13 @@ namespace Metasound
 			,	Bop(InBop)
 			,	StopEnvelopePos(-1)
 			,	ADSRDataReferences(InADSRData)
-			,	EnvelopeBuffer(FAudioBufferWriteRef::CreateNew(InSettings.FramesPerExecute))
+			,	EnvelopeBuffer(FAudioBufferWriteRef::CreateNew(InSettings.GetNumFramesPerBlock()))
 			{
-				check(EnvelopeBuffer->Num() == InSettings.FramesPerExecute);
+				check(EnvelopeBuffer->Num() == InSettings.GetNumFramesPerBlock());
 
 				OutputDataReferences.AddDataReadReference(TEXT("Envelope"), FAudioBufferReadRef(EnvelopeBuffer));
 
-				Envelope.Init(OperatorSettings.SampleRate);
+				Envelope.Init(OperatorSettings.GetSampleRate());
 			}
 
 			virtual const FDataReferenceCollection& GetInputs() const override
@@ -66,20 +66,20 @@ namespace Metasound
 
 				int32 BopIndex = 0;
 				int32 StartPos = 0;
-				int32 NextBop = BopIndex < Bop->Num() ? (*Bop)[BopIndex] : OperatorSettings.FramesPerExecute + 1;
-				int32 EndPos = FMath::Min(OperatorSettings.FramesPerExecute, NextBop);
+				int32 NextBop = BopIndex < Bop->Num() ? (*Bop)[BopIndex] : OperatorSettings.GetNumFramesPerBlock() + 1;
+				int32 EndPos = FMath::Min(OperatorSettings.GetNumFramesPerBlock(), NextBop);
 				
 				if (StopEnvelopePos > 0)
 				{
-					StopEnvelopePos -= OperatorSettings.FramesPerExecute;
+					StopEnvelopePos -= OperatorSettings.GetNumFramesPerBlock();
 					EndPos = FMath::Min(StopEnvelopePos, EndPos);
 				}
 
 				float* EnvelopeData = EnvelopeBuffer->GetData();
 
-				// TODO: FramesPerExecute need to be on 4 sample boundaries to allow for SIMD.
-				// TODO: just make FramesPerExecute an int32 instead of a uint32
-				while (StartPos < (int32)OperatorSettings.FramesPerExecute)
+				// TODO: GetNumFramesPerBlock() need to be on 4 sample boundaries to allow for SIMD.
+				// TODO: just make GetNumFramesPerBlock() an int32 instead of a uint32
+				while (StartPos < (int32)OperatorSettings.GetNumFramesPerBlock())
 				{
 					int32 i = StartPos;
 					for (; i < EndPos; i++)
@@ -94,19 +94,19 @@ namespace Metasound
 						// Stop sustain, start release.
 						Envelope.Stop();
 						StopEnvelopePos = -1;
-						EndPos = FMath::Min((int32)OperatorSettings.FramesPerExecute, NextBop);
+						EndPos = FMath::Min((int32)OperatorSettings.GetNumFramesPerBlock(), NextBop);
 					}
 
 					if (EndPos == NextBop)
 					{
 						// Trigger next envelope. 
 						Envelope.Start();
-						StopEnvelopePos = EndPos + FMath::RoundToInt(OperatorSettings.SampleRate * 0.001f * ADSMilliseconds);
+						StopEnvelopePos = EndPos + FMath::RoundToInt(OperatorSettings.GetSampleRate() * 0.001f * ADSMilliseconds);
 
 						BopIndex++;
 
-						NextBop = BopIndex < Bop->Num() ? (*Bop)[BopIndex] : OperatorSettings.FramesPerExecute + 1;
-						EndPos = FMath::Min(StopEnvelopePos, FMath::Min((int32)OperatorSettings.FramesPerExecute, NextBop));
+						NextBop = BopIndex < Bop->Num() ? (*Bop)[BopIndex] : OperatorSettings.GetNumFramesPerBlock() + 1;
+						EndPos = FMath::Min(StopEnvelopePos, FMath::Min((int32)OperatorSettings.GetNumFramesPerBlock(), NextBop));
 					}
 				}
 			}
diff --git a/Engine/Plugins/Runtime/Metasound/Source/MetasoundGraphCore/Private/MetasoundAudioFormats.cpp b/Engine/Plugins/Runtime/Metasound/Source/MetasoundGraphCore/Private/MetasoundAudioFormats.cpp
--- a/Engine/Plugins/Runtime/Metasound/Source/MetasoundGraphCore/Private/MetasoundAudioFormats.cpp
+++ b/Engine/Plugins/Runtime/Metasound/Source/MetasoundGraphCore/Private/MetasoundAudioFormats.cpp
@@ -53,7 +53,7 @@ namespace Metasound
 	}
 
 	FUnformattedAudio::FUnformattedAudio(int32 InInitialNumChannels, const FOperatorSettings& InSettings)
-		: FUnformattedAudio(InSettings.FramesPerExecute, InInitialNumChannels, 8)
+		: FUnformattedAudio(InSettings.GetNumFramesPerBlock(), InInitialNumChannels, 8)
 	{
 	}
 
@@ -95,7 +95,7 @@ namespace Metasound
 	}
 
 	FMultichannelAudioFormat::FMultichannelAudioFormat(int32 InNumChannels, const FOperatorSettings& InSettings)
-		: FMultichannelAudioFormat(InSettings.FramesPerExecute, InNumChannels)
+		: FMultichannelAudioFormat(InSettings.GetNumFramesPerBlock(), InNumChannels)
 	{}
 
 
diff --git a/Engine/Plugins/Runtime/Metasound/Source/MetasoundGraphCore/Private/MetasoundAudioMultiplyNode.cpp b/Engine/Plugins/Runtime/Metasound/Source/MetasoundGraphCore/Private/MetasoundAudioMultiplyNode.cpp
--- a/Engine/Plugins/Runtime/Metasound/Source/MetasoundGraphCore/Private/MetasoundAudioMultiplyNode.cpp
+++ b/Engine/Plugins/Runtime/Metasound/Source/MetasoundGraphCore/Private/MetasoundAudioMultiplyNode.cpp
@@ -20,9 +20,9 @@ namespace Metasound
 			,	InputBuffer2(InBuffer2)
 			,	OutputBuffer(FAudioBufferWriteRef::CreateNew(InSettings))
 			{
-				check(OutputBuffer->Num() == InSettings.FramesPerExecute);
-				check(InputBuffer1->Num() == InSettings.FramesPerExecute);
-				check(InputBuffer2->Num() == InSettings.FramesPerExecute);
+				check(OutputBuffer->Num() == InSettings.GetNumFramesPerBlock());
+				check(InputBuffer1->Num() == InSettings.GetNumFramesPerBlock());
+				check(InputBuffer2->Num() == InSettings.GetNumFramesPerBlock());
 
 				OutputDataReferences.AddDataReadReference(TEXT("Audio"), FAudioBufferReadRef(OutputBuffer));
 			}
@@ -43,7 +43,7 @@ namespace Metasound
 
 				// TODO: add buffer vector op to multipy two buffers (not in place). 
 				// TODO: what to do about RESTRICT? This could be something for the builder to keep in mind.
-				FMemory::Memcpy(OutputBuffer->GetData(), InputBuffer1->GetData(), sizeof(float) * OperatorSettings.FramesPerExecute);
+				FMemory::Memcpy(OutputBuffer->GetData(), InputBuffer1->GetData(), sizeof(float) * OperatorSettings.GetNumFramesPerBlock());
 				MultiplyBuffersInPlace(*InputBuffer2, *OutputBuffer);
 			}
 
@@ -66,11 +66,11 @@ namespace Metasound
 		FAudioBufferReadRef InputBuffer1 = FAudioBufferReadRef::CreateNew();
 		FAudioBufferReadRef InputBuffer2 = FAudioBufferReadRef::CreateNew();
 
-		FAudioBuffer Ones(InOperatorSettings.FramesPerExecute);
+		FAudioBuffer Ones(InOperatorSettings.GetNumFramesPerBlock());
 
 		// initialize default array to all ones. 
 		float* Data = Ones.GetData();
-		for (int32 i = 0; i < InOperatorSettings.FramesPerExecute; i++)
+		for (int32 i = 0; i < InOperatorSettings.GetNumFramesPerBlock(); i++)
 		{
 			Data[i] = 1.f;
 		}
diff --git a/Engine/Plugins/Runtime/Metasound/Source/MetasoundGraphCore/Private/MetasoundOperatorSettings.cpp b/Engine/Plugins/Runtime/Metasound/Source/MetasoundGraphCore/Private/MetasoundOperatorSettings.cpp
new file mode 100644
--- /dev/null
+++ b/Engine/Plugins/Runtime/Metasound/Source/MetasoundGraphCore/Private/MetasoundOperatorSettings.cpp
@@ -0,0 +1,107 @@
+// Copyright Epic Games, Inc. All Rights Reserved.
+
+#include "MetasoundOperatorSettings.h"
+
+#include "DSP/BufferVectorOperations.h"
+#include "CoreMinimal.h"
+
+namespace Metasound
+{
+	namespace MetasoundOperatorSettingsPrivate
+	{
+		// Fall back to this sample rate if an invalid sample rate is given.
+		static constexpr float DefaultSampleRate = 48000.f;
+		static constexpr float MinimumTargetBlockRate = 1.0f;
+
+		// Using same alignment as in BufferVectorOperations.h
+		static constexpr int32 FloatAlignment = AUDIO_SIMD_FLOAT_ALIGNMENT;
+
+		static constexpr int32 MinimumNumFramesPerBlock = FloatAlignment;
+		static constexpr int32 MaximumNumFramesPerBlock = 16384;
+	}
+
+	FOperatorSettings::FOperatorSettings(float InSampleRate, float InTargetBlockRate)
+	{
+		SetSampleRate(InSampleRate);
+		SetTargetBlockRate(InTargetBlockRate);
+	}
+
+	void FOperatorSettings::SetSampleRate(float InSampleRate)
+	{
+		using namespace MetasoundOperatorSettingsPrivate;
+
+		if (ensureMsgf(InSampleRate > 0.f, TEXT("Invalid sample rate %f. Defaulting to %f"), InSampleRate, DefaultSampleRate))
+		{
+			SampleRate = InSampleRate;
+		}
+		else
+		{
+			SampleRate = DefaultSampleRate;
+		}
+
+		// Update frames per block and actual block rate.
+		Update(); 
+	}
+
+	void FOperatorSettings::SetTargetBlockRate(float InTargetBlockRate)
+	{
+		using namespace MetasoundOperatorSettingsPrivate;
+
+		if (ensureMsgf(InTargetBlockRate > MinimumTargetBlockRate, TEXT("Invalid block rate %f. Defaulting to %f"), InTargetBlockRate, MinimumTargetBlockRate))
+		{
+			TargetBlockRate = InTargetBlockRate;
+		}
+		else
+		{
+			TargetBlockRate = MinimumTargetBlockRate;
+		}
+
+		// Update frames per block and actual block rate.
+		Update();
+	}
+
+	float FOperatorSettings::GetTargetBlockRate() const
+	{
+		return TargetBlockRate;
+	}
+
+	float FOperatorSettings::GetActualBlockRate() const
+	{
+		return ActualBlockRate;
+	}
+
+	void FOperatorSettings::Update()
+	{
+		using namespace MetasoundOperatorSettingsPrivate;
+
+		// Protects for divide by zero.
+		check(MinimumNumFramesPerBlock > 0);
+
+
+		float TargetNumFramesPerBlock = SampleRate / FMath::Max(TargetBlockRate, MinimumTargetBlockRate);
+
+		NumFramesPerBlock = RoundToAligned(FloatAlignment, FMath::RoundToInt(TargetNumFramesPerBlock));
+
+		// The minimum and maximum number of frames per a block must also satisfy alignment 
+		// requirements since we clamp to them after enforcing alignment.
+		checkf(RoundToAligned(FloatAlignment, MinimumNumFramesPerBlock) == MinimumNumFramesPerBlock, TEXT("MinimumNumFramesPerBlock (%d) is not aligned to required buffer alignment (%d)"), MinimumNumFramesPerBlock, FloatAlignment);
+		checkf(RoundToAligned(FloatAlignment, MaximumNumFramesPerBlock) == MaximumNumFramesPerBlock, TEXT("MaximumNumFramesPerBlock (%d) is not aligned to required buffer alignment (%d)"), MaximumNumFramesPerBlock, FloatAlignment);
+
+		NumFramesPerBlock = FMath::Clamp(NumFramesPerBlock, MinimumNumFramesPerBlock, MaximumNumFramesPerBlock);
+
+		ActualBlockRate = SampleRate / static_cast<float>(NumFramesPerBlock);
+	}
+
+	int32 FOperatorSettings::RoundToAligned(int32 Alignment, int32 InNum) const
+	{
+		check(Alignment >= 0);
+
+		// Need to add (Alignment / 2) in order to perform rounding operation.
+		// Otherwise, this function would be more like "FloorToAligned(...)"
+		int32 AlignedNum = InNum + (Alignment / 2);
+
+		AlignedNum = AlignedNum - (AlignedNum % Alignment);
+
+		return AlignedNum;
+	}
+}
diff --git a/Engine/Plugins/Runtime/Metasound/Source/MetasoundGraphCore/Private/MetasoundOscNode.cpp b/Engine/Plugins/Runtime/Metasound/Source/MetasoundGraphCore/Private/MetasoundOscNode.cpp
--- a/Engine/Plugins/Runtime/Metasound/Source/MetasoundGraphCore/Private/MetasoundOscNode.cpp
+++ b/Engine/Plugins/Runtime/Metasound/Source/MetasoundGraphCore/Private/MetasoundOscNode.cpp
@@ -20,7 +20,7 @@ namespace Metasound
 			,	Frequency(InFrequency)
 			,	AudioBuffer(FAudioBufferWriteRef::CreateNew(InSettings))
 			{
-				check(AudioBuffer->Num() == InSettings.FramesPerExecute);
+				check(AudioBuffer->Num() == InSettings.GetNumFramesPerBlock());
 
 				OutputDataReferences.AddDataReadReference(TEXT("Audio"), FAudioBufferReadRef(AudioBuffer));
 			}
@@ -37,10 +37,10 @@ namespace Metasound
 
 			void Execute()
 			{
-				const float PhaseDelta = Frequency->GetRadiansPerSample(OperatorSettings.SampleRate);
+				const float PhaseDelta = Frequency->GetRadiansPerSample(OperatorSettings.GetSampleRate());
 				float* Data = AudioBuffer->GetData();
 
-				for (int32 i = 0; i < OperatorSettings.FramesPerExecute; i++)
+				for (int32 i = 0; i < OperatorSettings.GetNumFramesPerBlock(); i++)
 				{
 					Data[i] = FMath::Sin(Phase);
 					Phase += PhaseDelta;
diff --git a/Engine/Plugins/Runtime/Metasound/Source/MetasoundGraphCore/Private/MetasoundPeriodicBopNode.cpp b/Engine/Plugins/Runtime/Metasound/Source/MetasoundGraphCore/Private/MetasoundPeriodicBopNode.cpp
--- a/Engine/Plugins/Runtime/Metasound/Source/MetasoundGraphCore/Private/MetasoundPeriodicBopNode.cpp
+++ b/Engine/Plugins/Runtime/Metasound/Source/MetasoundGraphCore/Private/MetasoundPeriodicBopNode.cpp
@@ -27,7 +27,7 @@ namespace Metasound
 			:	OperatorSettings(InSettings)
 			,	Bop(FBopWriteRef::CreateNew())
 			,	Period(InPeriod)
-			,	ExecuteDurationInSamples(InSettings.FramesPerExecute)
+			,	ExecuteDurationInSamples(InSettings.GetNumFramesPerBlock())
 			,	SampleCountdown(0.f)
 			{
 				OutputDataReferences.AddDataReadReference(TEXT("Bop"), FBopReadRef(Bop));
@@ -46,9 +46,9 @@ namespace Metasound
 			void Execute()
 			{
 				// Advance internal counter to get rid of old bops.
-				Bop->Advance(OperatorSettings.FramesPerExecute);
+				Bop->Advance(OperatorSettings.GetNumFramesPerBlock());
 
-				float PeriodInSamples = FMath::Max(Period->GetSeconds(), MinimumPeriodSeconds) * OperatorSettings.SampleRate;
+				float PeriodInSamples = FMath::Max(Period->GetSeconds(), MinimumPeriodSeconds) * OperatorSettings.GetSampleRate();
 
 				PeriodInSamples = FMath::Max(PeriodInSamples, MinimumPeriodSamples);
 
diff --git a/Engine/Plugins/Runtime/Metasound/Source/MetasoundGraphCore/Public/MetasoundAudioBuffer.h b/Engine/Plugins/Runtime/Metasound/Source/MetasoundGraphCore/Public/MetasoundAudioBuffer.h
--- a/Engine/Plugins/Runtime/Metasound/Source/MetasoundGraphCore/Public/MetasoundAudioBuffer.h
+++ b/Engine/Plugins/Runtime/Metasound/Source/MetasoundGraphCore/Public/MetasoundAudioBuffer.h
@@ -42,7 +42,7 @@ namespace Metasound
 			 */
 			FAudioBuffer(const FOperatorSettings& InSettings)
 			{
-				Buffer.AddUninitialized(InSettings.FramesPerExecute);
+				Buffer.AddUninitialized(InSettings.GetNumFramesPerBlock());
 
 #if METASOUNDGRAPHCORE_CHECKAUDIONUM
 				// InitialNum should not change during the life of an FAudioBuffer
diff --git a/Engine/Plugins/Runtime/Metasound/Source/MetasoundGraphCore/Public/MetasoundOperatorInterface.h b/Engine/Plugins/Runtime/Metasound/Source/MetasoundGraphCore/Public/MetasoundOperatorInterface.h
--- a/Engine/Plugins/Runtime/Metasound/Source/MetasoundGraphCore/Public/MetasoundOperatorInterface.h
+++ b/Engine/Plugins/Runtime/Metasound/Source/MetasoundGraphCore/Public/MetasoundOperatorInterface.h
@@ -3,8 +3,10 @@
 #pragma once
 
 #include "CoreMinimal.h"
+
 #include "MetasoundDataReference.h"
 #include "MetasoundDataReferenceCollection.h"
+#include "MetasoundOperatorSettings.h"
 
 namespace Metasound
 {
@@ -51,25 +53,6 @@ namespace Metasound
 		 */
 		virtual FExecuteFunction GetExecuteFunction() = 0;
 	};
-
-	/** FOperatorSettings
-	 *
-	 * Audio settings for operators.
-	 */
-	struct FOperatorSettings
-	{
-		/** The number of samples per a second of audio. */
-		float SampleRate = 0.f;
-
-		/** The number of audio frames produced during a single call to IOperators execution function. */
-		int32 FramesPerExecute = 0;
-
-		FOperatorSettings(float InSampleRate, uint32 InFramesPerExecute)
-			: SampleRate(InSampleRate)
-			, FramesPerExecute(InFramesPerExecute)
-		{
-		}
-	};
 }
 
 template <typename TDataType, typename TTypeToParse>
diff --git a/Engine/Plugins/Runtime/Metasound/Source/MetasoundGraphCore/Public/MetasoundOperatorSettings.h b/Engine/Plugins/Runtime/Metasound/Source/MetasoundGraphCore/Public/MetasoundOperatorSettings.h
new file mode 100644
--- /dev/null
+++ b/Engine/Plugins/Runtime/Metasound/Source/MetasoundGraphCore/Public/MetasoundOperatorSettings.h
@@ -0,0 +1,98 @@
+// Copyright Epic Games, Inc. All Rights Reserved.
+
+#pragma once
+
+#include "CoreMinimal.h"
+
+namespace Metasound
+{
+	/** FOperatorSettings
+	 *
+	 * Audio settings for Metasound IOperators including audio sample rate, 
+	 * block rate and frames per block.
+	 *
+	 * The audio SampleRate defines the number of audio samples per a second of
+	 * a mono waveform. 
+	 *
+	 * The BlockRate defines the number of times an IOperator's execution 
+	 * function will be triggered per a second of audio content.
+	 *
+	 * The NumFramesPerBlock defines the number of audio frames written-to 
+	 * and/or read-from during a single call to an IOperator's execution 
+	 * function.
+	 *
+	 * GetNumFramesPerBlock() is required to return a value that abides by the
+	 * alignment requirements of FAudioBuffer, AlignedFloatBuffer and the math
+	 * routines declared in BufferVectorOperations.h. This enables the vast 
+	 * majority of Metasound IOperators to take advantage of SIMD hardware
+	 * accelerations.  
+	 *
+	 * In order to achieve SIMD alignment, desired block rates must be slightly 
+	 * adjusted to match the boundaries of the given sample rate and SIMD 
+	 * alignment. GetActualBlockRate() returns the true block rate after 
+	 * adhering to sample rate and SIMD boundaries. 
+	 */
+	class METASOUNDGRAPHCORE_API FOperatorSettings
+	{
+		public:
+			/** FOperatorSettings constructor.
+			 *
+			 * @param InSampleRate      - Audio sample rate in Hz. 
+			 * @param InTargetBlockRate - The desired block rate in Hz.
+			 */
+			FOperatorSettings(float InSampleRate, float InTargetBlockRate);
+
+			/** Set the audio sample rate in Hz. */
+			void SetSampleRate(float InSampleRate);
+
+			/** Get the audio sample rate in Hz. */
+			FORCEINLINE float GetSampleRate() const
+			{
+				return SampleRate;
+			}
+
+			/** Set the target block rate in Hz. 
+			 *
+			 * In order to achieve SIMD alignment, target block rates must be 
+			 * adjusted to match the boundaries of the given sample rate and 
+			 * SIMD alignment. GetActualBlockRate() returns the true block rate 
+			 * after adhering to sample rate and SIMD boundaries. 
+			 */
+			void SetTargetBlockRate(float InTargetBlockRate);
+
+			/** Get the target block rate in Hz. */
+			float GetTargetBlockRate() const;
+
+			/** Get the actual block rate in Hz after adhering to sample rate 
+			 * and SIMD alignment.
+			 */
+			float GetActualBlockRate() const;
+
+			/** Get the number of audio frames in a block. 
+			 *
+			 * This function is required to return a value that abides by the
+			 * alignment requirements of FAudioBuffer, AlignedFloatBuffer and 
+			 * the math routines declared in BufferVectorOperations.h. This 
+			 * enables the vast majority of Metasound IOperators to take 
+			 * advantage of SIMD hardware accelerations.  
+			 */
+			FORCEINLINE int32 GetNumFramesPerBlock() const
+			{
+				return NumFramesPerBlock;
+			}
+
+		private:
+			// Update actual block rate and frames per block.
+			void Update();
+
+			int32 RoundToAligned(int32 Alignment, int32 InNum) const;
+
+			float SampleRate = 1.f;
+			
+			float TargetBlockRate = 1.f;
+
+			float ActualBlockRate = 1.f;
+
+			int32 NumFramesPerBlock = 0;
+	};
+}
diff --git a/Engine/Plugins/Runtime/Metasound/Source/MetasoundGraphCore/Public/MetasoundTime.h b/Engine/Plugins/Runtime/Metasound/Source/MetasoundGraphCore/Public/MetasoundTime.h
--- a/Engine/Plugins/Runtime/Metasound/Source/MetasoundGraphCore/Public/MetasoundTime.h
+++ b/Engine/Plugins/Runtime/Metasound/Source/MetasoundGraphCore/Public/MetasoundTime.h
@@ -303,7 +303,7 @@ namespace Metasound
 			 * @param InSettings the operator settings to use.
 			 */
 			TTime(float InSeconds, const FOperatorSettings& InSettings)
-				: TTime(InSeconds, ETimeResolution::Seconds, InSettings.SampleRate)
+				: TTime(InSeconds, ETimeResolution::Seconds, InSettings.GetSampleRate())
 			{}
 
 			/** Construct a sample count based TTime using a number of samples
