Fix PostLoad crash on asset load in-editor
#rnx

diff --git a/Engine/Plugins/Runtime/Metasound/Source/MetasoundEngine/Private/Metasound.cpp b/Engine/Plugins/Runtime/Metasound/Source/MetasoundEngine/Private/Metasound.cpp
--- a/Engine/Plugins/Runtime/Metasound/Source/MetasoundEngine/Private/Metasound.cpp
+++ b/Engine/Plugins/Runtime/Metasound/Source/MetasoundEngine/Private/Metasound.cpp
@@ -54,6 +54,7 @@ void UMetasound::SetMetadata(FMetasoundClassMetadata& InMetadata)
 void UMetasound::PostLoad()
 {
 	ConformDocumentToArchetype();
+	Super::PostLoad();
 }
 
 FMetasoundArchetype UMetasound::GetArchetype() const
