Beginings of WavePlayerNode and Wave Metasound types and unit tests

diff --git a/Engine/Plugins/Runtime/Metasound/Source/MetasoundEngine/MetasoundEngine.Build.cs b/Engine/Plugins/Runtime/Metasound/Source/MetasoundEngine/MetasoundEngine.Build.cs
--- a/Engine/Plugins/Runtime/Metasound/Source/MetasoundEngine/MetasoundEngine.Build.cs
+++ b/Engine/Plugins/Runtime/Metasound/Source/MetasoundEngine/MetasoundEngine.Build.cs
@@ -15,7 +15,8 @@ namespace UnrealBuildTool.Rules
 				{
 					"Core",
 					"MetasoundGraphCore",
-					"MetasoundGenerator"
+					"MetasoundGenerator",
+					"AudioExtentions"
 				}
 			);
 
diff --git a/Engine/Plugins/Runtime/Metasound/Source/MetasoundGraphCore/MetasoundGraphCore.Build.cs b/Engine/Plugins/Runtime/Metasound/Source/MetasoundGraphCore/MetasoundGraphCore.Build.cs
--- a/Engine/Plugins/Runtime/Metasound/Source/MetasoundGraphCore/MetasoundGraphCore.Build.cs
+++ b/Engine/Plugins/Runtime/Metasound/Source/MetasoundGraphCore/MetasoundGraphCore.Build.cs
@@ -9,12 +9,18 @@ namespace UnrealBuildTool.Rules
 	{
         public MetasoundGraphCore(ReadOnlyTargetRules Target) : base(Target)
 		{
-			//OptimizeCode = CodeOptimization.Never;
+            //OptimizeCode = CodeOptimization.Never;
+
+            PrivateDependencyModuleNames.AddRange(
+				new string[] {
+                    "AudioExtensions"
+                }
+			);
 
             PublicDependencyModuleNames.AddRange(
 				new string[] {
                     "Core",
-					"SignalProcessing"
+					"SignalProcessing",
                 }
             );
 		}
diff --git a/Engine/Plugins/Runtime/Metasound/Source/MetasoundGraphCore/Private/MetasoundGraphCoreModule.cpp b/Engine/Plugins/Runtime/Metasound/Source/MetasoundGraphCore/Private/MetasoundGraphCoreModule.cpp
--- a/Engine/Plugins/Runtime/Metasound/Source/MetasoundGraphCore/Private/MetasoundGraphCoreModule.cpp
+++ b/Engine/Plugins/Runtime/Metasound/Source/MetasoundGraphCore/Private/MetasoundGraphCoreModule.cpp
@@ -9,6 +9,7 @@
 #include "MetasoundDataReferenceTypes.h"
 #include "MetasoundTime.h"
 #include "MetasoundFrequency.h"
+#include "MetasoundWave.h"
 
 DEFINE_LOG_CATEGORY(LogMetasoundGraphCore);
 
@@ -29,6 +30,7 @@ namespace Metasound
 	REGISTER_METASOUND_DATATYPE(FDoubleTime, "Primitive:Time:HighResolution")
 	REGISTER_METASOUND_DATATYPE(FSampleTime, "Primitive:Time:SampleResolution")
 	REGISTER_METASOUND_DATATYPE(FFrequency, "Primitive:Frequency")
+	REGISTER_METASOUND_DATATYPE(FWave, "Primitive:Wave")
 }
 
 namespace Metasound 
diff --git a/Engine/Plugins/Runtime/Metasound/Source/MetasoundGraphCore/Private/MetasoundWavePlayerNode.cpp b/Engine/Plugins/Runtime/Metasound/Source/MetasoundGraphCore/Private/MetasoundWavePlayerNode.cpp
new file mode 100644
--- /dev/null
+++ b/Engine/Plugins/Runtime/Metasound/Source/MetasoundGraphCore/Private/MetasoundWavePlayerNode.cpp
@@ -0,0 +1,161 @@
+// Copyright Epic Games, Inc. All Rights Reserved.
+
+#include "MetasoundWavePlayerNode.h"
+#include "MetasoundExecutableOperator.h"
+#include "MetasoundPrimitives.h"
+#include "MetasoundWave.h"
+
+#include "IAudioCodecRegistry.h"
+#include "IAudioCodec.h"
+
+#define LOCTEXT_NAMESPACE "MetasoundWaveNode"
+
+namespace Metasound
+{
+	METASOUND_REGISTER_NODE(FWavePlayerNode)
+		
+	class FWavePlayerOperator : public TExecutableOperator<FWavePlayerOperator>
+	{
+		public:
+			FWavePlayerOperator(
+				const FOperatorSettings& InSettings, 
+				const FWaveReadRef& InWave, 
+				TUniquePtr<Audio::IDecoderInput>&& InDecoderInput, 
+				TUniquePtr<Audio::IDecoderOutput>&& InDecoderOutput,
+				TUniquePtr<Audio::IDecoder>&& InDecoder )
+				: OperatorSettings(InSettings)
+				, Decoder(MoveTemp(InDecoder))
+				, DecoderInput(MoveTemp(InDecoderInput))
+				, DecoderOutput(MoveTemp(InDecoderOutput))
+				, Wave(InWave)
+				, AudioBuffer(FAudioBufferWriteRef::CreateNew(InSettings))
+			{
+				check(AudioBuffer->Num() == InSettings.FramesPerExecute);
+
+				OutputDataReferences.AddDataReadReference(TEXT("Audio"), FAudioBufferReadRef(AudioBuffer));
+			}
+
+			virtual const FDataReferenceCollection& GetInputs() const override
+			{
+				return InputDataReferences;
+			}
+
+			virtual const FDataReferenceCollection& GetOutputs() const override
+			{
+				return OutputDataReferences;
+			}
+
+			void Execute()
+			{
+				// V1. Do the decode inline, this will sound bad.
+				Decoder->Decode();
+
+				TArrayView<float> Src;
+				float* Dst = AudioBuffer->GetData();
+				
+				int32 NumPopped = DecoderOutput->PopAudio(Src);
+				int32 NumFramesToCopy = FMath::Min(NumPopped,OperatorSettings.FramesPerExecute);
+
+				int32 i = 0;
+				for (; i < NumFramesToCopy; ++i)
+				{
+					Dst[i] = Src[i];
+				}
+				for ( ; i < OperatorSettings.FramesPerExecute; ++i)
+				{
+					Dst[i] = 0.0f;
+				}
+			}
+
+		private:
+			const FOperatorSettings OperatorSettings;
+			
+			Audio::ICodec::FDecoderPtr Decoder;
+			TUniquePtr<Audio::IDecoderInput> DecoderInput;
+			TUniquePtr<Audio::IDecoderOutput> DecoderOutput;
+			
+			FWaveReadRef Wave;
+			FAudioBufferWriteRef AudioBuffer;
+
+			FDataReferenceCollection InputDataReferences;
+			FDataReferenceCollection OutputDataReferences;
+	};
+
+	const FName FWavePlayerNode::ClassName = FName(TEXT("Wave"));
+
+	TUniquePtr<IOperator> FWavePlayerNode::FOperatorFactory::CreateOperator(
+		const INode& InNode, 
+		const FOperatorSettings& InOperatorSettings,
+		const FDataReferenceCollection& InInputDataReferences, 
+		TArray<TUniquePtr<IOperatorBuildError>>& OutErrors) 
+	{
+		using namespace Audio;
+
+		const FWavePlayerNode& WaveNode = static_cast<const FWavePlayerNode&>(InNode);
+		FWaveReadRef Wave = FWaveReadRef::CreateNew();
+
+		if (InInputDataReferences.ContainsDataReadReference<FWave>(TEXT("Wave")))
+		{
+			Wave = InInputDataReferences.GetDataReadReference<FWave>(TEXT("Wave"));
+		}
+
+		TUniquePtr<IDecoderInput> Input = IDecoderInput::Create(Wave->GetCompressedData(),0);
+		if(Input)
+		{
+			ICodecRegistry::FCodecPtr Codec = Audio::ICodecRegistry::Get().FindCodecByFromParsingInput(Input.Get());
+			if(Codec)
+			{
+				// V1, Ask for an output buffer the size of a frame.
+				IDecoderOutput::FRequirements Reqs { Float32_Interleaved, InOperatorSettings.FramesPerExecute };
+				TUniquePtr<IDecoderOutput> Output = IDecoderOutput::Create(Reqs);
+				TUniquePtr<IDecoder> Decoder = Codec->CreateDecoder(Input.Get(), Output.Get());
+
+				return MakeUnique<FWavePlayerOperator>(
+					InOperatorSettings, 
+					Wave, 
+					MoveTemp(Input), 
+					MoveTemp(Output), 
+					MoveTemp(Decoder)
+				);
+			}
+			else
+			{
+				// LogMsg; Failed to Find
+			}
+		}
+		else
+		{
+			// LogMsg; Failed to parse input stream, is it valid?
+		}
+
+		// Fail.
+		return nullptr;
+	}
+
+	FWavePlayerNode::FWavePlayerNode(const FString& InName)
+		:	FNode(InName)
+	{
+		AddInputDataVertex<FWave>(TEXT("Wave"), LOCTEXT("WaveTooltip", "The Wave to be decoded"));
+		AddOutputDataVertex<FAudioBuffer>(TEXT("Audio"), LOCTEXT("AudioTooltip", "The output audio"));
+	}
+
+	FWavePlayerNode::FWavePlayerNode(const FNodeInitData& InInitData)
+		: FWavePlayerNode(InInitData.InstanceName)
+	{
+	}
+
+	FWavePlayerNode::~FWavePlayerNode()
+	{
+	}
+	
+	const FName& FWavePlayerNode::GetClassName() const
+	{
+		return ::Metasound::FWavePlayerNode::ClassName;
+	}
+
+	IOperatorFactory& FWavePlayerNode::GetDefaultOperatorFactory() 
+	{
+		return Factory;
+	}
+}
+#undef LOCTEXT_NAMESPACE //MetasoundWaveNode
diff --git a/Engine/Plugins/Runtime/Metasound/Source/MetasoundGraphCore/Public/MetasoundWave.h b/Engine/Plugins/Runtime/Metasound/Source/MetasoundGraphCore/Public/MetasoundWave.h
new file mode 100644
--- /dev/null
+++ b/Engine/Plugins/Runtime/Metasound/Source/MetasoundGraphCore/Public/MetasoundWave.h
@@ -0,0 +1,18 @@
+// Copyright Epic Games, Inc. All Rights Reserved.
+
+#pragma once
+
+#include "MetasoundDataReference.h"
+#include "MetasoundDataTypeRegistrationMacro.h"
+
+namespace Metasound
+{
+	class FWave
+	{
+		TArray<uint8> CompressedBytes;
+	public:
+		const TArray<uint8>& GetCompressedData() const { return CompressedBytes; }
+	};
+
+	DECLARE_METASOUND_DATA_REFERENCE_TYPES(FWave, 0x0ddba11, FWaveTypeInfo, FWaveReadRef, FWaveWriteRef)
+}
diff --git a/Engine/Plugins/Runtime/Metasound/Source/MetasoundGraphCore/Public/MetasoundWavePlayerNode.h b/Engine/Plugins/Runtime/Metasound/Source/MetasoundGraphCore/Public/MetasoundWavePlayerNode.h
new file mode 100644
--- /dev/null
+++ b/Engine/Plugins/Runtime/Metasound/Source/MetasoundGraphCore/Public/MetasoundWavePlayerNode.h
@@ -0,0 +1,54 @@
+// Copyright Epic Games, Inc. All Rights Reserved.
+
+#pragma once
+
+#include "CoreMinimal.h"
+#include "MetasoundNode.h"
+#include "MetasoundOperatorInterface.h"
+#include "MetasoundBuilderInterface.h"
+#include "MetasoundDataReferenceCollection.h"
+
+namespace Metasound
+{
+	class METASOUNDGRAPHCORE_API FWavePlayerNode : public FNode
+	{
+			class FOperatorFactory : public IOperatorFactory
+			{
+				virtual TUniquePtr<IOperator> CreateOperator(const INode& InNode, const FOperatorSettings& InOperatorSettings, const FDataReferenceCollection& InInputDataReferences, TArray<TUniquePtr<IOperatorBuildError>>& OutErrors) override;
+			};
+
+		public:
+			static const FName ClassName;
+
+			FWavePlayerNode(const FString& InName);
+
+			// constructor used by the Metasound Frontend.
+			FWavePlayerNode(const FNodeInitData& InInitData);
+
+			virtual ~FWavePlayerNode();
+
+			const FName& GetClassName() const override;
+
+			IOperatorFactory& GetDefaultOperatorFactory() override;
+
+			const FText& GetDescription() const override
+			{
+				static const FText StaticDescription = NSLOCTEXT("MetasoundGraphCore", "Metasound_WavePlayerNodeDescription", "Plays a supplied Wave");
+				return StaticDescription;
+			}
+
+			const FText& GetAuthorName() const override
+			{
+				return PluginAuthor;
+			}
+
+			const FText& GetPromptIfMissing() const override
+			{
+				return PluginNodeMissingPrompt;
+			}
+
+		private:
+			float DefaultFrequency;
+			FOperatorFactory Factory;
+	};
+}
