Fix UBT warning

diff --git a/Engine/Plugins/Runtime/Metasound/Source/MetasoundEngine/MetasoundEngine.Build.cs b/Engine/Plugins/Runtime/Metasound/Source/MetasoundEngine/MetasoundEngine.Build.cs
--- a/Engine/Plugins/Runtime/Metasound/Source/MetasoundEngine/MetasoundEngine.Build.cs
+++ b/Engine/Plugins/Runtime/Metasound/Source/MetasoundEngine/MetasoundEngine.Build.cs
@@ -15,8 +15,7 @@ namespace UnrealBuildTool.Rules
 				{
 					"Core",
 					"MetasoundGraphCore",
-					"MetasoundGenerator",
-					"AudioExtentions"
+					"MetasoundGenerator"
 				}
 			);
 
diff --git a/Engine/Plugins/Runtime/Metasound/Source/MetasoundGraphCore/MetasoundGraphCore.Build.cs b/Engine/Plugins/Runtime/Metasound/Source/MetasoundGraphCore/MetasoundGraphCore.Build.cs
--- a/Engine/Plugins/Runtime/Metasound/Source/MetasoundGraphCore/MetasoundGraphCore.Build.cs
+++ b/Engine/Plugins/Runtime/Metasound/Source/MetasoundGraphCore/MetasoundGraphCore.Build.cs
@@ -20,7 +20,7 @@ namespace UnrealBuildTool.Rules
             PublicDependencyModuleNames.AddRange(
 				new string[] {
                     "Core",
-					"SignalProcessing",
+					"SignalProcessing"
                 }
             );
 		}
