Build Fix for FramesPerExecute getting renamed GetNumFramesPerBlock

diff --git a/Engine/Plugins/Runtime/Metasound/Source/MetasoundGraphCore/Private/MetasoundWavePlayerNode.cpp b/Engine/Plugins/Runtime/Metasound/Source/MetasoundGraphCore/Private/MetasoundWavePlayerNode.cpp
--- a/Engine/Plugins/Runtime/Metasound/Source/MetasoundGraphCore/Private/MetasoundWavePlayerNode.cpp
+++ b/Engine/Plugins/Runtime/Metasound/Source/MetasoundGraphCore/Private/MetasoundWavePlayerNode.cpp
@@ -30,7 +30,7 @@ namespace Metasound
 				, Wave(InWave)
 				, AudioBuffer(FAudioBufferWriteRef::CreateNew(InSettings))
 			{
-				check(AudioBuffer->Num() == InSettings.FramesPerExecute);
+				check(AudioBuffer->Num() == InSettings.GetNumFramesPerBlock());
 
 				OutputDataReferences.AddDataReadReference(TEXT("Audio"), FAudioBufferReadRef(AudioBuffer));
 			}
@@ -54,14 +54,14 @@ namespace Metasound
 				float* Dst = AudioBuffer->GetData();
 				
 				int32 NumPopped = DecoderOutput->PopAudio(Src);
-				int32 NumFramesToCopy = FMath::Min(NumPopped,OperatorSettings.FramesPerExecute);
+				int32 NumFramesToCopy = FMath::Min(NumPopped,OperatorSettings.GetNumFramesPerBlock());
 
 				int32 i = 0;
 				for (; i < NumFramesToCopy; ++i)
 				{
 					Dst[i] = Src[i];
 				}
-				for ( ; i < OperatorSettings.FramesPerExecute; ++i)
+				for ( ; i < OperatorSettings.GetNumFramesPerBlock(); ++i)
 				{
 					Dst[i] = 0.0f;
 				}
@@ -106,7 +106,7 @@ namespace Metasound
 			if(Codec)
 			{
 				// V1, Ask for an output buffer the size of a frame.
-				IDecoderOutput::FRequirements Reqs { Float32_Interleaved, InOperatorSettings.FramesPerExecute };
+				IDecoderOutput::FRequirements Reqs { Float32_Interleaved, InOperatorSettings.GetNumFramesPerBlock() };
 				TUniquePtr<IDecoderOutput> Output = IDecoderOutput::Create(Reqs);
 				TUniquePtr<IDecoder> Decoder = Codec->CreateDecoder(Input.Get(), Output.Get());
 
