Extract registry-related metasound files to Metasound Frontend, and move our core node library to MetasoundCoreNodeLibrary.

diff --git a/Engine/Plugins/Runtime/Metasound/Metasound.uplugin b/Engine/Plugins/Runtime/Metasound/Metasound.uplugin
--- a/Engine/Plugins/Runtime/Metasound/Metasound.uplugin
+++ b/Engine/Plugins/Runtime/Metasound/Metasound.uplugin
@@ -31,6 +31,11 @@
 			"Type" : "Runtime",
 			"LoadingPhase" : "Default"
 		},
+		{
+			"Name" : "MetasoundStandardNodes",
+			"Type" : "Runtime",
+			"LoadingPhase" : "Default"
+		},
 		{
 			"Name" : "MetasoundEngine",
 			"Type" : "Runtime",
diff --git a/Engine/Plugins/Runtime/Metasound/Source/MetasoundEngine/MetasoundEngine.Build.cs b/Engine/Plugins/Runtime/Metasound/Source/MetasoundEngine/MetasoundEngine.Build.cs
--- a/Engine/Plugins/Runtime/Metasound/Source/MetasoundEngine/MetasoundEngine.Build.cs
+++ b/Engine/Plugins/Runtime/Metasound/Source/MetasoundEngine/MetasoundEngine.Build.cs
@@ -26,6 +26,7 @@ namespace UnrealBuildTool.Rules
 					"CoreUObject",
 					"Engine",
 					"MetasoundFrontend",
+					"MetasoundStandardNodes",
 					"Serialization"
 				}
 			);
diff --git a/Engine/Plugins/Runtime/Metasound/Source/MetasoundFrontend/Private/MetasoundFrontend.cpp b/Engine/Plugins/Runtime/Metasound/Source/MetasoundFrontend/Private/MetasoundFrontend.cpp
--- a/Engine/Plugins/Runtime/Metasound/Source/MetasoundFrontend/Private/MetasoundFrontend.cpp
+++ b/Engine/Plugins/Runtime/Metasound/Source/MetasoundFrontend/Private/MetasoundFrontend.cpp
@@ -9,7 +9,6 @@
 #include "HAL/IConsoleManager.h"
 #include "Modules/ModuleManager.h"
 
-#include "MetasoundAudioFormats.h"
 #include "MetasoundOperatorBuilder.h"
 
 static int32 MetasoundUndoRollLimitCvar = 128;
diff --git a/Engine/Plugins/Runtime/Metasound/Source/MetasoundFrontend/Private/MetasoundFrontendRegistries.cpp b/Engine/Plugins/Runtime/Metasound/Source/MetasoundFrontend/Private/MetasoundFrontendRegistries.cpp
new file mode 100644
--- /dev/null
+++ b/Engine/Plugins/Runtime/Metasound/Source/MetasoundFrontend/Private/MetasoundFrontendRegistries.cpp
@@ -0,0 +1,321 @@
+// Copyright Epic Games, Inc. All Rights Reserved.
+
+#include "MetasoundFrontendRegistries.h"
+#include "CoreMinimal.h"
+#include "Misc/ScopeLock.h"
+#include "HAL/PlatformTime.h"
+
+#ifndef WITH_METASOUND_FRONTEND
+#define WITH_METASOUND_FRONTEND 0
+#endif
+
+FMetasoundFrontendRegistryContainer::FMetasoundFrontendRegistryContainer()
+	: bHasModuleBeenInitialized(false)
+{
+}
+
+FMetasoundFrontendRegistryContainer* FMetasoundFrontendRegistryContainer::LazySingleton = nullptr;
+
+FMetasoundFrontendRegistryContainer* FMetasoundFrontendRegistryContainer::Get()
+{
+	if (!LazySingleton)
+	{
+		LazySingleton = new FMetasoundFrontendRegistryContainer();
+	}
+
+	return LazySingleton;
+}
+
+void FMetasoundFrontendRegistryContainer::ShutdownMetasoundFrontend()
+{
+	if (LazySingleton)
+	{
+		delete LazySingleton;
+		LazySingleton = nullptr;
+	}
+}
+
+void FMetasoundFrontendRegistryContainer::InitializeFrontend()
+{
+	FScopeLock ScopeLock(&LazyInitCommandCritSection);
+
+	if (bHasModuleBeenInitialized)
+	{
+		// this function should only be called once.
+		checkNoEntry();
+	}
+
+	UE_LOG(LogTemp, Display, TEXT("Initializing Metasounds Frontend."));
+	uint64 CurrentTime = FPlatformTime::Cycles64();
+
+	for (TUniqueFunction<void()>& Command : LazyInitCommands)
+	{
+		Command();
+	}
+
+	LazyInitCommands.Empty();
+	bHasModuleBeenInitialized = true;
+
+	uint64 CyclesUsed = FPlatformTime::Cycles64() - CurrentTime;
+	UE_LOG(LogTemp, Display, TEXT("Initializing Metasounds Frontend took %f seconds."), FPlatformTime::ToSeconds64(CyclesUsed));
+}
+
+bool FMetasoundFrontendRegistryContainer::EnqueueInitCommand(TUniqueFunction<void()>&& InFunc)
+{
+	// if the module has been initalized already, we can safely call this function now.
+	if (bHasModuleBeenInitialized)
+	{
+		InFunc();
+	}
+
+	// otherwise, we enqueue the function to be executed after the frontend module has been initialized.
+
+	if (LazyInitCommands.Num() >= MaxNumNodesAndDatatypesToInitialize)
+	{
+		UE_LOG(LogTemp, Warning, TEXT("Registering more that %d nodes and datatypes for metasounds! Consider increasing MetasoundFrontendRegistryContainer::MaxNumNodesAndDatatypesToInitialize."));
+	}
+
+	LazyInitCommands.Add(MoveTemp(InFunc));
+	return true;
+}
+
+TMap<FMetasoundFrontendRegistryContainer::FNodeRegistryKey, FMetasoundFrontendRegistryContainer::FNodeRegistryElement>& FMetasoundFrontendRegistryContainer::GetExternalNodeRegistry()
+{
+	return ExternalNodeRegistry;
+}
+
+TUniquePtr<Metasound::INode> FMetasoundFrontendRegistryContainer::ConstructInputNode(const FName& InInputType, const Metasound::FInputNodeConstructorParams& InParams)
+{
+	if (ensureAlwaysMsgf(DataTypeRegistry.Contains(InInputType), TEXT("Couldn't find data type %s!"), *InInputType.ToString()))
+	{
+		return DataTypeRegistry[InInputType].InputNodeConstructor(InParams);
+	}
+	else
+	{
+		return nullptr;
+	}
+}
+
+TUniquePtr<Metasound::INode> FMetasoundFrontendRegistryContainer::ConstructOutputNode(const FName& InOutputType, const Metasound::FOutputNodeConstrutorParams& InParams)
+{
+	if (ensureAlwaysMsgf(DataTypeRegistry.Contains(InOutputType), TEXT("Couldn't find data type %s!"), *InOutputType.ToString()))
+	{
+		return DataTypeRegistry[InOutputType].OutputNodeConstructor(InParams);
+	}
+	else
+	{
+		return nullptr;
+	}
+}
+
+TUniquePtr<Metasound::INode> FMetasoundFrontendRegistryContainer::ConstructExternalNode(const FName& InNodeType, uint32 InNodeHash, const Metasound::FNodeInitData& InInitData)
+{
+	Metasound::Frontend::FNodeRegistryKey RegistryKey;
+	RegistryKey.NodeName = InNodeType;
+	RegistryKey.NodeHash = InNodeHash;
+
+	if (!ExternalNodeRegistry.Contains(RegistryKey))
+	{
+		return nullptr;
+	}
+	else
+	{
+		return ExternalNodeRegistry[RegistryKey].GetterCallback(InInitData);
+	}
+}
+
+Metasound::ELiteralArgType FMetasoundFrontendRegistryContainer::GetDesiredLiteralTypeForDataType(FName InDataType) const
+{
+	if (!DataTypeRegistry.Contains(InDataType))
+	{
+		return Metasound::ELiteralArgType::Invalid;
+	}
+
+	const FDataTypeRegistryElement& DataTypeInfo = DataTypeRegistry[InDataType];
+	
+	// If there's a designated preferred literal type for this datatype, use that.
+	if (DataTypeInfo.Info.PreferredLiteralType != Metasound::ELiteralArgType::None)
+	{
+		return DataTypeInfo.Info.PreferredLiteralType;
+	}
+
+	// Otherwise, we opt for the highest precision construction option available.
+	if (DataTypeInfo.Info.bIsStringParsable)
+	{
+		return Metasound::ELiteralArgType::String;
+	}
+	else if (DataTypeInfo.Info.bIsFloatParsable)
+	{
+		return Metasound::ELiteralArgType::Float;
+	}
+	else if (DataTypeInfo.Info.bIsIntParsable)
+	{
+		return Metasound::ELiteralArgType::Integer;
+	}
+	else if (DataTypeInfo.Info.bIsBoolParsable)
+	{
+		return Metasound::ELiteralArgType::Boolean;
+	}
+	else if (DataTypeInfo.Info.bIsConstructableWithSettings || DataTypeInfo.Info.bIsDefaultConstructible)
+	{
+		return Metasound::ELiteralArgType::None;
+	}
+	else
+	{
+		// if we ever hit this, something has gone terribly wrong with the REGISTER_METASOUND_DATATYPE macro.
+		// we should have failed to compile if any of these are false.
+		checkNoEntry();
+		return Metasound::ELiteralArgType::Invalid;
+	}
+}
+
+bool FMetasoundFrontendRegistryContainer::DoesDataTypeSupportLiteralType(FName InDataType, Metasound::ELiteralArgType InLiteralType) const
+{
+	if (!DataTypeRegistry.Contains(InDataType))
+	{
+		ensureAlwaysMsgf(false, TEXT("couldn't find DataType %s in the registry."), *InDataType.ToString());
+		return false;
+	}
+
+	const FDataTypeRegistryElement& DataTypeInfo = DataTypeRegistry[InDataType];
+	
+	switch (InLiteralType)
+	{
+		case Metasound::ELiteralArgType::Boolean:
+		{
+			return DataTypeInfo.Info.bIsBoolParsable;
+		}
+		case Metasound::ELiteralArgType::Integer:
+		{
+			return DataTypeInfo.Info.bIsIntParsable;
+		}
+		case Metasound::ELiteralArgType::Float:
+		{
+			return DataTypeInfo.Info.bIsFloatParsable;
+		}
+		case Metasound::ELiteralArgType::String:
+		{
+			return DataTypeInfo.Info.bIsStringParsable;
+		}
+		case Metasound::ELiteralArgType::None:
+		{
+			return DataTypeInfo.Info.bIsConstructableWithSettings || DataTypeInfo.Info.bIsDefaultConstructible;
+		}
+		case Metasound::ELiteralArgType::Invalid:
+		default:
+		{
+			return false;
+		}
+	}
+}
+
+bool FMetasoundFrontendRegistryContainer::RegisterDataType(const ::Metasound::FDataTypeRegistryInfo& InDataInfo, FInputNodeConstructorCallback&& InputNodeConstructor, FOutputNodeConstructorCallback&& OutputNodeConstructor)
+{
+	if (!ensureAlwaysMsgf(!DataTypeRegistry.Contains(InDataInfo.DataTypeName), TEXT("Name collision when trying to register Metasound Data Type %s! Make sure that you created a unique name for your data type, and that REGISTER_METASOUND_DATATYPE isn't called in a public header."), *InDataInfo.DataTypeName.ToString()))
+	{
+		// todo: capture callstack for previous declaration for non-shipping builds to help clarify who already registered this name for a type.
+		return false;
+	}
+	else
+	{
+		FDataTypeRegistryElement InElement = { MoveTemp(InputNodeConstructor), MoveTemp(OutputNodeConstructor), InDataInfo };
+		DataTypeRegistry.Add(InDataInfo.DataTypeName, MoveTemp(InElement));
+		UE_LOG(LogTemp, Display, TEXT("Registered Metasound Datatype %s."), *InDataInfo.DataTypeName.ToString());
+		return true;
+	}
+}
+
+bool FMetasoundFrontendRegistryContainer::RegisterExternalNode(FNodeGetterCallback&& InCallback)
+{
+	using FInputDataVertexCollection = Metasound::FInputDataVertexCollection;
+	using FOutputDataVertexCollection = Metasound::FOutputDataVertexCollection;
+
+	Metasound::FNodeInitData DummyInitData;
+	TUniquePtr<Metasound::INode> DummyNodePtr = InCallback(DummyInitData);
+	if (!ensureAlwaysMsgf(DummyNodePtr.IsValid(), TEXT("Invalid getter registered!")))
+	{
+		return false;
+	}
+
+	Metasound::INode& DummyNode = *DummyNodePtr;
+
+	// First, build the key.
+	const FName NodeName = DummyNode.GetClassName();
+	const FInputDataVertexCollection& Inputs = DummyNode.GetInputDataVertices();
+	const FOutputDataVertexCollection& Outputs = DummyNode.GetOutputDataVertices();
+
+	// Construct a hash using a combination of the class name, input names and output names.
+	uint32 NodeHash = FCrc::StrCrc32(*NodeName.ToString());
+
+	TArray<FName> InputTypes;
+	TArray<FName> OutputTypes;
+
+	for (auto& InputTuple : Inputs)
+	{
+		HashCombine(NodeHash, FCrc::StrCrc32(*InputTuple.Value.VertexName));
+		InputTypes.Add(InputTuple.Value.DataReferenceTypeName);
+	}
+
+	for (auto& OutputTuple : Outputs)
+	{
+		HashCombine(NodeHash, FCrc::StrCrc32(*OutputTuple.Value.VertexName));
+		OutputTypes.Add(OutputTuple.Value.DataReferenceTypeName);
+	}
+
+	FNodeRegistryKey InKey = { NodeName, NodeHash };
+
+	// check to see if an identical node was already registered, and log
+	ensureAlwaysMsgf(!ExternalNodeRegistry.Contains(InKey), TEXT("Node with identical name, inputs and outputs to node %s was already registered. The previously registered node will be overwritten. This could also happen because METASOUND_REGISTER_NODE is in a public header."), *NodeName.ToString());
+
+	UE_LOG(LogTemp, Display, TEXT("Registered Metasound Node %s"), *NodeName.ToString());
+
+	const bool bShouldLogInputsAndOutputs = true;
+	if (bShouldLogInputsAndOutputs)
+	{
+		UE_LOG(LogTemp, Display, TEXT("    %d inputs:"), Inputs.Num());
+		for (auto& InputTuple : Inputs)
+		{
+			UE_LOG(LogTemp, Display, TEXT("      %s (of type %s)"), *InputTuple.Value.VertexName, *InputTuple.Value.DataReferenceTypeName.ToString());
+		}
+
+		UE_LOG(LogTemp, Display, TEXT("    %d outputs:"), Outputs.Num());
+		for (auto& OutputTuple : Outputs)
+		{
+			UE_LOG(LogTemp, Display, TEXT("      %s (of type %s)"), *OutputTuple.Value.VertexName, *OutputTuple.Value.DataReferenceTypeName.ToString());
+		}
+	}
+
+	FNodeRegistryElement RegistryElement = FNodeRegistryElement(MoveTemp(InCallback));
+	RegistryElement.InputTypes = MoveTemp(InputTypes);
+	RegistryElement.OutputTypes = MoveTemp(OutputTypes);
+
+	ExternalNodeRegistry.Add(MoveTemp(InKey), MoveTemp(RegistryElement));
+
+	return true;
+}
+
+TArray<FName> FMetasoundFrontendRegistryContainer::GetAllValidDataTypes()
+{
+	TArray<FName> OutDataTypes;
+
+	for (auto& DataTypeTuple : DataTypeRegistry)
+	{
+		OutDataTypes.Add(DataTypeTuple.Key);
+	}
+
+	return OutDataTypes;
+}
+
+bool FMetasoundFrontendRegistryContainer::GetInfoForDataType(FName InDataType, Metasound::FDataTypeRegistryInfo& OutInfo)
+{
+	if (!DataTypeRegistry.Contains(InDataType))
+	{
+		return false;
+	}
+	else
+	{
+		OutInfo = DataTypeRegistry[InDataType].Info;
+		return true;
+	}
+}
+
diff --git a/Engine/Plugins/Runtime/Metasound/Source/MetasoundFrontend/Public/MetasoundDataTypeRegistrationMacro.h b/Engine/Plugins/Runtime/Metasound/Source/MetasoundFrontend/Public/MetasoundDataTypeRegistrationMacro.h
new file mode 100644
--- /dev/null
+++ b/Engine/Plugins/Runtime/Metasound/Source/MetasoundFrontend/Public/MetasoundDataTypeRegistrationMacro.h
@@ -0,0 +1,71 @@
+// Copyright Epic Games, Inc. All Rights Reserved.
+
+#pragma once
+
+#include "CoreMinimal.h"
+
+#include "MetasoundDataReference.h"
+#include "MetasoundOutputNode.h"
+#include "MetasoundInputNode.h"
+#include "MetasoundOperatorInterface.h"
+#include "MetasoundFrontendRegistries.h"
+
+typedef TUniqueFunction<TUniquePtr<Metasound::INode>(const ::Metasound::FInputNodeConstructorParams&)> FInputNodeConstructorCallback;
+typedef TUniqueFunction<TUniquePtr<Metasound::INode>(const ::Metasound::FOutputNodeConstrutorParams&)> FOutputNodeConstructorCallback;
+
+template<typename TDataType>
+bool RegisterDataTypeWithFrontend(::Metasound::ELiteralArgType PreferredArgType = ::Metasound::ELiteralArgType::None)
+{
+	// if we reenter this code (because DECLARE_METASOUND_DATA_REFERENCE_TYPES was called twice with the same type),
+	// we catch it here.
+	static bool bAlreadyRegisteredThisDataType = false;
+	if (bAlreadyRegisteredThisDataType)
+	{
+		UE_LOG(LogTemp, Display, TEXT("Tried to call REGISTER_METASOUND_DATATYPE twice with the same class %s. ignoring the second call. Likely because REGISTER_METASOUND_DATATYPE is in a header that's used in multiple modules. Consider moving it to a private header or cpp file."), ::Metasound::TDataReferenceTypeInfo<TDataType>::TypeName)
+		return false;
+	}
+
+	bAlreadyRegisteredThisDataType = true;
+
+	// Lambdas that generate our template-instantiated input and output nodes:
+	FInputNodeConstructorCallback InputNodeConstructor = [](const ::Metasound::FInputNodeConstructorParams& InParams) -> TUniquePtr<::Metasound::INode>
+	{
+		return TUniquePtr<::Metasound::INode>(new ::Metasound::TInputNode<TDataType>(InParams.InNodeName, InParams.InVertexName, InParams.InitParam, InParams.InSettings));
+	};
+
+	FOutputNodeConstructorCallback OutputNodeConstructor = [](const ::Metasound::FOutputNodeConstrutorParams& InParams) -> TUniquePtr<::Metasound::INode>
+	{
+		return TUniquePtr<::Metasound::INode>(new ::Metasound::TOutputNode<TDataType>(InParams.InNodeName, InParams.InVertexName));
+	};
+	
+	static FName DataTypeName = FName(::Metasound::TDataReferenceTypeInfo<TDataType>::TypeName);
+
+	::Metasound::FDataTypeRegistryInfo RegistryInfo;
+	RegistryInfo.DataTypeName = DataTypeName;
+	RegistryInfo.PreferredLiteralType = PreferredArgType;
+	RegistryInfo.bIsBoolParsable = ::Metasound::TDataReferenceTypeInfo<TDataType>::bIsBoolParsable;
+	RegistryInfo.bIsIntParsable = ::Metasound::TDataReferenceTypeInfo<TDataType>::bIsIntParsable;
+	RegistryInfo.bIsFloatParsable = ::Metasound::TDataReferenceTypeInfo<TDataType>::bIsFloatParsable;
+	RegistryInfo.bIsStringParsable = ::Metasound::TDataReferenceTypeInfo<TDataType>::bIsStringParsable;
+	RegistryInfo.bIsConstructableWithSettings = ::Metasound::TDataReferenceTypeInfo<TDataType>::bIsConstructableWithSettings;
+	RegistryInfo.bIsDefaultConstructible = ::Metasound::TDataReferenceTypeInfo<TDataType>::bCanUseDefaultConstructor;
+
+	bool bSucceeded = FMetasoundFrontendRegistryContainer::Get()->RegisterDataType(RegistryInfo, MoveTemp(InputNodeConstructor), MoveTemp(OutputNodeConstructor));
+	ensureAlwaysMsgf(bSucceeded, TEXT("Failed to register data type %s in the node registry!"), ::Metasound::TDataReferenceTypeInfo<TDataType>::TypeName);
+	return bSucceeded;
+}
+
+// This should be used to expose a datatype as a potential input or output for a metasound graph.
+// The first argument to the macro is the class to expose.
+// Optionally, a Metasound::ELiteralArgType can be passed in to designate a preferred literal type-
+// For example, if Metasound::ELiteralArgType::Float is passed in, we will default to using a float parameter to create this datatype.
+// If no argument is passed in, we will infer a literal type to use.
+// Metasound::ELiteralArgType::Invalid can be used to enforce that we don't provide space for a literal, in which case you should have a default constructor or a constructor that takes [const FOperatorSettings&] implemented.
+// If you pass in a preferred arg type, please make sure that the passed in datatype has a matching constructor, since we won't check this until runtime.
+#define REGISTER_METASOUND_DATATYPE(DataType, DataTypeName, ...) \
+	IMPL_METASOUND_DATA_TYPE(DataType, DataTypeName) \
+	static_assert(::Metasound::TDataReferenceTypeInfo<DataType>::bIsValidSpecialization, "Please call DECLARE_METASOUND_DATA_REFERENCE_TYPES(" #DataType "...) before calling REGISTER_METASOUND_DATATYPE(" #DataType")."); \
+	static constexpr bool bCanRegister##DataType = ::Metasound::TDataReferenceTypeInfo<DataType>::bIsStringParsable || ::Metasound::TDataReferenceTypeInfo<DataType>::bIsBoolParsable || ::Metasound::TDataReferenceTypeInfo<DataType>::bIsIntParsable || ::Metasound::TDataReferenceTypeInfo<DataType>::bIsFloatParsable || ::Metasound::TDataReferenceTypeInfo<DataType>::bIsConstructableWithSettings || ::Metasound::TDataReferenceTypeInfo<DataType>::bCanUseDefaultConstructor; \
+	static_assert(bCanRegister##DataType , "To register " #DataType " to be used as a Metasounds input or output type, it needs a default constructor or one of the following constructors must be implemented:  " #DataType "(const ::Metasound::FOperatorSettings& InSettings), " #DataType "(int32 InValue, const ::Metasound::FOperatorSettings& InSettings), " #DataType "(float InValue, const ::Metasound::FOperatorSettings& InSettings), " #DataType "(bool InValue, const ::Metasound::FOperatorSettings& InSettings), or " #DataType "(const FString& InString, const ::Metasound::FOperatorSettings& InSettings)."); \
+	static const bool bSuccessfullyRegistered##DataType = FMetasoundFrontendRegistryContainer::Get()->EnqueueInitCommand([](){ ::RegisterDataTypeWithFrontend<DataType>(__VA_ARGS__); }); // This static bool is useful for debugging, but also is the only way the compiler will let us call this function outside of an expression.
+
diff --git a/Engine/Plugins/Runtime/Metasound/Source/MetasoundFrontend/Public/MetasoundFrontendRegistries.h b/Engine/Plugins/Runtime/Metasound/Source/MetasoundFrontend/Public/MetasoundFrontendRegistries.h
new file mode 100644
--- /dev/null
+++ b/Engine/Plugins/Runtime/Metasound/Source/MetasoundFrontend/Public/MetasoundFrontendRegistries.h
@@ -0,0 +1,131 @@
+// Copyright Epic Games, Inc. All Rights Reserved.
+
+#pragma once
+
+#include "CoreMinimal.h"
+#include "MetasoundNodeInterface.h"
+#include "MetasoundDataReference.h"
+#include "MetasoundOperatorInterface.h"
+
+typedef TUniqueFunction<TUniquePtr<Metasound::INode>(const ::Metasound::FInputNodeConstructorParams&)> FInputNodeConstructorCallback;
+typedef TUniqueFunction<TUniquePtr<Metasound::INode>(const ::Metasound::FOutputNodeConstrutorParams&)> FOutputNodeConstructorCallback;
+
+typedef TUniqueFunction<TUniquePtr<Metasound::INode>(const Metasound::FNodeInitData&)> FNodeGetterCallback;
+
+namespace Metasound
+{
+	namespace Frontend
+	{
+		struct METASOUNDFRONTEND_API FNodeRegistryKey
+		{
+			// The class name for the node.
+			FName NodeName;
+
+			// A hash generated from the input types and output types for this node.
+			// TODO: Write up some tests to ensure this is deterministic.
+			uint32 NodeHash;
+
+			FORCEINLINE bool operator==(const FNodeRegistryKey& Other) const
+			{
+				return NodeHash == Other.NodeHash && NodeName == Other.NodeName;
+			}
+
+			friend uint32 GetTypeHash(const Metasound::Frontend::FNodeRegistryKey& InKey)
+			{
+				return InKey.NodeHash;
+			}
+		};
+
+
+		struct METASOUNDFRONTEND_API FNodeRegistryElement
+		{
+			// This lambda can be used to get an INodeBase for this specific node class.
+			FNodeGetterCallback GetterCallback;
+
+			TArray<FName> InputTypes;
+			TArray<FName> OutputTypes;
+
+			FNodeRegistryElement(FNodeGetterCallback&& InCallback)
+				: GetterCallback(MoveTemp(InCallback))
+			{
+			}
+		};
+	}
+}
+
+/*
+static uint32 METASOUNDGRAPHCORE_API GetTypeHash(const Metasound::Frontend::FNodeRegistryKey& InKey)
+{
+	return InKey.NodeHash;
+}
+*/
+
+/**
+ * Singleton registry for all types and nodes.
+ */
+class METASOUNDFRONTEND_API FMetasoundFrontendRegistryContainer
+{
+public:
+	static FMetasoundFrontendRegistryContainer* Get();
+	static void ShutdownMetasoundFrontend();
+
+	FMetasoundFrontendRegistryContainer(const FMetasoundFrontendRegistryContainer&) = delete;
+
+	void InitializeFrontend();
+	bool EnqueueInitCommand(TUniqueFunction<void()>&& InFunc);
+
+	using FNodeRegistryKey = Metasound::Frontend::FNodeRegistryKey;
+	using FNodeRegistryElement = Metasound::Frontend::FNodeRegistryElement;
+
+	TMap<FNodeRegistryKey, FNodeRegistryElement>& GetExternalNodeRegistry();
+
+	TUniquePtr<Metasound::INode> ConstructInputNode(const FName& InInputType, const Metasound::FInputNodeConstructorParams& InParams);
+	TUniquePtr<Metasound::INode> ConstructOutputNode(const FName& InOutputType, const Metasound::FOutputNodeConstrutorParams& InParams);
+
+	TUniquePtr<Metasound::INode> ConstructExternalNode(const FName& InNodeType, uint32 InNodeHash, const Metasound::FNodeInitData& InInitData);
+
+	// Get the desired kind of literal for a given data type. Returns EConstructorArgType::Invalid if the data type couldn't be found.
+	Metasound::ELiteralArgType GetDesiredLiteralTypeForDataType(FName InDataType) const;
+
+	// Get whether we can build a literal of this specific type for InDataType.
+	bool DoesDataTypeSupportLiteralType(FName InDataType, Metasound::ELiteralArgType InLiteralType) const;
+
+	bool RegisterDataType(const ::Metasound::FDataTypeRegistryInfo& InDataInfo, FInputNodeConstructorCallback&& InputNodeConstructor, FOutputNodeConstructorCallback&& OutputNodeConstructor);
+	bool RegisterExternalNode(FNodeGetterCallback&& InCallback);
+
+	// Return any data types that can be used as a metasound input type or output type.
+	TArray<FName> GetAllValidDataTypes();
+
+	// Get info about a specific data type (what kind of literals we can use, etc.)
+	// @returns false if InDataType wasn't found in the registry. 
+	bool GetInfoForDataType(FName InDataType, Metasound::FDataTypeRegistryInfo& OutInfo);
+
+private:
+	FMetasoundFrontendRegistryContainer();
+
+	static FMetasoundFrontendRegistryContainer* LazySingleton;
+
+	// This buffer is used to enqueue nodes and datatypes to register when the module has been initialized,
+	// in order to avoid bad behavior with ensures, logs, etc. on static initialization.
+	// The bad news is that TInlineAllocator is the safest allocator to use on static init.
+	// The good news is that none of these lambdas typically have captures, so this should have low memory overhead.
+	static constexpr int32 MaxNumNodesAndDatatypesToInitialize = 8192;
+	TArray<TUniqueFunction<void()>/*, TInlineAllocator<MaxNumNodesAndDatatypesToInitialize>*/> LazyInitCommands;
+	
+	FCriticalSection LazyInitCommandCritSection;
+	bool bHasModuleBeenInitialized;
+
+	// Registry in which we keep all information about nodes implemented in C++.
+	TMap<FNodeRegistryKey, FNodeRegistryElement> ExternalNodeRegistry;
+
+	struct FDataTypeRegistryElement
+	{
+		FInputNodeConstructorCallback InputNodeConstructor;
+		FOutputNodeConstructorCallback  OutputNodeConstructor;
+		Metasound::FDataTypeRegistryInfo Info;
+	};
+
+	TMap<FName, FDataTypeRegistryElement> DataTypeRegistry;
+};
+
+
diff --git a/Engine/Plugins/Runtime/Metasound/Source/MetasoundFrontend/Public/MetasoundInputNode.h b/Engine/Plugins/Runtime/Metasound/Source/MetasoundFrontend/Public/MetasoundInputNode.h
new file mode 100644
--- /dev/null
+++ b/Engine/Plugins/Runtime/Metasound/Source/MetasoundFrontend/Public/MetasoundInputNode.h
@@ -0,0 +1,217 @@
+// Copyright Epic Games, Inc. All Rights Reserved.
+#pragma once
+
+#include "CoreMinimal.h"
+#include "Internationalization/Text.h"
+#include "MetasoundOperatorInterface.h"
+#include "MetasoundNodeInterface.h"
+
+#define LOCTEXT_NAMESPACE "MetasoundGraphCore"
+
+
+namespace Metasound
+{
+	template<typename DataType>
+	class TInputNode : public INode
+	{
+		static_assert(TDataReferenceTypeInfo<DataType>::bIsValidSpecialization, "Please use DECLARE_METASOUND_DATA_REFERENCE_TYPES with this class before trying to create an input node with it.");
+
+		class FInputOperator : public IOperator
+		{
+			public:
+				using FDataWriteReference = TDataWriteReference<DataType>;
+
+				FInputOperator(const FString& InDataReferenceName, FDataWriteReference InDataReference)
+				{
+					Inputs.AddDataWriteReference<DataType>(InDataReferenceName, InDataReference);
+					Outputs.AddDataReadReference<DataType>(InDataReferenceName, InDataReference);
+				}
+
+				virtual ~FInputOperator() {}
+
+				virtual const FDataReferenceCollection& GetInputs() const override
+				{
+					return Inputs;
+				}
+
+				virtual const FDataReferenceCollection& GetOutputs() const override
+				{
+					return Outputs;
+				}
+
+				virtual FExecuteFunction GetExecuteFunction() override
+				{
+					return nullptr;
+				}
+
+			private:
+				FDataReferenceCollection Outputs;
+				FDataReferenceCollection Inputs;
+		};
+
+		// TODO: Would be nice if we didn't utilize the copy constructor by default here. But it's a nice catch-all
+		// for our use case.  Could store constructor parameters of the paramtype, but run the risk of holding rvalues
+		// or pointers to objects which are no longer valid. 
+
+		// TODO: instead of this copy factory, register a DataType factory lambda that uses a FDataTypeLiteralParam
+		class FCopyOperatorFactory : public IOperatorFactory
+		{
+			static_assert(TIsConstructible<DataType, const DataType&>::Value, "Data Type must be copy constructible!");
+				
+			public:
+				using FDataWriteReference = TDataWriteReference<DataType>;
+				using FInputNodeType = TInputNode<DataType>;
+
+				template< typename... ArgTypes >
+				FCopyOperatorFactory(ArgTypes&&... Args)
+				:	Data(Forward<ArgTypes>(Args)...)
+				{
+				}
+
+				/*
+				template<typename = typename TEnableIf< TIsConstructible<DataType>::Value >::Type >
+				FCopyOperatorFactory()
+				:	Data()
+				{
+				}
+				*/
+
+				FCopyOperatorFactory(FDataTypeLiteralParam InitParam, const FOperatorSettings& InSettings)
+					: Data(InitParam.ParseTo<DataType>(InSettings))
+				{
+				}
+
+				virtual TUniquePtr<IOperator> CreateOperator(const INode& InNode, const FOperatorSettings& InOperatorSettings, const FDataReferenceCollection& InInputDataReferences, TArray<TUniquePtr<IOperatorBuildError>>& OutErrors) override
+				{
+					// Use copy constructor to create a new parameter reference.
+					FDataWriteReference DataRef = FDataWriteReference::CreateNew(Data);
+					// TODO: Write special version of this for audio types since they will need to be constructed based upon the operator settings. 
+					// TODO: Need to ponder how inputs are initialized, but template specialization might do the trick. 
+
+					const FInputNodeType& InputNode = static_cast<const FInputNodeType&>(InNode);
+
+					return MakeUnique<FInputOperator>(InputNode.GetVertexName(), DataRef);
+				}
+
+			private:
+				DataType Data;
+		};
+
+		public:
+			// Constructors that need to get working on linux:
+			/*
+			template<typename = typename TEnableIf< TIsConstructible<DataType>::Value >::Type>
+			TInputNode(const FString& InNodeDescription, const FString& InVertexName)
+			:	NodeDescription(InNodeDescription)
+			,	VertexName(InVertexName)
+			,   Factory()
+			{
+				FOutputDataVertex OutputVertex = MakeOutputDataVertex<DataType>(VertexName, FText::GetEmpty());
+				FDataVertexKey OutputVertexKey = MakeDataVertexKey(OutputVertex);
+
+				Outputs.Add(OutputVertexKey, OutputVertex);
+
+				FInputDataVertex InputVertex = MakeInputDataVertex<DataType>(VertexName, FText::GetEmpty());
+				FDataVertexKey InputVertexKey = MakeDataVertexKey(InputVertex);
+
+				Inputs.Add(InputVertexKey, InputVertex);
+
+			}
+			*/
+
+			template<typename... ArgTypes>
+			TInputNode(const FString& InNodeDescription, const FString& InVertexName, ArgTypes&&... Args)
+			:	NodeDescription(InNodeDescription)
+			,	VertexName(InVertexName)
+			,	Factory(Forward<ArgTypes>(Args)...)
+			{
+				static_assert(TIsConstructible<DataType, ArgTypes...>::Value, "Tried to construct TInputNode<DataType> with arguments that don't match any constructor for DataType.");
+
+				FOutputDataVertex OutputVertex = MakeOutputDataVertex<DataType>(VertexName, FText::GetEmpty());
+				FDataVertexKey OutputVertexKey = MakeDataVertexKey(OutputVertex);
+
+				Outputs.Add(OutputVertexKey, OutputVertex);
+
+				FInputDataVertex InputVertex = MakeInputDataVertex<DataType>(VertexName, FText::GetEmpty());
+				FDataVertexKey InputVertexKey = MakeDataVertexKey(InputVertex);
+
+				Inputs.Add(InputVertexKey, InputVertex);
+
+			}
+
+			explicit TInputNode(const FString& InNodeDescription, const FString& InVertexName, FDataTypeLiteralParam InParam, const FOperatorSettings& InSettings)
+				: NodeDescription(InNodeDescription)
+				, VertexName(InVertexName)
+				, Factory(InParam, InSettings)
+			{
+				FOutputDataVertex OutputVertex = MakeOutputDataVertex<DataType>(VertexName, FText::GetEmpty());
+				FDataVertexKey OutputVertexKey = MakeDataVertexKey(OutputVertex);
+
+				Outputs.Add(OutputVertexKey, OutputVertex);
+
+				FInputDataVertex InputVertex = MakeInputDataVertex<DataType>(VertexName, FText::GetEmpty());
+				FDataVertexKey InputVertexKey = MakeDataVertexKey(InputVertex);
+
+				Inputs.Add(InputVertexKey, InputVertex);
+			}
+
+			virtual const FString& GetInstanceName() const override
+			{
+				return NodeDescription;
+			}
+
+			virtual const FName& GetClassName() const override
+			{
+				// TODO: although this is ok with MSVC's lax template instantiation, every other compiler will complain about TDataReferenceTypeInfo.
+				//static const FName ClassName = FName(FString(TEXT("Input_")) + FString(TDataReferenceTypeInfo<DataType>::TypeName));
+				static const FName ClassName("InputNode");
+				return ClassName;
+			}
+
+			virtual const FText& GetDescription() const override
+			{
+				static const FText Description = LOCTEXT("Metasound_InputNodeDescription", "Input into the parent Metasound graph.");
+				return Description;
+			}
+
+			virtual const FText& GetAuthorName() const override
+			{
+				return PluginAuthor;
+			}
+
+			virtual const FText& GetPromptIfMissing() const override
+			{
+				return PluginNodeMissingPrompt;
+			}
+
+			const FString& GetVertexName() const
+			{
+				return VertexName;
+			}
+
+			virtual const FInputDataVertexCollection& GetInputDataVertices() const override
+			{
+				return Inputs;
+			}
+
+			virtual const FOutputDataVertexCollection& GetOutputDataVertices() const override
+			{
+				return Outputs;
+			}
+
+			virtual IOperatorFactory& GetDefaultOperatorFactory() override
+			{
+				return Factory;
+			}
+
+		private:
+			FString NodeDescription;
+			FString VertexName;
+
+			FCopyOperatorFactory Factory;
+
+			FInputDataVertexCollection Inputs;
+			FOutputDataVertexCollection Outputs;
+	};
+} // namespace Metasound
+#undef LOCTEXT_NAMESPACE //MetasoundOutputNode
diff --git a/Engine/Plugins/Runtime/Metasound/Source/MetasoundFrontend/Public/MetasoundNodeRegistrationMacro.h b/Engine/Plugins/Runtime/Metasound/Source/MetasoundFrontend/Public/MetasoundNodeRegistrationMacro.h
new file mode 100644
--- /dev/null
+++ b/Engine/Plugins/Runtime/Metasound/Source/MetasoundFrontend/Public/MetasoundNodeRegistrationMacro.h
@@ -0,0 +1,52 @@
+// Copyright Epic Games, Inc. All Rights Reserved.
+
+#pragma once
+
+#include "CoreMinimal.h"
+#include "MetasoundNodeInterface.h"
+#include "Templates/IntegralConstant.h"
+#include "MetasoundFrontendRegistries.h"
+
+// Utility class to ensure that a node class can use the constructor the frontend uses.
+template <typename NodeClass>
+struct ConstructorTakesNodeInitData {
+
+	// Use SFINAE trick to see if we have a valid constructor:
+	template<typename T>
+	static uint16 TestForConstructor(decltype(T(::Metasound::FNodeInitData()))*);
+
+	template<typename T>
+	static uint8 TestForConstructor(...);
+
+	static const bool Value = sizeof(TestForConstructor<NodeClass>(nullptr)) == sizeof(uint16);
+};
+
+template <typename FNodeType>
+bool RegisterNodeWithFrontend()
+{
+	// if we reenter this code (because DECLARE_METASOUND_DATA_REFERENCE_TYPES was called twice with the same type),
+	// we catch it here.
+	static bool bAlreadyRegisteredThisDataType = false;
+	if (bAlreadyRegisteredThisDataType)
+	{
+		UE_LOG(LogTemp, Display, TEXT("Tried to call REGISTER_METASOUND_DATATYPE twice with the same class. ignoring the second call. Likely because METASOUND_REGISTER_NODE is in a header that's used in multiple modules. Consider moving it to a private header or cpp file."));
+		return false;
+	}
+
+	bAlreadyRegisteredThisDataType = true;
+
+	bool bSuccessfullyRegisteredNode = FMetasoundFrontendRegistryContainer::Get()->RegisterExternalNode(
+		[](const Metasound::FNodeInitData& InInitData) -> TUniquePtr<Metasound::INode>
+		{
+			return TUniquePtr<Metasound::INode>(new FNodeType(InInitData));
+		});
+
+	ensureAlwaysMsgf(bSuccessfullyRegisteredNode, TEXT("Registering node class failed. Please check the logs."));
+	return bSuccessfullyRegisteredNode;
+}
+
+#define METASOUND_REGISTER_NODE(NodeClass) \
+	 static_assert(std::is_base_of<::Metasound::INodeBase, NodeClass>::value, "To be registered as a  Metasound Node," #NodeClass "need to be a derived class from Metasound::INodeBase, Metasound::INode, or Metasound::FNode."); \
+	 static_assert(::ConstructorTakesNodeInitData<NodeClass>::Value, "In order to be registered as a Metasound Node, " #NodeClass " needs to implement the following public constructor: " #NodeClass "(const Metasound::FNodeInitData& InInitData);"); \
+	 static bool bSuccessfullyRegistered##NodeClass  = FMetasoundFrontendRegistryContainer::Get()->EnqueueInitCommand([](){ ::RegisterNodeWithFrontend<NodeClass>(); }); // This static bool is useful for debugging, but also is the only way the compiler will let us call this function outside of an expression.
+
diff --git a/Engine/Plugins/Runtime/Metasound/Source/MetasoundFrontend/Public/MetasoundOutputNode.h b/Engine/Plugins/Runtime/Metasound/Source/MetasoundFrontend/Public/MetasoundOutputNode.h
new file mode 100644
--- /dev/null
+++ b/Engine/Plugins/Runtime/Metasound/Source/MetasoundFrontend/Public/MetasoundOutputNode.h
@@ -0,0 +1,147 @@
+// Copyright Epic Games, Inc. All Rights Reserved.
+#pragma once
+
+#include "CoreMinimal.h"
+#include "MetasoundBuilderInterface.h"
+#include "MetasoundNodeInterface.h"
+#include "MetasoundOperatorInterface.h"
+
+#define LOCTEXT_NAMESPACE "MetasoundGraphCore"
+
+
+namespace Metasound
+{
+	template<typename DataType>
+	class TOutputNode : public INode
+	{
+		static_assert(TDataReferenceTypeInfo<DataType>::bIsValidSpecialization, "Please use DECLARE_METASOUND_DATA_REFERENCE_TYPES with this class before trying to create an output node with it.");
+			
+		class FOutputOperator : public IOperator
+		{
+			public:
+				using FDataReadReference = TDataReadReference<DataType>;
+
+				FOutputOperator(const FString& InDataReferenceName, FDataReadReference InDataReference)
+				{
+					Outputs.AddDataReadReference<DataType>(InDataReferenceName, InDataReference);
+				}
+
+				virtual ~FOutputOperator() {}
+
+				virtual const FDataReferenceCollection& GetInputs() const override
+				{
+					return Inputs;
+				}
+
+				virtual const FDataReferenceCollection& GetOutputs() const override
+				{
+					return Outputs;
+				}
+
+				virtual FExecuteFunction GetExecuteFunction() override
+				{
+					return nullptr;
+				}
+
+			private:
+				FDataReferenceCollection Inputs;
+				FDataReferenceCollection Outputs;
+		};
+
+		class FOutputOperatorFactory : public IOperatorFactory
+		{
+			public:
+				FOutputOperatorFactory(const FString& InDataReferenceName)
+				:	DataReferenceName(InDataReferenceName)
+				{
+				}
+
+				virtual TUniquePtr<IOperator> CreateOperator(const INode& InNode, const FOperatorSettings& InOperatorSettings, const FDataReferenceCollection& InInputDataReferences, TArray<TUniquePtr<IOperatorBuildError>>& OutErrors) override
+				{
+					if (!InInputDataReferences.ContainsDataReadReference<DataType>(DataReferenceName))
+					{
+						// TODO: Add build error.
+						return TUniquePtr<IOperator>(nullptr);
+					}
+
+					return MakeUnique<FOutputOperator>(DataReferenceName, InInputDataReferences.GetDataReadReference<DataType>(DataReferenceName));
+				}
+
+			private:
+				FString DataReferenceName;
+		};
+
+		public:
+			TOutputNode(const FString& InNodeDescription, const FString& InVertexName)
+			:	NodeDescription(InNodeDescription)
+			,	VertexName(InVertexName)
+			,	Factory(InVertexName)
+			{
+				FInputDataVertex InputVertex = MakeInputDataVertex<DataType>(VertexName, FText::GetEmpty());
+				FDataVertexKey InputKey = MakeDataVertexKey(InputVertex);
+
+				Inputs.Add(InputKey, InputVertex);
+
+				FOutputDataVertex OutputVertex = MakeOutputDataVertex<DataType>(VertexName, FText::GetEmpty());
+				FDataVertexKey OutputKey = MakeDataVertexKey(OutputVertex);
+
+				Outputs.Add(OutputKey, OutputVertex);
+			}
+
+			virtual const FName& GetClassName() const override
+			{
+				// TODO: Any special formatting for these node type names?
+				// TODO: although this is ok with MSVC's lax template instantiation, every other compiler will complain about TDataReferenceTypeInfo.
+				//static const FName ClassName = FName(FString(TEXT("Input_")) + FString(TDataReferenceTypeInfo<DataType>::TypeName));
+				static const FName ClassName("OutputNode");
+				return ClassName;
+			}
+
+			virtual const FString& GetInstanceName() const override
+			{
+				return NodeDescription;
+			}
+
+			virtual const FText& GetDescription() const override
+			{
+				static const FText Description = LOCTEXT("Metasound_InputNodeDescription", "Output from the parent Metasound graph.");
+				return Description;
+			}
+
+			virtual const FText& GetAuthorName() const override
+			{
+				return PluginAuthor;
+			}
+
+			virtual const FText& GetPromptIfMissing() const override
+			{
+				return PluginNodeMissingPrompt;
+			}
+
+			virtual const FInputDataVertexCollection& GetInputDataVertices() const override
+			{
+				return Inputs;
+			}
+
+			virtual const FOutputDataVertexCollection& GetOutputDataVertices() const override
+			{
+				return Outputs;
+			}
+
+			virtual IOperatorFactory& GetDefaultOperatorFactory() override
+			{
+				return Factory;
+			}
+
+		private:
+			FString NodeDescription;
+			FString VertexName;
+
+			FOutputOperatorFactory Factory;
+
+
+			FInputDataVertexCollection Inputs;
+			FOutputDataVertexCollection Outputs;
+	};
+}
+#undef LOCTEXT_NAMESPACE // MetasoundOutputNode
diff --git a/Engine/Plugins/Runtime/Metasound/Source/MetasoundGenerator/MetasoundGenerator.Build.cs b/Engine/Plugins/Runtime/Metasound/Source/MetasoundGenerator/MetasoundGenerator.Build.cs
--- a/Engine/Plugins/Runtime/Metasound/Source/MetasoundGenerator/MetasoundGenerator.Build.cs
+++ b/Engine/Plugins/Runtime/Metasound/Source/MetasoundGenerator/MetasoundGenerator.Build.cs
@@ -24,7 +24,8 @@ namespace UnrealBuildTool.Rules
 				new string[] {
                     "Core",
 					"AudioMixer",
-					"MetasoundGraphCore"
+					"MetasoundGraphCore",
+					"MetasoundStandardNodes"
                 }
             );
 		}
diff --git a/Engine/Plugins/Runtime/Metasound/Source/MetasoundGraphCore/Private/MetasoundADSRNode.cpp b/Engine/Plugins/Runtime/Metasound/Source/MetasoundGraphCore/Private/MetasoundADSRNode.cpp
deleted file mode 100644
diff --git a/Engine/Plugins/Runtime/Metasound/Source/MetasoundGraphCore/Private/MetasoundAudioFormats.cpp b/Engine/Plugins/Runtime/Metasound/Source/MetasoundGraphCore/Private/MetasoundAudioFormats.cpp
deleted file mode 100644
diff --git a/Engine/Plugins/Runtime/Metasound/Source/MetasoundGraphCore/Private/MetasoundAudioMultiplyNode.cpp b/Engine/Plugins/Runtime/Metasound/Source/MetasoundGraphCore/Private/MetasoundAudioMultiplyNode.cpp
deleted file mode 100644
diff --git a/Engine/Plugins/Runtime/Metasound/Source/MetasoundGraphCore/Private/MetasoundFrontendRegistries.cpp b/Engine/Plugins/Runtime/Metasound/Source/MetasoundGraphCore/Private/MetasoundFrontendRegistries.cpp
deleted file mode 100644
diff --git a/Engine/Plugins/Runtime/Metasound/Source/MetasoundGraphCore/Private/MetasoundGraphCoreModule.cpp b/Engine/Plugins/Runtime/Metasound/Source/MetasoundGraphCore/Private/MetasoundGraphCoreModule.cpp
--- a/Engine/Plugins/Runtime/Metasound/Source/MetasoundGraphCore/Private/MetasoundGraphCoreModule.cpp
+++ b/Engine/Plugins/Runtime/Metasound/Source/MetasoundGraphCore/Private/MetasoundGraphCoreModule.cpp
@@ -5,34 +5,9 @@
 #include "CoreMinimal.h"
 #include "Modules/ModuleInterface.h"
 #include "Modules/ModuleManager.h"
-#include "MetasoundPrimitives.h"
-#include "MetasoundDataReferenceTypes.h"
-#include "MetasoundTime.h"
-#include "MetasoundFrequency.h"
-#include "MetasoundWave.h"
 
 DEFINE_LOG_CATEGORY(LogMetasoundGraphCore);
 
-REGISTER_METASOUND_DATATYPE(bool, "Primitive:Bool", ::Metasound::ELiteralArgType::Boolean)
-
-REGISTER_METASOUND_DATATYPE(int32, "Primitive:Int32", ::Metasound::ELiteralArgType::Integer)
-
-REGISTER_METASOUND_DATATYPE(int64, "Primitive:Int64", ::Metasound::ELiteralArgType::Integer)
-
-REGISTER_METASOUND_DATATYPE(float, "Primitive:Float", ::Metasound::ELiteralArgType::Float)
-
-REGISTER_METASOUND_DATATYPE(double, "Primitive:Double", ::Metasound::ELiteralArgType::Float)
-
-namespace Metasound
-{
-	REGISTER_METASOUND_DATATYPE(FBop, "Primitive:Bop", ::Metasound::ELiteralArgType::Boolean)
-	REGISTER_METASOUND_DATATYPE(FFloatTime, "Primitive:Time")
-	REGISTER_METASOUND_DATATYPE(FDoubleTime, "Primitive:Time:HighResolution")
-	REGISTER_METASOUND_DATATYPE(FSampleTime, "Primitive:Time:SampleResolution")
-	REGISTER_METASOUND_DATATYPE(FFrequency, "Primitive:Frequency")
-	REGISTER_METASOUND_DATATYPE(FWave, "Primitive:Wave")
-}
-
 namespace Metasound 
 {
 	class FMetasoundGraphCoreModule : public IMetasoundGraphCoreModule
diff --git a/Engine/Plugins/Runtime/Metasound/Source/MetasoundGraphCore/Private/MetasoundOscNode.cpp b/Engine/Plugins/Runtime/Metasound/Source/MetasoundGraphCore/Private/MetasoundOscNode.cpp
deleted file mode 100644
diff --git a/Engine/Plugins/Runtime/Metasound/Source/MetasoundGraphCore/Private/MetasoundPeriodicBopNode.cpp b/Engine/Plugins/Runtime/Metasound/Source/MetasoundGraphCore/Private/MetasoundPeriodicBopNode.cpp
deleted file mode 100644
diff --git a/Engine/Plugins/Runtime/Metasound/Source/MetasoundGraphCore/Private/MetasoundWavePlayerNode.cpp b/Engine/Plugins/Runtime/Metasound/Source/MetasoundGraphCore/Private/MetasoundWavePlayerNode.cpp
deleted file mode 100644
diff --git a/Engine/Plugins/Runtime/Metasound/Source/MetasoundGraphCore/Public/MetasoundADSRNode.h b/Engine/Plugins/Runtime/Metasound/Source/MetasoundGraphCore/Public/MetasoundADSRNode.h
deleted file mode 100644
diff --git a/Engine/Plugins/Runtime/Metasound/Source/MetasoundGraphCore/Public/MetasoundAudioBuffer.h b/Engine/Plugins/Runtime/Metasound/Source/MetasoundGraphCore/Public/MetasoundAudioBuffer.h
deleted file mode 100644
diff --git a/Engine/Plugins/Runtime/Metasound/Source/MetasoundGraphCore/Public/MetasoundAudioFormats.h b/Engine/Plugins/Runtime/Metasound/Source/MetasoundGraphCore/Public/MetasoundAudioFormats.h
deleted file mode 100644
diff --git a/Engine/Plugins/Runtime/Metasound/Source/MetasoundGraphCore/Public/MetasoundAudioMultiplyNode.h b/Engine/Plugins/Runtime/Metasound/Source/MetasoundGraphCore/Public/MetasoundAudioMultiplyNode.h
deleted file mode 100644
diff --git a/Engine/Plugins/Runtime/Metasound/Source/MetasoundGraphCore/Public/MetasoundDataReference.h b/Engine/Plugins/Runtime/Metasound/Source/MetasoundGraphCore/Public/MetasoundDataReference.h
--- a/Engine/Plugins/Runtime/Metasound/Source/MetasoundGraphCore/Public/MetasoundDataReference.h
+++ b/Engine/Plugins/Runtime/Metasound/Source/MetasoundGraphCore/Public/MetasoundDataReference.h
@@ -14,7 +14,7 @@
 	template<> \
 	struct ::Metasound::TDataReferenceTypeInfo<DataType> \
 	{ \
-		static METASOUNDGRAPHCORE_API const TCHAR* TypeName; \
+		static METASOUNDSTANDARDNODES_API const TCHAR* TypeName; \
 		static constexpr ::Metasound::FDataTypeMagicNumber MagicNumber = (DataTypeMagicNumber); \
 		static constexpr bool bIsStringParsable = TTestIfDataTypeCtorIsImplemented<DataType, const FString&>::Value; \
 		static constexpr bool bIsBoolParsable = TTestIfDataTypeCtorIsImplemented<DataType, bool>::Value; \
diff --git a/Engine/Plugins/Runtime/Metasound/Source/MetasoundGraphCore/Public/MetasoundDataReferenceTypes.h b/Engine/Plugins/Runtime/Metasound/Source/MetasoundGraphCore/Public/MetasoundDataReferenceTypes.h
deleted file mode 100644
diff --git a/Engine/Plugins/Runtime/Metasound/Source/MetasoundGraphCore/Public/MetasoundDataTypeRegistrationMacro.h b/Engine/Plugins/Runtime/Metasound/Source/MetasoundGraphCore/Public/MetasoundDataTypeRegistrationMacro.h
deleted file mode 100644
diff --git a/Engine/Plugins/Runtime/Metasound/Source/MetasoundGraphCore/Public/MetasoundFrequency.h b/Engine/Plugins/Runtime/Metasound/Source/MetasoundGraphCore/Public/MetasoundFrequency.h
deleted file mode 100644
diff --git a/Engine/Plugins/Runtime/Metasound/Source/MetasoundGraphCore/Public/MetasoundFrontendRegistries.h b/Engine/Plugins/Runtime/Metasound/Source/MetasoundGraphCore/Public/MetasoundFrontendRegistries.h
deleted file mode 100644
diff --git a/Engine/Plugins/Runtime/Metasound/Source/MetasoundGraphCore/Public/MetasoundInputNode.h b/Engine/Plugins/Runtime/Metasound/Source/MetasoundGraphCore/Public/MetasoundInputNode.h
deleted file mode 100644
diff --git a/Engine/Plugins/Runtime/Metasound/Source/MetasoundGraphCore/Public/MetasoundNode.h b/Engine/Plugins/Runtime/Metasound/Source/MetasoundGraphCore/Public/MetasoundNode.h
--- a/Engine/Plugins/Runtime/Metasound/Source/MetasoundGraphCore/Public/MetasoundNode.h
+++ b/Engine/Plugins/Runtime/Metasound/Source/MetasoundGraphCore/Public/MetasoundNode.h
@@ -4,7 +4,6 @@
 
 #include "CoreMinimal.h"
 #include "MetasoundNodeInterface.h"
-#include "MetasoundNodeRegistrationMacro.h"
 
 namespace Metasound
 {
diff --git a/Engine/Plugins/Runtime/Metasound/Source/MetasoundGraphCore/Public/MetasoundNodeRegistrationMacro.h b/Engine/Plugins/Runtime/Metasound/Source/MetasoundGraphCore/Public/MetasoundNodeRegistrationMacro.h
deleted file mode 100644
diff --git a/Engine/Plugins/Runtime/Metasound/Source/MetasoundGraphCore/Public/MetasoundOscNode.h b/Engine/Plugins/Runtime/Metasound/Source/MetasoundGraphCore/Public/MetasoundOscNode.h
deleted file mode 100644
diff --git a/Engine/Plugins/Runtime/Metasound/Source/MetasoundGraphCore/Public/MetasoundOutputNode.h b/Engine/Plugins/Runtime/Metasound/Source/MetasoundGraphCore/Public/MetasoundOutputNode.h
deleted file mode 100644
diff --git a/Engine/Plugins/Runtime/Metasound/Source/MetasoundGraphCore/Public/MetasoundPeriodicBopNode.h b/Engine/Plugins/Runtime/Metasound/Source/MetasoundGraphCore/Public/MetasoundPeriodicBopNode.h
deleted file mode 100644
diff --git a/Engine/Plugins/Runtime/Metasound/Source/MetasoundGraphCore/Public/MetasoundPrimitives.h b/Engine/Plugins/Runtime/Metasound/Source/MetasoundGraphCore/Public/MetasoundPrimitives.h
deleted file mode 100644
diff --git a/Engine/Plugins/Runtime/Metasound/Source/MetasoundGraphCore/Public/MetasoundTime.h b/Engine/Plugins/Runtime/Metasound/Source/MetasoundGraphCore/Public/MetasoundTime.h
deleted file mode 100644
diff --git a/Engine/Plugins/Runtime/Metasound/Source/MetasoundGraphCore/Public/MetasoundWave.h b/Engine/Plugins/Runtime/Metasound/Source/MetasoundGraphCore/Public/MetasoundWave.h
deleted file mode 100644
diff --git a/Engine/Plugins/Runtime/Metasound/Source/MetasoundGraphCore/Public/MetasoundWavePlayerNode.h b/Engine/Plugins/Runtime/Metasound/Source/MetasoundGraphCore/Public/MetasoundWavePlayerNode.h
deleted file mode 100644
diff --git a/Engine/Plugins/Runtime/Metasound/Source/MetasoundStandardNodes/MetasoundStandardNodes.Build.cs b/Engine/Plugins/Runtime/Metasound/Source/MetasoundStandardNodes/MetasoundStandardNodes.Build.cs
new file mode 100644
--- /dev/null
+++ b/Engine/Plugins/Runtime/Metasound/Source/MetasoundStandardNodes/MetasoundStandardNodes.Build.cs
@@ -0,0 +1,34 @@
+// Copyright Epic Games, Inc. All Rights Reserved.
+
+using UnrealBuildTool;
+using System.IO;
+
+namespace UnrealBuildTool.Rules
+{
+	public class MetasoundStandardNodes : ModuleRules
+	{
+		public MetasoundStandardNodes(ReadOnlyTargetRules Target) : base(Target)
+		{
+			PublicDependencyModuleNames.AddRange
+			(
+				new string[]
+				{
+					"Core",
+					"Serialization",
+					"SignalProcessing",
+					"AudioExtensions",
+					"MetasoundFrontend"
+				}
+			);
+
+			PrivateDependencyModuleNames.AddRange
+			(
+				new string[]
+				{
+					"MetasoundGraphCore",
+					"CoreUObject"
+				}
+			);
+		}
+	}
+}
diff --git a/Engine/Plugins/Runtime/Metasound/Source/MetasoundStandardNodes/Private/MetasoundADSRNode.cpp b/Engine/Plugins/Runtime/Metasound/Source/MetasoundStandardNodes/Private/MetasoundADSRNode.cpp
new file mode 100644
--- /dev/null
+++ b/Engine/Plugins/Runtime/Metasound/Source/MetasoundStandardNodes/Private/MetasoundADSRNode.cpp
@@ -0,0 +1,215 @@
+// Copyright Epic Games, Inc. All Rights Reserved.
+
+#include "MetasoundADSRNode.h"
+#include "MetasoundExecutableOperator.h"
+#include "MetasoundTime.h"
+#include "MetasoundDataReferenceTypes.h"
+#include "MetasoundNodeRegistrationMacro.h"
+#include "DSP/Envelope.h"
+
+#define LOCTEXT_NAMESPACE "MetasoundADSRNode"
+
+
+
+namespace Metasound
+{
+	METASOUND_REGISTER_NODE(FADSRNode)
+
+	struct FADSRDataReferences
+	{
+		FFloatTimeReadRef Attack;
+		FFloatTimeReadRef Decay;
+		FFloatTimeReadRef Sustain;
+		FFloatTimeReadRef Release;
+	};
+
+	class FADSROperator : public TExecutableOperator<FADSROperator>
+	{
+		public:
+			FADSROperator(const FOperatorSettings& InSettings, const FBopReadRef& InBop, const FADSRDataReferences& InADSRData)
+			:	OperatorSettings(InSettings)
+			,	Bop(InBop)
+			,	StopEnvelopePos(-1)
+			,	ADSRDataReferences(InADSRData)
+			,	EnvelopeBuffer(FAudioBufferWriteRef::CreateNew(InSettings.GetNumFramesPerBlock()))
+			{
+				check(EnvelopeBuffer->Num() == InSettings.GetNumFramesPerBlock());
+
+				OutputDataReferences.AddDataReadReference(TEXT("Envelope"), FAudioBufferReadRef(EnvelopeBuffer));
+
+				Envelope.Init(OperatorSettings.GetSampleRate());
+			}
+
+			virtual const FDataReferenceCollection& GetInputs() const override
+			{
+				return InputDataReferences;
+			}
+
+			virtual const FDataReferenceCollection& GetOutputs() const override
+			{
+				return OutputDataReferences;
+			}
+
+			void Execute()
+			{
+				Envelope.SetAttackTime(ADSRDataReferences.Attack->GetMilliseconds());
+				Envelope.SetDecayTime(ADSRDataReferences.Decay->GetMilliseconds());
+				Envelope.SetReleaseTime(ADSRDataReferences.Release->GetMilliseconds());
+
+				float ADSMilliseconds = ADSRDataReferences.Attack->GetMilliseconds();
+				ADSMilliseconds += ADSRDataReferences.Decay->GetMilliseconds();
+				ADSMilliseconds += ADSRDataReferences.Sustain->GetMilliseconds();
+
+				// TODO: Bops need to be sorted.
+				// TODO: Make convenient way to do bop loops like this since it might happen alot.
+				// TODO: Bops need to be done on 4-sample boundaries to allow for SIMD.
+			
+
+				int32 BopIndex = 0;
+				int32 StartPos = 0;
+				int32 NextBop = BopIndex < Bop->Num() ? (*Bop)[BopIndex] : OperatorSettings.GetNumFramesPerBlock() + 1;
+				int32 EndPos = FMath::Min(OperatorSettings.GetNumFramesPerBlock(), NextBop);
+				
+				if (StopEnvelopePos > 0)
+				{
+					StopEnvelopePos -= OperatorSettings.GetNumFramesPerBlock();
+					EndPos = FMath::Min(StopEnvelopePos, EndPos);
+				}
+
+				float* EnvelopeData = EnvelopeBuffer->GetData();
+
+				// TODO: GetNumFramesPerBlock() need to be on 4 sample boundaries to allow for SIMD.
+				// TODO: just make GetNumFramesPerBlock() an int32 instead of a uint32
+				while (StartPos < (int32)OperatorSettings.GetNumFramesPerBlock())
+				{
+					int32 i = StartPos;
+					for (; i < EndPos; i++)
+					{
+						EnvelopeData[i] = Envelope.Generate();
+					}
+
+					StartPos = EndPos;
+					
+					if (EndPos == StopEnvelopePos)
+					{
+						// Stop sustain, start release.
+						Envelope.Stop();
+						StopEnvelopePos = -1;
+						EndPos = FMath::Min((int32)OperatorSettings.GetNumFramesPerBlock(), NextBop);
+					}
+
+					if (EndPos == NextBop)
+					{
+						// Trigger next envelope. 
+						Envelope.Start();
+						StopEnvelopePos = EndPos + FMath::RoundToInt(OperatorSettings.GetSampleRate() * 0.001f * ADSMilliseconds);
+
+						BopIndex++;
+
+						NextBop = BopIndex < Bop->Num() ? (*Bop)[BopIndex] : OperatorSettings.GetNumFramesPerBlock() + 1;
+						EndPos = FMath::Min(StopEnvelopePos, FMath::Min((int32)OperatorSettings.GetNumFramesPerBlock(), NextBop));
+					}
+				}
+			}
+
+		private:
+			const FOperatorSettings OperatorSettings;
+
+			// TODO: write envelope gen for metasound more suited to this processing structure. 
+			Audio::FEnvelope Envelope;
+
+			FBopReadRef Bop;
+			int32 StopEnvelopePos;
+			FADSRDataReferences ADSRDataReferences;
+			FAudioBufferWriteRef EnvelopeBuffer;
+
+			FDataReferenceCollection OutputDataReferences;
+			FDataReferenceCollection InputDataReferences;
+	};
+
+	const FName FADSRNode::ClassName = FName(TEXT("ADSR"));
+
+	TUniquePtr<IOperator> FADSRNode::FOperatorFactory::CreateOperator(const INode& InNode, const FOperatorSettings& InOperatorSettings, const FDataReferenceCollection& InInputDataReferences, TArray<TUniquePtr<IOperatorBuildError>>& OutErrors) 
+	{
+		const FADSRNode& ADSRNode = static_cast<const FADSRNode&>(InNode);
+
+		FBopReadRef Bop = FBopReadRef::CreateNew();
+
+		// TODO: Could no-op this if the bop is not connected.
+		SetReadableRefIfInCollection(TEXT("Bop"), InInputDataReferences, Bop);
+
+		FADSRDataReferences ADSRDataReferences = 
+			{
+				FFloatTimeReadRef::CreateNew(ADSRNode.GetDefaultAttackMs(), ETimeResolution::Milliseconds),
+				FFloatTimeReadRef::CreateNew(ADSRNode.GetDefaultDecayMs(), ETimeResolution::Milliseconds),
+				FFloatTimeReadRef::CreateNew(ADSRNode.GetDefaultSustainMs(), ETimeResolution::Milliseconds),
+				FFloatTimeReadRef::CreateNew(ADSRNode.GetDefaultReleaseMs(), ETimeResolution::Milliseconds)
+			};
+
+		// TODO: If none of these are connected, could pregenerate ADSR envelope and return a different operator. 
+		SetReadableRefIfInCollection(TEXT("Attack"), InInputDataReferences, ADSRDataReferences.Attack);
+		SetReadableRefIfInCollection(TEXT("Decay"), InInputDataReferences, ADSRDataReferences.Decay);
+		SetReadableRefIfInCollection(TEXT("Sustain"), InInputDataReferences, ADSRDataReferences.Sustain);
+		SetReadableRefIfInCollection(TEXT("Release"), InInputDataReferences, ADSRDataReferences.Release);
+
+		return MakeUnique<FADSROperator>(InOperatorSettings, Bop, ADSRDataReferences);
+	}
+
+	FADSRNode::FADSRNode(const FString& InName, float InDefaultAttackMs, float InDefaultDecayMs, float InDefaultSustainMs, float InDefaultReleaseMs)
+	:	FNode(InName)
+	,	DefaultAttackMs(InDefaultAttackMs)
+	,	DefaultDecayMs(InDefaultDecayMs)
+	,	DefaultSustainMs(InDefaultSustainMs)
+	,	DefaultReleaseMs(InDefaultReleaseMs)
+	{
+		AddInputDataVertex<FBop>(TEXT("Bop"), LOCTEXT("BopTooltip", "Trigger for envelope."));
+		AddInputDataVertex<FFloatTime>(TEXT("Attack"), LOCTEXT("AttackTooltip", "Attack time in milliseconds."));
+		AddInputDataVertex<FFloatTime>(TEXT("Decay"), LOCTEXT("DecayTooltip", "Decay time in milliseconds."));
+		AddInputDataVertex<FFloatTime>(TEXT("Sustain"), LOCTEXT("SustainTooltip", "Sustain time in milliseconds."));
+		AddInputDataVertex<FFloatTime>(TEXT("Release"), LOCTEXT("ReleaseTooltip", "Release time in milliseconds."));
+
+		AddOutputDataVertex<FAudioBuffer>(TEXT("Envelope"), LOCTEXT("EnvelopeTooltip", "The output envelope"));
+	}
+
+	FADSRNode::FADSRNode(const FNodeInitData& InitData)
+		: FADSRNode(InitData.InstanceName, 10.0f, 20.0f, 50.0f, 20.0f)
+	{
+	}
+
+	FADSRNode::~FADSRNode()
+	{
+	}
+
+
+	float FADSRNode::GetDefaultAttackMs() const
+	{
+		return DefaultAttackMs;
+	}
+
+	float FADSRNode::GetDefaultDecayMs() const
+	{
+		return DefaultDecayMs;
+	}
+
+	float FADSRNode::GetDefaultSustainMs() const
+	{
+		return DefaultSustainMs;
+	}
+
+	float FADSRNode::GetDefaultReleaseMs() const
+	{
+		return DefaultReleaseMs;
+	}
+
+	const FName& FADSRNode::GetClassName() const
+	{
+		return ClassName;
+	}
+
+	IOperatorFactory& FADSRNode::GetDefaultOperatorFactory() 
+	{
+		return Factory;
+	}
+}
+
+#undef LOCTEXT_NAMESPACE
diff --git a/Engine/Plugins/Runtime/Metasound/Source/MetasoundStandardNodes/Private/MetasoundAudioFormats.cpp b/Engine/Plugins/Runtime/Metasound/Source/MetasoundStandardNodes/Private/MetasoundAudioFormats.cpp
new file mode 100644
--- /dev/null
+++ b/Engine/Plugins/Runtime/Metasound/Source/MetasoundStandardNodes/Private/MetasoundAudioFormats.cpp
@@ -0,0 +1,142 @@
+// Copyright Epic Games, Inc. All Rights Reserved.
+
+#include "MetasoundAudioFormats.h"
+
+#include "CoreMinimal.h"
+#include "MetasoundDataReference.h"
+#include "MetasoundAudioBuffer.h"
+
+struct FMyType
+{
+public:
+	FMyType() = delete;
+
+private:
+	FMyType Get();
+	friend class FMyOtherType;
+};
+
+namespace Metasound
+{
+	REGISTER_METASOUND_DATATYPE(FAudioBuffer, "Audio:Buffer")
+	REGISTER_METASOUND_DATATYPE(FUnformattedAudio, "Audio:Unformatted")
+	REGISTER_METASOUND_DATATYPE(FMultichannelAudioFormat, "Audio:Multichannel")
+
+
+
+	// TODO: Neither of these are trivially constructible, becuase they are actually views into a preexisting buffer.
+	//REGISTER_METASOUND_DATATYPE(FMonoAudioFormat)
+	//REGISTER_METASOUND_DATATYPE(FStereoAudioFormat)
+	IMPL_METASOUND_DATA_TYPE(FMonoAudioFormat, "Audio:Mono")
+	IMPL_METASOUND_DATA_TYPE(FStereoAudioFormat, "Audio:Stereo")
+
+	/* FUnformattedAudio */
+
+	FUnformattedAudio::FUnformattedAudio(int32 InNumFrames, int32 InNumChannels, int32 InMaxNumChannels)
+	:	NumFrames(InNumFrames)
+	,	NumChannels(0)
+	,	MaxNumChannels(InMaxNumChannels)
+	{
+		NumFrames = FMath::Max(NumFrames, 0);
+		MaxNumChannels = FMath::Max(MaxNumChannels, 0);
+
+		for (int32 i = 0; i < MaxNumChannels; i++)
+		{
+			FAudioBufferWriteRef Audio = FAudioBufferWriteRef::CreateNew(NumFrames);
+			Audio->Zero();
+
+			WritableBufferStorage.Add(Audio);
+			ReadableBufferStorage.Add(Audio);
+		}
+
+		SetNumChannels(InNumChannels);
+	}
+
+	FUnformattedAudio::FUnformattedAudio(int32 InInitialNumChannels, const FOperatorSettings& InSettings)
+		: FUnformattedAudio(InSettings.GetNumFramesPerBlock(), InInitialNumChannels, 8)
+	{
+	}
+
+	int32 FUnformattedAudio::SetNumChannels(int32 InNumChannels)
+	{
+		NumChannels = FMath::Max(0, FMath::Min(InNumChannels, MaxNumChannels));
+
+		ReadableBuffers = TArrayView<const FAudioBufferReadRef>(ReadableBufferStorage.GetData(), NumChannels);
+		WritableBuffers = TArrayView<const FAudioBufferWriteRef>(WritableBufferStorage.GetData(), NumChannels);
+
+		return NumChannels;
+	}
+
+
+	/* FMultichannelAudioFormat */
+
+	FMultichannelAudioFormat::FMultichannelAudioFormat()
+	:	NumChannels(0)
+	{
+	}
+
+	FMultichannelAudioFormat::FMultichannelAudioFormat(int32 InNumFrames, int32 InNumChannels)
+	:	NumChannels(InNumChannels)
+	{
+		NumChannels = FMath::Max(0, NumChannels);
+		InNumFrames = FMath::Max(0, InNumFrames);
+
+		for (int32 i = 0; i < NumChannels; i++)
+		{
+			FAudioBufferWriteRef Audio = FAudioBufferWriteRef::CreateNew(InNumFrames);
+			Audio->Zero();
+
+			WritableBufferStorage.Add(Audio);
+			ReadableBufferStorage.Add(Audio);
+		}
+
+		WritableBuffers = WritableBufferStorage;
+		ReadableBuffers = ReadableBufferStorage;
+	}
+
+	FMultichannelAudioFormat::FMultichannelAudioFormat(int32 InNumChannels, const FOperatorSettings& InSettings)
+		: FMultichannelAudioFormat(InSettings.GetNumFramesPerBlock(), InNumChannels)
+	{}
+
+
+	FMultichannelAudioFormat::FMultichannelAudioFormat(TArrayView<const FAudioBufferWriteRef> InWriteRefs)
+	:	NumChannels(InWriteRefs.Num())
+	{
+		if (NumChannels > 0)
+		{
+			const int32 NumFrames = InWriteRefs[0]->Num();
+
+			for (const FAudioBufferWriteRef& Ref : InWriteRefs)
+			{
+				checkf(NumFrames == Ref->Num(), TEXT("All buffers must have same number of frames (%d != %d)"), NumFrames, Ref->Num());
+
+				WritableBufferStorage.Add(Ref);
+				ReadableBufferStorage.Add(Ref);
+			}
+
+			WritableBuffers = WritableBufferStorage;
+			ReadableBuffers = ReadableBufferStorage;
+		}
+	}
+
+	FMultichannelAudioFormat::FMultichannelAudioFormat(TArrayView<const FAudioBufferReadRef> InReadRefs)
+	:	NumChannels(InReadRefs.Num())
+	{
+		if (NumChannels > 0)
+		{
+			const int32 NumFrames = InReadRefs[0]->Num();
+
+			for (const FAudioBufferReadRef& Ref : InReadRefs)
+			{
+				checkf(NumFrames == Ref->Num(), TEXT("All buffers must have same number of frames (%d != %d)"), NumFrames, Ref->Num());
+
+				WritableBufferStorage.Add(WriteCast(Ref));
+				ReadableBufferStorage.Add(Ref);
+			}
+
+			WritableBuffers = WritableBufferStorage;
+			ReadableBuffers = ReadableBufferStorage;
+		}
+	}
+
+}
diff --git a/Engine/Plugins/Runtime/Metasound/Source/MetasoundStandardNodes/Private/MetasoundAudioMultiplyNode.cpp b/Engine/Plugins/Runtime/Metasound/Source/MetasoundStandardNodes/Private/MetasoundAudioMultiplyNode.cpp
new file mode 100644
--- /dev/null
+++ b/Engine/Plugins/Runtime/Metasound/Source/MetasoundStandardNodes/Private/MetasoundAudioMultiplyNode.cpp
@@ -0,0 +1,120 @@
+// Copyright Epic Games, Inc. All Rights Reserved.
+
+#include "MetasoundAudioMultiplyNode.h"
+#include "MetasoundExecutableOperator.h"
+#include "MetasoundDataReferenceTypes.h"
+#include "MetasoundNodeRegistrationMacro.h"
+#include "DSP/BufferVectorOperations.h"
+
+#define LOCTEXT_NAMESPACE "MetasoundAudioMultiplyNode"
+
+namespace Metasound
+{
+	METASOUND_REGISTER_NODE(FAudioMultiplyNode);
+
+	class FAudioMultiplyOperator : public TExecutableOperator<FAudioMultiplyOperator>
+	{
+		public:
+			FAudioMultiplyOperator(const FOperatorSettings& InSettings, const FAudioBufferReadRef& InBuffer1, const FAudioBufferReadRef& InBuffer2)
+			:	OperatorSettings(InSettings)
+			,	InputBuffer1(InBuffer1)
+			,	InputBuffer2(InBuffer2)
+			,	OutputBuffer(FAudioBufferWriteRef::CreateNew(InSettings))
+			{
+				check(OutputBuffer->Num() == InSettings.GetNumFramesPerBlock());
+				check(InputBuffer1->Num() == InSettings.GetNumFramesPerBlock());
+				check(InputBuffer2->Num() == InSettings.GetNumFramesPerBlock());
+
+				OutputDataReferences.AddDataReadReference(TEXT("Audio"), FAudioBufferReadRef(OutputBuffer));
+			}
+
+			virtual const FDataReferenceCollection& GetInputs() const override
+			{
+				return InputDataReferences;
+			}
+
+			virtual const FDataReferenceCollection& GetOutputs() const override
+			{
+				return OutputDataReferences;
+			}
+
+			void Execute()
+			{
+				using namespace Audio;
+
+				// TODO: add buffer vector op to multipy two buffers (not in place). 
+				// TODO: what to do about RESTRICT? This could be something for the builder to keep in mind.
+				FMemory::Memcpy(OutputBuffer->GetData(), InputBuffer1->GetData(), sizeof(float) * OperatorSettings.GetNumFramesPerBlock());
+				MultiplyBuffersInPlace(*InputBuffer2, *OutputBuffer);
+			}
+
+		private:
+			const FOperatorSettings OperatorSettings;
+
+			FAudioBufferReadRef InputBuffer1;
+			FAudioBufferReadRef InputBuffer2;
+			FAudioBufferWriteRef OutputBuffer;
+			FDataReferenceCollection OutputDataReferences;
+			FDataReferenceCollection InputDataReferences;
+	};
+
+	const FName FAudioMultiplyNode::ClassName = FName(TEXT("AudioMultiply"));
+
+	TUniquePtr<IOperator> FAudioMultiplyNode::FOperatorFactory::CreateOperator(const INode& InNode, const FOperatorSettings& InOperatorSettings, const FDataReferenceCollection& InInputDataReferences, TArray<TUniquePtr<IOperatorBuildError>>& OutErrors) 
+	{
+		const FAudioMultiplyNode& AudioMultiplyNode = static_cast<const FAudioMultiplyNode&>(InNode);
+
+		FAudioBufferReadRef InputBuffer1 = FAudioBufferReadRef::CreateNew();
+		FAudioBufferReadRef InputBuffer2 = FAudioBufferReadRef::CreateNew();
+
+		FAudioBuffer Ones(InOperatorSettings.GetNumFramesPerBlock());
+
+		// initialize default array to all ones. 
+		float* Data = Ones.GetData();
+		for (int32 i = 0; i < InOperatorSettings.GetNumFramesPerBlock(); i++)
+		{
+			Data[i] = 1.f;
+		}
+
+		if (!SetReadableRefIfInCollection(TEXT("InputBuffer1"), InInputDataReferences, InputBuffer1))
+		{
+			InputBuffer1 = FAudioBufferReadRef::CreateNew(Ones);
+		}
+
+		if (!SetReadableRefIfInCollection(TEXT("InputBuffer2"), InInputDataReferences, InputBuffer2))
+		{
+			InputBuffer2 = FAudioBufferReadRef::CreateNew(Ones);
+		}
+
+		return MakeUnique<FAudioMultiplyOperator>(InOperatorSettings, InputBuffer1, InputBuffer2);
+	}
+
+	FAudioMultiplyNode::FAudioMultiplyNode(const FString& InName)
+	:	FNode(InName)
+	{
+		AddInputDataVertex<FAudioBuffer>(TEXT("InputBuffer1"), LOCTEXT("InputBuffer1Tooltip", "Input buffer to multiply."));
+		AddInputDataVertex<FAudioBuffer>(TEXT("InputBuffer2"), LOCTEXT("InputBuffer2Tooltip", "Input buffer to multiply."));
+
+		AddOutputDataVertex<FAudioBuffer>(TEXT("Audio"), LOCTEXT("OutpuBufferTooltip", "The output audio."));
+	}
+
+	FAudioMultiplyNode::FAudioMultiplyNode(const FNodeInitData& InInitData)
+		: FAudioMultiplyNode(InInitData.InstanceName)
+	{}
+
+	FAudioMultiplyNode::~FAudioMultiplyNode()
+	{
+	}
+
+	const FName& FAudioMultiplyNode::GetClassName() const
+	{
+		return ClassName;
+	}
+
+	IOperatorFactory& FAudioMultiplyNode::GetDefaultOperatorFactory() 
+	{
+		return Factory;
+	}
+}
+
+#undef LOCTEXT_NAMESPACE //MetasoundAudioMultiplyNode
diff --git a/Engine/Plugins/Runtime/Metasound/Source/MetasoundStandardNodes/Private/MetasoundOscNode.cpp b/Engine/Plugins/Runtime/Metasound/Source/MetasoundStandardNodes/Private/MetasoundOscNode.cpp
new file mode 100644
--- /dev/null
+++ b/Engine/Plugins/Runtime/Metasound/Source/MetasoundStandardNodes/Private/MetasoundOscNode.cpp
@@ -0,0 +1,112 @@
+// Copyright Epic Games, Inc. All Rights Reserved.
+
+#include "MetasoundOscNode.h"
+#include "MetasoundExecutableOperator.h"
+#include "MetasoundPrimitives.h"
+#include "MetasoundNodeRegistrationMacro.h"
+
+#define LOCTEXT_NAMESPACE "MetasoundOscNode"
+
+namespace Metasound
+{
+	METASOUND_REGISTER_NODE(FOscNode)
+
+	class FOscOperator : public TExecutableOperator<FOscOperator>
+	{
+		public:
+			FOscOperator(const FOperatorSettings& InSettings, const FFrequencyReadRef& InFrequency)
+			:	OperatorSettings(InSettings)
+			,	TwoPi(2.f * PI)
+			,	Phase(0.f)
+			,	Frequency(InFrequency)
+			,	AudioBuffer(FAudioBufferWriteRef::CreateNew(InSettings))
+			{
+				check(AudioBuffer->Num() == InSettings.GetNumFramesPerBlock());
+
+				OutputDataReferences.AddDataReadReference(TEXT("Audio"), FAudioBufferReadRef(AudioBuffer));
+			}
+
+			virtual const FDataReferenceCollection& GetInputs() const override
+			{
+				return InputDataReferences;
+			}
+
+			virtual const FDataReferenceCollection& GetOutputs() const override
+			{
+				return OutputDataReferences;
+			}
+
+			void Execute()
+			{
+				const float PhaseDelta = Frequency->GetRadiansPerSample(OperatorSettings.GetSampleRate());
+				float* Data = AudioBuffer->GetData();
+
+				for (int32 i = 0; i < OperatorSettings.GetNumFramesPerBlock(); i++)
+				{
+					Data[i] = FMath::Sin(Phase);
+					Phase += PhaseDelta;
+				}
+
+				Phase -= FMath::FloorToFloat(Phase / TwoPi) * TwoPi;
+			}
+
+		private:
+			const FOperatorSettings OperatorSettings;
+			const float TwoPi;
+			float Phase;
+
+			FFrequencyReadRef Frequency;
+			FAudioBufferWriteRef AudioBuffer;
+
+			FDataReferenceCollection InputDataReferences;
+			FDataReferenceCollection OutputDataReferences;
+	};
+
+	const FName FOscNode::ClassName = FName(TEXT("Osc"));
+
+	TUniquePtr<IOperator> FOscNode::FOperatorFactory::CreateOperator(const INode& InNode, const FOperatorSettings& InOperatorSettings, const FDataReferenceCollection& InInputDataReferences, TArray<TUniquePtr<IOperatorBuildError>>& OutErrors) 
+	{
+		const FOscNode& OscNode = static_cast<const FOscNode&>(InNode);
+		FFrequencyReadRef Frequency = FFrequencyReadRef::CreateNew(OscNode.GetDefaultFrequency(), EFrequencyResolution::Hertz);
+
+		if (InInputDataReferences.ContainsDataReadReference<FFrequency>(TEXT("Frequency")))
+		{
+			Frequency = InInputDataReferences.GetDataReadReference<FFrequency>(TEXT("Frequency"));
+		}
+
+		return MakeUnique<FOscOperator>(InOperatorSettings, Frequency);
+	}
+
+	FOscNode::FOscNode(const FString& InName, float InDefaultFrequency)
+		:	FNode(InName)
+		,	DefaultFrequency(InDefaultFrequency)
+	{
+		AddInputDataVertex<FFrequency>(TEXT("Frequency"), LOCTEXT("FrequencyTooltip", "The frequency of oscillator."));
+		AddOutputDataVertex<FAudioBuffer>(TEXT("Audio"), LOCTEXT("AudioTooltip", "The output audio"));
+	}
+
+	FOscNode::FOscNode(const FNodeInitData& InInitData)
+		: FOscNode(InInitData.InstanceName, 440.0f)
+	{
+	}
+
+	FOscNode::~FOscNode()
+	{
+	}
+
+	float FOscNode::GetDefaultFrequency() const
+	{
+		return DefaultFrequency;
+	}
+
+	const FName& FOscNode::GetClassName() const
+	{
+		return ::Metasound::FOscNode::ClassName;
+	}
+
+	IOperatorFactory& FOscNode::GetDefaultOperatorFactory() 
+	{
+		return Factory;
+	}
+}
+#undef LOCTEXT_NAMESPACE //MetasoundOscNode
diff --git a/Engine/Plugins/Runtime/Metasound/Source/MetasoundStandardNodes/Private/MetasoundPeriodicBopNode.cpp b/Engine/Plugins/Runtime/Metasound/Source/MetasoundStandardNodes/Private/MetasoundPeriodicBopNode.cpp
new file mode 100644
--- /dev/null
+++ b/Engine/Plugins/Runtime/Metasound/Source/MetasoundStandardNodes/Private/MetasoundPeriodicBopNode.cpp
@@ -0,0 +1,129 @@
+// Copyright Epic Games, Inc. All Rights Reserved.
+
+#include "MetasoundPeriodicBopNode.h"
+
+#include "CoreMinimal.h"
+#include "MetasoundNode.h"
+#include "MetasoundOperatorInterface.h"
+#include "MetasoundBuilderInterface.h"
+#include "MetasoundNodeInterface.h"
+#include "MetasoundDataReferenceCollection.h"
+#include "MetasoundExecutableOperator.h"
+#include "MetasoundPrimitives.h"
+#include "MetasoundNodeRegistrationMacro.h"
+
+#define LOCTEXT_NAMESPACE "MetasoundPeriodicBopNode"
+
+namespace Metasound
+{
+	METASOUND_REGISTER_NODE(FPeriodicBopNode)
+
+	class FPeriodicBopOperator : public TExecutableOperator<FPeriodicBopOperator>
+	{
+		public:
+			static constexpr float MinimumPeriodSeconds = 0.001f;
+			static constexpr float MinimumPeriodSamples = 20.f;
+
+			FPeriodicBopOperator(const FOperatorSettings& InSettings, const FFloatTimeReadRef& InPeriod)
+			:	OperatorSettings(InSettings)
+			,	Bop(FBopWriteRef::CreateNew())
+			,	Period(InPeriod)
+			,	ExecuteDurationInSamples(InSettings.GetNumFramesPerBlock())
+			,	SampleCountdown(0.f)
+			{
+				OutputDataReferences.AddDataReadReference(TEXT("Bop"), FBopReadRef(Bop));
+			}
+
+			virtual const FDataReferenceCollection& GetInputs() const override
+			{
+				return InputDataReferences;
+			}
+
+			virtual const FDataReferenceCollection& GetOutputs() const override
+			{
+				return OutputDataReferences;
+			}
+
+			void Execute()
+			{
+				// Advance internal counter to get rid of old bops.
+				Bop->Advance(OperatorSettings.GetNumFramesPerBlock());
+
+				float PeriodInSamples = FMath::Max(Period->GetSeconds(), MinimumPeriodSeconds) * OperatorSettings.GetSampleRate();
+
+				PeriodInSamples = FMath::Max(PeriodInSamples, MinimumPeriodSamples);
+
+				while ((SampleCountdown - ExecuteDurationInSamples) <= 0.f)
+				{
+					uint32 Frame = FMath::RoundToInt(SampleCountdown);
+					Bop->BopFrame(Frame);
+					SampleCountdown += PeriodInSamples;
+				}
+
+				SampleCountdown -= ExecuteDurationInSamples;
+			}
+
+		private:
+
+			FOperatorSettings OperatorSettings;
+			FBopWriteRef Bop;
+			FFloatTimeReadRef Period;
+
+			float ExecuteDurationInSamples;
+			float SampleCountdown;
+
+			FDataReferenceCollection InputDataReferences;
+			FDataReferenceCollection OutputDataReferences;
+	};
+
+
+	TUniquePtr<IOperator> FPeriodicBopNode::FOperatorFactory::CreateOperator(const INode& InNode, const FOperatorSettings& InOperatorSettings, const FDataReferenceCollection& InInputDataReferences, TArray<TUniquePtr<IOperatorBuildError>>& OutErrors)
+	{
+		const FPeriodicBopNode& PeriodicBopNode = static_cast<const FPeriodicBopNode&>(InNode);
+
+		// Note- won't this cause the Period input to not update after we build the operator?
+		FFloatTimeReadRef Period = FFloatTimeReadRef::CreateNew(PeriodicBopNode.GetDefaultPeriodInSeconds(), ETimeResolution::Seconds);
+
+		if (InInputDataReferences.ContainsDataReadReference<FFloatTime>(TEXT("Period")))
+		{
+			Period = InInputDataReferences.GetDataReadReference<FFloatTime>(TEXT("Period"));
+		}
+
+		return MakeUnique<FPeriodicBopOperator>(InOperatorSettings, Period);
+	}
+
+	const FName FPeriodicBopNode::ClassName = FName(TEXT("PeriodicBop"));
+
+	FPeriodicBopNode::FPeriodicBopNode(const FString& InName, float InDefaultPeriodInSeconds)
+	:	FNode(InName)
+	,	DefaultPeriod(InDefaultPeriodInSeconds)
+	{
+		AddInputDataVertex<FFloatTime>(TEXT("Period"), LOCTEXT("PeriodTooltip", "The period of the bops in seconds."));
+		AddOutputDataVertex<FBop>(TEXT("Bop"), LOCTEXT("BopTooltip", "The output bop"));
+	}
+
+	FPeriodicBopNode::FPeriodicBopNode(const FNodeInitData& InInitData)
+		: FPeriodicBopNode(InInitData.InstanceName, 1.0f)
+	{}
+
+	FPeriodicBopNode::~FPeriodicBopNode()
+	{
+	}
+
+	float FPeriodicBopNode::GetDefaultPeriodInSeconds() const
+	{
+		return DefaultPeriod;
+	}
+
+	const FName& FPeriodicBopNode::GetClassName() const
+	{
+		return ClassName;
+	}
+
+	IOperatorFactory& FPeriodicBopNode::GetDefaultOperatorFactory()
+	{
+		return Factory;
+	}
+}
+
+#undef LOCTEXT_NAMESPACE //MetasoundPeriodicBopNode
diff --git a/Engine/Plugins/Runtime/Metasound/Source/MetasoundStandardNodes/Private/MetasoundStandardNodesModule.cpp b/Engine/Plugins/Runtime/Metasound/Source/MetasoundStandardNodes/Private/MetasoundStandardNodesModule.cpp
new file mode 100644
--- /dev/null
+++ b/Engine/Plugins/Runtime/Metasound/Source/MetasoundStandardNodes/Private/MetasoundStandardNodesModule.cpp
@@ -0,0 +1,38 @@
+// Copyright Epic Games, Inc. All Rights Reserved.
+
+#include "MetasoundGraphCoreModule.h"
+
+#include "CoreMinimal.h"
+#include "Modules/ModuleInterface.h"
+#include "Modules/ModuleManager.h"
+#include "MetasoundPrimitives.h"
+#include "MetasoundDataReferenceTypes.h"
+#include "MetasoundTime.h"
+#include "MetasoundFrequency.h"
+#include "MetasoundWave.h"
+
+REGISTER_METASOUND_DATATYPE(bool, "Primitive:Bool", ::Metasound::ELiteralArgType::Boolean)
+REGISTER_METASOUND_DATATYPE(int32, "Primitive:Int32", ::Metasound::ELiteralArgType::Integer)
+REGISTER_METASOUND_DATATYPE(int64, "Primitive:Int64", ::Metasound::ELiteralArgType::Integer)
+REGISTER_METASOUND_DATATYPE(float, "Primitive:Float", ::Metasound::ELiteralArgType::Float)
+REGISTER_METASOUND_DATATYPE(double, "Primitive:Double", ::Metasound::ELiteralArgType::Float)
+
+namespace Metasound
+{
+	REGISTER_METASOUND_DATATYPE(FBop, "Primitive:Bop", ::Metasound::ELiteralArgType::Boolean)
+	REGISTER_METASOUND_DATATYPE(FFloatTime, "Primitive:Time")
+	REGISTER_METASOUND_DATATYPE(FDoubleTime, "Primitive:Time:HighResolution")
+	REGISTER_METASOUND_DATATYPE(FSampleTime, "Primitive:Time:SampleResolution")
+	REGISTER_METASOUND_DATATYPE(FFrequency, "Primitive:Frequency")
+	REGISTER_METASOUND_DATATYPE(FWave, "Primitive:Wave")
+}
+
+namespace Metasound 
+{
+	class FMetasoundStandardNodesModule : public IMetasoundGraphCoreModule
+	{
+	};
+}
+
+IMPLEMENT_MODULE(Metasound::FMetasoundStandardNodesModule, MetasoundStandardNodes);
+
diff --git a/Engine/Plugins/Runtime/Metasound/Source/MetasoundStandardNodes/Private/MetasoundWavePlayerNode.cpp b/Engine/Plugins/Runtime/Metasound/Source/MetasoundStandardNodes/Private/MetasoundWavePlayerNode.cpp
new file mode 100644
--- /dev/null
+++ b/Engine/Plugins/Runtime/Metasound/Source/MetasoundStandardNodes/Private/MetasoundWavePlayerNode.cpp
@@ -0,0 +1,161 @@
+// Copyright Epic Games, Inc. All Rights Reserved.
+
+#include "MetasoundWavePlayerNode.h"
+#include "MetasoundExecutableOperator.h"
+#include "MetasoundPrimitives.h"
+#include "MetasoundWave.h"
+
+#include "IAudioCodecRegistry.h"
+#include "IAudioCodec.h"
+
+#define LOCTEXT_NAMESPACE "MetasoundWaveNode"
+
+namespace Metasound
+{
+	METASOUND_REGISTER_NODE(FWavePlayerNode)
+		
+	class FWavePlayerOperator : public TExecutableOperator<FWavePlayerOperator>
+	{
+		public:
+			FWavePlayerOperator(
+				const FOperatorSettings& InSettings, 
+				const FWaveReadRef& InWave, 
+				TUniquePtr<Audio::IDecoderInput>&& InDecoderInput, 
+				TUniquePtr<Audio::IDecoderOutput>&& InDecoderOutput,
+				TUniquePtr<Audio::IDecoder>&& InDecoder )
+				: OperatorSettings(InSettings)
+				, Decoder(MoveTemp(InDecoder))
+				, DecoderInput(MoveTemp(InDecoderInput))
+				, DecoderOutput(MoveTemp(InDecoderOutput))
+				, Wave(InWave)
+				, AudioBuffer(FAudioBufferWriteRef::CreateNew(InSettings))
+			{
+				check(AudioBuffer->Num() == InSettings.GetNumFramesPerBlock());
+
+				OutputDataReferences.AddDataReadReference(TEXT("Audio"), FAudioBufferReadRef(AudioBuffer));
+			}
+
+			virtual const FDataReferenceCollection& GetInputs() const override
+			{
+				return InputDataReferences;
+			}
+
+			virtual const FDataReferenceCollection& GetOutputs() const override
+			{
+				return OutputDataReferences;
+			}
+
+			void Execute()
+			{
+				// V1. Do the decode inline, this will sound bad.
+				Decoder->Decode();
+
+				TArrayView<float> Src;
+				float* Dst = AudioBuffer->GetData();
+				
+				int32 NumPopped = DecoderOutput->PopAudio(Src);
+				int32 NumFramesToCopy = FMath::Min(NumPopped,OperatorSettings.GetNumFramesPerBlock());
+
+				int32 i = 0;
+				for (; i < NumFramesToCopy; ++i)
+				{
+					Dst[i] = Src[i];
+				}
+				for ( ; i < OperatorSettings.GetNumFramesPerBlock(); ++i)
+				{
+					Dst[i] = 0.0f;
+				}
+			}
+
+		private:
+			const FOperatorSettings OperatorSettings;
+			
+			Audio::ICodec::FDecoderPtr Decoder;
+			TUniquePtr<Audio::IDecoderInput> DecoderInput;
+			TUniquePtr<Audio::IDecoderOutput> DecoderOutput;
+			
+			FWaveReadRef Wave;
+			FAudioBufferWriteRef AudioBuffer;
+
+			FDataReferenceCollection InputDataReferences;
+			FDataReferenceCollection OutputDataReferences;
+	};
+
+	const FName FWavePlayerNode::ClassName = FName(TEXT("Wave"));
+
+	TUniquePtr<IOperator> FWavePlayerNode::FOperatorFactory::CreateOperator(
+		const INode& InNode, 
+		const FOperatorSettings& InOperatorSettings,
+		const FDataReferenceCollection& InInputDataReferences, 
+		TArray<TUniquePtr<IOperatorBuildError>>& OutErrors) 
+	{
+		using namespace Audio;
+
+		const FWavePlayerNode& WaveNode = static_cast<const FWavePlayerNode&>(InNode);
+		FWaveReadRef Wave = FWaveReadRef::CreateNew();
+
+		if (InInputDataReferences.ContainsDataReadReference<FWave>(TEXT("Wave")))
+		{
+			Wave = InInputDataReferences.GetDataReadReference<FWave>(TEXT("Wave"));
+		}
+
+		TUniquePtr<IDecoderInput> Input = IDecoderInput::Create(Wave->GetCompressedData(),0);
+		if(Input)
+		{
+			ICodecRegistry::FCodecPtr Codec = Audio::ICodecRegistry::Get().FindCodecByFromParsingInput(Input.Get());
+			if(Codec)
+			{
+				// V1, Ask for an output buffer the size of a frame.
+				IDecoderOutput::FRequirements Reqs { Float32_Interleaved, InOperatorSettings.GetNumFramesPerBlock() };
+				TUniquePtr<IDecoderOutput> Output = IDecoderOutput::Create(Reqs);
+				TUniquePtr<IDecoder> Decoder = Codec->CreateDecoder(Input.Get(), Output.Get());
+
+				return MakeUnique<FWavePlayerOperator>(
+					InOperatorSettings, 
+					Wave, 
+					MoveTemp(Input), 
+					MoveTemp(Output), 
+					MoveTemp(Decoder)
+				);
+			}
+			else
+			{
+				// LogMsg; Failed to Find
+			}
+		}
+		else
+		{
+			// LogMsg; Failed to parse input stream, is it valid?
+		}
+
+		// Fail.
+		return nullptr;
+	}
+
+	FWavePlayerNode::FWavePlayerNode(const FString& InName)
+		:	FNode(InName)
+	{
+		AddInputDataVertex<FWave>(TEXT("Wave"), LOCTEXT("WaveTooltip", "The Wave to be decoded"));
+		AddOutputDataVertex<FAudioBuffer>(TEXT("Audio"), LOCTEXT("AudioTooltip", "The output audio"));
+	}
+
+	FWavePlayerNode::FWavePlayerNode(const FNodeInitData& InInitData)
+		: FWavePlayerNode(InInitData.InstanceName)
+	{
+	}
+
+	FWavePlayerNode::~FWavePlayerNode()
+	{
+	}
+	
+	const FName& FWavePlayerNode::GetClassName() const
+	{
+		return ::Metasound::FWavePlayerNode::ClassName;
+	}
+
+	IOperatorFactory& FWavePlayerNode::GetDefaultOperatorFactory() 
+	{
+		return Factory;
+	}
+}
+#undef LOCTEXT_NAMESPACE //MetasoundWaveNode
diff --git a/Engine/Plugins/Runtime/Metasound/Source/MetasoundStandardNodes/Public/MetasoundADSRNode.h b/Engine/Plugins/Runtime/Metasound/Source/MetasoundStandardNodes/Public/MetasoundADSRNode.h
new file mode 100644
--- /dev/null
+++ b/Engine/Plugins/Runtime/Metasound/Source/MetasoundStandardNodes/Public/MetasoundADSRNode.h
@@ -0,0 +1,97 @@
+// Copyright Epic Games, Inc. All Rights Reserved.
+#pragma once
+
+#include "CoreMinimal.h"
+#include "Internationalization/Text.h"
+#include "MetasoundNode.h"
+#include "MetasoundOperatorInterface.h"
+#include "MetasoundBuilderInterface.h"
+#include "MetasoundDataReferenceCollection.h"
+
+
+namespace Metasound
+{
+	/** FADSRNode
+	 *
+	 *  Creates an Attack, Decay Sustain, Release audio processer node. 
+	 */
+	class METASOUNDSTANDARDNODES_API FADSRNode : public FNode
+	{
+		// The operator factory for this node.
+		class FOperatorFactory : public IOperatorFactory
+		{
+			public:
+				virtual TUniquePtr<IOperator> CreateOperator(const INode& InNode, const FOperatorSettings& InOperatorSettings, const FDataReferenceCollection& InInputDataReferences, TArray<TUniquePtr<IOperatorBuildError>>& OutErrors) override;
+
+				template<typename DataType>
+				void SetReadableRefIfInCollection(const FString& InName, const FDataReferenceCollection& InCollection, TDataReadReference<DataType>& DataRef)
+				{
+					// TODO: add a helper function to FDataReferenceCollection to do a SetDataIfContains.
+					if (InCollection.ContainsDataReadReference<DataType>(InName))
+					{
+						DataRef = InCollection.GetDataReadReference<DataType>(InName);
+					}
+				}
+		};
+
+		public:
+			/** Class name for FADSRNode */
+			static const FName ClassName;
+
+			/** FADSR node constructor.
+			 *
+			 * @param InName - Name of this node.
+			 * @param InDefaultAttackMs - Default attack in milliseconds.
+			 * @param InDefaultDecayMs - Default decay in milliseconds.
+			 * @param InDefaultSustainMs - Default sustain in milliseconds.
+			 * @param InDefaultReleaseMs - Default release in milliseconds.
+			 */
+			FADSRNode(const FString& InName, float InDefaultAttackMs, float InDefaultDecayMs, float InDefaultSustainMs, float InDefaultReleaseMs);
+
+			/**
+			 * Constructor used by the Metasound Frontend.
+			 */
+			FADSRNode(const FNodeInitData& InitData);
+
+			virtual ~FADSRNode();
+
+			/** Return default attack time in milliseconds */
+			float GetDefaultAttackMs() const;
+			/** Return default attack time in milliseconds */
+			float GetDefaultDecayMs() const;
+			/** Return default attack time in milliseconds */
+			float GetDefaultSustainMs() const;
+			/** Return default attack time in milliseconds */
+			float GetDefaultReleaseMs() const;
+
+			/** Returns the type name of the FADSRNode */
+			virtual const FName& GetClassName() const override;
+
+			virtual const FText& GetDescription() const override
+			{
+				static const FText StaticDescription = NSLOCTEXT("MetasoundGraphCore", "Metasound_ADSRNodeDescription", "Emits an ADSR (Attack, decay, sustain, & release) envelope when bopped.");
+				return StaticDescription;
+			}
+
+			virtual const FText& GetAuthorName() const override
+			{
+				return PluginAuthor;
+			}
+
+			virtual const FText& GetPromptIfMissing() const override
+			{
+				return PluginNodeMissingPrompt;
+			}
+
+			/** Return a factory for building a metasound operator. */
+			virtual IOperatorFactory& GetDefaultOperatorFactory() override;
+
+		private:
+			float DefaultAttackMs;
+			float DefaultDecayMs;
+			float DefaultSustainMs;
+			float DefaultReleaseMs;
+
+			FOperatorFactory Factory;
+	};
+} // namespace Metasound
diff --git a/Engine/Plugins/Runtime/Metasound/Source/MetasoundStandardNodes/Public/MetasoundAudioBuffer.h b/Engine/Plugins/Runtime/Metasound/Source/MetasoundStandardNodes/Public/MetasoundAudioBuffer.h
new file mode 100644
--- /dev/null
+++ b/Engine/Plugins/Runtime/Metasound/Source/MetasoundStandardNodes/Public/MetasoundAudioBuffer.h
@@ -0,0 +1,158 @@
+// Copyright Epic Games, Inc. All Rights Reserved.
+
+#pragma once
+
+#include "CoreMinimal.h"
+#include "DSP/BufferVectorOperations.h"
+#include "MetasoundGraphCoreModule.h"
+#include "MetasoundDataReference.h"
+#include "MetasoundDataTypeRegistrationMacro.h"
+
+/** Define which determines whether to check that the size of the audio buffer has not changed since initialization */
+#define METASOUNDGRAPHCORE_CHECKAUDIONUM !(UE_BUILD_SHIPPING || UE_BUILD_TEST)
+
+namespace Metasound 
+{
+	/**  FAudioBuffer
+	 *
+	 * FAudioBuffer is the default buffer for passing audio data between nodes.  It should not be resized. 
+	 * It is compatible with functions accepting const references to AlignedFloatBuffer (const AlignedFloatBuffer&) 
+	 * arguments via an implicit conversion operator which exposes the underlying AlignedFloatBuffer container.
+	 *
+	 */
+	class METASOUNDSTANDARDNODES_API FAudioBuffer
+	{
+		public:
+			/** Create an FAudioBuffer with a specific number of samples. 
+			 *
+			 * @param InNumSamples - Number of samples in buffer.
+			 */
+			explicit FAudioBuffer(int32 InNumSamples)
+			{
+				Buffer.AddUninitialized(InNumSamples);
+
+#if METASOUNDGRAPHCORE_CHECKAUDIONUM
+				// InitialNum should not change during the life of an FAudioBuffer
+				InitialNum = Buffer.Num();
+#endif				
+			}
+
+			/**
+			 * This is the constructor used by the frontend.
+			 */
+			FAudioBuffer(const FOperatorSettings& InSettings)
+			{
+				Buffer.AddUninitialized(InSettings.GetNumFramesPerBlock());
+
+#if METASOUNDGRAPHCORE_CHECKAUDIONUM
+				// InitialNum should not change during the life of an FAudioBuffer
+				InitialNum = Buffer.Num();
+#endif	
+			}
+
+			/** Allow all other constructors of Audio::AlignedFloatBuffer to be forwarded to
+			 * underlying container */
+			/*
+			template <
+				typename... ArgTypes,
+				typename = typename TEnableIf<
+					TAndValue<
+						sizeof...(ArgTypes) != 0,
+						TOrValue<
+							sizeof...(ArgTypes) != 1,
+							TNot<
+								TIsIntegral<
+									typename TDecay<
+										typename TNthTypeFromParameterPack<0, ArgTypes...>::Type
+									>::Type
+								>
+							>
+						>
+					>::Value
+				>::Type
+			>
+			FAudioBuffer(ArgTypes&&... Args)
+			:	Buffer(Forward<ArgTypes>(Args)...)
+			{
+#if METASOUNDGRAPHCORE_CHECKAUDIONUM
+				// InitialNum should not change during the life of an FAudioBuffer
+				InitialNum = Buffer.Num();
+#endif				
+			}
+			*/
+
+			FAudioBuffer()
+			{
+#if METASOUNDGRAPHCORE_CHECKAUDIONUM
+				// InitialNum should not change during the life of an FAudioBuffer
+				InitialNum = Buffer.Num();
+#endif				
+			}
+
+			/** Return a pointer to the audio float data. */
+			FORCEINLINE const float* GetData() const
+			{
+#if METASOUNDGRAPHCORE_CHECKAUDIONUM
+				UE_CLOG(InitialNum != Buffer.Num(), LogMetasoundGraphCore, Error, TEXT("Metasound audio buffer size change detected.  Audio Buffers should not be resized."));
+#endif				
+				return Buffer.GetData();
+			}
+
+			/** Return a pointer to the audio float data. */
+			FORCEINLINE float* GetData()
+			{
+#if METASOUNDGRAPHCORE_CHECKAUDIONUM
+				UE_CLOG(InitialNum != Buffer.Num(), LogMetasoundGraphCore, Error, TEXT("Metasound audio buffer size change detected.  Audio Buffers should not be resized."));
+#endif				
+				return Buffer.GetData();
+			}
+
+			/** Return the number of samples in the audio buffer. */
+			FORCEINLINE int32 Num() const
+			{
+#if METASOUNDGRAPHCORE_CHECKAUDIONUM
+				UE_CLOG(InitialNum != Buffer.Num(), LogMetasoundGraphCore, Error, TEXT("Metasound audio buffer size change detected.  Audio Buffers should not be resized."));
+#endif				
+				return Buffer.Num();
+			}
+
+			/** Implicit conversion to Audio::AlignedFloatBuffer */
+			FORCEINLINE operator const Audio::AlignedFloatBuffer& () const
+			{
+#if METASOUNDGRAPHCORE_CHECKAUDIONUM
+				UE_CLOG(InitialNum != Buffer.Num(), LogMetasoundGraphCore, Error, TEXT("Metasound audio buffer size change detected.  Audio Buffers should not be resized."));
+#endif				
+				return Buffer;
+			}
+
+			/** Implicit conversion to Audio::AlignedFloatBuffer 
+			 *
+			 * WARNING: if the buffer is resized, it will cause errors. 
+			 */
+			FORCEINLINE operator Audio::AlignedFloatBuffer& ()
+			{
+#if METASOUNDGRAPHCORE_CHECKAUDIONUM
+				UE_CLOG(InitialNum != Buffer.Num(), LogMetasoundGraphCore, Error, TEXT("Metasound audio buffer size change detected.  Audio Buffers should not be resized."));
+#endif				
+				return Buffer;
+			}
+
+			FORCEINLINE void Zero()
+			{
+				if (Buffer.Num() > 0)
+				{
+					FMemory::Memset(Buffer.GetData(), 0, sizeof(float) * Buffer.Num());
+				}
+			}
+
+		private:
+
+			Audio::AlignedFloatBuffer Buffer;
+
+#if METASOUNDGRAPHCORE_CHECKAUDIONUM
+			int32 InitialNum;
+#endif				
+	};
+
+	DECLARE_METASOUND_DATA_REFERENCE_TYPES(FAudioBuffer, 0x1b02dac2 , FAudioBufferTypeInfo, FAudioBufferReadRef, FAudioBufferWriteRef);
+};
diff --git a/Engine/Plugins/Runtime/Metasound/Source/MetasoundStandardNodes/Public/MetasoundAudioFormats.h b/Engine/Plugins/Runtime/Metasound/Source/MetasoundStandardNodes/Public/MetasoundAudioFormats.h
new file mode 100644
--- /dev/null
+++ b/Engine/Plugins/Runtime/Metasound/Source/MetasoundStandardNodes/Public/MetasoundAudioFormats.h
@@ -0,0 +1,470 @@
+// Copyright Epic Games, Inc. All Rights Reserved.
+
+#pragma once
+
+#include "CoreMinimal.h"
+#include "MetasoundDataReference.h"
+#include "MetasoundAudioBuffer.h"
+#include "MetasoundDataTypeRegistrationMacro.h"
+
+namespace Metasound
+{
+	/** FUnformattedAudio
+	 *
+	 * FUnformattedAudio represents deinterleaved multichannel audio which can dynamically change
+	 * the number of channels of audio. In possible, using FMultichannelAudioFormat is preferred over
+	 * using FUnformattedAudio when the number of audio channels is known at construction.
+	 *
+	 * The audio buffers in FUnformattedAudio are shared data references which can be accessed outside
+	 * of the FUnformattedAudio. All audio buffers within a FUnformattedAudio object must contain the 
+	 * same number of audio frames.
+	 */
+	class METASOUNDSTANDARDNODES_API FUnformattedAudio
+	{
+		public:
+			/** FUnformattedAudio Constructor.
+			 * 
+			 * @param InNumFrames - Number of frames to hold in each buffer.
+			 * @param InNumChannels - Initial number of audio channels.
+			 * @param InMaxNumChannels - Maximum number of audio channels.
+			 */
+			FUnformattedAudio(int32 InNumFrames, int32 InNumChannels, int32 InMaxNumChannels);
+
+			/**
+			 * FUnformattedAudio Constructor used by the metasound frontend.
+			 *
+			 * @param InNumChannels - initial number of audio channels.
+			 * @param InSettings - Operator Settings passed in on construction.
+			 */
+			explicit FUnformattedAudio(int32 InNumChannels, const FOperatorSettings& InSettings);
+
+
+			/** Sets the number of audio channels.
+			 * 
+			 * The requested number of channels is internally clamped to the supported range 
+			 * of [0, GetMaxNumChannels()].
+			 *
+			 * @parma InNumChannels - The desired number of channels.
+			 *
+			 * @return The actual number of channels. This may differ from the requested number 
+			 * 		   of channels if the requested number exceeded the maximum number of channels.
+			 */
+			int32 SetNumChannels(int32 InNumChannels);
+
+			/** Return the number of channels. */
+			int32 GetNumChannels() const { return NumChannels; }
+
+			/** Return the maximum number of channels. */
+			int32 GetMaxNumChannels() const { return MaxNumChannels; }
+
+			/** Return an array view of the readable buffer references.
+			 *
+			 * This array will have GetNumChannels() elements.
+			 */
+			const TArrayView<const FAudioBufferReadRef>& GetBuffers() const { return ReadableBuffers; }
+
+			/** Return an array view of the writable buffer references.
+			 *
+			 * This array will have GetNumChannels() elements.
+			 */
+			const TArrayView<const FAudioBufferWriteRef>& GetBuffers() { return WritableBuffers; }
+
+			/** Return an array of the readable buffer reference storage.
+			 *
+			 * This array will have GetMaxNumChannels() elements.
+			 */
+			const TArray<FAudioBufferReadRef>& GetStorage() const { return ReadableBufferStorage; }
+
+			/** Return an array of the writable buffer reference storage.
+			 *
+			 * This array will have GetMaxNumChannels() elements.
+			 */
+			const TArray<FAudioBufferWriteRef>& GetStorage() { return WritableBufferStorage; }
+			
+		private:
+			int32 NumFrames;
+			int32 NumChannels;
+			int32 MaxNumChannels;
+
+			TArrayView<const FAudioBufferReadRef> ReadableBuffers;
+			TArrayView<const FAudioBufferWriteRef> WritableBuffers;
+
+			TArray<FAudioBufferReadRef> ReadableBufferStorage;
+			TArray<FAudioBufferWriteRef> WritableBufferStorage;
+	};
+
+
+	/** FMultichannelAudioFormat
+	 *
+	 * FMultichannelAudioFormat represents deinterleaved multichannel audio which supports a constant
+	 * number of channels for the lifetime of the object. 
+	 *
+	 * The audio buffers in FMultichannelAudioFormat are shared data references which can be accessed outside
+	 * of the FMultichannelAudioFormat. All audio buffers within a FMultichannelAudioFormat object must contain the 
+	 * same number of audio frames.
+	 */
+	class METASOUNDSTANDARDNODES_API FMultichannelAudioFormat
+	{
+		public:
+			FMultichannelAudioFormat();
+
+			/** FMultichannelAudioFormat Constructor.
+			 *
+			 * @param InNumFrames - The number of frames per an audio buffer.
+			 * @param InNumChannels - The number of audio channels.
+			 */
+			FMultichannelAudioFormat(int32 InNumFrames, int32 InNumChannels);
+
+			/**
+			 * FMultichannelAudioFormat Constructor used by the metasound frontend.
+			 *
+			 * @param InNumChannels - initial number of audio channels.
+			 * @param InSettings - Operator Settings passed in on construction.
+			 */
+			explicit FMultichannelAudioFormat(int32 InNumChannels, const FOperatorSettings& InSettings);
+
+			/** FMultichannelAudioFormat Constructor.
+			 *
+			 * This constructor accepts an array of writable audio buffer references. Each 
+			 * buffer in the array must contain equal number of frames. 
+			 *
+			 * @param InWriteRefs - An array of writable audio buffer references. 
+			 */
+			FMultichannelAudioFormat(TArrayView<const FAudioBufferWriteRef> InWriteRefs);
+
+			// Enable the copy constructor.
+			FMultichannelAudioFormat(const FMultichannelAudioFormat& InOther) = default;
+
+			// Disable move constructor as incoming object should not be altered. 
+			//FMultichannelAudioFormat(FMultichannelAudioFormat&& InOther) = delete;
+
+			// Disable equal operator so channel count does not change
+			FMultichannelAudioFormat& operator=(const FMultichannelAudioFormat& Other) = delete;
+
+			// Disable move operator so channel count does not change
+			//FMultichannelAudioFormat& operator=(FMultichannelAudioFormat&& Other) = delete;
+
+			/** Return the number of audio channels. */
+			int32 GetNumChannels() const { return NumChannels; }
+
+			/** Return an array view of the readable buffer references.
+			 *
+			 * This array will have GetNumChannels() elements.
+			 */
+			const TArrayView<const FAudioBufferReadRef>& GetBuffers() const { return ReadableBuffers; }
+
+			/** Return an array view of the writable buffer references.
+			 *
+			 * This array will have GetNumChannels() elements.
+			 */
+			const TArrayView<const FAudioBufferWriteRef>& GetBuffers() { return WritableBuffers; }
+
+			/** Return an array of the readable buffer reference storage.
+			 *
+			 * This array will have GetNumChannels() elements.
+			 */
+			const TArray<FAudioBufferReadRef>& GetStorage() const { return ReadableBufferStorage; }
+
+			/** Return an array of the writable buffer reference storage.
+			 *
+			 * This array will have GetNumChannels() elements.
+			 */
+			const TArray<FAudioBufferWriteRef>& GetStorage() { return WritableBufferStorage; }
+
+		private:
+			// Friendship with the data reference class gives it access to the protected constructor
+			// for the scenario where a TDataReadReference<FMulitchannelAudioFormat> is constructed
+			// with FAudioBufferReadRefs. The constructor cannot be public as it would provide writable
+			// access to the passed in audio buffers, even though the passed buffers explicitly were
+			// read references.
+			friend class TDataReadReference<FMultichannelAudioFormat>;
+
+			// Construct a FMultichannelAudioFormat with an array of readable buffers. 
+			// 
+			// This constructor should only be used when it can be assured that the constructed object
+			// will not provide writable access to the passed in audio buffers. 
+			FMultichannelAudioFormat(TArrayView<const FAudioBufferReadRef> InReadRefs);
+
+			int32 NumChannels;
+
+			TArrayView<const FAudioBufferWriteRef> WritableBuffers;
+			TArrayView<const FAudioBufferReadRef> ReadableBuffers;
+
+			TArray<FAudioBufferWriteRef> WritableBufferStorage;
+			TArray<FAudioBufferReadRef> ReadableBufferStorage;
+	};
+
+	
+	/** A TStaticChannelAudioFormat represents deinterleaved multichannel audio 
+	 * where the number of channels is known at compile time. This is primarily 
+	 * useful to define such cases as Stereo, Mono, Qaud, 5.1, etc. 
+	 *
+	 * The audio buffers in FMultichannelAudioFormat are shared data references 
+	 * which can be accessed outside of the FMultichannelAudioFormat. All audio 
+	 * buffers within a FMultichannelAudioFormat object must contain the same 
+	 * number of audio frames.
+	 */
+	template<int32 NumChannels>
+	class TStaticChannelAudioFormat
+	{
+		public:
+			/** TStaticChannelAudioFormat Constructor
+			 *
+			 * @param InNumFrames - The number of frames per an audio buffer.
+			 */
+			TStaticChannelAudioFormat(int32 InNumFrames)
+			{
+				static_assert(NumChannels > 0, "NumChannels must be greater than zero");
+
+				InNumFrames = FMath::Max(InNumFrames, 0);
+
+				for (int32 i = 0; i < NumChannels; i++)
+				{
+					FAudioBufferWriteRef Audio = FAudioBufferWriteRef::CreateNew(InNumFrames);
+					Audio->Zero();
+
+					WritableBufferStorage.Add(Audio);
+					ReadableBufferStorage.Add(Audio);
+				}
+
+				WritableBuffers = WritableBufferStorage;
+				ReadableBuffers = ReadableBufferStorage;
+			}
+
+			/** Return the number of audio channels. */
+			int32 GetNumChannels() const
+			{
+				return NumChannels;
+			}
+
+			/** Return an readable buffer reference for a specific channel. */
+			template<int32 ChannelIndex>
+			FAudioBufferReadRef GetBuffer() const
+			{
+				static_assert(ChannelIndex >= 0, "Index must be within range of channels");
+				static_assert(ChannelIndex < NumChannels, "Index must be within range of channels");
+
+				return ReadableBuffers.GetData()[ChannelIndex];
+			}
+
+			/** Return an writable buffer reference for a specific channel. */
+			template<int32 ChannelIndex>
+			FAudioBufferWriteRef GetBuffer() 
+			{
+				static_assert(ChannelIndex >= 0, "Index must be within range of channels");
+				static_assert(ChannelIndex < NumChannels, "Index must be within range of channels");
+
+				return WritableBuffers.GetData()[ChannelIndex];
+			}
+
+			/** Return an array view of the readable buffer references.
+			 *
+			 * This array will have GetNumChannels() elements.
+			 */
+			const TArrayView<const FAudioBufferReadRef> GetBuffers() const
+			{
+				return ReadableBuffers;
+			}
+
+			/** Return an array view of the writable buffer references.
+			 *
+			 * This array will have GetNumChannels() elements.
+			 */
+			const TArrayView<const FAudioBufferWriteRef> GetBuffers() 
+			{
+				return WritableBuffers;
+			}
+
+			/** Return an array of the readable buffer references.
+			 *
+			 * This array will have GetNumChannels() elements.
+			 */
+			const TArray<FAudioBufferReadRef> GetStorage() const
+			{
+				return ReadableBufferStorage;
+			}
+
+			/** Return an array of the writable buffer references.
+			 *
+			 * This array will have GetNumChannels() elements.
+			 */
+			const TArray<FAudioBufferWriteRef> GetStorage()
+			{
+				return WritableBufferStorage;
+			}
+
+		protected:
+
+			// TStaticChannelAudioFormat constructor with an array of writable buffers. 
+			TStaticChannelAudioFormat(const FAudioBufferWriteRef(&InBuffers)[NumChannels])
+			{
+				int32 NumFrames = 0;
+				if (NumChannels > 0)
+				{
+					NumFrames = InBuffers[0]->Num();
+				}
+
+				for (int32 i = 0; i < NumChannels; i++)
+				{
+					checkf(NumFrames == InBuffers[i]->Num(), TEXT("All buffers must have same number of frames (%d != %d)"), NumFrames, InBuffers[i]->Num());
+
+					WritableBufferStorage.Add(InBuffers[i]);
+					ReadableBufferStorage.Add(InBuffers[i]);
+				}
+
+				WritableBuffers = WritableBufferStorage;
+				ReadableBuffers = ReadableBufferStorage;
+			}
+
+		private:
+			TArrayView<const FAudioBufferWriteRef> WritableBuffers;
+			TArrayView<const FAudioBufferReadRef> ReadableBuffers;
+
+			TArray<FAudioBufferWriteRef> WritableBufferStorage;
+			TArray<FAudioBufferReadRef> ReadableBufferStorage;
+	};
+
+	/** FMonoAudioFormat represents mono audio containing one channel of audio.
+	 *
+	 * The audio buffer is a shared data references which can be accessed 
+	 * outside of the FMonoAudioFormat. 
+	 */
+	class FMonoAudioFormat : public TStaticChannelAudioFormat<1>
+	{
+			// This is used to grant specific classes the ability to create FMonoAudioFormat
+			// using a FAudioBufferReadRef
+			//enum EPrivateToken { Token };
+			//static const EPrivateToken PrivateToken = EPrivateToken::Token;
+
+		public:
+			using Super = TStaticChannelAudioFormat<1>;
+
+			// Inherit constructors of base class.
+			using Super::Super;
+
+			/** FMonoAudioFormat Construtor
+			 *
+			 * Construct with a single writable audio buffer reference.
+			 *
+			 * @param InAudio - A writable audio buffer reference.
+			 */
+			FMonoAudioFormat(const FAudioBufferWriteRef& InAudio)
+			:	Super({InAudio})
+			{
+			}
+
+			// TODO: clean up this mess.
+			// Construct a FMonoAudioFormat a readable buffer reference.
+			// 
+			// This constructor should only be used when it can be assured 
+			// that the constructed object will not provide writable access 
+			// to the passed in audio buffers. 
+			/*
+			FMonoAudioFormat(const FAudioBufferReadRef& InAudio, EPrivateToken InToken)
+			:	Super({WriteCast(InAudio)})
+			{
+			}
+			*/
+
+			/** Return writable audio buffer reference of center channel. */
+			FAudioBufferWriteRef GetCenter() { return GetBuffer<0>(); }
+
+			/** Return readable audio buffer reference of center channel. */
+			FAudioBufferReadRef GetCenter() const { return GetBuffer<0>(); }
+
+		private:
+
+
+			// Friendship with the data reference class gives it access to the 
+			// protected constructor for the scenario where a 
+			// TDataReadReference<*> is constructed with FAudioBufferReadRefs. 
+			// The constructor cannot be public as it would provide writable
+			// access to the passed in audio buffers, even though the passed 
+			// buffers explicitly were read references.
+			//friend class TDataReadReference<FMonoAudioFormat>;
+	};
+
+	// Template specialization for TDataReadReference<FMonoAudioFormat> constructor to use
+	// private token.
+	/*
+	template<>
+	template<>
+	TDataReadReference<FMonoAudioFormat>::TDataReadReference(const FAudioBufferReadRef& InReadBuffer)
+	:	FDataReference(InReadBuffer, FMonoAudioFormat::PrivateToken)
+	{
+	}
+	*/
+	
+
+	/** FStereoAudioFormat represents stereo audio containing two channels of 
+	 * audio.
+	 *
+	 * The audio buffers are shared data references which can be accessed 
+	 * outside of the FStereoAudioFormat. 
+	 */
+	class FStereoAudioFormat : public TStaticChannelAudioFormat<2>
+	{
+		public:
+			using Super = TStaticChannelAudioFormat<2>;
+
+			// Inherit constructors of base class.
+			using Super::Super;
+
+			/** FStereoAudioFormat Construtor
+			 *
+			 * Construct with a two writable audio buffer reference.
+			 *
+			 * @param InLeftAudio - A writable audio buffer reference for the left channel.
+			 * @param InRightAudio - A writable audio buffer reference for the right channel.
+			 */
+			FStereoAudioFormat(const FAudioBufferWriteRef& InLeftAudio, const FAudioBufferWriteRef& InRightAudio)
+			:	Super({InLeftAudio, InRightAudio})
+			{
+			}
+
+			/** Return writable audio buffer reference of left channel. */
+			FAudioBufferWriteRef GetLeft() { return GetBuffer<0>(); }
+
+			/** Return readable audio buffer reference of left channel. */
+			FAudioBufferReadRef GetLeft() const { return GetBuffer<0>(); }
+
+			/** Return writable audio buffer reference of right channel. */
+			FAudioBufferWriteRef GetRight() { return GetBuffer<1>(); }
+
+			/** Return readable audio buffer reference of right channel. */
+			FAudioBufferReadRef GetRight() const { return GetBuffer<1>(); }
+
+		protected:
+
+			// Friendship with the data reference class gives it access to the 
+			// protected constructor for the scenario where a 
+			// TDataReadReference<*> is constructed with FAudioBufferReadRefs. 
+			// The constructor cannot be public as it would provide writable
+			// access to the passed in audio buffers, even though the passed 
+			// buffers explicitly were read references.
+			friend class TDataReadReference<FStereoAudioFormat>;
+
+			// Construct a FStereoAudioFormat a readable buffer reference.
+			// 
+			// This constructor should only be used when it can be assured 
+			// that the constructed object will not provide writable access 
+			// to the passed in audio buffers. 
+			FStereoAudioFormat(const FAudioBufferReadRef& InLeftAudio, const FAudioBufferReadRef& InRightAudio)
+			:	Super({WriteCast(InLeftAudio), WriteCast(InRightAudio)})
+			{
+			}
+	};
+
+
+	DECLARE_METASOUND_DATA_REFERENCE_TYPES(FUnformattedAudio, 0xd78a3ed1 , FUnformattedAudioTypeInfo, FUnformattedAudioReadRef, FUnformattedAudioWriteRef);
+	
+
+	DECLARE_METASOUND_DATA_REFERENCE_TYPES(FMultichannelAudioFormat, 0x56bdcbe0 , FMultichannelAudioFormatTypeInfo, FMultichannelAudioFormatReadRef, FMultichannelAudioFormatWriteRef);
+	
+
+	DECLARE_METASOUND_DATA_REFERENCE_TYPES(FMonoAudioFormat, 0x6f468c8c, FMonoAudioFormatTypeInfo, FMonoAudioFormatReadRef, FMonoAudioFormatWriteRef);
+
+	DECLARE_METASOUND_DATA_REFERENCE_TYPES(FStereoAudioFormat, 0xb55304e2 , FStereoAudioFormatTypeInfo, FStereoAudioFormatReadRef, FStereoAudioFormatWriteRef);
+	
+	
+}
diff --git a/Engine/Plugins/Runtime/Metasound/Source/MetasoundStandardNodes/Public/MetasoundAudioMultiplyNode.h b/Engine/Plugins/Runtime/Metasound/Source/MetasoundStandardNodes/Public/MetasoundAudioMultiplyNode.h
new file mode 100644
--- /dev/null
+++ b/Engine/Plugins/Runtime/Metasound/Source/MetasoundStandardNodes/Public/MetasoundAudioMultiplyNode.h
@@ -0,0 +1,69 @@
+// Copyright Epic Games, Inc. All Rights Reserved.
+#pragma once
+
+#include "CoreMinimal.h"
+#include "Internationalization/Text.h"
+#include "MetasoundNode.h"
+#include "MetasoundOperatorInterface.h"
+#include "MetasoundBuilderInterface.h"
+#include "MetasoundDataReferenceCollection.h"
+
+
+namespace Metasound
+{
+	class METASOUNDSTANDARDNODES_API FAudioMultiplyNode : public FNode
+	{
+		class FOperatorFactory : public IOperatorFactory
+		{
+			public:
+				virtual TUniquePtr<IOperator> CreateOperator(const INode& InNode, const FOperatorSettings& InOperatorSettings, const FDataReferenceCollection& InInputDataReferences, TArray<TUniquePtr<IOperatorBuildError>>& OutErrors) override;
+
+				template<typename ParamType>
+				bool SetReadableRefIfInCollection(const FString& InParamName, const FDataReferenceCollection& InCollection, TDataReadReference<ParamType>& ParamRef)
+				{
+					// TODO: add a helper function to FDataReferenceCollection to do a SetParamIfContains.
+					if (InCollection.ContainsDataReadReference<ParamType>(InParamName))
+					{
+						ParamRef = InCollection.GetDataReadReference<ParamType>(InParamName);
+						return true;
+					}
+
+					return false;
+				}
+
+		};
+
+		public:
+			static const FName ClassName;
+
+			FAudioMultiplyNode(const FString& InName);
+
+			// Constructor used by the Metasound Frontend.
+			FAudioMultiplyNode(const FNodeInitData& InInitData);
+
+			virtual ~FAudioMultiplyNode();
+
+			virtual const FName& GetClassName() const override;
+
+			virtual IOperatorFactory& GetDefaultOperatorFactory() override;
+
+			virtual const FText& GetDescription() const override
+			{
+				static const FText StaticDescription = NSLOCTEXT("MetasoundGraphCore", "Metasound_AudioMultipleNodeDescription", "Multiplies two audio signals together (useful for amplitude modulation).");
+				return StaticDescription;
+			}
+
+			virtual const FText& GetAuthorName() const override
+			{
+				return PluginAuthor;
+			}
+
+			virtual const FText& GetPromptIfMissing() const override
+			{
+				return PluginNodeMissingPrompt;
+			}
+
+		private:
+			FOperatorFactory Factory;
+	};
+} // namespace Metasound
diff --git a/Engine/Plugins/Runtime/Metasound/Source/MetasoundStandardNodes/Public/MetasoundDataReferenceTypes.h b/Engine/Plugins/Runtime/Metasound/Source/MetasoundStandardNodes/Public/MetasoundDataReferenceTypes.h
new file mode 100644
--- /dev/null
+++ b/Engine/Plugins/Runtime/Metasound/Source/MetasoundStandardNodes/Public/MetasoundDataReferenceTypes.h
@@ -0,0 +1,89 @@
+// Copyright Epic Games, Inc. All Rights Reserved.
+
+#pragma once
+
+#include "MetasoundDataReference.h"
+#include "MetasoundAudioBuffer.h"
+#include "MetasoundDataTypeRegistrationMacro.h"
+
+namespace Metasound
+{
+	class FBop
+	{
+		public:
+			// Constructor that, if called with true, bops the beginning of the callback.
+			FBop(bool bShouldBop)
+			{
+				if (bShouldBop)
+				{
+					BopFrame(0);
+				}
+			}
+
+			// Constructor that bops a specific audio frame within the callback.
+			FBop(uint32 InFrameToBop)
+			{
+				BopFrame(InFrameToBop);
+			}
+
+			// The default Bop Constructor 
+			FBop()
+				: FBop(false)
+			{}
+
+			void BopFrame(uint32 InFrameToBop)
+			{
+				// TODO: need to maintain a sorted array. 
+				// Can likely implement a faster version.
+				bool bNeedsSort = false;
+
+				if (Frames.Num() > 0)
+				{
+					bNeedsSort = InFrameToBop < Frames.Last();
+				}
+
+				Frames.Add(InFrameToBop);
+				Frames.Sort();
+			}
+
+			void Advance(uint32 InNumFrames)
+			{
+				Frames.RemoveAllSwap([InNumFrames](uint32 Frame) { return Frame < InNumFrames; }, false /* bAllowShrinking */);
+
+				for (uint32& Frame : Frames)
+				{
+					Frame -= InNumFrames;
+				}
+			}
+
+			int32 Num() const
+			{
+				return Frames.Num();
+			}
+
+			uint32 operator[](int32 Index) const
+			{
+				return Frames[Index];
+			}
+
+			bool IsBopped() const
+			{
+				return Frames.Num() > 0;
+			}
+
+			operator bool() const 
+			{
+				return IsBopped();
+			}
+
+			void Reset()
+			{
+				Frames.Reset();
+			}
+
+		private:
+			TArray<uint32> Frames;
+	};
+
+	DECLARE_METASOUND_DATA_REFERENCE_TYPES(FBop, 2, FBopTypeInfo, FBopReadRef, FBopWriteRef);
+}
diff --git a/Engine/Plugins/Runtime/Metasound/Source/MetasoundStandardNodes/Public/MetasoundFrequency.h b/Engine/Plugins/Runtime/Metasound/Source/MetasoundStandardNodes/Public/MetasoundFrequency.h
new file mode 100644
--- /dev/null
+++ b/Engine/Plugins/Runtime/Metasound/Source/MetasoundStandardNodes/Public/MetasoundFrequency.h
@@ -0,0 +1,154 @@
+// Copyright Epic Games, Inc. All Rights Reserved.
+
+#pragma once
+
+#include "MetasoundDataReference.h"
+#include "MetasoundDataTypeRegistrationMacro.h"
+
+namespace Metasound
+{
+	/** Resolution of frequency value */
+	enum class EFrequencyResolution : uint8
+	{
+		/** Frequency value given in hertz. */
+		Hertz,
+
+		/** Frequency value given in kilohertz. */
+		Kilohertz,
+
+		/** Frequency value given in megahertz. */
+		Megahertz
+	};
+
+	/** FFrequency represents a frequency value. It provides clearly defined 
+	 * getters and setters as well as several convenience functions for 
+	 * converting frequencies to radians per sample..
+	 */
+	class FFrequency
+	{
+		public:
+			/** Construct default frequency of 0 hz. */
+			FFrequency()
+			:	FrequencyInHz(0.f)
+			{
+			}
+
+			/** Construct a frequency with a given value and resolution.
+			 *
+			 * @param InValue - The value of the frequency at the given 
+			 * 					resolution.
+			 * @param InResolution - The resolution of the given value.
+			 */
+			FFrequency(float InValue, EFrequencyResolution InResolution)
+			:	FrequencyInHz(0.f)
+			{
+				switch (InResolution)
+				{
+					case EFrequencyResolution::Megahertz:
+						SetMegahertz(InValue);
+						break;
+
+					case EFrequencyResolution::Kilohertz:
+						SetKilohertz(InValue);
+						break;
+
+					case EFrequencyResolution::Hertz:
+					default:
+						SetHertz(InValue);
+						break;
+				}
+			}
+
+			/**
+			 * FFrequency constructor used to pass in float literals from the metasound frontend.
+			 */
+			FFrequency(float InValue)
+				: FFrequency(InValue, EFrequencyResolution::Hertz)
+			{}
+
+			/** Set the frequency in hertz. */
+			void SetHertz(float InHz)
+			{
+				FrequencyInHz = InHz;	
+			}
+
+			/** Set the frequency in kilohertz. */
+			void SetKilohertz(float InKHz)
+			{
+				FrequencyInHz = InKHz * 1e3f;
+			}
+
+			/** Set the frequency in megahertz. */
+			void SetMegahertz(float InMHz)
+			{
+				FrequencyInHz = InMHz * 1e6f; 
+			}
+
+			/** Set the frequency in radians per a sample.
+			 *
+			 * @param InRadians - The number of radians per a sample.
+			 * @param InSampleRate - The number of samples per a second.
+			 */
+			void SetRadiansPerSample(float InRadians, float InSampleRate)
+			{
+				FrequencyInHz = InSampleRate * InRadians / (2.f * PI);
+			}
+
+			/** Return the frequency in hertz. */
+			float GetHertz() const
+			{
+				return FrequencyInHz;
+			}
+
+			/** Return the frequency in kilohertz. */
+			float GetKilohertz() const
+			{
+				return FrequencyInHz * 1e-3f;
+			}
+
+			/** Return the frequency in megahertz. */
+			float GetMegahertz() const
+			{
+				return FrequencyInHz * 1e-6f;
+			}
+
+			/** Return the radians per a sample.
+			 *
+			 * @param InSampleRate - The number of samples per a second.
+			 *
+			 * @return The phase increment as radians per a sample.
+			 */
+			float GetRadiansPerSample(float InSampleRate) const
+			{
+				if (!ensure(InSampleRate > 0.f))
+				{
+					InSampleRate = SMALL_NUMBER;
+				}
+
+				return (FrequencyInHz * 2.f * PI) / InSampleRate;
+			}
+
+			/** Return the number of radians for a given number of samples.
+			 *
+			 * @param InNumSamples - The number of samples.
+			 * @param InSampleRate - The number of samples per a second.
+			 *
+			 * @return The number of radians for the given samples.
+			 */
+			float GetRadians(float InNumSamples, float InSampleRate) const
+			{
+				if (!ensure(InSampleRate > 0.f))
+				{
+					InSampleRate = SMALL_NUMBER;
+				}
+
+				return (InNumSamples * FrequencyInHz * 2.f * PI) / InSampleRate;
+			}
+		  
+		private:
+
+			float FrequencyInHz;
+	};
+
+	DECLARE_METASOUND_DATA_REFERENCE_TYPES(FFrequency, 0xf1c430a4, FFrequencyTypeInfo, FFrequencyReadRef, FFrequencyWriteRef)
+}
diff --git a/Engine/Plugins/Runtime/Metasound/Source/MetasoundStandardNodes/Public/MetasoundOscNode.h b/Engine/Plugins/Runtime/Metasound/Source/MetasoundStandardNodes/Public/MetasoundOscNode.h
new file mode 100644
--- /dev/null
+++ b/Engine/Plugins/Runtime/Metasound/Source/MetasoundStandardNodes/Public/MetasoundOscNode.h
@@ -0,0 +1,56 @@
+// Copyright Epic Games, Inc. All Rights Reserved.
+
+#pragma once
+
+#include "CoreMinimal.h"
+#include "MetasoundNode.h"
+#include "MetasoundOperatorInterface.h"
+#include "MetasoundBuilderInterface.h"
+#include "MetasoundDataReferenceCollection.h"
+
+namespace Metasound
+{
+	class METASOUNDSTANDARDNODES_API FOscNode : public FNode
+	{
+			class FOperatorFactory : public IOperatorFactory
+			{
+				virtual TUniquePtr<IOperator> CreateOperator(const INode& InNode, const FOperatorSettings& InOperatorSettings, const FDataReferenceCollection& InInputDataReferences, TArray<TUniquePtr<IOperatorBuildError>>& OutErrors) override;
+			};
+
+		public:
+			static const FName ClassName;
+
+			FOscNode(const FString& InName, float InDefaultFrequency);
+
+			// constructor used by the Metasound Frontend.
+			FOscNode(const FNodeInitData& InInitData);
+
+			virtual ~FOscNode();
+
+			float GetDefaultFrequency() const;
+
+			virtual const FName& GetClassName() const override;
+
+			virtual IOperatorFactory& GetDefaultOperatorFactory() override;
+
+			virtual const FText& GetDescription() const override
+			{
+				static const FText StaticDescription = NSLOCTEXT("MetasoundGraphCore", "Metasound_OscNodeDescription", "Emits an audio signal of a sinusoid.");
+				return StaticDescription;
+			}
+
+			virtual const FText& GetAuthorName() const override
+			{
+				return PluginAuthor;
+			}
+
+			virtual const FText& GetPromptIfMissing() const override
+			{
+				return PluginNodeMissingPrompt;
+			}
+
+		private:
+			float DefaultFrequency;
+			FOperatorFactory Factory;
+	};
+}
diff --git a/Engine/Plugins/Runtime/Metasound/Source/MetasoundStandardNodes/Public/MetasoundPeriodicBopNode.h b/Engine/Plugins/Runtime/Metasound/Source/MetasoundStandardNodes/Public/MetasoundPeriodicBopNode.h
new file mode 100644
--- /dev/null
+++ b/Engine/Plugins/Runtime/Metasound/Source/MetasoundStandardNodes/Public/MetasoundPeriodicBopNode.h
@@ -0,0 +1,56 @@
+// Copyright Epic Games, Inc. All Rights Reserved.
+
+#pragma once
+
+#include "CoreMinimal.h"
+#include "MetasoundNode.h"
+#include "MetasoundBuilderInterface.h"
+#include "MetasoundOperatorInterface.h"
+#include "MetasoundNodeInterface.h"
+#include "MetasoundDataReferenceCollection.h"
+#include "MetasoundDataReferenceTypes.h"
+
+namespace Metasound
+{
+	class METASOUNDSTANDARDNODES_API FPeriodicBopNode : public FNode
+	{
+		class FOperatorFactory : public IOperatorFactory
+		{
+			virtual TUniquePtr<IOperator> CreateOperator(const INode& InNode, const FOperatorSettings& InOperatorSettings, const FDataReferenceCollection& InInputDataReferences, TArray<TUniquePtr<IOperatorBuildError>>& OutErrors) override;
+		};
+
+		public:
+			static const FName ClassName;
+
+			FPeriodicBopNode(const FString& InName, float InDefaultPeriodInSeconds);
+			FPeriodicBopNode(const FNodeInitData& InInitData);
+
+			virtual ~FPeriodicBopNode();
+
+			float GetDefaultPeriodInSeconds() const;
+
+			virtual const FName& GetClassName() const override;
+
+			virtual IOperatorFactory& GetDefaultOperatorFactory() override;
+
+			virtual const FText& GetDescription() const override
+			{
+				static const FText StaticDescription = NSLOCTEXT("MetasoundGraphCore", "Metasound_PeriodicBopNodeDescription", "Emits a bop periodically based on the period duration given.");
+				return StaticDescription;
+			}
+
+			virtual const FText& GetAuthorName() const override
+			{
+				return PluginAuthor;
+			}
+
+			virtual const FText& GetPromptIfMissing() const override
+			{
+				return PluginNodeMissingPrompt;
+			}
+
+		private:
+			float DefaultPeriod;
+			FOperatorFactory Factory;
+	};
+}
diff --git a/Engine/Plugins/Runtime/Metasound/Source/MetasoundStandardNodes/Public/MetasoundPrimitives.h b/Engine/Plugins/Runtime/Metasound/Source/MetasoundStandardNodes/Public/MetasoundPrimitives.h
new file mode 100644
--- /dev/null
+++ b/Engine/Plugins/Runtime/Metasound/Source/MetasoundStandardNodes/Public/MetasoundPrimitives.h
@@ -0,0 +1,25 @@
+// Copyright Epic Games, Inc. All Rights Reserved.
+
+#pragma once
+
+#include "CoreMinimal.h"
+#include "MetasoundDataReference.h"
+
+#include "MetasoundTime.h"
+#include "MetasoundFrequency.h"
+#include "MetasoundAudioBuffer.h"
+#include "MetasoundAudioFormats.h"
+
+namespace Metasound
+{
+	DECLARE_METASOUND_DATA_REFERENCE_TYPES(bool, 0xcc9d715e, FBoolTypeInfo, FBoolReadRef, FBoolWriteRef);
+	
+	DECLARE_METASOUND_DATA_REFERENCE_TYPES(int32, 0xf9920034, FInt32TypeInfo, FInt32ReadRef, FInt32WriteRef);
+
+	DECLARE_METASOUND_DATA_REFERENCE_TYPES(int64, 0xfc8f8ada, FInt64TypeInfo, FInt64ReadRef, FInt64WriteRef);
+
+	DECLARE_METASOUND_DATA_REFERENCE_TYPES(float, 0xe7adc4db, FFloatTypeInfo, FFloatReadRef, FFloatWriteRef);
+
+	DECLARE_METASOUND_DATA_REFERENCE_TYPES(double, 0x8f48a053, FDoubleTypeInfo, FDoubleReadRef, FDoubleWriteRef);
+}
+
diff --git a/Engine/Plugins/Runtime/Metasound/Source/MetasoundStandardNodes/Public/MetasoundTime.h b/Engine/Plugins/Runtime/Metasound/Source/MetasoundStandardNodes/Public/MetasoundTime.h
new file mode 100644
--- /dev/null
+++ b/Engine/Plugins/Runtime/Metasound/Source/MetasoundStandardNodes/Public/MetasoundTime.h
@@ -0,0 +1,642 @@
+// Copyright Epic Games, Inc. All Rights Reserved.
+
+#pragma once
+
+#include "MetasoundDataReference.h"
+#include "MetasoundDataTypeRegistrationMacro.h"
+
+namespace Metasound
+{ 
+	/** Resolution of time value */
+	enum class ETimeResolution : uint8
+	{
+		/** Time value given in seconds. */
+		Seconds,
+
+		/** Time value given in milliseconds. */
+		Milliseconds,
+
+		/** Time value given in microseconds. */
+		Microseconds
+	};
+
+	/** TTime represents a time value. It provides basic mathematical 
+	 * operations, clearly defined getters and setters, and interoperability 
+	 * with TTime objects utilizing different template arguments. 
+	 *
+	 * QuantizationType declares the type of object which stores time 
+	 * information.
+	 *
+	 * TimeType declares the type of object which represent time values 
+	 * returned from TTime's getters and accepted by TTime's setters.
+	 */
+	template<typename QuantizationType, typename TimeType>
+	class TTime
+	{
+			static_assert(TIsFloatingPoint<TimeType>::Value, "TimeType must be floating point");
+
+		public:
+			using FQuantizationType = QuantizationType;
+			using FTimeTYpe = TimeType;
+
+			/** Default constructor */
+			TTime()
+			:	Value(0)
+			{}
+			
+			/** Initialize with a time value and time resolution.
+			 *
+			 * @param InValue - The initial value of the time object for the 
+			 * 					given resolution.
+			 * @param InResolution - The Resolution of the given InValue.
+			 */
+			TTime(QuantizationType InValue, ETimeResolution InResolution)
+			:	Value(0)
+			{
+				// Different setters called based upon resolution.
+				switch (InResolution)
+				{
+					case ETimeResolution::Microseconds:
+						SetMicroseconds(InValue);
+						break;
+
+					case ETimeResolution::Milliseconds:
+						SetMilliseconds(InValue);
+						break;
+
+					case ETimeResolution::Seconds:
+					default:
+						SetSeconds(InValue);
+				}
+			}
+
+			/** Constructor used by the metasound backend. */
+			TTime(QuantizationType InSeconds)
+				: TTime(InSeconds, ETimeResolution::Seconds)
+			{}
+
+			/** Returns the time as seconds. */
+			TimeType GetSeconds() const 
+			{ 
+				return static_cast<TimeType>(Value);
+			}
+
+			/** Returns the time as milliseconds. */
+			TimeType GetMilliseconds() const 
+			{ 
+				return static_cast<TimeType>(Value * QuantizationType(1e3));
+			}
+
+			/** Returns the time as microseconds. */
+			TimeType GetMicroseconds() const 
+			{ 
+				return static_cast<TimeType>(Value * QuantizationType(1e6));
+			}
+
+			/** Returns the number of samples which represent the current time 
+			 * value.
+			 *
+			 * @param InSampleRate - The sample rate to use when calculating
+			 * 						 the number of samples. 
+			 *
+			 * @return The number of samples. 
+			 */
+			QuantizationType GetNumSamples(TimeType InSampleRate) const
+			{
+				return static_cast<QuantizationType>(Value * InSampleRate);
+			}
+
+			/** Set the time value in seconds. */
+			void SetSeconds(TimeType InSeconds) 
+			{ 
+				Value = static_cast<QuantizationType>(InSeconds);
+			}
+
+			/** Set the time value in milliseconds. */
+			void SetMilliseconds(TimeType InMilliseconds) 
+			{ 
+				Value = static_cast<QuantizationType>(InMilliseconds * TimeType(1e-3));
+			}
+
+			/** Set the time value in microseconds. */
+			void SetMicroseconds(TimeType InMicroseconds) 
+			{ 
+				Value = static_cast<QuantizationType>(InMicroseconds * TimeType(1e-6));
+			}
+
+			/** Set the number of samples to represent. 
+			 *
+			 * @param InNumSamples - The number of samples for the TTime object 
+			 * 						 to represent.
+			 * @param InSampleRate - The sample rate to use when calculating
+			 * 						 the number of samples. 
+			 */
+			void SetNumSamples(QuantizationType InNumSamples, TimeType InSampleRate)
+			{ 
+				SetSeconds(InNumSamples / FMath::Max(InSampleRate, static_cast<TimeType>(SMALL_NUMBER)));
+			}
+
+			/** Assignment operator. */
+			template<typename OtherQuantizationType, typename OtherTimeType>
+			TTime& operator=(const TTime<OtherQuantizationType, OtherTimeType>& InOther)
+			{
+				Value = static_cast<QuantizationType>(InOther.GetSeconds());
+				return *this;
+			}
+
+			/** Addition assignment operator. 
+			 * 
+			 * Adds the Other's time to this time.
+			 */
+			template<typename OtherQuantizationType, typename OtherTimeType>
+			TTime& operator+=(const TTime<OtherQuantizationType, OtherTimeType>& InOther)
+			{
+				Value += static_cast<QuantizationType>(InOther.GetSeconds());
+				return *this;
+			}
+
+			/** Subtraction assignment operator. 
+			 * 
+			 * Subtracts the Other's time from this time.
+			 */
+			template<typename OtherQuantizationType, typename OtherTimeType>
+			TTime& operator-=(const TTime<OtherQuantizationType, OtherTimeType>& InOther)
+			{
+				Value -= static_cast<QuantizationType>(InOther.GetSeconds());
+				return *this;
+			}
+
+			/* Multiplication assignment operator. 
+			 *
+			 * Multiplies this time by the value. 
+			 */
+			template<typename OtherType>
+			TTime& operator*=(OtherType InValue)
+			{
+				static_assert(TIsArithmetic<OtherType>::Value, "TTime can only be multiplied with arithmetic types");
+
+				Value *= InValue;
+				return *this;
+			}
+
+			/* Division assignment operator. 
+			 *
+			 * Divides this time by the value. 
+			 */
+			template<typename OtherType>
+			TTime& operator/=(OtherType InValue)
+			{
+				static_assert(TIsArithmetic<OtherType>::Value, "TTime can only be divided with arithmetic types");
+
+				Value /= InValue;
+				return *this;
+			}
+
+			/* Comparison operator global friend functions */
+
+			template<typename LHSQuantizationType, typename LHSTimeType, typename RHSQuantizationType, typename RHSTimeType>
+			friend bool operator<(const TTime<LHSQuantizationType, LHSTimeType>& LHS, const TTime<RHSQuantizationType, RHSTimeType>& RHS);
+
+			template<typename LHSQuantizationType, typename LHSTimeType, typename RHSQuantizationType, typename RHSTimeType>
+			friend bool operator>(const TTime<LHSQuantizationType, LHSTimeType>& LHS, const TTime<RHSQuantizationType, RHSTimeType>& RHS);
+
+			template<typename LHSQuantizationType, typename LHSTimeType, typename RHSQuantizationType, typename RHSTimeType>
+			friend bool operator<=(const TTime<LHSQuantizationType, LHSTimeType>& LHS, const TTime<RHSQuantizationType, RHSTimeType>& RHS);
+
+			template<typename LHSQuantizationType, typename LHSTimeType, typename RHSQuantizationType, typename RHSTimeType>
+			friend bool operator>=(const TTime<LHSQuantizationType, LHSTimeType>& LHS, const TTime<RHSQuantizationType, RHSTimeType>& RHS);
+
+			template<typename LHSQuantizationType, typename LHSTimeType, typename RHSQuantizationType, typename RHSTimeType>
+			friend bool operator==(const TTime<LHSQuantizationType, LHSTimeType>& LHS, const TTime<RHSQuantizationType, RHSTimeType>& RHS);
+
+			template<typename LHSQuantizationType, typename LHSTimeType, typename RHSQuantizationType, typename RHSTimeType>
+			friend bool operator!=(const TTime<LHSQuantizationType, LHSTimeType>& LHS, const TTime<RHSQuantizationType, RHSTimeType>& RHS);
+
+			/* Arithmetic operator global friend functions */
+
+			template<typename LHSQuantizationType, typename LHSTimeType, typename RHSQuantizationType, typename RHSTimeType>
+			friend TTime<LHSQuantizationType, LHSTimeType> operator+(TTime<LHSQuantizationType, LHSTimeType> LHS, const TTime<RHSQuantizationType, RHSTimeType>& RHS);
+
+			template<typename LHSQuantizationType, typename LHSTimeType, typename RHSQuantizationType, typename RHSTimeType>
+			friend TTime<LHSQuantizationType, LHSTimeType> operator-(TTime<LHSQuantizationType, LHSTimeType> LHS, const TTime<RHSQuantizationType, RHSTimeType>& RHS);
+
+			template<typename LHSQuantizationType, typename LHSTimeType, typename OtherType>
+			friend TTime<LHSQuantizationType, LHSTimeType> operator*(TTime<LHSQuantizationType, LHSTimeType> LHS, OtherType RHS);
+
+			template<typename RHSQuantizationType, typename RHSTimeType, typename OtherType>
+			friend TTime<RHSQuantizationType, RHSTimeType> operator*(OtherType LHS, const TTime<RHSQuantizationType, RHSTimeType>& RHS);
+
+			template<typename LHSQuantizationType, typename LHSTimeType, typename OtherType>
+			friend TTime<LHSQuantizationType, LHSTimeType> operator/(TTime<LHSQuantizationType, LHSTimeType> LHS, OtherType RHS);
+
+
+		private:
+
+			QuantizationType Value;
+	};
+
+
+	// holds on to a count of samples.
+	struct FTimeSampleCounter
+	{
+		// defines the type for counting samples.
+		typedef int32 FCountType;
+		FCountType Num;
+	};
+
+	/** TTime sample counting template specialization holds time values which
+	 * are quantized to sample boundaries. It provides basic mathematical 
+	 * operations, clearly defined getters and setters, and interoperability 
+	 * with generic TTime objects utilizing different template arguments. 
+	 *
+	 * This specialization differs distinctly by storing the number of samples
+	 * rather than the number of seconds in the quantization type. Additionally,
+	 * instances have an associated sample rate member which is used to convert
+	 * between TimeType and the sample count.
+	 *
+	 * TimeType declares the type of object which represent time values 
+	 * returned from TTime's getters.
+	 */
+	template<typename TimeType>
+	class TTime<FTimeSampleCounter, TimeType>
+	{
+		public:
+			static_assert(TIsFloatingPoint<TimeType>::Value, "TimeType must be floating point");
+
+			// using types from namespace.
+			using FCountType = FTimeSampleCounter::FCountType;
+			
+			/** Construct a sample count based TTime with a given time value, 
+			 * time resolution, and sample rate.
+			 *
+			 * @param InValue - The initial time value.
+			 * @param InResolution - The time resolution of InValue.
+			 * @param InSampleRate - The sample rate associated with this time 
+			 * 						 object.
+			 */
+			TTime(TimeType InValue, ETimeResolution InResolution, TimeType InSampleRate)
+			{
+				// Need to set sample rate before setting time or else the 
+				// resulting sample count will be incorrect.
+				SetSampleRate(InSampleRate);
+
+				switch (InResolution)
+				{
+					case ETimeResolution::Microseconds:
+						SetMicroseconds(InValue);
+						break;
+
+					case ETimeResolution::Milliseconds:
+						SetMilliseconds(InValue);
+						break;
+
+					case ETimeResolution::Seconds:
+					default:
+						SetSeconds(InValue);
+				}
+			}
+
+			/**
+			 * Constructor used by the Metasound Frontend.
+			 *
+			 * @param InSeconds The number in time to initialize to.
+			 * @param InSettings the operator settings to use.
+			 */
+			TTime(float InSeconds, const FOperatorSettings& InSettings)
+				: TTime(InSeconds, ETimeResolution::Seconds, InSettings.GetSampleRate())
+			{}
+
+			/** Construct a sample count based TTime using a number of samples
+			 * and a sample rate.
+			 *
+			 * @param InNumSamples - The initial number of samples.
+			 * @param InSampleRate - The sample rate of this instance.
+			 */
+			TTime(FCountType InNumSamples, TimeType InSampleRate)
+			{
+				SetSampleRate(InSampleRate);
+
+				SetNumSamples(InNumSamples);
+			}
+
+			/** Set the sample rate of this object without changing the number
+			 * of samples. This will result in the object representing a 
+			 * different amount of time. 
+			 */
+			void SetSampleRate(TimeType InSampleRate)
+			{
+				if (!ensure(InSampleRate > 0))
+				{
+					InSampleRate = SMALL_NUMBER;
+				}
+
+				SampleRate = InSampleRate;
+			}
+
+			/** Set the sample rate of this object and change the number
+			 * of samples. This will result in the object representing 
+			 * approximately the same amount of time as it did before the 
+			 * sample rate was changed.
+			 */
+			void SetSampleRateAndUpdateSampleCount(TimeType InSampleRate)
+			{
+				if (!ensure(InSampleRate > 0))
+				{
+					InSampleRate = SMALL_NUMBER;
+				}
+
+				Value.Num = GetNumSamples(InSampleRate);
+				
+				SampleRate = InSampleRate;
+			}
+
+			/** Return the sample rate. */
+			TimeType GetSampleRate() const
+			{
+				return SampleRate;
+			}
+
+			/** Return the time as seconds. */
+			TimeType GetSeconds() const 
+			{ 
+				return static_cast<TimeType>(Value.Num / SampleRate);
+			}
+
+			/** Return the time as milliseconds. */
+			TimeType GetMilliseconds() const 
+			{ 
+				return static_cast<TimeType>(Value.Num * 1e3 / SampleRate);
+			}
+
+			/** Return the time as microseconds. */
+			TimeType GetMicroseconds() const 
+			{ 
+				return static_cast<TimeType>(Value.Num * 1e6 / SampleRate);
+			}
+
+			/** Return the number of samples. */
+			FCountType GetNumSamples() const
+			{ 
+				return Value.Num;
+			}
+
+			/** Return the number of samples which represent this time duration
+			 * using a different sample rate. 
+			 *
+			 * @param InOtherSampleRate - Sample rate to use when calculating 
+			 * 							  the number of samples.  
+			 */
+			FCountType GetNumSamples(TimeType InOtherSampleRate) const
+			{ 
+				if (InOtherSampleRate == SampleRate)
+				{
+					return Value.Num;
+				}
+				else
+				{
+					return FMath::RoundToInt(GetSeconds() * InOtherSampleRate);
+				}
+			}
+
+			/** Set the time value in seconds. */
+			void SetSeconds(TimeType InSeconds) 
+			{ 
+				Value.Num = FMath::RoundToInt(InSeconds * SampleRate);
+			}
+
+			/** Set the time value in milliseconds. */
+			void SetMilliseconds(TimeType InMilliseconds) 
+			{ 
+				Value.Num = FMath::RoundToInt(InMilliseconds * SampleRate * TimeType(1e-3));
+			}
+
+			/** Set the time value in microseconds. */
+			void SetMicroseconds(TimeType InMicroseconds) 
+			{ 
+				Value.Num = FMath::RoundToInt(InMicroseconds * SampleRate * TimeType(1e-6));
+			}
+
+			/** Set the number of samples. */
+			void SetNumSamples(FCountType InNumSamples)
+			{ 
+				Value.Num = InNumSamples;
+			}
+
+			/** Assignment operator 
+			 *
+			 * This assignment preserves this objects sample rate when being 
+			 * assigned to TTime objects of different quantization types.
+			 * */
+			template<typename OtherQuantizationType, typename OtherTimeType>
+			TTime& operator=(const TTime<OtherQuantizationType, OtherTimeType>& Other)
+			{
+				Value.Num = Other.GetNumSamples(SampleRate);
+
+				return *this;
+			}
+
+			/** Assignment operator 
+			 *
+			 * This assignment copies the sample rate and sample count of TTime
+			 * objects of equal quantization types.
+			 */
+			template<typename OtherTimeType>
+			TTime& operator=(const TTime<FTimeSampleCounter, OtherTimeType>& Other)
+			{
+				Value.Num = Other.Value.Num;
+				SampleRate = Other.SampleRate;
+
+				return *this;
+			}
+
+			/** Addition assignment operator. */
+			template<typename OtherQuantizationType, typename OtherTimeType>
+			TTime& operator+=(const TTime<OtherQuantizationType, OtherTimeType>& Other)
+			{
+				Value.Num += Other.GetNumSamples(SampleRate);
+
+				return *this;
+			}
+
+			/** Subtraction assignment operator. */
+			template<typename OtherQuantizationType, typename OtherTimeType>
+			TTime& operator-=(const TTime<OtherQuantizationType, OtherTimeType>& Other)
+			{
+				Value.Num -= Other.GetNumSamples(SampleRate);
+
+				return *this;
+			}
+
+			/** Multiplication assignment operator. */
+			template<typename OtherType>
+			TTime& operator*=(OtherType Other)
+			{
+				static_assert(TIsArithmetic<OtherType>::Value, "TTime can only be multiplied with arithmetic types");
+
+				Value.Num *= Other;
+				return *this;
+			}
+
+			/** Division assignment operator. */
+			template<typename OtherType>
+			TTime& operator/=(OtherType Other)
+			{
+				static_assert(TIsArithmetic<OtherType>::Value, "TTime can only be divided with arithmetic types");
+
+				Value.Num /= Other;
+				return *this;
+			}
+
+			/* Comparison operator global friend functions */
+			
+			template<typename LHSQuantizationType, typename LHSTimeType, typename RHSQuantizationType, typename RHSTimeType>
+			friend bool operator<(const TTime<LHSQuantizationType, LHSTimeType>& LHS, const TTime<RHSQuantizationType, RHSTimeType>& RHS);
+
+			template<typename LHSQuantizationType, typename LHSTimeType, typename RHSQuantizationType, typename RHSTimeType>
+			friend bool operator>(const TTime<LHSQuantizationType, LHSTimeType>& LHS, const TTime<RHSQuantizationType, RHSTimeType>& RHS);
+
+			template<typename LHSQuantizationType, typename LHSTimeType, typename RHSQuantizationType, typename RHSTimeType>
+			friend bool operator<=(const TTime<LHSQuantizationType, LHSTimeType>& LHS, const TTime<RHSQuantizationType, RHSTimeType>& RHS);
+
+			template<typename LHSQuantizationType, typename LHSTimeType, typename RHSQuantizationType, typename RHSTimeType>
+			friend bool operator>=(const TTime<LHSQuantizationType, LHSTimeType>& LHS, const TTime<RHSQuantizationType, RHSTimeType>& RHS);
+
+			template<typename LHSQuantizationType, typename LHSTimeType, typename RHSQuantizationType, typename RHSTimeType>
+			friend bool operator==(const TTime<LHSQuantizationType, LHSTimeType>& LHS, const TTime<RHSQuantizationType, RHSTimeType>& RHS);
+
+			template<typename LHSQuantizationType, typename LHSTimeType, typename RHSQuantizationType, typename RHSTimeType>
+			friend bool operator!=(const TTime<LHSQuantizationType, LHSTimeType>& LHS, const TTime<RHSQuantizationType, RHSTimeType>& RHS);
+
+			/* Arithmetic operator global friend functions */
+
+			template<typename LHSQuantizationType, typename LHSTimeType, typename RHSQuantizationType, typename RHSTimeType>
+			friend TTime<LHSQuantizationType, LHSTimeType> operator+(TTime<LHSQuantizationType, LHSTimeType> LHS, const TTime<RHSQuantizationType, RHSTimeType>& RHS);
+
+			template<typename LHSQuantizationType, typename LHSTimeType, typename RHSQuantizationType, typename RHSTimeType>
+			friend TTime<LHSQuantizationType, LHSTimeType> operator-(TTime<LHSQuantizationType, LHSTimeType> LHS, const TTime<RHSQuantizationType, RHSTimeType>& RHS);
+
+			template<typename LHSQuantizationType, typename LHSTimeType, typename OtherType>
+			friend TTime<LHSQuantizationType, LHSTimeType> operator*(TTime<LHSQuantizationType, LHSTimeType> LHS, OtherType RHS);
+
+			template<typename RHSQuantizationType, typename RHSTimeType, typename OtherType>
+			friend TTime<RHSQuantizationType, RHSTimeType> operator*(OtherType LHS, const TTime<RHSQuantizationType, RHSTimeType>& RHS);
+
+			template<typename LHSQuantizationType, typename LHSTimeType, typename OtherType>
+			friend TTime<LHSQuantizationType, LHSTimeType> operator/(TTime<LHSQuantizationType, LHSTimeType> LHS, OtherType RHS);
+
+		private:
+
+			FTimeSampleCounter Value = {0};
+			TimeType SampleRate = static_cast<TimeType>(SMALL_NUMBER);
+	};
+
+
+	/** Less-than comparison of TTime */
+	template<typename LHSQuantizationType, typename LHSTimeType, typename RHSQuantizationType, typename RHSTimeType>
+	bool operator<(const TTime<LHSQuantizationType, LHSTimeType>& LHS, const TTime<RHSQuantizationType, RHSTimeType>& RHS)
+	{ 
+		return LHS.GetSeconds() < RHS.GetSeconds();
+	}
+
+	/** Greater-than comparison of TTime */
+	template<typename LHSQuantizationType, typename LHSTimeType, typename RHSQuantizationType, typename RHSTimeType>
+	bool operator>(const TTime<LHSQuantizationType, LHSTimeType>& LHS, const TTime<RHSQuantizationType, RHSTimeType>& RHS)
+	{ 
+		return RHS < LHS; 
+	}
+
+	/** Less-than-or-equal-to comparison of TTime */
+	template<typename LHSQuantizationType, typename LHSTimeType, typename RHSQuantizationType, typename RHSTimeType>
+	bool operator<=(const TTime<LHSQuantizationType, LHSTimeType>& LHS, const TTime<RHSQuantizationType, RHSTimeType>& RHS)
+	{ 
+		return !(LHS > RHS); 
+	}
+
+	/** Greater-than-or-equal-to comparison of TTime */
+	template<typename LHSQuantizationType, typename LHSTimeType, typename RHSQuantizationType, typename RHSTimeType>
+	bool operator>=(const TTime<LHSQuantizationType, LHSTimeType>& LHS, const TTime<RHSQuantizationType, RHSTimeType>& RHS)
+	{ 
+		return !(LHS < RHS); 
+	}
+
+	/** Equal comparison of TTime */
+	template<typename LHSQuantizationType, typename LHSTimeType, typename RHSQuantizationType, typename RHSTimeType>
+	bool operator==(const TTime<LHSQuantizationType, LHSTimeType>& LHS, const TTime<RHSQuantizationType, RHSTimeType>& RHS)
+	{
+		return LHS.GetSeconds() == RHS.GetSeconds();
+	}
+
+	/** Not-equal comparison of TTime */
+	template<typename LHSQuantizationType, typename LHSTimeType, typename RHSQuantizationType, typename RHSTimeType>
+	bool operator!=(const TTime<LHSQuantizationType, LHSTimeType>& LHS, const TTime<RHSQuantizationType, RHSTimeType>& RHS)
+	{ 
+		return !(LHS == RHS); 
+	}
+
+	/** Addition operator of TTime */
+	template<typename LHSQuantizationType, typename LHSTimeType, typename RHSQuantizationType, typename RHSTimeType>
+	TTime<LHSQuantizationType, LHSTimeType> operator+(TTime<LHSQuantizationType, LHSTimeType> LHS, const TTime<RHSQuantizationType, RHSTimeType>& RHS)
+	{
+		LHS += RHS;
+		return LHS;
+	}
+
+	/** Subtraction operator of TTime */
+	template<typename LHSQuantizationType, typename LHSTimeType, typename RHSQuantizationType, typename RHSTimeType>
+	TTime<LHSQuantizationType, LHSTimeType> operator-(TTime<LHSQuantizationType, LHSTimeType> LHS, const TTime<RHSQuantizationType, RHSTimeType>& RHS)
+	{
+		LHS -= RHS;
+		return LHS;
+	}
+
+	/** Multiplication operator of TTime */
+	template<typename LHSQuantizationType, typename LHSTimeType, typename OtherType>
+	TTime<LHSQuantizationType, LHSTimeType> operator*(TTime<LHSQuantizationType, LHSTimeType> LHS, OtherType RHS)
+	{
+		static_assert(TIsArithmetic<OtherType>::Value, "TTime can only be multiplied with arithmetic types");
+
+		LHS *= RHS;
+		return LHS;
+	}
+
+	/** Division operator of TTime */
+	template<typename LHSQuantizationType, typename LHSTimeType, typename OtherType>
+	TTime<LHSQuantizationType, LHSTimeType> operator/(TTime<LHSQuantizationType, LHSTimeType> LHS, OtherType RHS)
+	{
+		static_assert(TIsArithmetic<OtherType>::Value, "TTime can only be divided with arithmetic types");
+
+		LHS /= RHS;
+		return LHS;
+	}
+
+	/** Multiplication operator of TTime */
+	template<typename RHSQuantizationType, typename RHSTimeType, typename OtherType>
+	TTime<RHSQuantizationType, RHSTimeType> operator*(OtherType LHS, const TTime<RHSQuantizationType, RHSTimeType>& RHS)
+	{
+		static_assert(TIsArithmetic<OtherType>::Value, "TTime can only be multiplied with arithmetic types");
+
+		return RHS * LHS;
+	}
+	
+
+	/** FFloatTime stores and represents all time values as 32bit floating 
+	 * point values. 
+	 */
+	typedef TTime<float, float> FFloatTime;
+
+	DECLARE_METASOUND_DATA_REFERENCE_TYPES(FFloatTime, 0x2490848e, FFloatTimeTypeInfo, FFloatTimeReadRef, FFloatTimeWriteRef);
+
+	/** FDoubleTime stores and represents all time values as 64bit floating
+	 * point values.
+	 */
+	typedef TTime<double, double> FDoubleTime;
+
+	DECLARE_METASOUND_DATA_REFERENCE_TYPES(FDoubleTime, 0x407fcff1, FDoubleTimeTypeInfo, FDoubleTimeReadRef, FDoubleTimeWriteRef);
+
+	/** FSampleTime stores time quantized to sample boundaries. It returns time
+	 * as 32bit floating point values.
+	 */
+	typedef TTime<FTimeSampleCounter, float> FSampleTime;
+
+	DECLARE_METASOUND_DATA_REFERENCE_TYPES(FSampleTime, 0xf1de8f2d, FSampleTimeTypeInfo, FSampleTimeReadRef, FSampleTimeWriteRef);
+}
diff --git a/Engine/Plugins/Runtime/Metasound/Source/MetasoundStandardNodes/Public/MetasoundWave.h b/Engine/Plugins/Runtime/Metasound/Source/MetasoundStandardNodes/Public/MetasoundWave.h
new file mode 100644
--- /dev/null
+++ b/Engine/Plugins/Runtime/Metasound/Source/MetasoundStandardNodes/Public/MetasoundWave.h
@@ -0,0 +1,18 @@
+// Copyright Epic Games, Inc. All Rights Reserved.
+
+#pragma once
+
+#include "MetasoundDataReference.h"
+#include "MetasoundDataTypeRegistrationMacro.h"
+
+namespace Metasound
+{
+	class FWave
+	{
+		TArray<uint8> CompressedBytes;
+	public:
+		const TArray<uint8>& GetCompressedData() const { return CompressedBytes; }
+	};
+
+	DECLARE_METASOUND_DATA_REFERENCE_TYPES(FWave, 0x0ddba11, FWaveTypeInfo, FWaveReadRef, FWaveWriteRef)
+}
diff --git a/Engine/Plugins/Runtime/Metasound/Source/MetasoundStandardNodes/Public/MetasoundWavePlayerNode.h b/Engine/Plugins/Runtime/Metasound/Source/MetasoundStandardNodes/Public/MetasoundWavePlayerNode.h
new file mode 100644
--- /dev/null
+++ b/Engine/Plugins/Runtime/Metasound/Source/MetasoundStandardNodes/Public/MetasoundWavePlayerNode.h
@@ -0,0 +1,54 @@
+// Copyright Epic Games, Inc. All Rights Reserved.
+
+#pragma once
+
+#include "CoreMinimal.h"
+#include "MetasoundNode.h"
+#include "MetasoundOperatorInterface.h"
+#include "MetasoundBuilderInterface.h"
+#include "MetasoundDataReferenceCollection.h"
+
+namespace Metasound
+{
+	class METASOUNDSTANDARDNODES_API FWavePlayerNode : public FNode
+	{
+			class FOperatorFactory : public IOperatorFactory
+			{
+				virtual TUniquePtr<IOperator> CreateOperator(const INode& InNode, const FOperatorSettings& InOperatorSettings, const FDataReferenceCollection& InInputDataReferences, TArray<TUniquePtr<IOperatorBuildError>>& OutErrors) override;
+			};
+
+		public:
+			static const FName ClassName;
+
+			FWavePlayerNode(const FString& InName);
+
+			// constructor used by the Metasound Frontend.
+			FWavePlayerNode(const FNodeInitData& InInitData);
+
+			virtual ~FWavePlayerNode();
+
+			const FName& GetClassName() const override;
+
+			IOperatorFactory& GetDefaultOperatorFactory() override;
+
+			const FText& GetDescription() const override
+			{
+				static const FText StaticDescription = NSLOCTEXT("MetasoundGraphCore", "Metasound_WavePlayerNodeDescription", "Plays a supplied Wave");
+				return StaticDescription;
+			}
+
+			const FText& GetAuthorName() const override
+			{
+				return PluginAuthor;
+			}
+
+			const FText& GetPromptIfMissing() const override
+			{
+				return PluginNodeMissingPrompt;
+			}
+
+		private:
+			float DefaultFrequency;
+			FOperatorFactory Factory;
+	};
+}
