minifix: fix comment for UMetasoundSource

diff --git a/Engine/Plugins/Runtime/Metasound/Source/MetasoundEngine/Public/MetasoundSource.h b/Engine/Plugins/Runtime/Metasound/Source/MetasoundEngine/Public/MetasoundSource.h
--- a/Engine/Plugins/Runtime/Metasound/Source/MetasoundEngine/Public/MetasoundSource.h
+++ b/Engine/Plugins/Runtime/Metasound/Source/MetasoundEngine/Public/MetasoundSource.h
@@ -17,8 +17,7 @@
 
 
 /**
- * This asset type is used for Metasound assets that can only be used as nodes in other Metasound graphs.
- * Because of this, they can have any inputs or outputs they need.
+ * This Metasound type can be played as an audio source.
  */
 UCLASS(hidecategories = object, BlueprintType)
 class METASOUNDENGINE_API UMetasoundSource : public USoundWaveProcedural, public FMetasoundAssetBase
