Microfix: better pin colors for various metasound datatypes (Time, Frequency, etc.)

diff --git a/Engine/Plugins/Runtime/Metasound/Source/MetasoundStandardNodes/Private/MetasoundStandardNodesModule.cpp b/Engine/Plugins/Runtime/Metasound/Source/MetasoundStandardNodes/Private/MetasoundStandardNodesModule.cpp
--- a/Engine/Plugins/Runtime/Metasound/Source/MetasoundStandardNodes/Private/MetasoundStandardNodesModule.cpp
+++ b/Engine/Plugins/Runtime/Metasound/Source/MetasoundStandardNodes/Private/MetasoundStandardNodesModule.cpp
@@ -20,10 +20,10 @@ REGISTER_METASOUND_DATATYPE(double, "Primitive:Double", ::Metasound::ELiteralArg
 namespace Metasound
 {
 	REGISTER_METASOUND_DATATYPE(FBop, "Primitive:Bop", ::Metasound::ELiteralArgType::Boolean)
-	REGISTER_METASOUND_DATATYPE(FFloatTime, "Primitive:Time")
-	REGISTER_METASOUND_DATATYPE(FDoubleTime, "Primitive:Time:HighResolution")
-	REGISTER_METASOUND_DATATYPE(FSampleTime, "Primitive:Time:SampleResolution")
-	REGISTER_METASOUND_DATATYPE(FFrequency, "Primitive:Frequency")
+	REGISTER_METASOUND_DATATYPE(FFloatTime, "Primitive:Time", ::Metasound::ELiteralArgType::Float)
+	REGISTER_METASOUND_DATATYPE(FDoubleTime, "Primitive:Time:HighResolution", ::Metasound::ELiteralArgType::Float)
+	REGISTER_METASOUND_DATATYPE(FSampleTime, "Primitive:Time:SampleResolution", ::Metasound::ELiteralArgType::Float)
+	REGISTER_METASOUND_DATATYPE(FFrequency, "Primitive:Frequency", ::Metasound::ELiteralArgType::Float)
 	REGISTER_METASOUND_DATATYPE(FWave, "Primitive:Wave")
 }
 
