CIS non-unity build fix to get includes for metasound wave node player

diff --git a/Engine/Plugins/Runtime/Metasound/Source/MetasoundStandardNodes/Private/MetasoundWavePlayerNode.cpp b/Engine/Plugins/Runtime/Metasound/Source/MetasoundStandardNodes/Private/MetasoundWavePlayerNode.cpp
--- a/Engine/Plugins/Runtime/Metasound/Source/MetasoundStandardNodes/Private/MetasoundWavePlayerNode.cpp
+++ b/Engine/Plugins/Runtime/Metasound/Source/MetasoundStandardNodes/Private/MetasoundWavePlayerNode.cpp
@@ -4,6 +4,7 @@
 #include "MetasoundExecutableOperator.h"
 #include "MetasoundPrimitives.h"
 #include "MetasoundWave.h"
+#include "MetasoundNodeRegistrationMacro.h"
 
 #include "IAudioCodecRegistry.h"
 #include "IAudioCodec.h"
