Fix for incremental compile of UE4Game Android

diff --git a/Engine/Plugins/Runtime/Metasound/Source/MetasoundEngine/Public/MetasoundArchetypeRegistration.h b/Engine/Plugins/Runtime/Metasound/Source/MetasoundEngine/Public/MetasoundArchetypeRegistration.h
--- a/Engine/Plugins/Runtime/Metasound/Source/MetasoundEngine/Public/MetasoundArchetypeRegistration.h
+++ b/Engine/Plugins/Runtime/Metasound/Source/MetasoundEngine/Public/MetasoundArchetypeRegistration.h
@@ -4,6 +4,7 @@
 #include "CoreMinimal.h"
 #include "MetasoundAssetBase.h"
 #include "MetasoundFrontend.h"
+#include "Misc/PackageName.h"
 
 #if WITH_EDITOR
 #include "AssetRegistryModule.h"
