Adding Trajan algorithm to find graph cycles

diff --git a/Engine/Plugins/Runtime/Metasound/Source/MetasoundGraphCore/Private/MetasoundBuildError.cpp b/Engine/Plugins/Runtime/Metasound/Source/MetasoundGraphCore/Private/MetasoundBuildError.cpp
--- a/Engine/Plugins/Runtime/Metasound/Source/MetasoundGraphCore/Private/MetasoundBuildError.cpp
+++ b/Engine/Plugins/Runtime/Metasound/Source/MetasoundGraphCore/Private/MetasoundBuildError.cpp
@@ -72,7 +72,7 @@ namespace Metasound
 
 	const FName FGraphCycleError::ErrorType = FName(TEXT("MetasoundGraphCycleError"));
 
-	FGraphCycleError::FGraphCycleError(const TArray<INode*>& InNodes, const TArray<FDataEdge>& InEdges)
+	FGraphCycleError::FGraphCycleError(TArrayView<INode const* const> InNodes, const TArray<FDataEdge>& InEdges)
 	:	FBuildErrorBase(ErrorType, LOCTEXT("GraphCycleError", "Graph contains cycles."))
 	{
 		AddNodes(InNodes);
diff --git a/Engine/Plugins/Runtime/Metasound/Source/MetasoundGraphCore/Private/MetasoundOperatorBuilder.cpp b/Engine/Plugins/Runtime/Metasound/Source/MetasoundGraphCore/Private/MetasoundOperatorBuilder.cpp
--- a/Engine/Plugins/Runtime/Metasound/Source/MetasoundGraphCore/Private/MetasoundOperatorBuilder.cpp
+++ b/Engine/Plugins/Runtime/Metasound/Source/MetasoundGraphCore/Private/MetasoundOperatorBuilder.cpp
@@ -4,6 +4,7 @@
 #include "MetasoundBuilderInterface.h"
 #include "MetasoundOperatorInterface.h"
 #include "MetasoundGraphOperator.h"
+#include "MetasoundTarjan.h"
 
 #include "MetasoundBuildError.h"
 
@@ -59,7 +60,7 @@ namespace Metasound
 
 		// TODO: Validate edges.
 
-		bSuccess = TopologicalSort(InGraph.GetDataEdges(), SortedNodes, OutErrors);
+		bSuccess = TopologicalSort(InGraph, SortedNodes, OutErrors);
 		if (!bSuccess)
 		{
 			return TUniquePtr<IOperator>(nullptr);
@@ -91,7 +92,7 @@ namespace Metasound
 		return true;
 	}
 
-	bool FOperatorBuilder::TopologicalSort(const TArray<FDataEdge>& InEdges, TArray<INode*>& OutNodes, TArray<FBuildErrorPtr>& OutErrors) const
+	bool FOperatorBuilder::TopologicalSort(const IGraph& InGraph, TArray<INode*>& OutNodes, TArray<FBuildErrorPtr>& OutErrors) const
 	{
 		using namespace OperatorBuilderPrivate;
 		using FNodePair = TPair<INode*, INode*>;
@@ -101,7 +102,7 @@ namespace Metasound
 		TArray<INode*> UniqueNodes;
 		FNodeMultiMap Dependencies;
 
-		for (const FDataEdge& Edge : InEdges)
+		for (const FDataEdge& Edge : InGraph.GetDataEdges())
 		{
 			if (nullptr == Edge.To.Node)
 			{
@@ -134,13 +135,24 @@ namespace Metasound
 
 			if (0 == IndependentNodes.Num())
 			{
-				// add build error. likely a cycle in the graph.
+				// likely a cycle in the graph.
 				
-				// Get edges in the cylce
-				TArray<FDataEdge> CycleEdges;
-				GetEdgesBetweenNodes(UniqueNodes, InEdges, CycleEdges);
+				bool bExcludeSingleVertex = true;
+
+				// Try to find cycles
+				TArray<FMetasoundStronglyConnectedComponent> Cycles;
 
-				AddBuildError<FGraphCycleError>(OutErrors, UniqueNodes, CycleEdges);
+				if(FTarjan::FindStronglyConnectedComponents(InGraph, Cycles, bExcludeSingleVertex))
+				{
+					for (const FMetasoundStronglyConnectedComponent& Cycle : Cycles)
+					{
+						AddBuildError<FGraphCycleError>(OutErrors, Cycle.Nodes, Cycle.Edges);
+					}
+				}
+				else
+				{
+					AddBuildError<FInternalError>(OutErrors, TEXT(__FILE__), __LINE__);
+				}
 
 				return false;
 			}
@@ -213,8 +225,6 @@ namespace Metasound
 				return false;
 			}
 
-
-
 			bool bSuccess = OutCollection.AddDataReadReferenceFrom(Edge->To.Vertex.VertexName, FromDataReferenceCollection, Edge->From.Vertex.VertexName, Edge->From.Vertex.DataReferenceTypeName);
 
 			if (!bSuccess)
diff --git a/Engine/Plugins/Runtime/Metasound/Source/MetasoundGraphCore/Private/MetasoundTarjan.cpp b/Engine/Plugins/Runtime/Metasound/Source/MetasoundGraphCore/Private/MetasoundTarjan.cpp
new file mode 100644
--- /dev/null
+++ b/Engine/Plugins/Runtime/Metasound/Source/MetasoundGraphCore/Private/MetasoundTarjan.cpp
@@ -0,0 +1,273 @@
+// Copyright Epic Games, Inc. All Rights Reserved.
+
+#include "MetasoundTarjan.h"
+
+namespace Metasound
+{
+	namespace MetasoundTarjanPrivate
+	{
+		struct FTarjanVertexInfo
+		{
+			// Tag value given to a vertex.
+			int32 Tag = INDEX_NONE;
+
+			// Lowest tag in strongly connected components
+			int32 LowLinkTag = INDEX_NONE;
+
+			// Flag for global stack
+			bool bIsOnStack = false;
+
+			// Immediate children of vertex in DAG.
+			TArray<int32> Children;
+		};
+
+		struct FTarjanAlgoImplSettings 
+		{
+			// Single vertices can be returned as strongly connected components
+			// but usually they are not wanted by the caller. Setting this to
+			// true will require all strongly connected components to contain
+			// two or more vertices.
+			bool bExcludeSingleVertex;
+		};
+
+		// Private implementation of algorithm.
+		class FTarjanAlgoImpl
+		{
+			public:
+
+			FTarjanAlgoImpl(const FTarjanAlgoImplSettings& InSettings)
+			:	Settings(InSettings)
+			{
+			}
+
+			bool FindStronglyConnectedComponents(const TSet<FTarjanEdge>& InEdges, TArray<FTarjanStronglyConnectedComponent>& OutComponents)
+			{
+				using FElement = TMap<int32, FTarjanVertexInfo>::ElementType;
+
+				// Initialize internal data stores and tracking info.
+				Init(InEdges);
+				
+				int32 OriginalNumOutComponents = OutComponents.Num();
+
+				// Iterate through all vertices. Update(...) does a 
+				// depth-first-search which may result some vertices getting 
+				// inspected during recursive calls to Update(...).
+				for (FElement& Element : VertexInfos)
+				{
+					// Only examine vertice if it hasn't already been inspected.
+					if (Element.Value.Tag == INDEX_NONE)
+					{
+						Update(Element.Key, Element.Value, OutComponents);
+					}
+				}
+				
+				// Check if the size of OutComponents has grown.
+				return OutComponents.Num() > OriginalNumOutComponents;
+			}
+
+			private:
+
+			void Init(const TSet<FTarjanEdge>& InEdges)
+			{
+				CurrentTag = 0;
+
+				VertexInfos.Reset();
+				VertexStack.Reset();
+
+				for (const FTarjanEdge& Edge : InEdges)
+				{
+					FTarjanVertexInfo& Info = VertexInfos.FindOrAdd(Edge.Get<0>());
+					Info.Children.AddUnique(Edge.Get<1>());
+
+					VertexInfos.FindOrAdd(Edge.Get<1>());
+				}
+			}
+
+			void Update(int32 InVertex, FTarjanVertexInfo& InVertexInfo, TArray<FTarjanStronglyConnectedComponent>& OutComponents)
+			{
+				InVertexInfo.Tag = CurrentTag;
+				InVertexInfo.LowLinkTag = CurrentTag;
+				InVertexInfo.bIsOnStack = true;
+				VertexStack.Add(InVertex);
+
+				// Tag needs to always be increasing during subsequent calls to
+				// Update(...).
+				CurrentTag++;
+
+				// Perform depth first traversal through the graph looking for 
+				// children that are already on the stack (which denotes a 
+				// strongly connected component)
+				for (int32 ChildVertex : InVertexInfo.Children)
+				{
+					FTarjanVertexInfo& ChildInfo = VertexInfos[ChildVertex];
+
+					if (ChildInfo.Tag == INDEX_NONE)
+					{
+						// If ChildInfo.Tag == INDEX_NONE, then this child is 
+						// not on the stack because it has never been inspected.
+
+						// Continue in depth-first-traversal. 
+						Update(ChildVertex, ChildInfo, OutComponents);
+
+						// Update to use the lowest link tag found in children.
+						InVertexInfo.LowLinkTag = FMath::Min(InVertexInfo.LowLinkTag, ChildInfo.LowLinkTag);
+					}
+					else if(ChildInfo.bIsOnStack)
+					{
+						// This child is on the stack and represents a strongly 
+						// connected component. 
+						InVertexInfo.LowLinkTag = FMath::Min(InVertexInfo.LowLinkTag, ChildInfo.Tag);
+					}
+				}
+
+				if (InVertexInfo.LowLinkTag == InVertexInfo.Tag)
+				{
+					// Create a strongly connected component object if this is 
+					// the root vertex of a strongly connected component.
+					AddStronglyConnectedComponent(InVertex, OutComponents);
+				}
+			}
+
+			// Adds the current strongly connected component to the 
+			// OutComponents array.
+			void AddStronglyConnectedComponent(int32 InRootVertex, TArray<FTarjanStronglyConnectedComponent>& OutComponents)
+			{
+				// When a single vertex represents the entire strongly connected
+				// component, the root vertex is the same as the vertex at the
+				// top of the stack.
+				bool bIsSingleVertex = (VertexStack.Top() == InRootVertex);
+
+				if (bIsSingleVertex && Settings.bExcludeSingleVertex)
+				{
+					// This strongly connected component is skipped because it 
+					// consists of a single vertex. 
+					//
+					// The stack still needs to be updated as well as the vertex
+					// info for the root vertex.
+					int32 TopVertex = VertexStack.Pop();	
+					FTarjanVertexInfo& TopInfo = VertexInfos[TopVertex];
+					TopInfo.bIsOnStack = false;
+				}
+				else
+				{
+					bool bDidCompleteLoop = false;
+
+					// Create a new strongly connected component subgraph.
+					FTarjanStronglyConnectedComponent& StronglyConnectedComponent = OutComponents.AddDefaulted_GetRef();
+
+					// Pop off the top of the vertex stack until we reach the
+					// root vertex.
+					while (!bDidCompleteLoop)
+					{
+						int32 TopVertex = VertexStack.Pop();
+
+						FTarjanVertexInfo& TopInfo = VertexInfos[TopVertex];
+						TopInfo.bIsOnStack = false;
+
+						StronglyConnectedComponent.Vertices.Add(TopVertex);
+
+						// The unwinding continues until the current vertex is 
+						// removed from the stack.
+						bDidCompleteLoop = (TopVertex == InRootVertex);
+					}
+
+					// Add the edges associated with the strongly connected
+					// component. 
+					for (int32 ComponentVertex : StronglyConnectedComponent.Vertices)
+					{
+						FTarjanVertexInfo& ComponentInfo = VertexInfos[ComponentVertex];
+
+						for (int32 ChildVertex : ComponentInfo.Children)
+						{
+							if (StronglyConnectedComponent.Vertices.Contains(ChildVertex))
+							{
+								StronglyConnectedComponent.Edges.Emplace(ComponentVertex, ChildVertex);
+							}
+						}
+					}
+				}
+			}
+
+			FTarjanAlgoImplSettings Settings;
+
+			int32 CurrentTag = 0;
+
+			TMap<int32, FTarjanVertexInfo> VertexInfos;
+
+			TArray<int32> VertexStack;
+		};
+	}
+
+	bool FTarjan::FindStronglyConnectedComponents(const TSet<FTarjanEdge>& InEdges, TArray<FTarjanStronglyConnectedComponent>& OutComponents, bool bExcludeSingleVertex)
+	{
+		using namespace MetasoundTarjanPrivate;
+
+		FTarjanAlgoImplSettings Settings;
+		Settings.bExcludeSingleVertex = bExcludeSingleVertex;
+
+		FTarjanAlgoImpl TarjanAlgo(Settings);
+
+		return TarjanAlgo.FindStronglyConnectedComponents(InEdges, OutComponents);
+	}
+
+	bool FTarjan::FindStronglyConnectedComponents(const IGraph& InGraph, TArray<FMetasoundStronglyConnectedComponent>& OutComponents, bool bExcludeSingleVertex)
+	{
+		
+		using FNodePair = TTuple<const INode*, const INode*>;
+
+		TArray<const INode*> UniqueNodes;
+		TSet<FTarjanEdge> UniqueEdges;
+
+		TMultiMap<FNodePair, FDataEdge> NodePairEdges;
+
+		// To find metasound specific strongly connected components, the nodes
+		// and edges must be translated to integers to work with the underlying 
+		// algorithm.
+		for (const FDataEdge& Edge : InGraph.GetDataEdges())
+		{
+			if ((nullptr != Edge.From.Node) && (nullptr != Edge.To.Node))
+			{
+				int32 FromTarjanVertex = UniqueNodes.AddUnique(Edge.From.Node);	
+				int32 ToTarjanVertex = UniqueNodes.AddUnique(Edge.To.Node);	
+
+				UniqueEdges.Add(FTarjanEdge(FromTarjanVertex, ToTarjanVertex));
+
+				NodePairEdges.Emplace(FNodePair(Edge.From.Node, Edge.To.Node), Edge);
+			}
+		}
+
+		TArray<FTarjanStronglyConnectedComponent> StronglyConnectedComponents;
+
+		// Run tarjan on metasound derived graph edges 
+		if (FindStronglyConnectedComponents(UniqueEdges, StronglyConnectedComponents, bExcludeSingleVertex))
+		{
+			// If strongly connected components are found, they must be converted
+			// back into metasound types. 
+			for (const FTarjanStronglyConnectedComponent& Component : StronglyConnectedComponents)
+			{
+				FMetasoundStronglyConnectedComponent& MetasoundGraphComponent = OutComponents.AddDefaulted_GetRef();
+
+				for (int32 Vertex : Component.Vertices)
+				{
+					MetasoundGraphComponent.Nodes.Add(UniqueNodes[Vertex]);
+				}
+
+				for (const FTarjanEdge& Edge : Component.Edges)
+				{
+					TArray<FDataEdge> MetasoundComponentEdges;
+
+					const INode* FromNode = UniqueNodes[Edge.Get<0>()];
+					const INode* ToNode = UniqueNodes[Edge.Get<1>()];
+
+					NodePairEdges.MultiFind(FNodePair(FromNode, ToNode), MetasoundComponentEdges);
+					MetasoundGraphComponent.Edges.Append(MetasoundComponentEdges);
+				}
+			}
+
+			return true;
+		}
+
+		// No strongly connected components found.
+		return false;
+	}
+}
diff --git a/Engine/Plugins/Runtime/Metasound/Source/MetasoundGraphCore/Public/MetasoundBuildError.h b/Engine/Plugins/Runtime/Metasound/Source/MetasoundGraphCore/Public/MetasoundBuildError.h
--- a/Engine/Plugins/Runtime/Metasound/Source/MetasoundGraphCore/Public/MetasoundBuildError.h
+++ b/Engine/Plugins/Runtime/Metasound/Source/MetasoundGraphCore/Public/MetasoundBuildError.h
@@ -79,7 +79,7 @@ namespace Metasound
 		public:
 			static const FName ErrorType;
 
-			FGraphCycleError(const TArray<INode*>& InNodes, const TArray<FDataEdge>& InEdges);
+			FGraphCycleError(TArrayView<INode const* const> InNodes, const TArray<FDataEdge>& InEdges);
 
 			virtual ~FGraphCycleError() = default;
 
diff --git a/Engine/Plugins/Runtime/Metasound/Source/MetasoundGraphCore/Public/MetasoundOperatorBuilder.h b/Engine/Plugins/Runtime/Metasound/Source/MetasoundGraphCore/Public/MetasoundOperatorBuilder.h
--- a/Engine/Plugins/Runtime/Metasound/Source/MetasoundGraphCore/Public/MetasoundOperatorBuilder.h
+++ b/Engine/Plugins/Runtime/Metasound/Source/MetasoundGraphCore/Public/MetasoundOperatorBuilder.h
@@ -41,7 +41,7 @@ namespace Metasound
 
 			bool GroupInputEdges(const TArray<FDataEdge>& InEdges, FNodeEdgeMultiMap& OutNodeInputs, TArray<FBuildErrorPtr>& OutErrors) const;
 
-			bool TopologicalSort(const TArray<FDataEdge>& InEdges, TArray<INode*>& OutNodes, TArray<FBuildErrorPtr>& OutErrors) const;
+			bool TopologicalSort(const IGraph& InGraph, TArray<INode*>& OutNodes, TArray<FBuildErrorPtr>& OutErrors) const;
 
 			bool GatherInputDataReferences(const INode* InNode, const FNodeEdgeMultiMap& InEdgeMap, const FNodeDataReferenceMap& InDataReferenceMap, FDataReferenceCollection& OutCollection, TArray<FBuildErrorPtr>& OutErrors) const;
 
diff --git a/Engine/Plugins/Runtime/Metasound/Source/MetasoundGraphCore/Public/MetasoundTarjan.h b/Engine/Plugins/Runtime/Metasound/Source/MetasoundGraphCore/Public/MetasoundTarjan.h
new file mode 100644
--- /dev/null
+++ b/Engine/Plugins/Runtime/Metasound/Source/MetasoundGraphCore/Public/MetasoundTarjan.h
@@ -0,0 +1,71 @@
+// Copyright Epic Games, Inc. All Rights Reserved.
+
+#pragma once
+
+#include "CoreMinimal.h"
+#include "MetasoundNodeInterface.h"
+
+namespace Metasound
+{
+	/** A pair of int32s represent a directed edge. The first value represents 
+	 * the source vertex, and the second value represents the destination 
+	 * vertex.
+	 */
+	typedef TTuple<int32, int32> FTarjanEdge;
+
+	/** A strongly connected component contains a subgraph of strongly connected
+	 * vertices and their corresponding edges.
+	 */
+	struct FTarjanStronglyConnectedComponent
+	{
+		/** Vertices in the strongly connected component. */
+		TArray<int32> Vertices;
+
+		/** Edges in the strongly connected component. */
+		TArray<FTarjanEdge> Edges;
+	};
+
+	/** A strongly connected component containing Metasound INodes and 
+	 * FDataEdges.
+	 */
+	struct FMetasoundStronglyConnectedComponent
+	{
+		/** Nodes in the strongly connected component. */
+		TArray<const INode*> Nodes;
+
+		/** FDataEdges in the strongly connected component. */
+		TArray<FDataEdge> Edges;
+	};
+
+	/** The Tarjan algorithm identifies strongly connected component in a 
+	 * directed graph. Strongly connected components consist of a set vertices 
+	 * where each vertex is reachable from every other vertex in the directed 
+	 * graph. A strongly connected component of a directed graph generally 
+	 * represents a cycle in the graph. Technically, a single vertex is 
+	 * considered a strongly connected component, but for most use-cases that 
+	 * information is not needed. By default, single vertices are not returned
+	 * as strongly connected components.
+	 */
+	struct METASOUNDGRAPHCORE_API FTarjan
+	{
+		/** Find strongly connected components given a set of FTarjanEdges.
+		 *
+		 * @param InEdges - Edges in the directed graph.
+		 * @param OutComponents - Strongly connected components found in the graph are be added to this array.
+		 * @param bExcludeSingleVertex - If true, single vertices are not be returned as strongly connected components. If false, single vertices may be returned as strongly connected components. 
+		 *
+		 * @return True if one or more strongly connected components are added to OutComponents. False otherwise. 
+		 */
+		static bool FindStronglyConnectedComponents(const TSet<FTarjanEdge>& InEdges, TArray<FTarjanStronglyConnectedComponent>& OutComponents, bool bExcludeSingleVertex = true);
+
+		/** Find strongly connected components given a Metasound IGraph.
+		 *
+		 * @param InGraph - Graph to be analyzed. 
+		 * @param OutComponents - Strongly connected components found in the graph are be added to this array.
+		 * @param bExcludeSingleVertex - If true, single vertices are not be returned as strongly connected components. If false, single vertices may be returned as strongly connected components. 
+		 *
+		 * @return True if one or more strongly connected components are added to OutComponents. False otherwise. 
+		 */
+		static bool FindStronglyConnectedComponents(const IGraph& InGraph, TArray<FMetasoundStronglyConnectedComponent>& OutComponents, bool bExcludeSingleVertex = true);
+	};
+}
