Metasound Literal Editor Updates
- Hide superfluous type info & non-preferred literal data
- Allow user to rename nodes with display name, which is independent from node string name (and supports localization)
- Add better membership checks to editor detail customizations to ensure document member renames/data org doesn't break editor
#rnx

diff --git a/Engine/Plugins/Runtime/Metasound/Source/MetasoundEditor/Private/MetasoundDetailCustomization.cpp b/Engine/Plugins/Runtime/Metasound/Source/MetasoundEditor/Private/MetasoundDetailCustomization.cpp
--- a/Engine/Plugins/Runtime/Metasound/Source/MetasoundEditor/Private/MetasoundDetailCustomization.cpp
+++ b/Engine/Plugins/Runtime/Metasound/Source/MetasoundEditor/Private/MetasoundDetailCustomization.cpp
@@ -2,40 +2,174 @@
 #include "MetasoundDetailCustomization.h"
 
 #include "CoreMinimal.h"
+#include "Delegates/Delegate.h"
 #include "DetailCategoryBuilder.h"
 #include "DetailLayoutBuilder.h"
+#include "DetailWidgetRow.h"
 #include "IDetailGroup.h"
 #include "Metasound.h"
+#include "MetasoundAssetBase.h"
 #include "PropertyEditorDelegates.h"
 #include "PropertyHandle.h"
 #include "Templates/Casts.h"
 #include "Templates/SharedPointer.h"
+#include "UObject/WeakObjectPtrTemplates.h"
 
 
 #define LOCTEXT_NAMESPACE "MetasoundEditor"
 
+namespace Metasound
+{
+	namespace Editor
+	{
+		FName BuildChildPath(const FString& InBasePath, FName InPropertyName)
+		{
+			return FName(InBasePath + TEXT(".") + InPropertyName.ToString());
+		}
+
+		FName BuildChildPath(const FName& InBasePath, FName InPropertyName)
+		{
+			return FName(InBasePath.ToString() + TEXT(".") + InPropertyName.ToString());
+		}
+
+		template <typename T>
+		void BuildFixedArray(IDetailLayoutBuilder& InDetailLayout, FName CategoryName, FName InPropertyName)
+		{
+			IDetailCategoryBuilder& CategoryBuilder = InDetailLayout.EditCategory(CategoryName);
+			TSharedPtr<IPropertyHandle> ParentProperty = InDetailLayout.GetProperty(InPropertyName);
+			TSharedPtr<IPropertyHandleArray> ArrayHandle = ParentProperty->AsArray();
+
+			uint32 NumElements = 0;
+			ArrayHandle->GetNumElements(NumElements);
+			for (int32 i = 0; i < static_cast<int32>(NumElements); ++i)
+			{
+				TMap<FName, TSharedPtr<IPropertyHandle>> ChildProperties;
+				TSharedRef<IPropertyHandle> ArrayItemHandle = ArrayHandle->GetElement(i);
+				uint32 NumChildren = 0;
+				ArrayItemHandle->GetNumChildren(NumChildren);
+				for (int32 j = 0; j < static_cast<int32>(NumChildren); ++j)
+				{
+					TSharedPtr<IPropertyHandle> PropertyHandle = ArrayItemHandle->GetChildHandle(j);
+					FProperty* Property = PropertyHandle->GetProperty();
+					if (ensureAlways(Property))
+					{
+						const FName PropertyName = Property->GetFName();
+						ChildProperties.Add(PropertyName, PropertyHandle);
+					}
+				}
+
+				TSharedPtr<IPropertyHandle> TypeProperty = ChildProperties.FindChecked(GET_MEMBER_NAME_CHECKED(T, TypeName));
+				TSharedPtr<IPropertyHandle> NameProperty = ChildProperties.FindChecked(GET_MEMBER_NAME_CHECKED(T, DisplayName));
+
+				FDetailWidgetRow& NameRow = CategoryBuilder.AddCustomRow(ParentProperty->GetPropertyDisplayName())
+					.NameContent()
+					[
+						SNew(STextBlock)
+						.Font(IDetailLayoutBuilder::GetDetailFontBold())
+						.Text(TAttribute<FText>::Create([i, NameProperty, TypeProperty]()
+						{
+							FString TypeNameString;
+							ensureAlways(TypeProperty->GetValue(TypeNameString) == FPropertyAccess::Success);
+
+							// Remove namespace info to keep concise
+							TypeNameString.RightChopInline(TypeNameString.Find(TEXT(":"), ESearchCase::IgnoreCase, ESearchDir::FromEnd) + 1);
+
+							FText DisplayName;
+							ensureAlways(NameProperty->GetValue(DisplayName) == FPropertyAccess::Success);
+
+							return FText::Format(LOCTEXT("Metasound_FixedArray_Format", "{0}. {1} ({2})"), FText::AsNumber(i + 1), DisplayName, FText::FromString(TypeNameString));
+						}))
+						.ToolTipText(TAttribute<FText>::Create([ParentProperty, TypeProperty]()
+						{
+							FString TypeNameString;
+							ensureAlways(TypeProperty->GetValue(TypeNameString) == FPropertyAccess::Success);
+
+							return FText::Format(LOCTEXT("", "Type {0}\n"), FText::FromString(TypeNameString), ParentProperty->GetToolTipText());
+						}))
+					];
+
+				for (const TPair<FName, TSharedPtr<IPropertyHandle>>& Pair : ChildProperties)
+				{
+					if (Pair.Key != GET_MEMBER_NAME_CHECKED(T, TypeName))
+					{
+						CategoryBuilder.AddProperty(Pair.Value);
+					}
+				}
+			}
+
+			FSimpleDelegate RefreshDelegate = FSimpleDelegate::CreateLambda([DetailLayout = &InDetailLayout]()
+				{
+					DetailLayout->ForceRefreshDetails();
+				});
+			ArrayHandle->SetOnNumElementsChanged(RefreshDelegate);
+		};
+	}
+}
+
+FMetasoundDetailCustomization::FMetasoundDetailCustomization(FName InDocumentPropertyName)
+	: IDetailCustomization()
+	, DocumentPropertyName(InDocumentPropertyName)
+{
+}
+
+FName FMetasoundDetailCustomization::GetMetadataRootClassPath() const
+{
+	return Metasound::Editor::BuildChildPath(DocumentPropertyName, GET_MEMBER_NAME_CHECKED(FMetasoundDocument, RootClass));
+}
+
+FName FMetasoundDetailCustomization::GetMetadataPropertyPath() const
+{
+	const FName RootClass = FName(GetMetadataRootClassPath());
+	return Metasound::Editor::BuildChildPath(RootClass, GET_MEMBER_NAME_CHECKED(FMetasoundClassDescription, Metadata));
+}
+
 void FMetasoundDetailCustomization::CustomizeDetails(IDetailLayoutBuilder& DetailLayout)
 {
+	using namespace Metasound::Editor;
+
 	// General Category
 	IDetailCategoryBuilder& GeneralCategoryBuilder = DetailLayout.EditCategory("General");
 
-	TSharedPtr<IPropertyHandle> AuthorHandle = DetailLayout.GetProperty("RootMetasoundDocument.RootClass.Metadata.AuthorName");
-	TSharedPtr<IPropertyHandle> DescHandle = DetailLayout.GetProperty("RootMetasoundDocument.RootClass.Metadata.MetasoundDescription");
-	TSharedPtr<IPropertyHandle> NodeTypeHandle = DetailLayout.GetProperty("RootMetasoundDocument.RootClass.Metadata.NodeType");
+	TArray<TWeakObjectPtr<UObject>> ObjectsCustomized;
+	DetailLayout.GetObjectsBeingCustomized(ObjectsCustomized);
+
+	const FName AuthorPropertyPath = BuildChildPath(GetMetadataPropertyPath(), GET_MEMBER_NAME_CHECKED(FMetasoundClassMetadata, AuthorName));
+	const FName DescPropertyPath = BuildChildPath(GetMetadataPropertyPath(), GET_MEMBER_NAME_CHECKED(FMetasoundClassMetadata, MetasoundDescription));
+	const FName NodeTypePropertyPath = BuildChildPath(GetMetadataPropertyPath(), GET_MEMBER_NAME_CHECKED(FMetasoundClassMetadata, NodeType));
+
+	TSharedPtr<IPropertyHandle> AuthorHandle = DetailLayout.GetProperty(AuthorPropertyPath);
+	TSharedPtr<IPropertyHandle> DescHandle = DetailLayout.GetProperty(DescPropertyPath);
+	TSharedPtr<IPropertyHandle> NodeTypeHandle = DetailLayout.GetProperty(NodeTypePropertyPath);
 
 	GeneralCategoryBuilder.AddProperty(NodeTypeHandle);
 	GeneralCategoryBuilder.AddProperty(AuthorHandle);
 	GeneralCategoryBuilder.AddProperty(DescHandle);
 
-	IDetailCategoryBuilder& InputsCategoryBuilder = DetailLayout.EditCategory("Inputs");
-	TSharedPtr<IPropertyHandle> InputsHandle = DetailLayout.GetProperty("RootMetasoundDocument.RootClass.Inputs");
-	InputsCategoryBuilder.AddProperty(InputsHandle);
+	// Input/Output Categories
+
+	const FName InputsPropertyPath = BuildChildPath(GetMetadataRootClassPath(), GET_MEMBER_NAME_CHECKED(FMetasoundClassDescription, Inputs));
+	const FName OutputsPropertyPath = BuildChildPath(GetMetadataRootClassPath(), GET_MEMBER_NAME_CHECKED(FMetasoundClassDescription, Outputs));
 
-	IDetailCategoryBuilder& OutputsCategoryBuilder = DetailLayout.EditCategory("Outputs");
-	TSharedPtr<IPropertyHandle> OutputsHandle = DetailLayout.GetProperty("RootMetasoundDocument.RootClass.Outputs");
-	OutputsCategoryBuilder.AddProperty(OutputsHandle);
+	BuildFixedArray<FMetasoundInputDescription>(DetailLayout, "Inputs", InputsPropertyPath);
+	BuildFixedArray<FMetasoundOutputDescription>(DetailLayout, "Outputs", OutputsPropertyPath);
 
 	// Hack to hide parent structs for nested metadata properties
 	DetailLayout.HideCategory("Hidden");
+
+	// Hack to hide categories brought in from UMetasoundSource inherited from USoundBase
+	DetailLayout.HideCategory("Analysis");
+	DetailLayout.HideCategory("Attenuation");
+	DetailLayout.HideCategory("Curves");
+	DetailLayout.HideCategory("Debug");
+	DetailLayout.HideCategory("Effects");
+	DetailLayout.HideCategory("File Path");
+	DetailLayout.HideCategory("Format");
+	DetailLayout.HideCategory("Info");
+	DetailLayout.HideCategory("Loading");
+	DetailLayout.HideCategory("Modulation");
+	DetailLayout.HideCategory("Playback");
+	DetailLayout.HideCategory("Sound");
+	DetailLayout.HideCategory("Subtitles");
+	DetailLayout.HideCategory("Voice Management");
 }
 #undef LOCTEXT_NAMESPACE
\ No newline at end of file
diff --git a/Engine/Plugins/Runtime/Metasound/Source/MetasoundEditor/Private/MetasoundDetailCustomization.h b/Engine/Plugins/Runtime/Metasound/Source/MetasoundEditor/Private/MetasoundDetailCustomization.h
--- a/Engine/Plugins/Runtime/Metasound/Source/MetasoundEditor/Private/MetasoundDetailCustomization.h
+++ b/Engine/Plugins/Runtime/Metasound/Source/MetasoundEditor/Private/MetasoundDetailCustomization.h
@@ -2,6 +2,7 @@
 #pragma once
 
 #include "IDetailCustomization.h"
+#include "UObject/NameTypes.h"
 
 // Forward Declarations
 class IDetailLayoutBuilder;
@@ -10,7 +11,15 @@ class IDetailLayoutBuilder;
 class FMetasoundDetailCustomization : public IDetailCustomization
 {
 public:
+	FMetasoundDetailCustomization(FName InDocumentPropertyName);
+
 	// IDetailCustomization interface
 	virtual void CustomizeDetails(IDetailLayoutBuilder& DetailLayout) override;
 	// End of IDetailCustomization interface
+
+private:
+	FName GetMetadataRootClassPath() const;
+	FName GetMetadataPropertyPath() const;
+
+	FName DocumentPropertyName;
 };
\ No newline at end of file
diff --git a/Engine/Plugins/Runtime/Metasound/Source/MetasoundEditor/Private/MetasoundEditorGraphBuilder.cpp b/Engine/Plugins/Runtime/Metasound/Source/MetasoundEditor/Private/MetasoundEditorGraphBuilder.cpp
--- a/Engine/Plugins/Runtime/Metasound/Source/MetasoundEditor/Private/MetasoundEditorGraphBuilder.cpp
+++ b/Engine/Plugins/Runtime/Metasound/Source/MetasoundEditor/Private/MetasoundEditorGraphBuilder.cpp
@@ -5,13 +5,14 @@
 #include "EdGraph/EdGraphNode.h"
 #include "EdGraph/EdGraphPin.h"
 #include "Metasound.h"
+#include "MetasoundAssetBase.h"
 #include "MetasoundEditorGraph.h"
 #include "MetasoundEditorGraphNode.h"
+#include "MetasoundEditorModule.h"
 #include "MetasoundFrontend.h"
+#include "Modules/ModuleManager.h"
 #include "ScopedTransaction.h"
 #include "Templates/Tuple.h"
-#include "MetasoundEditorModule.h"
-#include "Modules/ModuleManager.h"
 
 
 #define LOCTEXT_NAMESPACE "MetasoundEditor"
@@ -72,7 +73,59 @@ namespace Metasound
 			FMetasoundAssetBase* MetasoundAsset = Frontend::GetObjectAsAssetBase(&InMetasound);
 			check(MetasoundAsset);
 
-			Frontend::FNodeHandle NodeHandle = MetasoundAsset->GetRootGraphHandle().AddNewInput(Description);
+			Frontend::FGraphHandle GraphHandle = MetasoundAsset->GetRootGraphHandle();
+			Frontend::FNodeHandle NodeHandle = GraphHandle.AddNewInput(Description);
+			if (!NodeHandle.IsValid())
+			{
+				return nullptr;
+			}
+
+			IMetasoundEditorModule& EditorModule = FModuleManager::GetModuleChecked<IMetasoundEditorModule>("MetasoundEditor");
+			FEditorDataType DataType = EditorModule.FindDataType(InTypeName);
+
+			Metasound::FDataTypeLiteralParam LiteralParam = Frontend::GetDefaultParamForDataType(InTypeName);
+			if (!ensureAlways(LiteralParam.IsValid()))
+			{
+				return nullptr;
+			}
+
+			switch (LiteralParam.ConstructorArgType)
+			{
+				case ELiteralArgType::Boolean:
+				{
+					GraphHandle.SetInputToLiteral(InName, false);
+				}
+				break;
+
+				case ELiteralArgType::Float:
+				{
+					GraphHandle.SetInputToLiteral(InName, 0.0f);
+				}
+				break;
+
+				case ELiteralArgType::Integer:
+				{
+					GraphHandle.SetInputToLiteral(InName, 0);
+				}
+				break;
+
+				case ELiteralArgType::String:
+				{
+					GraphHandle.SetInputToLiteral(InName, FString(TEXT("")));
+				}
+				break;
+
+				case ELiteralArgType::Invalid:
+				case ELiteralArgType::None:
+				default:
+				{
+					static_assert(static_cast<int32>(ELiteralArgType::Invalid) == 5, "Possible missing ELiteralArgType case coverage");
+				}
+				break;
+			}
+
+			GraphHandle.SetInputDisplayName(InName, FText::FromString(InName));
+
 			return AddNode(InMetasound, Location, NodeHandle, bInSelectNewNode);
 		}
 
@@ -86,7 +139,11 @@ namespace Metasound
 			FMetasoundAssetBase* MetasoundAsset = Frontend::GetObjectAsAssetBase(&InMetasound);
 			check(MetasoundAsset);
 
-			Frontend::FNodeHandle NodeHandle = MetasoundAsset->GetRootGraphHandle().AddNewOutput(Description);
+			Frontend::FGraphHandle GraphHandle = MetasoundAsset->GetRootGraphHandle();
+			Frontend::FNodeHandle NodeHandle = GraphHandle.AddNewOutput(Description);
+			
+			GraphHandle.SetOutputDisplayName(InName, FText::FromString(InName));
+
 			return AddNode(InMetasound, Location, NodeHandle, bInSelectNewNode);
 		}
 
@@ -147,8 +204,10 @@ namespace Metasound
 
 			Graph->Nodes.Reset();
 
-			// TODO: Space graph nodes in a readable way
-			FVector2D Location = FVector2D::ZeroVector;
+			// TODO: Space graph nodes in a procedural and readable way
+			FVector2D InputNodeLocation = FVector2D::ZeroVector;
+			FVector2D OpNodeLocation = FVector2D(250.0f, 0.0f);
+			FVector2D OutputNodeLocation = FVector2D(500.0f, 0.0f);
 
 			struct FNodePair
 			{
@@ -161,6 +220,22 @@ namespace Metasound
 			for (FNodeHandle& NodeHandle : NodeHandles)
 			{
 				const FNodeClassInfo ClassInfo = NodeHandle.GetClassInfo();
+				FVector2D Location;
+				if (ClassInfo.NodeType == EMetasoundClassType::Input)
+				{
+					Location = InputNodeLocation;
+					InputNodeLocation.Y += 100.0f;
+				}
+				else if (ClassInfo.NodeType == EMetasoundClassType::Output)
+				{
+					Location = OutputNodeLocation;
+					OutputNodeLocation.Y += 100.0f;
+				}
+				else
+				{
+					Location = OpNodeLocation;
+					OpNodeLocation.Y += 100.0f;
+				}
 				UEdGraphNode* NewNode = AddNode(InMetasound, Location, NodeHandle, false /* bInSelectNewNode */);
 				NewIdNodeMap.Add(NodeHandle.GetNodeID(), FNodePair { NodeHandle, NewNode });
 			}
@@ -213,12 +288,28 @@ namespace Metasound
 
 			IMetasoundEditorModule& EditorModule = FModuleManager::GetModuleChecked<IMetasoundEditorModule>("MetasoundEditor");
 
+			const FMetasoundAssetBase* MetasoundAsset = Frontend::GetObjectAsAssetBase(&InGraphNode.GetMetasoundChecked());
+			check(MetasoundAsset);
+
+			Frontend::FGraphHandle GraphHandle = MetasoundAsset->GetRootGraphHandle();
+
 			TArray<Frontend::FInputHandle> InputHandles = InNodeHandle.GetAllInputs();
 			for (int32 i = 0; i < InputHandles.Num(); ++i)
 			{
 				const Frontend::FInputHandle& InputHandle = InputHandles[i];
+
+				FText DisplayName;
+				if (InGraphNode.GetNodeHandle().GetNodeType() == EMetasoundClassType::Input)
+				{
+					DisplayName = GraphHandle.GetInputDisplayName(InputHandle.GetInputName());
+				}
+				else
+				{
+					DisplayName = FText::FromString(*InputHandle.GetInputName());
+				}
+
 				FEdGraphPinType PinType("MetasoundNode", NAME_None, nullptr, EPinContainerType::None, false, FEdGraphTerminalType());
-				UEdGraphPin* NewPin = InGraphNode.CreatePin(EGPD_Input, PinType, *InputHandle.GetInputName(), i);
+				UEdGraphPin* NewPin = InGraphNode.CreatePin(EGPD_Input, PinType, *DisplayName.ToString(), i);
 				if (ensureAlways(NewPin))
 				{
 					NewPin->PinToolTip = InputHandle.GetInputTooltip().ToString();
@@ -233,7 +324,18 @@ namespace Metasound
 			{
 				const Frontend::FOutputHandle& OutputHandle = OutputHandles[i];
 				FEdGraphPinType PinType("MetasoundNode", NAME_None, nullptr, EPinContainerType::None, false, FEdGraphTerminalType());
-				UEdGraphPin* NewPin = InGraphNode.CreatePin(EGPD_Output, PinType, *OutputHandle.GetOutputName(), i);
+
+				FText DisplayName;
+				if (InGraphNode.GetNodeHandle().GetNodeType() == EMetasoundClassType::Output)
+				{
+					DisplayName = GraphHandle.GetInputDisplayName(OutputHandle.GetOutputName());
+				}
+				else
+				{
+					DisplayName = FText::FromString(*OutputHandle.GetOutputName());
+				}
+
+				UEdGraphPin* NewPin = InGraphNode.CreatePin(EGPD_Output, PinType, *DisplayName.ToString(), i);
 				if (ensureAlways(NewPin))
 				{
 					NewPin->PinToolTip = OutputHandle.GetOutputTooltip().ToString();
diff --git a/Engine/Plugins/Runtime/Metasound/Source/MetasoundEditor/Private/MetasoundEditorGraphNode.cpp b/Engine/Plugins/Runtime/Metasound/Source/MetasoundEditor/Private/MetasoundEditorGraphNode.cpp
--- a/Engine/Plugins/Runtime/Metasound/Source/MetasoundEditor/Private/MetasoundEditorGraphNode.cpp
+++ b/Engine/Plugins/Runtime/Metasound/Source/MetasoundEditor/Private/MetasoundEditorGraphNode.cpp
@@ -190,13 +190,15 @@ FText UMetasoundEditorGraphNode::GetNodeTitle(ENodeTitleType::Type TitleType) co
 	{
 		case EMetasoundClassType::Input:
 		{
-			return FText::Format(LOCTEXT("MetasoundGraphNode_TitleFormat", "Input {0}"), FText::FromString(NodeHandle.GetNodeName()));
+			FGraphHandle GraphHandle = GetRootGraphHandle();
+			return FText::Format(LOCTEXT("MetasoundGraphNode_TitleFormat", "Input {0}"), GraphHandle.GetInputDisplayName(NodeHandle.GetNodeName()));
 		}
 		break;
 
 		case EMetasoundClassType::Output:
 		{
-			return FText::Format(LOCTEXT("MetasoundGraphNode_TitleFormat", "Output {0}"), FText::FromString(NodeHandle.GetNodeName()));
+			FGraphHandle GraphHandle = GetRootGraphHandle();
+			return FText::Format(LOCTEXT("MetasoundGraphNode_TitleFormat", "Output {0}"), GraphHandle.GetOutputDisplayName(NodeHandle.GetNodeName()));
 		}
 		break;
 
@@ -284,23 +286,20 @@ FText UMetasoundEditorGraphNode::GetTooltipText() const
 {
 	using namespace Metasound::Frontend;
 
-	// TODO: Add author to ToolTip
 	FNodeHandle NodeHandle = GetNodeHandle();
-
-	const FMetasoundAssetBase* MetasoundAsset = GetObjectAsAssetBase(&GetMetasoundChecked());
-	check(MetasoundAsset);
-
 	switch (NodeHandle.GetNodeType())
 	{
 		case EMetasoundClassType::Input:
 		{
-			return MetasoundAsset->GetInputToolTip(NodeHandle.GetNodeName());
+			FGraphHandle GraphHandle = GetRootGraphHandle();
+			return GraphHandle.GetInputToolTip(NodeHandle.GetNodeName());
 		}
 		break;
 
 		case EMetasoundClassType::Output:
 		{
-			return MetasoundAsset->GetOutputToolTip(NodeHandle.GetNodeName());
+			FGraphHandle GraphHandle = GetRootGraphHandle();
+			return GraphHandle.GetOutputToolTip(NodeHandle.GetNodeName());
 		}
 		break;
 
diff --git a/Engine/Plugins/Runtime/Metasound/Source/MetasoundEditor/Private/MetasoundEditorGraphSchema.cpp b/Engine/Plugins/Runtime/Metasound/Source/MetasoundEditor/Private/MetasoundEditorGraphSchema.cpp
--- a/Engine/Plugins/Runtime/Metasound/Source/MetasoundEditor/Private/MetasoundEditorGraphSchema.cpp
+++ b/Engine/Plugins/Runtime/Metasound/Source/MetasoundEditor/Private/MetasoundEditorGraphSchema.cpp
@@ -24,6 +24,7 @@
 #include "UObject/NoExportTypes.h"
 #include "UObject/UObjectHash.h"
 #include "UObject/UObjectIterator.h"
+#include "../../MetasoundFrontend/Public/MetasoundAssetBase.h"
 
 #define LOCTEXT_NAMESPACE "MetasoundEditor"
 
@@ -404,7 +405,13 @@ const FPinConnectionResponse UMetasoundEditorGraphSchema::CanCreateConnection(co
 	Metasound::Frontend::FConnectability Connectability = InputHandle.CanConnectTo(OutputHandle);
 	if (Connectability.Connectable != Metasound::Frontend::FConnectability::EConnectable::Yes)
 	{
-		return FPinConnectionResponse(CONNECT_RESPONSE_DISALLOW, LOCTEXT("ConnectionTypeIncompatible", "Connection pin types cannot be connected"));
+		const FName InputType = InputHandle.GetInputType();
+		const FName OutputType = OutputHandle.GetOutputType();
+		return FPinConnectionResponse(CONNECT_RESPONSE_DISALLOW, FText::Format(
+			LOCTEXT("ConnectionTypeIncompatibleFormat", "Output pin of type '{0}' cannot be connected to input pin of type '{1}'"),
+			FText::FromName(OutputType),
+			FText::FromName(InputType)
+		));
 	}
 
 	// Break existing connections on inputs only - multiple output connections are acceptable
@@ -461,6 +468,41 @@ bool UMetasoundEditorGraphSchema::ShouldHidePinDefaultValue(UEdGraphPin* Pin) co
 	return false;
 }
 
+FText UMetasoundEditorGraphSchema::GetPinDisplayName(const UEdGraphPin* Pin) const
+{
+	using namespace Metasound::Frontend;
+
+	check(Pin);
+
+	UMetasoundEditorGraphNode* Node = CastChecked<UMetasoundEditorGraphNode>(Pin->GetOwningNode());
+	FNodeHandle NodeHandle = Node->GetNodeHandle();
+	const EMetasoundClassType ClassType = NodeHandle.GetClassInfo().NodeType;
+
+	switch (ClassType)
+	{
+		case EMetasoundClassType::Input:
+		{
+			FGraphHandle GraphHandle = Node->GetRootGraphHandle();
+			return GraphHandle.GetInputDisplayName(NodeHandle.GetNodeName());
+		}
+
+		case EMetasoundClassType::Output:
+		{
+			FGraphHandle GraphHandle = Node->GetRootGraphHandle();
+			return GraphHandle.GetOutputDisplayName(NodeHandle.GetNodeName());
+		}
+
+		case EMetasoundClassType::External:
+		case EMetasoundClassType::MetasoundGraph:
+		case EMetasoundClassType::Invalid:
+		default:
+		{
+			static_assert(static_cast<int32>(EMetasoundClassType::Invalid) == 4, "Possible missing EMetasoundClassType case coverage");
+			return Super::GetPinDisplayName(Pin);
+		}
+	}
+}
+
 FLinearColor UMetasoundEditorGraphSchema::GetPinTypeColor(const FEdGraphPinType& PinType) const
 {
 	return Metasound::Editor::GetPinCategoryColor(PinType);
diff --git a/Engine/Plugins/Runtime/Metasound/Source/MetasoundEditor/Private/MetasoundEditorModule.cpp b/Engine/Plugins/Runtime/Metasound/Source/MetasoundEditor/Private/MetasoundEditorModule.cpp
--- a/Engine/Plugins/Runtime/Metasound/Source/MetasoundEditor/Private/MetasoundEditorModule.cpp
+++ b/Engine/Plugins/Runtime/Metasound/Source/MetasoundEditor/Private/MetasoundEditorModule.cpp
@@ -14,6 +14,7 @@
 #include "MetasoundSource.h"
 #include "MetasoundAssetTypeActions.h"
 #include "MetasoundDetailCustomization.h"
+#include "MetasoundLiteralDescriptionDetailCustomization.h"
 #include "MetasoundEditorGraphBuilder.h"
 #include "MetasoundEditorGraphSchema.h"
 #include "Modules/ModuleInterface.h"
@@ -157,11 +158,15 @@ namespace Metasound
 				
 				PropertyModule.RegisterCustomClassLayout(
 					UMetasound::StaticClass()->GetFName(),
-					FOnGetDetailCustomizationInstance::CreateLambda([]() { return MakeShared<FMetasoundDetailCustomization>(); }));
+					FOnGetDetailCustomizationInstance::CreateLambda([]() { return MakeShared<FMetasoundDetailCustomization>(UMetasound::GetDocumentPropertyName()); }));
 
 				PropertyModule.RegisterCustomClassLayout(
 					UMetasoundSource::StaticClass()->GetFName(),
-					FOnGetDetailCustomizationInstance::CreateLambda([]() { return MakeShared<FMetasoundDetailCustomization>(); }));
+					FOnGetDetailCustomizationInstance::CreateLambda([]() { return MakeShared<FMetasoundDetailCustomization>(UMetasoundSource::GetDocumentPropertyName()); }));
+
+				PropertyModule.RegisterCustomPropertyTypeLayout(
+					"MetasoundLiteralDescription",
+					FOnGetPropertyTypeCustomizationInstance::CreateLambda([]() { return MakeShared<FMetasoundLiteralDescriptionDetailCustomization>(); }));
 
 				StyleSet = MakeShared<FSlateStyle>();
 
diff --git a/Engine/Plugins/Runtime/Metasound/Source/MetasoundEditor/Private/MetasoundLiteralDescriptionDetailCustomization.cpp b/Engine/Plugins/Runtime/Metasound/Source/MetasoundEditor/Private/MetasoundLiteralDescriptionDetailCustomization.cpp
new file mode 100644
--- /dev/null
+++ b/Engine/Plugins/Runtime/Metasound/Source/MetasoundEditor/Private/MetasoundLiteralDescriptionDetailCustomization.cpp
@@ -0,0 +1,174 @@
+// Copyright Epic Games, Inc. All Rights Reserved.
+#include "MetasoundLiteralDescriptionDetailCustomization.h"
+
+#include "CoreMinimal.h"
+#include "Delegates/Delegate.h"
+#include "DetailCategoryBuilder.h"
+#include "DetailLayoutBuilder.h"
+#include "DetailWidgetRow.h"
+#include "IDetailChildrenBuilder.h"
+#include "IDetailGroup.h"
+#include "Layout/Visibility.h"
+#include "MetasoundFrontendDataLayout.h"
+#include "Misc/AssertionMacros.h"
+#include "Misc/Attribute.h"
+#include "PropertyEditorDelegates.h"
+#include "PropertyHandle.h"
+#include "Templates/Casts.h"
+#include "Templates/SharedPointer.h"
+#include "Widgets/Text/STextBlock.h"
+
+
+#define LOCTEXT_NAMESPACE "MetasoundEditor"
+
+namespace Metasound
+{
+	namespace Editor
+	{
+		TSharedPtr<IPropertyHandle> GetLiteralHandleForType(TSharedRef<IPropertyHandle> StructPropertyHandle, EMetasoundLiteralType LiteralType)
+		{
+			switch (LiteralType)
+			{
+			case EMetasoundLiteralType::Bool:
+			{
+				return StructPropertyHandle->GetChildHandle(GET_MEMBER_NAME_CHECKED(FMetasoundLiteralDescription, AsBool));
+			}
+
+			case EMetasoundLiteralType::Float:
+			{
+				return StructPropertyHandle->GetChildHandle(GET_MEMBER_NAME_CHECKED(FMetasoundLiteralDescription, AsFloat));
+			}
+
+			case EMetasoundLiteralType::Integer:
+			{
+				return StructPropertyHandle->GetChildHandle(GET_MEMBER_NAME_CHECKED(FMetasoundLiteralDescription, AsInteger));
+			}
+
+			case EMetasoundLiteralType::String:
+			{
+				return StructPropertyHandle->GetChildHandle(GET_MEMBER_NAME_CHECKED(FMetasoundLiteralDescription, AsString));
+			}
+
+			case EMetasoundLiteralType::None:
+			default:
+				static_assert(static_cast<int32>(EMetasoundLiteralType::Invalid) == 5, "Possible missing case coverage for EMetasoundLiteralType");
+				return TSharedPtr<IPropertyHandle>();
+			}
+		}
+
+		TSharedPtr<IPropertyHandle> GetActiveLiteralHandle(TSharedRef<IPropertyHandle> StructPropertyHandle, TSharedRef<IPropertyHandle> TypeEnumHandle)
+		{
+			uint8 Value = static_cast<uint8>(EMetasoundLiteralType::None);
+			if (TypeEnumHandle->GetValue(Value) == FPropertyAccess::Result::Success)
+			{
+				const EMetasoundLiteralType LiteralType = static_cast<EMetasoundLiteralType>(Value);
+				return GetLiteralHandleForType(StructPropertyHandle, LiteralType);
+			}
+
+			return TSharedPtr<IPropertyHandle>();
+		}
+
+		bool IsLiteralTypeActive(TSharedRef<IPropertyHandle> TypeEnumHandle, EMetasoundLiteralType LiteralType)
+		{
+			uint8 Value = static_cast<uint8>(EMetasoundLiteralType::None);
+			TypeEnumHandle->GetValue(Value);
+
+			return Value == static_cast<uint8>(LiteralType);
+		}
+	} // namespace Editor
+} // namespace Metasound
+
+
+void FMetasoundLiteralDescriptionDetailCustomization::CustomizeHeader(TSharedRef<IPropertyHandle> StructPropertyHandle, FDetailWidgetRow& HeaderRow, IPropertyTypeCustomizationUtils& StructCustomizationUtils)
+{
+}
+
+void FMetasoundLiteralDescriptionDetailCustomization::CustomizeChildren(TSharedRef<IPropertyHandle> StructPropertyHandle, IDetailChildrenBuilder& ChildBuilder, IPropertyTypeCustomizationUtils& StructCustomizationUtils)
+{
+	TSharedRef<IPropertyHandle> TypeEnumHandle = StructPropertyHandle->GetChildHandle(GET_MEMBER_NAME_CHECKED(FMetasoundLiteralDescription, LiteralType)).ToSharedRef();
+
+	struct FTypeEditorInfo
+	{
+		EMetasoundLiteralType LiteralType;
+		TSharedPtr<IPropertyHandle> PropertyHandle;
+		FText DisplayName;
+	};
+
+	static_assert(static_cast<int32>(EMetasoundLiteralType::Invalid) == 5, "Possible missing property descriptor for literal customization display");
+	const TArray<FTypeEditorInfo> TypePropertyInfo =
+	{
+		FTypeEditorInfo
+		{
+			EMetasoundLiteralType::Bool,
+			StructPropertyHandle->GetChildHandle(GET_MEMBER_NAME_CHECKED(FMetasoundLiteralDescription, AsBool)),
+			LOCTEXT("Metasound_LiteralDisplayNameBool", "Bool")
+		},
+
+		FTypeEditorInfo
+		{
+			EMetasoundLiteralType::Integer,
+			StructPropertyHandle->GetChildHandle(GET_MEMBER_NAME_CHECKED(FMetasoundLiteralDescription, AsInteger)),
+			LOCTEXT("Metasound_LiteralDisplayNameInteger", "Int32")
+		},
+
+		FTypeEditorInfo
+		{
+			EMetasoundLiteralType::Float,
+			StructPropertyHandle->GetChildHandle(GET_MEMBER_NAME_CHECKED(FMetasoundLiteralDescription, AsFloat)),
+			LOCTEXT("Metasound_LiteralDisplayNameFloat", "Float")
+		},
+
+		FTypeEditorInfo
+		{
+			EMetasoundLiteralType::String,
+			StructPropertyHandle->GetChildHandle(GET_MEMBER_NAME_CHECKED(FMetasoundLiteralDescription, AsString)),
+			LOCTEXT("Metasound_LiteralDisplayNameString", "String")
+		},
+	};
+
+	for (const FTypeEditorInfo& Info : TypePropertyInfo)
+	{
+		ChildBuilder.AddCustomRow(StructPropertyHandle->GetPropertyDisplayName())
+			.NameContent()
+			[
+				SNew(SHorizontalBox)
+				+ SHorizontalBox::Slot()
+					.FillWidth(0.5f)
+					.Padding(1.0f, 0.0f, 0.0f, 0.0f)
+					.VAlign(VAlign_Center)
+					.HAlign(HAlign_Right)
+					[
+						SNew(STextBlock)
+						.Font(IDetailLayoutBuilder::GetDetailFont())
+						.Text(LOCTEXT("Metasound_LiteralDefault", "Default"))
+						.ToolTipText(TAttribute<FText>::Create([StructPropertyHandle, TypeEnumHandle]()
+						{
+							TSharedPtr<IPropertyHandle> VisibleHandle = Metasound::Editor::GetActiveLiteralHandle(StructPropertyHandle, TypeEnumHandle);
+							if (VisibleHandle.IsValid())
+							{
+								return VisibleHandle->GetToolTipText();
+							}
+							return FText::GetEmpty();
+						}))
+					]
+			]
+			.ValueContent()
+			.MinDesiredWidth(120.0f)
+			[
+				SNew(SHorizontalBox)
+				+ SHorizontalBox::Slot()
+					.FillWidth(0.5f)
+					.Padding(1.0f, 0.0f, 0.0f, 0.0f)
+					.VAlign(VAlign_Center)
+					[
+						Info.PropertyHandle->CreatePropertyValueWidget()
+					]
+			]
+			.Visibility(TAttribute<EVisibility>::Create([TypeEnumHandle, LiteralType = Info.LiteralType]()
+			{
+				const bool bIsLiteralActive = Metasound::Editor::IsLiteralTypeActive(TypeEnumHandle, LiteralType);
+				return bIsLiteralActive ? EVisibility::Visible : EVisibility::Hidden;
+			}));
+	}
+}
+#undef LOCTEXT_NAMESPACE // MetasoundEditor
\ No newline at end of file
diff --git a/Engine/Plugins/Runtime/Metasound/Source/MetasoundEditor/Private/MetasoundLiteralDescriptionDetailCustomization.h b/Engine/Plugins/Runtime/Metasound/Source/MetasoundEditor/Private/MetasoundLiteralDescriptionDetailCustomization.h
new file mode 100644
--- /dev/null
+++ b/Engine/Plugins/Runtime/Metasound/Source/MetasoundEditor/Private/MetasoundLiteralDescriptionDetailCustomization.h
@@ -0,0 +1,22 @@
+// Copyright Epic Games, Inc. All Rights Reserved.
+#pragma once
+
+#include "IPropertyTypeCustomization.h"
+#include "Templates/Casts.h"
+#include "Templates/SharedPointer.h"
+
+
+// Forward Declarations
+class IPropertyHandle;
+class FDetailWidgetRow;
+class IPropertyTypeCustomizationUtils;
+
+
+class FMetasoundLiteralDescriptionDetailCustomization : public IPropertyTypeCustomization
+{
+public:
+	//~ Begin IPropertyTypeCustomization
+	virtual void CustomizeHeader(TSharedRef<IPropertyHandle> StructPropertyHandle, FDetailWidgetRow& HeaderRow, IPropertyTypeCustomizationUtils& StructCustomizationUtils) override;
+	virtual void CustomizeChildren(TSharedRef<IPropertyHandle> StructPropertyHandle, IDetailChildrenBuilder& ChildBuilder, IPropertyTypeCustomizationUtils& StructCustomizationUtils) override;
+	//~ End IPropertyTypeCustomization
+};
\ No newline at end of file
diff --git a/Engine/Plugins/Runtime/Metasound/Source/MetasoundEditor/Public/MetasoundEditorGraphSchema.h b/Engine/Plugins/Runtime/Metasound/Source/MetasoundEditor/Public/MetasoundEditorGraphSchema.h
--- a/Engine/Plugins/Runtime/Metasound/Source/MetasoundEditor/Public/MetasoundEditorGraphSchema.h
+++ b/Engine/Plugins/Runtime/Metasound/Source/MetasoundEditor/Public/MetasoundEditorGraphSchema.h
@@ -216,6 +216,7 @@ class UMetasoundEditorGraphSchema : public UEdGraphSchema
 	//~ Begin EdGraphSchema Interface
 	virtual void GetGraphContextActions(FGraphContextMenuBuilder& ContextMenuBuilder) const override;
 	virtual void GetContextMenuActions(class UToolMenu* Menu, class UGraphNodeContextMenuContext* Context) const override;
+	virtual FText GetPinDisplayName(const UEdGraphPin* Pin) const override;
 	virtual void CreateDefaultNodesForGraph(UEdGraph& Graph) const override;
 	virtual const FPinConnectionResponse CanCreateConnection(const UEdGraphPin* A, const UEdGraphPin* B) const override;
 	virtual bool TryCreateConnection(UEdGraphPin* A, UEdGraphPin* B) const override;
diff --git a/Engine/Plugins/Runtime/Metasound/Source/MetasoundEngine/Private/MetasoundSource.cpp b/Engine/Plugins/Runtime/Metasound/Source/MetasoundEngine/Private/MetasoundSource.cpp
--- a/Engine/Plugins/Runtime/Metasound/Source/MetasoundEngine/Private/MetasoundSource.cpp
+++ b/Engine/Plugins/Runtime/Metasound/Source/MetasoundEngine/Private/MetasoundSource.cpp
@@ -124,21 +124,24 @@ FMetasoundArchetype UMetasoundSource::GetArchetype() const
 
 	FMetasoundInputDescription OnPlayBop;
 	OnPlayBop.Name = GetOnPlayInputName();
+	OnPlayBop.DisplayName = FText::FromString(OnPlayBop.Name);
 	OnPlayBop.TypeName = Metasound::Frontend::GetDataTypeName<Metasound::FBop>();
-	OnPlayBop.ToolTip = LOCTEXT("OnPlayBopToolTip", "This input is bopped when this source is first played.");
+	OnPlayBop.ToolTip = LOCTEXT("OnPlayBopToolTip", "Bop executed when this source is first played.");
 
 	Archetype.RequiredInputs.Add(OnPlayBop);
 
 	FMetasoundOutputDescription GeneratedAudio;
 	GeneratedAudio.Name = GetAudioOutputName();
+	GeneratedAudio.DisplayName = FText::FromString(GeneratedAudio.Name);
 	GeneratedAudio.TypeName = Metasound::Frontend::GetDataTypeName<Metasound::FAudioBuffer>();
-	GeneratedAudio.ToolTip = LOCTEXT("GeneratedAudioToolTip", "This should be the resulting output audio from this source.");
+	GeneratedAudio.ToolTip = LOCTEXT("GeneratedAudioToolTip", "The resulting output audio from this source.");
 	Archetype.RequiredOutputs.Add(GeneratedAudio);
 
 	FMetasoundOutputDescription OnFinished;
 	OnFinished.Name = GetIsFinishedOutputName();
+	OnFinished.DisplayName = FText::FromString(OnFinished.Name);
 	OnFinished.TypeName = Metasound::Frontend::GetDataTypeName<Metasound::FBop>();
-	OnFinished.ToolTip = LOCTEXT("OnFinishedToolTip", "This output can be bopped to stop the sound.");
+	OnFinished.ToolTip = LOCTEXT("OnFinishedToolTip", "Bop executed to initiate stopping the source.");
 	Archetype.RequiredOutputs.Add(OnFinished);
 
 	return Archetype;
diff --git a/Engine/Plugins/Runtime/Metasound/Source/MetasoundEngine/Public/Metasound.h b/Engine/Plugins/Runtime/Metasound/Source/MetasoundEngine/Public/Metasound.h
--- a/Engine/Plugins/Runtime/Metasound/Source/MetasoundEngine/Public/Metasound.h
+++ b/Engine/Plugins/Runtime/Metasound/Source/MetasoundEngine/Public/Metasound.h
@@ -9,6 +9,7 @@
 
 #if WITH_EDITORONLY_DATA
 #include "EdGraph/EdGraph.h"
+#include "Misc/AssertionMacros.h"
 #endif // WITH_EDITORONLY_DATA
 
 #include "Metasound.generated.h"
@@ -40,6 +41,12 @@ public:
 	UMetasound(const FObjectInitializer& ObjectInitializer);
 
 #if WITH_EDITORONLY_DATA
+	// Returns document name (for editor purposes, and avoids making document public for edit
+	// while allowing editor to reference directly)
+	static FName GetDocumentPropertyName()
+	{
+		return GET_MEMBER_NAME_CHECKED(UMetasound, RootMetasoundDocument);
+	}
 
 	// Returns the graph associated with this Metasound. Graph is required to be referenced on
 	// Metasound UObject for editor serialization purposes.
@@ -53,7 +60,6 @@ public:
 	// Metasound UObject for editor serialization purposes.
 	// @param Editor graph associated with UMetasound.
 	virtual void SetGraph(UEdGraph* InGraph) override;
-
 #endif // WITH_EDITORONLY_DATA
 
 	FMetasoundDocument& GetDocument() override
diff --git a/Engine/Plugins/Runtime/Metasound/Source/MetasoundEngine/Public/MetasoundSource.h b/Engine/Plugins/Runtime/Metasound/Source/MetasoundEngine/Public/MetasoundSource.h
--- a/Engine/Plugins/Runtime/Metasound/Source/MetasoundEngine/Public/MetasoundSource.h
+++ b/Engine/Plugins/Runtime/Metasound/Source/MetasoundEngine/Public/MetasoundSource.h
@@ -11,6 +11,7 @@
 
 #if WITH_EDITORONLY_DATA
 #include "EdGraph/EdGraph.h"
+#include "Misc/AssertionMacros.h"
 #endif // WITH_EDITORONLY_DATA
 
 #include "MetasoundSource.generated.h"
@@ -37,6 +38,12 @@ public:
 	UMetasoundSource(const FObjectInitializer& ObjectInitializer);
 
 #if WITH_EDITORONLY_DATA
+	// Returns document name (for editor purposes, and avoids making document public for edit
+	// while allowing editor to reference directly)
+	static FName GetDocumentPropertyName()
+	{
+		return GET_MEMBER_NAME_CHECKED(UMetasoundSource, RootMetasoundDocument);
+	}
 
 	// Returns the graph associated with this Metasound. Graph is required to be referenced on
 	// Metasound UObject for editor serialization purposes.
diff --git a/Engine/Plugins/Runtime/Metasound/Source/MetasoundFrontend/Private/MetasoundAssetBase.cpp b/Engine/Plugins/Runtime/Metasound/Source/MetasoundFrontend/Private/MetasoundAssetBase.cpp
--- a/Engine/Plugins/Runtime/Metasound/Source/MetasoundFrontend/Private/MetasoundAssetBase.cpp
+++ b/Engine/Plugins/Runtime/Metasound/Source/MetasoundFrontend/Private/MetasoundAssetBase.cpp
@@ -46,36 +46,6 @@ FMetasoundClassMetadata FMetasoundAssetBase::GetMetadata()
 	return GetDocument().RootClass.Metadata;
 }
 
-const FText& FMetasoundAssetBase::GetInputToolTip(FString InputName) const
-{
-	const FMetasoundDocument& RootMetasoundDocument = GetDocument();
-
-	for (const FMetasoundInputDescription& Desc : RootMetasoundDocument.RootClass.Inputs)
-	{
-		if (Desc.Name == InputName)
-		{
-			return Desc.ToolTip;
-		}
-	}
-
-	return FText::GetEmpty();
-}
-
-const FText& FMetasoundAssetBase::GetOutputToolTip(FString OutputName) const
-{
-	const FMetasoundDocument& RootMetasoundDocument = GetDocument();
-
-	for (const FMetasoundOutputDescription& Desc : RootMetasoundDocument.RootClass.Outputs)
-	{
-		if (Desc.Name == OutputName)
-		{
-			return Desc.ToolTip;
-		}
-	}
-
-	return FText::GetEmpty();
-}
-
 TWeakPtr<Metasound::Frontend::FDescriptionAccessPoint> FMetasoundAssetBase::GetGraphAccessPoint() const
 {
 	return TWeakPtr<Metasound::Frontend::FDescriptionAccessPoint>(AccessPoint);
diff --git a/Engine/Plugins/Runtime/Metasound/Source/MetasoundFrontend/Private/MetasoundFrontend.cpp b/Engine/Plugins/Runtime/Metasound/Source/MetasoundFrontend/Private/MetasoundFrontend.cpp
--- a/Engine/Plugins/Runtime/Metasound/Source/MetasoundFrontend/Private/MetasoundFrontend.cpp
+++ b/Engine/Plugins/Runtime/Metasound/Source/MetasoundFrontend/Private/MetasoundFrontend.cpp
@@ -1127,7 +1127,6 @@ namespace Metasound
 			return NodePtr->UniqueID;
 		}
 
-
 		const FString& FNodeHandle::GetNodeName() const
 		{
 			if (!IsValid())
@@ -2381,6 +2380,92 @@ namespace Metasound
 			return Builder.BuildGraphOperator(GraphToBuild, OutBuildErrors);
 		}
 
+		const FText& FGraphHandle::GetInputDisplayName(FString InputName) const
+		{
+			const FMetasoundDocument& RootMetasoundDocument = *OwningDocument;
+
+			for (const FMetasoundInputDescription& Desc : RootMetasoundDocument.RootClass.Inputs)
+			{
+				if (Desc.Name == InputName)
+				{
+					return Desc.DisplayName;
+				}
+			}
+
+			return FText::GetEmpty();
+		}
+
+		const FText& FGraphHandle::GetInputToolTip(FString InputName) const
+		{
+			const FMetasoundDocument& RootMetasoundDocument = *OwningDocument;
+
+			for (const FMetasoundInputDescription& Desc : RootMetasoundDocument.RootClass.Inputs)
+			{
+				if (Desc.Name == InputName)
+				{
+					return Desc.ToolTip;
+				}
+			}
+
+			return FText::GetEmpty();
+		}
+
+		const FText& FGraphHandle::GetOutputDisplayName(FString OutputName) const
+		{
+			const FMetasoundDocument& RootMetasoundDocument = *OwningDocument;
+
+			for (const FMetasoundOutputDescription& Desc : RootMetasoundDocument.RootClass.Outputs)
+			{
+				if (Desc.Name == OutputName)
+				{
+					return Desc.DisplayName;
+				}
+			}
+
+			return FText::GetEmpty();
+		}
+
+		const FText& FGraphHandle::GetOutputToolTip(FString OutputName) const
+		{
+			const FMetasoundDocument& RootMetasoundDocument = *OwningDocument;
+
+			for (const FMetasoundOutputDescription& Desc : RootMetasoundDocument.RootClass.Outputs)
+			{
+				if (Desc.Name == OutputName)
+				{
+					return Desc.ToolTip;
+				}
+			}
+
+			return FText::GetEmpty();
+		}
+
+		void FGraphHandle::SetInputDisplayName(FString InName, const FText& InDisplayName)
+		{
+			FMetasoundDocument& RootMetasoundDocument = *OwningDocument;
+			for (FMetasoundInputDescription& Desc : RootMetasoundDocument.RootClass.Inputs)
+			{
+				if (Desc.Name == InName)
+				{
+					Desc.DisplayName = InDisplayName;
+					break;
+				}
+			}
+		}
+
+		void FGraphHandle::SetOutputDisplayName(FString InName, const FText& InDisplayName)
+		{
+			FMetasoundDocument& RootMetasoundDocument = *OwningDocument;
+			for (FMetasoundOutputDescription& Desc : RootMetasoundDocument.RootClass.Outputs)
+			{
+				if (Desc.Name == InName)
+				{
+					Desc.DisplayName = InDisplayName;
+					break;
+				}
+			}
+		}
+
 		FGraphHandle GetGraphHandleForClass(const FMetasoundClassDescription& InClass)
 		{
 			ensureAlwaysMsgf(false, TEXT("Implement Me!"));
diff --git a/Engine/Plugins/Runtime/Metasound/Source/MetasoundFrontend/Public/MetasoundAssetBase.h b/Engine/Plugins/Runtime/Metasound/Source/MetasoundFrontend/Public/MetasoundAssetBase.h
--- a/Engine/Plugins/Runtime/Metasound/Source/MetasoundFrontend/Public/MetasoundAssetBase.h
+++ b/Engine/Plugins/Runtime/Metasound/Source/MetasoundFrontend/Public/MetasoundAssetBase.h
@@ -45,9 +45,6 @@ public:
 	// Returns the root class metadata
 	FMetasoundClassMetadata GetMetadata();
 
-	const FText& GetInputToolTip(FString InputName) const;
-	const FText& GetOutputToolTip(FString OutputName) const;
-
 	// Imports the asset from a JSON file at provided path
 	bool ImportFromJSON(const FString& InAbsolutePath);
 
diff --git a/Engine/Plugins/Runtime/Metasound/Source/MetasoundFrontend/Public/MetasoundFrontend.h b/Engine/Plugins/Runtime/Metasound/Source/MetasoundFrontend/Public/MetasoundFrontend.h
--- a/Engine/Plugins/Runtime/Metasound/Source/MetasoundFrontend/Public/MetasoundFrontend.h
+++ b/Engine/Plugins/Runtime/Metasound/Source/MetasoundFrontend/Public/MetasoundFrontend.h
@@ -379,6 +379,24 @@ namespace Metasound
 			bool SetInputToLiteral(const FString& InInputName, UObject* InValue);
 			bool SetInputToLiteral(const FString& InInputName, TArray<UObject*> InValue);
 
+			// Returns the display name for the input with the given name
+			const FText& GetInputDisplayName(FString InputName) const;
+
+			// Returns the display name for the output with the given name
+			const FText& GetOutputDisplayName(FString OutputName) const;
+
+			// Returns the tool tip for the input with the given name
+			const FText& GetInputToolTip(FString InputName) const;
+
+			// Returns the tool tip for the output with the given name
+			const FText& GetOutputToolTip(FString OutputName) const;
+
+			// Set the display name for the input with the given name
+			void SetInputDisplayName(FString InName, const FText& InDisplayName);
+
+			// Set the display name for the output with the given name
+			void SetOutputDisplayName(FString InName, const FText& InDisplayName);
+
 			// This can be used to clear the current literal for a given input.
 			// @returns false if the input name couldn't be found.
 			bool ClearLiteralForInput(const FString& InInputName);
diff --git a/Engine/Plugins/Runtime/Metasound/Source/MetasoundFrontend/Public/MetasoundFrontendBaseClasses.h b/Engine/Plugins/Runtime/Metasound/Source/MetasoundFrontend/Public/MetasoundFrontendBaseClasses.h
--- a/Engine/Plugins/Runtime/Metasound/Source/MetasoundFrontend/Public/MetasoundFrontendBaseClasses.h
+++ b/Engine/Plugins/Runtime/Metasound/Source/MetasoundFrontend/Public/MetasoundFrontendBaseClasses.h
@@ -17,40 +17,40 @@ namespace Metasound
 		void InitializeFrontend();
 
 		/** Use this to get a reference to the static registry of all nodes implemented in C++. */
-		TMap<FNodeRegistryKey, FNodeRegistryElement>& GetExternalNodeRegistry();
+		METASOUNDFRONTEND_API TMap<FNodeRegistryKey, FNodeRegistryElement>& GetExternalNodeRegistry();
 
 		// Convenience functions to create an INode corresponding to a specific input or output for a metasound graph.
 		// @returns nullptr if the type given wasn't found.
-		TUniquePtr<INode> ConstructInputNode(const FName& InInputType, const FInputNodeConstructorParams& InParams);
-		TUniquePtr<INode> ConstructOutputNode(const FName& InOutputType, const FOutputNodeConstrutorParams& InParams);
+		METASOUNDFRONTEND_API TUniquePtr<INode> ConstructInputNode(const FName& InInputType, const FInputNodeConstructorParams& InParams);
+		METASOUNDFRONTEND_API TUniquePtr<INode> ConstructOutputNode(const FName& InOutputType, const FOutputNodeConstrutorParams& InParams);
 
 		// Convenience functions to create an INodeB corresponding to a specific externally declared node type.
 		// InNodeType and InNodeHash can be retrieved from the FNodeClassInfo generated from the node registry queries in the metasound frontend (GetAllAvailableNodeClasses, GetAllNodeClassesInNamespace, etc.)
 		// @returns nullptr if the type given wasn't found.
-		TUniquePtr<INode> ConstructExternalNode(const FName& InNodeType, uint32 InNodeHash, const FNodeInitData& InInitData);
+		METASOUNDFRONTEND_API TUniquePtr<INode> ConstructExternalNode(const FName& InNodeType, uint32 InNodeHash, const FNodeInitData& InInitData);
 
 		// Utility functions for setting serialized literals.
-		void SetLiteralDescription(FMetasoundLiteralDescription& OutDescription, bool InValue);
-		void SetLiteralDescription(FMetasoundLiteralDescription& OutDescription, int32 InValue);
-		void SetLiteralDescription(FMetasoundLiteralDescription& OutDescription, float InValue);
-		void SetLiteralDescription(FMetasoundLiteralDescription& OutDescription, const FString& InValue);
-		void ClearLiteralDescription(FMetasoundLiteralDescription& OutDescription);
+		METASOUNDFRONTEND_API void SetLiteralDescription(FMetasoundLiteralDescription& OutDescription, bool InValue);
+		METASOUNDFRONTEND_API void SetLiteralDescription(FMetasoundLiteralDescription& OutDescription, int32 InValue);
+		METASOUNDFRONTEND_API void SetLiteralDescription(FMetasoundLiteralDescription& OutDescription, float InValue);
+		METASOUNDFRONTEND_API void SetLiteralDescription(FMetasoundLiteralDescription& OutDescription, const FString& InValue);
+		METASOUNDFRONTEND_API void ClearLiteralDescription(FMetasoundLiteralDescription& OutDescription);
 
 		// Utility functions for building a ::Metasound::FDataInitParam corresponding to a literal.
 
 		// Return the literal description parsed into a init param. 
 		// @Returns an invalid init param if the data type couldn't be found, or if the literal type was incompatible with the data type.
-		FDataTypeLiteralParam GetLiteralParamForDataType(FName InDataType, const FMetasoundLiteralDescription& InDescription);
+		METASOUNDFRONTEND_API FDataTypeLiteralParam GetLiteralParamForDataType(FName InDataType, const FMetasoundLiteralDescription& InDescription);
 
-		bool DoesDataTypeSupportLiteralType(FName InDataType, EMetasoundLiteralType InLiteralType);
-		bool DoesDataTypeSupportLiteralType(FName InDataType, ELiteralArgType InLiteralType);
+		METASOUNDFRONTEND_API bool DoesDataTypeSupportLiteralType(FName InDataType, EMetasoundLiteralType InLiteralType);
+		METASOUNDFRONTEND_API bool DoesDataTypeSupportLiteralType(FName InDataType, ELiteralArgType InLiteralType);
 		
 		// Returns a literal param without any type checking.
-		FDataTypeLiteralParam GetLiteralParam(const FMetasoundLiteralDescription& InDescription);
+		METASOUNDFRONTEND_API FDataTypeLiteralParam GetLiteralParam(const FMetasoundLiteralDescription& InDescription);
 
 		// Returns the defaulted version of a literal param for the given data type.
 		// @Returns an invalid init param if the data type couldn't be found.
-		FDataTypeLiteralParam GetDefaultParamForDataType(FName InDataType);
+		METASOUNDFRONTEND_API FDataTypeLiteralParam GetDefaultParamForDataType(FName InDataType);
 
 		// Utility base class for classes that want to support naive multi-level undo/redo.
 		class METASOUNDFRONTEND_API ITransactable : public TSharedFromThis<ITransactable>
diff --git a/Engine/Plugins/Runtime/Metasound/Source/MetasoundFrontend/Public/MetasoundFrontendDataLayout.h b/Engine/Plugins/Runtime/Metasound/Source/MetasoundFrontend/Public/MetasoundFrontendDataLayout.h
--- a/Engine/Plugins/Runtime/Metasound/Source/MetasoundFrontend/Public/MetasoundFrontendDataLayout.h
+++ b/Engine/Plugins/Runtime/Metasound/Source/MetasoundFrontend/Public/MetasoundFrontendDataLayout.h
@@ -44,7 +44,7 @@ struct FMetasoundClassMetadata
 	FString NodeName;
 
 	// This will always be set to EMetasoundObjectType::Metasound.
-	UPROPERTY(VisibleAnywhere, Category = General)
+	UPROPERTY(VisibleAnywhere, Category = General, meta = (DisplayName = "Type"))
 	EMetasoundClassType NodeType;
 
 	// Optional longform description of what this Metasound does. Can be displayed in a tooltip.
@@ -62,8 +62,8 @@ struct FMetasoundClassMetadata
 };
 
 
-// This ustruct represents the serialized version of our variant literal type.
-// Currently, we only support bools, integers, floats, and strings as literals.
+// Represents the serialized version of variant literal types. Currently, only
+// support bools, integers, floats, and strings are supported as literals.
 USTRUCT()
 struct FMetasoundLiteralDescription
 {
@@ -72,19 +72,19 @@ struct FMetasoundLiteralDescription
 	// HACK: Once we make a property inspector for these we should hide them.
 
 	// The actual type of this literal.
-	UPROPERTY(EditAnywhere, Category = Parameters)
+	UPROPERTY(EditAnywhere, Category = Customized)
 	EMetasoundLiteralType LiteralType;
 
-	UPROPERTY(EditAnywhere, Category = Parameters)
+	UPROPERTY(EditAnywhere, Category = Customized)
 	bool AsBool;
 
-	UPROPERTY(EditAnywhere, Category = Parameters)
+	UPROPERTY(EditAnywhere, Category = Customized)
 	int32 AsInteger;
 
-	UPROPERTY(EditAnywhere, Category = Parameters)
+	UPROPERTY(EditAnywhere, Category = Customized)
 	float AsFloat;
 
-	UPROPERTY(EditAnywhere, Category = Parameters)
+	UPROPERTY(EditAnywhere, Category = Customized)
 	FString AsString;
 };
 
@@ -93,12 +93,16 @@ struct FMetasoundInputDescription
 {
 	GENERATED_BODY()
 
-	// The descriptive name of this input.
-	UPROPERTY(EditAnywhere, Category = Parameters)
+	// The name of this input.
+	UPROPERTY()
 	FString Name;
 
+	// The name of this input displayed to users.
+	UPROPERTY(EditAnywhere, Category = Parameters, meta = (DisplayName = "Name"))
+	FText DisplayName;
+
 	// The type of this input. This type matches with a type declared via the DECLARE_METASOUND_DATA_REFERENCE_TYPES macro.
-	UPROPERTY(VisibleAnywhere, Category = Parameters)
+	UPROPERTY(EditAnywhere, Category = Parameters)
 	FName TypeName;
 
 	// Optional description text about this input.
@@ -117,11 +121,15 @@ struct FMetasoundOutputDescription
 	GENERATED_BODY()
 
 	// The descriptive name of this output.
-	UPROPERTY(EditAnywhere, Category = Parameters)
+	UPROPERTY()
 	FString Name;
 
+	// The name of this input displayed to users.
+	UPROPERTY(EditAnywhere, Category = Parameters, meta = (DisplayName = "Name"))
+	FText DisplayName;
+
 	// The type of this output. This type matches with a type declared via the DECLARE_METASOUND_DATA_REFERENCE_TYPES macro.
-	UPROPERTY(VisibleAnywhere, Category = Parameters)
+	UPROPERTY(EditAnywhere, Category = Parameters)
 	FName TypeName;
 
 	// Optional description text about this output.
