Waveplayer node unit test and fixes around conversions

diff --git a/Engine/Plugins/Runtime/Metasound/Source/MetasoundStandardNodes/Private/MetasoundWave.cpp b/Engine/Plugins/Runtime/Metasound/Source/MetasoundStandardNodes/Private/MetasoundWave.cpp
new file mode 100644
--- /dev/null
+++ b/Engine/Plugins/Runtime/Metasound/Source/MetasoundStandardNodes/Private/MetasoundWave.cpp
@@ -0,0 +1,24 @@
+// Copyright Epic Games, Inc. All Rights Reserved.
+
+#include "MetasoundWave.h"
+#include "MetasoundPrimitives.h"
+#include "IAudioCodec.h"
+
+namespace Metasound
+{	
+	// Simple for now, just wrap around the TArray<uint8>
+	class FWaveDecoderInput : public Audio::FDecoderInputArrayView
+	{
+		FWaveReadRef WaveRef;	// These are not thread safe, and will need to be
+	public:
+		FWaveDecoderInput(const FWaveReadRef& InWaveRef)
+			: Audio::FDecoderInputArrayView(MakeArrayView(InWaveRef->CompressedBytes.GetData(), InWaveRef->CompressedBytes.Num()), 0)
+			, WaveRef(InWaveRef)
+		{}
+	};
+
+	FWave::FDecoderInputPtr FWave::CreateDecoderInput(const FWaveReadRef& InWaveRef) 
+	{
+		return MakeShared<FWaveDecoderInput, ESPMode::ThreadSafe>(InWaveRef);
+	}
+}
diff --git a/Engine/Plugins/Runtime/Metasound/Source/MetasoundStandardNodes/Private/MetasoundWavePlayerNode.cpp b/Engine/Plugins/Runtime/Metasound/Source/MetasoundStandardNodes/Private/MetasoundWavePlayerNode.cpp
--- a/Engine/Plugins/Runtime/Metasound/Source/MetasoundStandardNodes/Private/MetasoundWavePlayerNode.cpp
+++ b/Engine/Plugins/Runtime/Metasound/Source/MetasoundStandardNodes/Private/MetasoundWavePlayerNode.cpp
@@ -5,6 +5,7 @@
 #include "MetasoundPrimitives.h"
 #include "MetasoundWave.h"
 #include "MetasoundNodeRegistrationMacro.h"
+#include "MetasoundBuildError.h"
 
 #include "IAudioCodecRegistry.h"
 #include "IAudioCodec.h"
@@ -14,72 +15,109 @@
 namespace Metasound
 {
 	METASOUND_REGISTER_NODE(FWavePlayerNode)
-		
+	
+	// WavePlayer custom error 
+	class FWavePlayerError : public FBuildErrorBase
+	{
+	public:
+		FWavePlayerError(FName InName, FText InErrorDescription)
+			: FBuildErrorBase(InName, InErrorDescription)
+		{}
+		virtual ~FWavePlayerError() = default;
+	};
+	
 	class FWavePlayerOperator : public TExecutableOperator<FWavePlayerOperator>
 	{
-		public:
-			FWavePlayerOperator(
-				const FOperatorSettings& InSettings, 
-				const FWaveReadRef& InWave, 
-				TUniquePtr<Audio::IDecoderInput>&& InDecoderInput, 
-				TUniquePtr<Audio::IDecoderOutput>&& InDecoderOutput,
-				TUniquePtr<Audio::IDecoder>&& InDecoder )
-				: OperatorSettings(InSettings)
-				, Decoder(MoveTemp(InDecoder))
-				, DecoderInput(MoveTemp(InDecoderInput))
-				, DecoderOutput(MoveTemp(InDecoderOutput))
-				, Wave(InWave)
-				, AudioBuffer(FAudioBufferWriteRef::CreateNew(InSettings))
-			{
-				check(AudioBuffer->Num() == InSettings.GetNumFramesPerBlock());
+	public:
+
+		// Silent setup
+		FWavePlayerOperator(
+			const FOperatorSettings& InSettings, 
+			const FWaveReadRef& InWave )
+			: OperatorSettings(InSettings)
+			, Wave(InWave)
+			, AudioBuffer(FAudioBufferWriteRef::CreateNew(InSettings))
+		{
+			check(AudioBuffer->Num() == InSettings.GetNumFramesPerBlock());
+			OutputDataReferences.AddDataReadReference(TEXT("Audio"), FAudioBufferReadRef(AudioBuffer));
+		}
+		
+		FWavePlayerOperator(
+			const FOperatorSettings& InSettings,
+			const FWaveReadRef& InWave, 
+			FWave::FDecoderInputPtr&& InDecoderInput, 
+			TUniquePtr<Audio::IDecoderOutput>&& InDecoderOutput,
+			TUniquePtr<Audio::IDecoder>&& InDecoder )
+			: OperatorSettings(InSettings)
+			, Wave(InWave)
+			, AudioBuffer(FAudioBufferWriteRef::CreateNew(InSettings))
+			, Decoder(MoveTemp(InDecoder))
+			, DecoderInput(MoveTemp(InDecoderInput))
+			, DecoderOutput(MoveTemp(InDecoderOutput))
+		{
+			check(AudioBuffer->Num() == InSettings.GetNumFramesPerBlock());
+			OutputDataReferences.AddDataReadReference(TEXT("Audio"), FAudioBufferReadRef(AudioBuffer));
+		}
 
-				OutputDataReferences.AddDataReadReference(TEXT("Audio"), FAudioBufferReadRef(AudioBuffer));
-			}
+		virtual const FDataReferenceCollection& GetInputs() const override
+		{
+			return InputDataReferences;
+		}
 
-			virtual const FDataReferenceCollection& GetInputs() const override
-			{
-				return InputDataReferences;
-			}
+		virtual const FDataReferenceCollection& GetOutputs() const override
+		{
+			return OutputDataReferences;
+		}
 
-			virtual const FDataReferenceCollection& GetOutputs() const override
-			{
-				return OutputDataReferences;
-			}
+		// Scheduler/decoder holds weakref, and pins the sharedptr to the output when it writes.
+		// FBufferedDecoderOutput
+		// 1 write [... .. .. ]
+		// n free
+		// m queued 
+		// 1 read [...........]
+
+		// ApuMemory 2kb compress -> 512 frames
+
+		// TCircularAudioBuffer //consider for Output object.
+		
+		void Execute()
+		{
+			int32 i = 0;
+			float* Dst = AudioBuffer->GetData();
 
-			void Execute()
+			// If we don't have a valid state, just output silence.
+			if (Decoder && DecoderOutput && DecoderInput )
 			{
 				// V1. Do the decode inline, this will sound bad.
 				Decoder->Decode();
 
 				TArrayView<float> Src;
-				float* Dst = AudioBuffer->GetData();
-				
 				int32 NumPopped = DecoderOutput->PopAudio(Src);
-				int32 NumFramesToCopy = FMath::Min(NumPopped,OperatorSettings.GetNumFramesPerBlock());
+				int32 NumFramesToCopy = FMath::Min(NumPopped, OperatorSettings.GetNumFramesPerBlock());
 
-				int32 i = 0;
 				for (; i < NumFramesToCopy; ++i)
 				{
 					Dst[i] = Src[i];
 				}
-				for ( ; i < OperatorSettings.GetNumFramesPerBlock(); ++i)
-				{
-					Dst[i] = 0.0f;
-				}
 			}
+			for ( ; i < OperatorSettings.GetNumFramesPerBlock(); ++i)
+			{
+				Dst[i] = 0.0f;
+			}
+		}
+
+	private:
+		const FOperatorSettings OperatorSettings;
+
+		FWaveReadRef Wave;
+		FAudioBufferWriteRef AudioBuffer;
+		FDataReferenceCollection InputDataReferences;
+		FDataReferenceCollection OutputDataReferences;
 
-		private:
-			const FOperatorSettings OperatorSettings;
-			
-			Audio::ICodec::FDecoderPtr Decoder;
-			TUniquePtr<Audio::IDecoderInput> DecoderInput;
-			TUniquePtr<Audio::IDecoderOutput> DecoderOutput;
-			
-			FWaveReadRef Wave;
-			FAudioBufferWriteRef AudioBuffer;
-
-			FDataReferenceCollection InputDataReferences;
-			FDataReferenceCollection OutputDataReferences;
+		// Decoder/IO. 
+		Audio::ICodec::FDecoderPtr Decoder;
+		FWave::FDecoderInputPtr DecoderInput;
+		TUniquePtr<Audio::IDecoderOutput> DecoderOutput;
 	};
 
 	const FName FWavePlayerNode::ClassName = FName(TEXT("Wave"));
@@ -100,11 +138,11 @@ namespace Metasound
 			Wave = InInputDataReferences.GetDataReadReference<FWave>(TEXT("Wave"));
 		}
 
-		TUniquePtr<IDecoderInput> Input = IDecoderInput::Create(Wave->GetCompressedData(),0);
-		if(Input)
+		FWave::FDecoderInputPtr Input = FWave::CreateDecoderInput(Wave);
+		if (Input)
 		{
-			ICodecRegistry::FCodecPtr Codec = Audio::ICodecRegistry::Get().FindCodecByFromParsingInput(Input.Get());
-			if(Codec)
+			ICodecRegistry::FCodecPtr Codec = ICodecRegistry::Get().FindCodecByFromParsingInput(Input.Get());
+			if (Codec)
 			{
 				// V1, Ask for an output buffer the size of a frame.
 				IDecoderOutput::FRequirements Reqs { Float32_Interleaved, InOperatorSettings.GetNumFramesPerBlock() };
@@ -121,16 +159,18 @@ namespace Metasound
 			}
 			else
 			{
-				// LogMsg; Failed to Find
+				OutErrors.Add(MakeUnique<FWavePlayerError>(TEXT("FailedToFindCodec"),
+					LOCTEXT("FailedToFindCodec", "Failed to find codec for opening the supplied Wave")));
 			}
 		}
 		else
 		{
-			// LogMsg; Failed to parse input stream, is it valid?
+			OutErrors.Add(MakeUnique<FWavePlayerError>(TEXT("FailedToParseInput"),
+				LOCTEXT("FailedToParseInput", "Failed to parse the compressed data")));
 		}
 
-		// Fail.
-		return nullptr;
+		// Create the player without any inputs, will just produce silence.
+		return MakeUnique<FWavePlayerOperator>(InOperatorSettings,Wave);
 	}
 
 	FWavePlayerNode::FWavePlayerNode(const FString& InName)
diff --git a/Engine/Plugins/Runtime/Metasound/Source/MetasoundStandardNodes/Public/MetasoundWave.h b/Engine/Plugins/Runtime/Metasound/Source/MetasoundStandardNodes/Public/MetasoundWave.h
--- a/Engine/Plugins/Runtime/Metasound/Source/MetasoundStandardNodes/Public/MetasoundWave.h
+++ b/Engine/Plugins/Runtime/Metasound/Source/MetasoundStandardNodes/Public/MetasoundWave.h
@@ -5,14 +5,35 @@
 #include "MetasoundDataReference.h"
 #include "MetasoundDataTypeRegistrationMacro.h"
 
+// Forward declares
+namespace Audio
+{	
+	struct IDecoderInput;
+}
+
 namespace Metasound
 {
+	// Forward declare ReadRef
+	class FWave;
+	typedef TDataReadReference<FWave> FWaveReadRef;
+		
 	class FWave
 	{
 		TArray<uint8> CompressedBytes;
+		friend class FWaveDecoderInput;
 	public:
-		const TArray<uint8>& GetCompressedData() const { return CompressedBytes; }
-	};
+		FWave() = default;
 
+		// For testing only.
+		FWave(const TArray<uint8>& InBytes) 
+			: CompressedBytes(InBytes) 
+		{}
+
+		using FDecoderInputPtr = TSharedPtr<Audio::IDecoderInput,ESPMode::ThreadSafe>;
+
+		// Factory function to create a Decoder input
+		static FDecoderInputPtr CreateDecoderInput(
+			const FWaveReadRef& InWaveRef);
+	};
 	DECLARE_METASOUND_DATA_REFERENCE_TYPES(FWave, 0x0ddba11, FWaveTypeInfo, FWaveReadRef, FWaveWriteRef)
 }
diff --git a/Engine/Source/Runtime/AudioExtensions/Public/IAudioCodec.h b/Engine/Source/Runtime/AudioExtensions/Public/IAudioCodec.h
--- a/Engine/Source/Runtime/AudioExtensions/Public/IAudioCodec.h
+++ b/Engine/Source/Runtime/AudioExtensions/Public/IAudioCodec.h
@@ -431,7 +431,7 @@ namespace Audio
 			const FDecodedFormatInfo& InFormat) const override
 		{
 			return Reqs;
-		}
+		}	
 
 		template<typename From, typename To>
 		static constexpr float GetConvertScalar()
@@ -442,11 +442,11 @@ namespace Audio
 			}
 			else if( TIsSame<From,int16>::Value && TIsSame<To,float>::Value)
 			{
-				return 1.f / static_cast<float>(FMath::Abs(TNumericLimits<int16>::Min()));
+				return 1.f / (static_cast<float>(TNumericLimits<int16>::Max()) + 1.f);
 			}
 			else if( TIsSame<From,float>::Value && TIsSame<To,int16>::Value)
 			{
-				return static_cast<float>(FMath::Abs(TNumericLimits<int16>::Min()));
+				return static_cast<float>(TNumericLimits<int16>::Max()) + 1.f;
 			}
 		}
 
@@ -725,4 +725,4 @@ namespace Audio
 
 } // namespace Audio
 
-#pragma endregion Codec
\ No newline at end of file
+#pragma endregion Codec
