- Mock up displaying required I/O to UI for MetasoundSources
- Disallow removing required I/O from MetasoundSources
- Add functions to expose/iterate required I/O archetypes from Frontend
- Enable copy graph nodes using JSON MS Document Schema and pruning unused nodes
- Mock in paste code (Ethan to add validation or "fix-up" logic on MSDoc side to avoid duplication of node/I/O Ids, names, etc.)

diff --git a/Engine/Plugins/Runtime/Metasound/Source/MetasoundEditor/Private/MetasoundDetailCustomization.cpp b/Engine/Plugins/Runtime/Metasound/Source/MetasoundEditor/Private/MetasoundDetailCustomization.cpp
--- a/Engine/Plugins/Runtime/Metasound/Source/MetasoundEditor/Private/MetasoundDetailCustomization.cpp
+++ b/Engine/Plugins/Runtime/Metasound/Source/MetasoundEditor/Private/MetasoundDetailCustomization.cpp
@@ -1,6 +1,7 @@
 // Copyright Epic Games, Inc. All Rights Reserved.
 #include "MetasoundDetailCustomization.h"
 
+#include "Containers/Set.h"
 #include "CoreMinimal.h"
 #include "Delegates/Delegate.h"
 #include "DetailCategoryBuilder.h"
@@ -9,10 +10,13 @@
 #include "IDetailGroup.h"
 #include "Metasound.h"
 #include "MetasoundAssetBase.h"
+#include "MetasoundFrontend.h"
 #include "PropertyEditorDelegates.h"
 #include "PropertyHandle.h"
+#include "SlateCore/Public/Styling/SlateColor.h"
 #include "Templates/Casts.h"
 #include "Templates/SharedPointer.h"
+#include "UObject/WeakObjectPtr.h"
 #include "UObject/WeakObjectPtrTemplates.h"
 #include "Widgets/Text/STextBlock.h"
 
@@ -34,7 +38,7 @@ namespace Metasound
 		}
 
 		template <typename T>
-		void BuildFixedArray(IDetailLayoutBuilder& InDetailLayout, FName CategoryName, FName InPropertyName)
+		void BuildIOFixedArray(IDetailLayoutBuilder& InDetailLayout, FName CategoryName, FName InPropertyName, const TSet<FString>& RequiredValues, bool bIsInput)
 		{
 			IDetailCategoryBuilder& CategoryBuilder = InDetailLayout.EditCategory(CategoryName);
 			TSharedPtr<IPropertyHandle> ParentProperty = InDetailLayout.GetProperty(InPropertyName);
@@ -60,117 +64,156 @@ namespace Metasound
 				}
 
 				TSharedPtr<IPropertyHandle> TypeProperty = ChildProperties.FindChecked(GET_MEMBER_NAME_CHECKED(T, TypeName));
-				TSharedPtr<IPropertyHandle> NameProperty = ChildProperties.FindChecked(GET_MEMBER_NAME_CHECKED(T, DisplayName));
-
-				FDetailWidgetRow& NameRow = CategoryBuilder.AddCustomRow(ParentProperty->GetPropertyDisplayName())
-					.NameContent()
-					[
-						SNew(STextBlock)
-						.Font(IDetailLayoutBuilder::GetDetailFontBold())
-						.Text(TAttribute<FText>::Create([i, NameProperty, TypeProperty]()
-						{
-							FString TypeNameString;
-							ensureAlways(TypeProperty->GetValue(TypeNameString) == FPropertyAccess::Success);
+				TSharedPtr<IPropertyHandle> ToolTipProperty = ChildProperties.FindChecked(GET_MEMBER_NAME_CHECKED(T, ToolTip));
+				TSharedPtr<IPropertyHandle> DisplayNameProperty = ChildProperties.FindChecked(GET_MEMBER_NAME_CHECKED(T, DisplayName));
+				TSharedPtr<IPropertyHandle> NameProperty = ChildProperties.FindChecked(GET_MEMBER_NAME_CHECKED(T, Name));
+
+				FString Name;
+				const bool bNameFound = NameProperty->GetValue(Name) == FPropertyAccess::Success;
+				const bool bIsRequired = bNameFound && RequiredValues.Contains(Name);
+
+				CategoryBuilder.AddCustomRow(ParentProperty->GetPropertyDisplayName()).NameContent()
+				[
+					SNew(STextBlock)
+					.Font(IDetailLayoutBuilder::GetDetailFontBold())
+					.Text(TAttribute<FText>::Create([i, bIsRequired, DisplayNameProperty, TypeProperty]()
+					{
+						FName TypeName;
+						TypeProperty->GetValue(TypeName);
+						FString TypeNameString = TypeName.ToString();
 
-							// Remove namespace info to keep concise
-							TypeNameString.RightChopInline(TypeNameString.Find(TEXT(":"), ESearchCase::IgnoreCase, ESearchDir::FromEnd) + 1);
+						// Remove namespace info to keep concise
+						TypeNameString.RightChopInline(TypeNameString.Find(TEXT(":"), ESearchCase::IgnoreCase, ESearchDir::FromEnd) + 1);
 
-							FText DisplayName;
-							ensureAlways(NameProperty->GetValue(DisplayName) == FPropertyAccess::Success);
+						FText DisplayName;
+						DisplayNameProperty->GetValue(DisplayName);
 
-							return FText::Format(LOCTEXT("Metasound_FixedArray_Format", "{0}. {1} ({2})"), FText::AsNumber(i + 1), DisplayName, FText::FromString(TypeNameString));
-						}))
-						.ToolTipText(TAttribute<FText>::Create([ParentProperty, TypeProperty]()
+						if (bIsRequired)
+						{
+							return FText::Format(LOCTEXT("Metasound_FixedIOArrayRequiredEntry_Format", "{0}. {1} ({2}, Required)"), FText::AsNumber(i + 1), DisplayName, FText::FromString(TypeNameString));
+						}
+						else
 						{
-							FString TypeNameString;
-							ensureAlways(TypeProperty->GetValue(TypeNameString) == FPropertyAccess::Success);
+							return FText::Format(LOCTEXT("Metasound_FixedIOArray_Format", "{0}. {1} ({2})"), FText::AsNumber(i + 1), DisplayName, FText::FromString(TypeNameString));
+						}
+					}))
+					.ToolTipText(TAttribute<FText>::Create([ToolTipProperty]()
+					{
+						FText ToolTip;
+						ToolTipProperty->GetValue(ToolTip);
+						return ToolTip;
+					}))
+				];
 
-							return FText::Format(LOCTEXT("", "Type {0}\n"), FText::FromString(TypeNameString), ParentProperty->GetToolTipText());
-						}))
-					];
+				if (!bIsRequired)
+				{
+					CategoryBuilder.AddProperty(DisplayNameProperty);
+					CategoryBuilder.AddProperty(ToolTipProperty);
+				}
 
-				for (const TPair<FName, TSharedPtr<IPropertyHandle>>& Pair : ChildProperties)
+				if (bIsInput)
 				{
-					if (Pair.Key != GET_MEMBER_NAME_CHECKED(T, TypeName))
-					{
-						CategoryBuilder.AddProperty(Pair.Value);
-					}
+					TSharedPtr<IPropertyHandle> LiteralValueProperty = ChildProperties.FindChecked(GET_MEMBER_NAME_CHECKED(FMetasoundInputDescription, LiteralValue));
+					CategoryBuilder.AddProperty(LiteralValueProperty);
 				}
 			}
 
 			FSimpleDelegate RefreshDelegate = FSimpleDelegate::CreateLambda([DetailLayout = &InDetailLayout]()
-				{
-					DetailLayout->ForceRefreshDetails();
-				});
+			{
+				DetailLayout->ForceRefreshDetails();
+			});
 			ArrayHandle->SetOnNumElementsChanged(RefreshDelegate);
 		};
-	}
-}
 
-FMetasoundDetailCustomization::FMetasoundDetailCustomization(FName InDocumentPropertyName)
-	: IDetailCustomization()
-	, DocumentPropertyName(InDocumentPropertyName)
-{
-}
+		FMetasoundDetailCustomization::FMetasoundDetailCustomization(FName InDocumentPropertyName)
+			: IDetailCustomization()
+			, DocumentPropertyName(InDocumentPropertyName)
+		{
+		}
 
-FName FMetasoundDetailCustomization::GetMetadataRootClassPath() const
-{
-	return Metasound::Editor::BuildChildPath(DocumentPropertyName, GET_MEMBER_NAME_CHECKED(FMetasoundDocument, RootClass));
-}
+		FName FMetasoundDetailCustomization::GetMetadataRootClassPath() const
+		{
+			return Metasound::Editor::BuildChildPath(DocumentPropertyName, GET_MEMBER_NAME_CHECKED(FMetasoundDocument, RootClass));
+		}
 
-FName FMetasoundDetailCustomization::GetMetadataPropertyPath() const
-{
-	const FName RootClass = FName(GetMetadataRootClassPath());
-	return Metasound::Editor::BuildChildPath(RootClass, GET_MEMBER_NAME_CHECKED(FMetasoundClassDescription, Metadata));
-}
+		FName FMetasoundDetailCustomization::GetMetadataPropertyPath() const
+		{
+			const FName RootClass = FName(GetMetadataRootClassPath());
+			return Metasound::Editor::BuildChildPath(RootClass, GET_MEMBER_NAME_CHECKED(FMetasoundClassDescription, Metadata));
+		}
 
-void FMetasoundDetailCustomization::CustomizeDetails(IDetailLayoutBuilder& DetailLayout)
-{
-	using namespace Metasound::Editor;
-
-	// General Category
-	IDetailCategoryBuilder& GeneralCategoryBuilder = DetailLayout.EditCategory("General");
-
-	TArray<TWeakObjectPtr<UObject>> ObjectsCustomized;
-	DetailLayout.GetObjectsBeingCustomized(ObjectsCustomized);
-
-	const FName AuthorPropertyPath = BuildChildPath(GetMetadataPropertyPath(), GET_MEMBER_NAME_CHECKED(FMetasoundClassMetadata, AuthorName));
-	const FName DescPropertyPath = BuildChildPath(GetMetadataPropertyPath(), GET_MEMBER_NAME_CHECKED(FMetasoundClassMetadata, MetasoundDescription));
-	const FName NodeTypePropertyPath = BuildChildPath(GetMetadataPropertyPath(), GET_MEMBER_NAME_CHECKED(FMetasoundClassMetadata, NodeType));
-
-	TSharedPtr<IPropertyHandle> AuthorHandle = DetailLayout.GetProperty(AuthorPropertyPath);
-	TSharedPtr<IPropertyHandle> DescHandle = DetailLayout.GetProperty(DescPropertyPath);
-	TSharedPtr<IPropertyHandle> NodeTypeHandle = DetailLayout.GetProperty(NodeTypePropertyPath);
-
-	GeneralCategoryBuilder.AddProperty(NodeTypeHandle);
-	GeneralCategoryBuilder.AddProperty(AuthorHandle);
-	GeneralCategoryBuilder.AddProperty(DescHandle);
-
-	// Input/Output Categories
-
-	const FName InputsPropertyPath = BuildChildPath(GetMetadataRootClassPath(), GET_MEMBER_NAME_CHECKED(FMetasoundClassDescription, Inputs));
-	const FName OutputsPropertyPath = BuildChildPath(GetMetadataRootClassPath(), GET_MEMBER_NAME_CHECKED(FMetasoundClassDescription, Outputs));
-
-	BuildFixedArray<FMetasoundInputDescription>(DetailLayout, "Inputs", InputsPropertyPath);
-	BuildFixedArray<FMetasoundOutputDescription>(DetailLayout, "Outputs", OutputsPropertyPath);
-
-	// Hack to hide parent structs for nested metadata properties
-	DetailLayout.HideCategory("Hidden");
-
-	// Hack to hide categories brought in from UMetasoundSource inherited from USoundBase
-	DetailLayout.HideCategory("Analysis");
-	DetailLayout.HideCategory("Attenuation");
-	DetailLayout.HideCategory("Curves");
-	DetailLayout.HideCategory("Debug");
-	DetailLayout.HideCategory("Effects");
-	DetailLayout.HideCategory("File Path");
-	DetailLayout.HideCategory("Format");
-	DetailLayout.HideCategory("Info");
-	DetailLayout.HideCategory("Loading");
-	DetailLayout.HideCategory("Modulation");
-	DetailLayout.HideCategory("Playback");
-	DetailLayout.HideCategory("Sound");
-	DetailLayout.HideCategory("Subtitles");
-	DetailLayout.HideCategory("Voice Management");
-}
+		void FMetasoundDetailCustomization::CustomizeDetails(IDetailLayoutBuilder& DetailLayout)
+		{
+			using namespace Metasound::Editor;
+
+			// General Category
+			IDetailCategoryBuilder& GeneralCategoryBuilder = DetailLayout.EditCategory("General");
+
+			TArray<TWeakObjectPtr<UObject>> ObjectsCustomized;
+			DetailLayout.GetObjectsBeingCustomized(ObjectsCustomized);
+
+			const FName AuthorPropertyPath = BuildChildPath(GetMetadataPropertyPath(), GET_MEMBER_NAME_CHECKED(FMetasoundClassMetadata, AuthorName));
+			const FName DescPropertyPath = BuildChildPath(GetMetadataPropertyPath(), GET_MEMBER_NAME_CHECKED(FMetasoundClassMetadata, MetasoundDescription));
+			const FName NodeTypePropertyPath = BuildChildPath(GetMetadataPropertyPath(), GET_MEMBER_NAME_CHECKED(FMetasoundClassMetadata, NodeType));
+
+			TSharedPtr<IPropertyHandle> AuthorHandle = DetailLayout.GetProperty(AuthorPropertyPath);
+			TSharedPtr<IPropertyHandle> DescHandle = DetailLayout.GetProperty(DescPropertyPath);
+			TSharedPtr<IPropertyHandle> NodeTypeHandle = DetailLayout.GetProperty(NodeTypePropertyPath);
+
+			GeneralCategoryBuilder.AddProperty(NodeTypeHandle);
+			GeneralCategoryBuilder.AddProperty(AuthorHandle);
+			GeneralCategoryBuilder.AddProperty(DescHandle);
+
+			// Input/Output Categories
+
+			// If editing multiple metasound objects, all should be the same type, so safe to just check first in array for
+			// required inputs/outputs
+			TArray<TWeakObjectPtr<UObject>> Objects;
+			TSet<FString> RequiredInputs;
+			TSet<FString> RequiredOutputs;
+			DetailLayout.GetObjectsBeingCustomized(Objects);
+			if (Objects.Num() > 0)
+			{
+				FMetasoundAssetBase* MetasoundAsset = Frontend::GetObjectAsAssetBase(Objects[0].Get());
+				check(MetasoundAsset);
+
+				Frontend::FGraphHandle GraphHandle = MetasoundAsset->GetRootGraphHandle();
+				for (const FMetasoundInputDescription& Desc : GraphHandle.GetRequiredInputs())
+				{
+					RequiredInputs.Add(Desc.Name);
+				}
+				for (const FMetasoundOutputDescription& Desc : GraphHandle.GetRequiredOutputs())
+				{
+					RequiredOutputs.Add(Desc.Name);
+				}
+			}
+
+			const FName InputsPropertyPath = BuildChildPath(GetMetadataRootClassPath(), GET_MEMBER_NAME_CHECKED(FMetasoundClassDescription, Inputs));
+			const FName OutputsPropertyPath = BuildChildPath(GetMetadataRootClassPath(), GET_MEMBER_NAME_CHECKED(FMetasoundClassDescription, Outputs));
+
+			BuildIOFixedArray<FMetasoundInputDescription>(DetailLayout, "Inputs", InputsPropertyPath, RequiredInputs, true /* bIsInput */);
+			BuildIOFixedArray<FMetasoundOutputDescription>(DetailLayout, "Outputs", OutputsPropertyPath, RequiredOutputs, false /* bIsInput */);
+
+			// Hack to hide parent structs for nested metadata properties
+			DetailLayout.HideCategory("CustomView");
+
+			// Hack to hide categories brought in from UMetasoundSource inherited from USoundBase
+			DetailLayout.HideCategory("Analysis");
+			DetailLayout.HideCategory("Attenuation");
+			DetailLayout.HideCategory("Curves");
+			DetailLayout.HideCategory("Debug");
+			DetailLayout.HideCategory("Effects");
+			DetailLayout.HideCategory("File Path");
+			DetailLayout.HideCategory("Format");
+			DetailLayout.HideCategory("Info");
+			DetailLayout.HideCategory("Loading");
+			DetailLayout.HideCategory("Modulation");
+			DetailLayout.HideCategory("Playback");
+			DetailLayout.HideCategory("Sound");
+			DetailLayout.HideCategory("SoundWave");
+			DetailLayout.HideCategory("Subtitles");
+			DetailLayout.HideCategory("Voice Management");
+		}
+	} // namespace Editor
+} // namespace Metasound
 #undef LOCTEXT_NAMESPACE
diff --git a/Engine/Plugins/Runtime/Metasound/Source/MetasoundEditor/Private/MetasoundDetailCustomization.h b/Engine/Plugins/Runtime/Metasound/Source/MetasoundEditor/Private/MetasoundDetailCustomization.h
--- a/Engine/Plugins/Runtime/Metasound/Source/MetasoundEditor/Private/MetasoundDetailCustomization.h
+++ b/Engine/Plugins/Runtime/Metasound/Source/MetasoundEditor/Private/MetasoundDetailCustomization.h
@@ -7,19 +7,24 @@
 // Forward Declarations
 class IDetailLayoutBuilder;
 
-
-class FMetasoundDetailCustomization : public IDetailCustomization
+namespace Metasound
 {
-public:
-	FMetasoundDetailCustomization(FName InDocumentPropertyName);
+	namespace Editor
+	{
+		class FMetasoundDetailCustomization : public IDetailCustomization
+		{
+		public:
+			FMetasoundDetailCustomization(FName InDocumentPropertyName);
 
-	// IDetailCustomization interface
-	virtual void CustomizeDetails(IDetailLayoutBuilder& DetailLayout) override;
-	// End of IDetailCustomization interface
+			// IDetailCustomization interface
+			virtual void CustomizeDetails(IDetailLayoutBuilder& DetailLayout) override;
+			// End of IDetailCustomization interface
 
-private:
-	FName GetMetadataRootClassPath() const;
-	FName GetMetadataPropertyPath() const;
+		private:
+			FName GetMetadataRootClassPath() const;
+			FName GetMetadataPropertyPath() const;
 
-	FName DocumentPropertyName;
-};
\ No newline at end of file
+			FName DocumentPropertyName;
+		};
+	} // namespace Editor
+} // namespace Metasound
\ No newline at end of file
diff --git a/Engine/Plugins/Runtime/Metasound/Source/MetasoundEditor/Private/MetasoundEditor.cpp b/Engine/Plugins/Runtime/Metasound/Source/MetasoundEditor/Private/MetasoundEditor.cpp
--- a/Engine/Plugins/Runtime/Metasound/Source/MetasoundEditor/Private/MetasoundEditor.cpp
+++ b/Engine/Plugins/Runtime/Metasound/Source/MetasoundEditor/Private/MetasoundEditor.cpp
@@ -10,6 +10,7 @@
 #include "Framework/Commands/GenericCommands.h"
 #include "Framework/Docking/TabManager.h"
 #include "Framework/MultiBox/MultiBoxBuilder.h"
+#include "Framework/Notifications/NotificationManager.h"
 #include "GraphEditor.h"
 #include "GraphEditorActions.h"
 #include "HAL/PlatformApplicationMisc.h"
@@ -28,6 +29,8 @@
 #include "SMetasoundPalette.h"
 #include "SNodePanel.h"
 #include "Widgets/Docking/SDockTab.h"
+#include "Widgets/Notifications/SNotificationList.h"
+
 
 #define LOCTEXT_NAMESPACE "MetasoundEditor"
 
@@ -91,26 +94,34 @@ namespace Metasound
 			GEditor->UnregisterForUndo(this);
 		}
 
-		void FEditor::InitMetasoundEditor(const EToolkitMode::Type Mode, const TSharedPtr<IToolkitHost>& InitToolkitHost, UObject* ObjectToEdit)
+		Frontend::FGraphHandle FEditor::InitMetasound(UObject& InMetasound)
 		{
-			checkf(Frontend::IsObjectAMetasoundArchetype(ObjectToEdit), TEXT("Object passed in was not registered as a valid metasound archetype!"));
-			
-			Metasound = ObjectToEdit;
-			FMetasoundAssetBase* MetasoundAsset = Frontend::GetObjectAsAssetBase(ObjectToEdit);
+			FMetasoundAssetBase* MetasoundAsset = Frontend::GetObjectAsAssetBase(&InMetasound);
 			check(MetasoundAsset);
 
-			// Support undo/redo
-			ObjectToEdit->SetFlags(RF_Transactional);
-
 			if (!MetasoundAsset->GetGraph())
 			{
-				UMetasoundEditorGraph* Graph = NewObject<UMetasoundEditorGraph>(Metasound);
-				Graph->ParentMetasound = Metasound;
+				UMetasoundEditorGraph* Graph = NewObject<UMetasoundEditorGraph>(&InMetasound);
+				Graph->ParentMetasound = &InMetasound;
 				Graph->Schema = UMetasoundEditorGraphSchema::StaticClass();
 				MetasoundAsset->SetGraph(Graph);
-				FGraphBuilder::RebuildGraph(*Metasound);
+				FGraphBuilder::RebuildGraph(InMetasound);
 			}
 
+			return MetasoundAsset->GetRootGraphHandle();
+		}
+
+		void FEditor::InitMetasoundEditor(const EToolkitMode::Type Mode, const TSharedPtr<IToolkitHost>& InitToolkitHost, UObject* ObjectToEdit)
+		{
+			check(ObjectToEdit);
+			checkf(Frontend::IsObjectAMetasoundArchetype(ObjectToEdit), TEXT("Object passed in was not registered as a valid metasound archetype!"));
+			
+			// Support undo/redo
+			ObjectToEdit->SetFlags(RF_Transactional);
+
+			Metasound = ObjectToEdit;
+			InitMetasound(*Metasound);
+
 			GEditor->RegisterForUndo(this);
 
 			FGraphEditorCommands::Register();
@@ -353,7 +364,7 @@ namespace Metasound
 
 				FMetasoundDocument MetasoundDoc;
 
-				if (Frontend::ImportJSONToMetasound(InputPath, MetasoundDoc))
+				if (Frontend::ImportJSONAssetToMetasound(InputPath, MetasoundDoc))
 				{
 					Frontend::GetObjectForDocument(MetasoundDoc, OutputPath);
 				}
@@ -378,7 +389,7 @@ namespace Metasound
 
 			// TODO: We could just make this an object.
 			const FString Path = FPaths::ProjectSavedDir() / TEXT("Metasounds") + FPaths::ChangeExtension(Metasound->GetPathName(), MetasoundExtension);
-			InMetasoundAsset->ExportToJSON(Path);
+			InMetasoundAsset->GetRootGraphHandle().ExportToJSONAsset(Path);
 		}
 
 		void FEditor::Play()
@@ -517,7 +528,7 @@ namespace Metasound
 
 				GraphEditorCommands->MapAction(FGenericCommands::Get().Delete,
 					FExecuteAction::CreateSP(this, &FEditor::DeleteSelectedNodes),
-					FCanExecuteAction::CreateLambda([this]() { return MetasoundGraphEditor->GetSelectedNodes().Num() > 0; }));
+					FCanExecuteAction::CreateLambda([this]() { return CanDeleteNodes(); }));
 
 				GraphEditorCommands->MapAction(FGenericCommands::Get().Copy,
 					FExecuteAction::CreateSP(this, &FEditor::CopySelectedNodes),
@@ -525,15 +536,15 @@ namespace Metasound
 
 				GraphEditorCommands->MapAction(FGenericCommands::Get().Cut,
 					FExecuteAction::CreateSP(this, &FEditor::CutSelectedNodes),
-					FCanExecuteAction::CreateSP(this, &FEditor::CanCutNodes));
+					FCanExecuteAction::CreateLambda([this]() { return CanCopyNodes() && CanDeleteNodes(); }));
 
 				GraphEditorCommands->MapAction(FGenericCommands::Get().Paste,
-					FExecuteAction::CreateSP(this, &FEditor::PasteNodes),
+					FExecuteAction::CreateLambda([this]() { PasteNodes(); }),
 					FCanExecuteAction::CreateSP(this, &FEditor::CanPasteNodes));
 
 				GraphEditorCommands->MapAction(FGenericCommands::Get().Duplicate,
-					FExecuteAction::CreateSP(this, &FEditor::DuplicateNodes),
-					FCanExecuteAction::CreateSP(this, &FEditor::CanDuplicateNodes));
+					FExecuteAction::CreateLambda([this] { CopySelectedNodes(); PasteNodes(); }),
+					FCanExecuteAction::CreateLambda([this]() { return CanCopyNodes(); }));
 
 				// Alignment Commands
 				GraphEditorCommands->MapAction(FGraphEditorCommands::Get().AlignNodesTop,
@@ -624,8 +635,6 @@ namespace Metasound
 
 		void FEditor::DeleteSelectedNodes()
 		{
-			using namespace Metasound::Frontend;
-
 			const FScopedTransaction Transaction(LOCTEXT("MetasoundEditorDeleteSelectedNode", "Delete Selected Metasound Node(s)"));
 
 			UMetasoundEditorGraph* Graph = CastChecked<UMetasoundEditorGraph>(MetasoundGraphEditor->GetCurrentGraph());
@@ -635,20 +644,52 @@ namespace Metasound
 			MetasoundGraphEditor->ClearSelectionSet();
 
 			UObject* ParentMetasoundObject = Graph->ParentMetasound;
-
-			if (FMetasoundAssetBase* MetasoundAsset = GetObjectAsAssetBase(ParentMetasoundObject))
+			check(ParentMetasoundObject);
+			if (FMetasoundAssetBase* MetasoundAsset = Frontend::GetObjectAsAssetBase(ParentMetasoundObject))
 			{
-				FGraphHandle GraphHandle = MetasoundAsset->GetRootGraphHandle();
+				Frontend::FGraphHandle GraphHandle = MetasoundAsset->GetRootGraphHandle();
 				for (FGraphPanelSelectionSet::TConstIterator NodeIt(SelectedNodes); NodeIt; ++NodeIt)
 				{
 					UMetasoundEditorGraphNode* Node = CastChecked<UMetasoundEditorGraphNode>(*NodeIt);
+					Frontend::FNodeHandle NodeHandle = Node->GetNodeHandle();
+
+					if (NodeHandle.GetNodeType() == EMetasoundClassType::Input)
+					{
+						if (GraphHandle.IsRequiredInput(NodeHandle.GetNodeName()))
+						{
+							FNotificationInfo Info(FText::Format(LOCTEXT("Metasounds_CannotDeleteRequiredInput",
+								"'Required Input '{0}' cannot be deleted."), GraphHandle.GetInputDisplayName(NodeHandle.GetNodeName())));
+							Info.bFireAndForget = true;
+							Info.ExpireDuration = 2.0f;
+							Info.bUseThrobber = true;
+							FSlateNotificationManager::Get().AddNotification(Info);
+							continue;
+						}
+					}
+
+					if (NodeHandle.GetNodeType() == EMetasoundClassType::Output)
+					{
+						if (GraphHandle.IsRequiredOutput(NodeHandle.GetNodeName()))
+						{
+							FNotificationInfo Info(FText::Format(LOCTEXT("Metasounds_CannotDeleteRequiredOutput",
+								"'Required Output '{0}' cannot be deleted."), GraphHandle.GetOutputDisplayName(NodeHandle.GetNodeName())));
+							Info.bFireAndForget = true;
+							Info.ExpireDuration = 2.0f;
+							Info.bUseThrobber = true;
+							FSlateNotificationManager::Get().AddNotification(Info);
+							continue;
+						}
+					}
+
 					FGraphBuilder::DeleteNode(*Node, false /* bInRecordTransaction */);
 				}
 			}
 		}
 
-		void FEditor::DeleteSelectedDuplicatableNodes()
+		void FEditor::CutSelectedNodes()
 		{
+			CopySelectedNodes();
+
 			// Cache off the old selection
 			const FGraphPanelSelectionSet OldSelectedNodes = MetasoundGraphEditor->GetSelectedNodes();
 
@@ -659,7 +700,7 @@ namespace Metasound
 			for (FGraphPanelSelectionSet::TConstIterator SelectedIter(OldSelectedNodes); SelectedIter; ++SelectedIter)
 			{
 				UEdGraphNode* Node = Cast<UEdGraphNode>(*SelectedIter);
-				if ((Node != nullptr) && Node->CanDuplicateNode())
+				if (Node && Node->CanUserDeleteNode())
 				{
 					MetasoundGraphEditor->SetNodeSelection(Node, true);
 				}
@@ -669,7 +710,7 @@ namespace Metasound
 				}
 			}
 
-			// Delete the duplicatable nodes
+			// Delete the deleteable nodes
 			DeleteSelectedNodes();
 
 			// Reselect whatever's left from the original selection after the deletion
@@ -684,44 +725,61 @@ namespace Metasound
 			}
 		}
 
-		void FEditor::CutSelectedNodes()
-		{
-			CopySelectedNodes();
-			// Cut should only delete nodes that can be duplicated
-			DeleteSelectedDuplicatableNodes();
-		}
-
-		bool FEditor::CanCutNodes() const
-		{
-			return CanCopyNodes();
-		}
-
 		void FEditor::CopySelectedNodes()
 		{
+			using namespace Metasound::Frontend;
+
 			// Export the selected nodes and place the text on the clipboard
 			const FGraphPanelSelectionSet SelectedNodes = MetasoundGraphEditor->GetSelectedNodes();
 
-			FString ExportedText;
-
+			TMap<uint32, FNodeHandle> HandlesToCopy;
 			for (FGraphPanelSelectionSet::TConstIterator SelectedIter(SelectedNodes); SelectedIter; ++SelectedIter)
 			{
 				if (UMetasoundEditorGraphNode* Node = Cast<UMetasoundEditorGraphNode>(*SelectedIter))
 				{
+					HandlesToCopy.Add(Node->GetNodeID(), Node->GetNodeHandle());
 					Node->PrepareForCopying();
 				}
 			}
 
-			FEdGraphUtilities::ExportNodesToText(SelectedNodes, /*out*/ ExportedText);
-			FPlatformApplicationMisc::ClipboardCopy(*ExportedText);
+			FMetasoundAssetBase* MetasoundAsset = Metasound::Frontend::GetObjectAsAssetBase(Metasound);
+			check(MetasoundAsset);
+			FGraphHandle GraphHandle = MetasoundAsset->GetRootGraphHandle();
 
-			// Make sure Metasound remains the owner of the copied nodes
-			for (FGraphPanelSelectionSet::TConstIterator SelectedIter(SelectedNodes); SelectedIter; ++SelectedIter)
+			UMetasound* MetasoundCopy = NewObject<UMetasound>();
+			FGraphHandle CopyGraphHandle = InitMetasound(*MetasoundCopy);
+
+			GraphHandle.CopyGraph(CopyGraphHandle);
+
+			TArray<FNodeHandle> AllCopiedNodes = CopyGraphHandle.GetAllNodes();
+			for (FNodeHandle& Node : AllCopiedNodes)
 			{
-				if (UMetasoundEditorGraphNode* Node = Cast<UMetasoundEditorGraphNode>(*SelectedIter))
+				if (!HandlesToCopy.Contains(Node.GetNodeID()))
 				{
-					Node->PostCopyNode();
+					if (Node.GetNodeType() == EMetasoundClassType::Input)
+					{
+						CopyGraphHandle.RemoveInput(Node.GetNodeName());
+					}
+					else if (Node.GetNodeType() == EMetasoundClassType::Output)
+					{
+						CopyGraphHandle.RemoveOutput(Node.GetNodeName());
+					}
+					else
+					{
+						CopyGraphHandle.RemoveNode(Node);
+					}
 				}
 			}
+
+			FString ExportedEdText;
+			FEdGraphUtilities::ExportNodesToText(SelectedNodes, ExportedEdText);
+
+			FMetasoundEditorData EditorData;
+			EditorData.GraphData = ExportedEdText;
+			CopyGraphHandle.SetEditorData(EditorData);
+
+			const FString JSONToCopy = CopyGraphHandle.ExportToJSON();
+			FPlatformApplicationMisc::ClipboardCopy(*JSONToCopy);
 		}
 
 		bool FEditor::CanCopyNodes() const
@@ -739,109 +797,116 @@ namespace Metasound
 			return false;
 		}
 
-		void FEditor::PasteNodes()
-		{
-			PasteNodesAtLocation(MetasoundGraphEditor->GetPasteLocation());
-		}
-
-		void FEditor::PasteNodesAtLocation(const FVector2D& Location)
+		bool FEditor::CanDeleteNodes() const
 		{
-			FMetasoundAssetBase* MetasoundAsset = Frontend::GetObjectAsAssetBase(Metasound);
-
-			if (!ensureAlways(MetasoundAsset))
+			if (MetasoundGraphEditor->GetSelectedNodes().Num() == 0)
 			{
-				return;
+				return false;
 			}
 
-			UEdGraph* Graph = MetasoundAsset->GetGraph();
-			if (!Graph)
-			{
-				return;
-			}
-
-			// Undo/Redo support
-			const FScopedTransaction Transaction(NSLOCTEXT("UnrealEd", "MetasoundEditorPaste", "Paste Metasound Node"));
-			Graph->Modify();
-			Metasound->Modify();
-
-			// Clear the selection set (newly pasted stuff will be selected)
-			MetasoundGraphEditor->ClearSelectionSet();
-
-			// Grab the text to paste from the clipboard.
-			FString TextToImport;
-			FPlatformApplicationMisc::ClipboardPaste(TextToImport);
-
-			// Import the nodes
-			TSet<UEdGraphNode*> PastedNodes;
-			FEdGraphUtilities::ImportNodesFromText(Graph, TextToImport, /*out*/ PastedNodes);
-
-			//Average position of nodes so we can move them while still maintaining relative distances to each other
-			FVector2D AvgNodePosition(0.0f, 0.0f);
-
-			for (TSet<UEdGraphNode*>::TIterator It(PastedNodes); It; ++It)
-			{
-				UEdGraphNode* Node = *It;
-				AvgNodePosition.X += Node->NodePosX;
-				AvgNodePosition.Y += Node->NodePosY;
-			}
-
-			if (PastedNodes.Num() > 0)
-			{
-				float InvNumNodes = 1.0f / float(PastedNodes.Num());
-				AvgNodePosition.X *= InvNumNodes;
-				AvgNodePosition.Y *= InvNumNodes;
-			}
-
-			for (TSet<UEdGraphNode*>::TIterator It(PastedNodes); It; ++It)
+			const FGraphPanelSelectionSet& SelectedNodes = MetasoundGraphEditor->GetSelectedNodes();
+			for (FGraphPanelSelectionSet::TConstIterator SelectedIter(SelectedNodes); SelectedIter; ++SelectedIter)
 			{
-				UEdGraphNode* Node = *It;
-
-				if (UMetasoundEditorGraphNode* SoundGraphNode = Cast<UMetasoundEditorGraphNode>(Node))
+				UEdGraphNode* Node = Cast<UEdGraphNode>(*SelectedIter);
+				if (Node && Node->CanUserDeleteNode())
 				{
-					// TODO: Add newly referenced nodes to MS reference list
+					return true;
 				}
-
-				// Select the newly pasted stuff
-				MetasoundGraphEditor->SetNodeSelection(Node, true);
-
-				Node->NodePosX = (Node->NodePosX - AvgNodePosition.X) + Location.X;
-				Node->NodePosY = (Node->NodePosY - AvgNodePosition.Y) + Location.Y;
-
-				Node->SnapToGrid(SNodePanel::GetSnapGridSize());
-
-				// Give new node a different Guid from the old one
-				Node->CreateNewGuid();
 			}
-
-			// Update UI
-			MetasoundGraphEditor->NotifyGraphChanged();
-
-			Metasound->PostEditChange();
-			Metasound->MarkPackageDirty();
-		}
-
-		bool FEditor::CanPasteNodes() const
-		{
-			FString ClipboardContent;
-			FPlatformApplicationMisc::ClipboardPaste(ClipboardContent);
-
-			// TODO: Implement
-		// 	const bool bCanPasteNodes = FEdGraphUtilities::CanImportNodesFromText(Metasound->MetasoundGraph, ClipboardContent);
-		// 	return bCanPasteNodes;
-
 			return false;
 		}
 
-		void FEditor::DuplicateNodes()
+		void FEditor::PasteNodes(const FVector2D* InLocation)
+		{
+// 			FVector2D Location;
+// 			if (InLocation)
+// 			{
+// 				Location = *InLocation;
+// 			}
+// 			else
+// 			{
+// 				check(MetasoundGraphEditor);
+// 				Location = MetasoundGraphEditor->GetPasteLocation();
+// 			}
+// 
+// 			FMetasoundAssetBase* MetasoundAsset = Frontend::GetObjectAsAssetBase(Metasound);
+// 			check(MetasoundAsset);
+// 
+// 			UEdGraph* Graph = MetasoundAsset->GetGraph();
+// 			if (!Graph)
+// 			{
+// 				return;
+// 			}
+// 
+// 			// Undo/Redo support
+// 			const FScopedTransaction Transaction(NSLOCTEXT("UnrealEd", "MetasoundEditorPaste", "Paste Metasound Node"));
+// 			Graph->Modify();
+// 			Metasound->Modify();
+// 
+// 			// Clear the selection set (newly pasted stuff will be selected)
+// 			MetasoundGraphEditor->ClearSelectionSet();
+// 
+//			static_assert(false, "TODO: Delete this and assert & insert nodes/inputs/outputs from 'PastedDocument' into 'Metasound' doc here. Run fix-up on new nodes to generate ids ensuring no duplicates with nodes that pre-existed.");
+// 
+// 			// Import the nodes
+// 			TSet<UEdGraphNode*> PastedNodes;
+// 			FEdGraphUtilities::ImportNodesFromText(Graph, PastedDocument.EditorData.GraphData, PastedNodes);
+// 
+// 			// Clear out, no reason to hold a bunch of references to data no longer needed.
+// 			PastedDocument = FMetasoundDocument();
+// 
+// 			FVector2D AvgNodePosition = FVector2D::ZeroVector;
+// 			for (UEdGraphNode* Node : PastedNodes)
+// 			{
+// 				AvgNodePosition.X += Node->NodePosX;
+// 				AvgNodePosition.Y += Node->NodePosY;
+// 			}
+// 
+// 			if (PastedNodes.Num() > 0)
+// 			{
+// 				float InvNumNodes = 1.0f / PastedNodes.Num();
+// 				AvgNodePosition.X *= InvNumNodes;
+// 				AvgNodePosition.Y *= InvNumNodes;
+// 			}
+// 
+// 			for (UEdGraphNode* Node : PastedNodes)
+// 			{
+// 				// Select the newly pasted stuff
+// 				MetasoundGraphEditor->SetNodeSelection(Node, true);
+// 
+// 				Node->NodePosX = (Node->NodePosX - AvgNodePosition.X) + Location.X;
+// 				Node->NodePosY = (Node->NodePosY - AvgNodePosition.Y) + Location.Y;
+// 
+// 				Node->SnapToGrid(SNodePanel::GetSnapGridSize());
+// 				Node->CreateNewGuid();
+// 			}
+// 
+// 			MetasoundGraphEditor->NotifyGraphChanged();
+// 
+// 			Metasound->PostEditChange();
+// 			Metasound->MarkPackageDirty();
+		}
+
+		bool FEditor::CanPasteNodes()
 		{
-			// Copy and paste current selection
-			CopySelectedNodes();
-			PasteNodes();
-		}
+			return false;
 
-		bool FEditor::CanDuplicateNodes() const
-		{
-			return CanCopyNodes();
+// 			static_assert(false, "Delete this assert once 'PasteNodes' includes doc fix-ups as described above.");
+// 			FString ClipboardContent;
+// 			FPlatformApplicationMisc::ClipboardPaste(ClipboardContent);
+// 
+// 			FMetasoundAssetBase* AssetBase = Frontend::GetObjectAsAssetBase(Metasound);
+// 			check(AssetBase);
+// 
+// 			PastedDocument = FMetasoundDocument();
+// 			const bool bCanImportDocument = Metasound::Frontend::ImportJSONToMetasound(ClipboardContent, PastedDocument);
+// 			if (!bCanImportDocument)
+// 			{
+// 				return false;
+// 			}
+// 
+// 		 	const bool bCanPasteNodes = FEdGraphUtilities::CanImportNodesFromText(AssetBase->GetGraph(), PastedDocument.EditorData.GraphData);
+// 		 	return bCanPasteNodes;
 		}
 
 		void FEditor::UndoGraphAction()
diff --git a/Engine/Plugins/Runtime/Metasound/Source/MetasoundEditor/Private/MetasoundEditor.h b/Engine/Plugins/Runtime/Metasound/Source/MetasoundEditor/Private/MetasoundEditor.h
--- a/Engine/Plugins/Runtime/Metasound/Source/MetasoundEditor/Private/MetasoundEditor.h
+++ b/Engine/Plugins/Runtime/Metasound/Source/MetasoundEditor/Private/MetasoundEditor.h
@@ -7,6 +7,7 @@
 #include "GraphEditor.h"
 #include "IDetailsView.h"
 #include "IMetasoundEditor.h"
+#include "MetasoundFrontend.h"
 #include "Misc/NotifyHook.h"
 #include "SMetasoundPalette.h"
 #include "Textures/SlateIcon.h"
@@ -67,10 +68,9 @@ namespace Metasound
 			// End of FEditorUndoClient
 
 			/** Whether pasting the currently selected nodes is permissible */
-			bool CanPasteNodes() const;
+			bool CanPasteNodes();
 
-			/** Paste the contents of the clipboard at the provided location */
-			void PasteNodesAtLocation(const FVector2D& Location);
+			void PasteNodes(const FVector2D* Location = nullptr);
 
 			int32 GetNumNodesSelected() const
 			{
@@ -93,14 +93,9 @@ namespace Metasound
 			/** Delete the currently selected nodes */
 			void DeleteSelectedNodes();
 
-			/** Delete only the currently selected nodes that can be duplicated */
-			void DeleteSelectedDuplicatableNodes();
-
 			/** Cut the currently selected nodes */
 			void CutSelectedNodes();
 
-			/** Whether we are able to cut the currently selected nodes */
-			bool CanCutNodes() const;
 
 			/** Copy the currently selected nodes */
 			void CopySelectedNodes();
@@ -108,14 +103,8 @@ namespace Metasound
 			/** Whether copying the currently selected nodes is permissible */
 			bool CanCopyNodes() const;
 
-			/** Paste the contents of the clipboard */
-			void PasteNodes();
-
-			/** Duplicate the currently selected nodes */
-			void DuplicateNodes();
-
-			/** Whether we are able to duplicate the currently selected nodes */
-			bool CanDuplicateNodes() const;
+			/** Whether the currently selected nodes can be deleted */
+			bool CanDeleteNodes() const;
 
 			/** Called to undo the last action */
 			void UndoGraphAction();
@@ -124,6 +113,8 @@ namespace Metasound
 			void RedoGraphAction();
 
 		private:
+			static Frontend::FGraphHandle InitMetasound(UObject& InMetasound);
+
 			/** FNotifyHook interface */
 			virtual void NotifyPostChange(const FPropertyChangedEvent& PropertyChangedEvent, FProperty* PropertyThatChanged) override;
 
@@ -198,6 +189,9 @@ namespace Metasound
 
 			/** Whether or not metasound being edited is valid */
 			bool bPassedValidation = true;
+
+			/** Document used when pasting from clipboard to avoid deserializing twice */
+			FMetasoundDocument PastedDocument;
 		};
 	} // namespace Editor
 } // namespace Metasound
diff --git a/Engine/Plugins/Runtime/Metasound/Source/MetasoundEditor/Private/MetasoundEditorGraphBuilder.cpp b/Engine/Plugins/Runtime/Metasound/Source/MetasoundEditor/Private/MetasoundEditorGraphBuilder.cpp
--- a/Engine/Plugins/Runtime/Metasound/Source/MetasoundEditor/Private/MetasoundEditorGraphBuilder.cpp
+++ b/Engine/Plugins/Runtime/Metasound/Source/MetasoundEditor/Private/MetasoundEditorGraphBuilder.cpp
@@ -31,8 +31,6 @@ namespace Metasound
 		{
 			const FScopedTransaction Transaction(LOCTEXT("AddMetasoundGraphNode", "Add Metasound Node"));
 
-			const FString EdNodeName = InNodeHandle.GetNodeClassName() + TEXT("_") + FString::FromInt(InNodeHandle.GetNodeID());
-
 			FMetasoundAssetBase* MetasoundAsset = Frontend::GetObjectAsAssetBase(&InMetasound);
 			check(MetasoundAsset);
 
@@ -55,15 +53,87 @@ namespace Metasound
 		}
 
 		UEdGraphNode* FGraphBuilder::AddNode(UObject& InMetasound, const FVector2D& Location, const Frontend::FNodeClassInfo& InClassInfo, bool bInSelectNewNode)
+		{
+			Frontend::FNodeHandle NodeHandle = AddNodeHandle(InMetasound, InClassInfo);
+			return AddNode(InMetasound, Location, NodeHandle, bInSelectNewNode);
+		}
+
+		Frontend::FNodeHandle FGraphBuilder::AddNodeHandle(UObject& InMetasound, const Frontend::FNodeClassInfo& InClassInfo)
 		{
 			FMetasoundAssetBase* MetasoundAsset = Frontend::GetObjectAsAssetBase(&InMetasound);
 			check(MetasoundAsset);
 
-			Frontend::FNodeHandle NodeHandle = MetasoundAsset->GetRootGraphHandle().AddNewNode(InClassInfo);
-			return AddNode(InMetasound, Location, NodeHandle, bInSelectNewNode);
+			return MetasoundAsset->GetRootGraphHandle().AddNewNode(InClassInfo);
+		}
+
+		FString FGraphBuilder::GetDataTypeDisplayName(const FName& InDataTypeName)
+		{
+			FString CategoryString = InDataTypeName.ToString();
+			int32 Index = 0;
+			CategoryString.FindLastChar(':', Index);
+
+			return CategoryString.RightChop(Index + 1);
+		}
+
+		TArray<FString> FGraphBuilder::GetDataTypeNameCategories(const FName& InDataTypeName)
+		{
+			FString CategoryString = InDataTypeName.ToString();
+
+			TArray<FString> Categories;
+			CategoryString.ParseIntoArray(Categories, TEXT(":"));
+
+			if (Categories.Num() > 0)
+			{
+				// Remove name
+				Categories.RemoveAt(Categories.Num() - 1);
+			}
+
+			return Categories;
+		}
+
+		FString FGraphBuilder::GenerateUniqueInputName(const UObject& InMetasound, const FName InBaseName)
+		{
+			FString NameBase = GetDataTypeDisplayName(InBaseName);
+			const FMetasoundAssetBase* MetasoundAsset = Frontend::GetObjectAsAssetBase(&InMetasound);
+			check(MetasoundAsset);
+
+			const Frontend::FGraphHandle GraphHandle = MetasoundAsset->GetRootGraphHandle();
+
+			int32 i = 1;
+			FString NewNodeName = NameBase + FString::Printf(TEXT("_%02d"), i);
+			while (GraphHandle.ContainsInputNodeWithName(NewNodeName))
+			{
+				NewNodeName = NameBase + FString::Printf(TEXT("_%02d"), ++i);
+			}
+
+			return NewNodeName;
+		}
+
+		FString FGraphBuilder::GenerateUniqueOutputName(const UObject& InMetasound, const FName InBaseName)
+		{
+			FString NameBase = GetDataTypeDisplayName(InBaseName);
+			const FMetasoundAssetBase* MetasoundAsset = Frontend::GetObjectAsAssetBase(&InMetasound);
+			check(MetasoundAsset);
+
+			const Frontend::FGraphHandle GraphHandle = MetasoundAsset->GetRootGraphHandle();
+
+			int32 i = 1;
+			FString NewNodeName = NameBase + FString::Printf(TEXT("_%02d"), i);
+			while (GraphHandle.ContainsOutputNodeWithName(NewNodeName))
+			{
+				NewNodeName = NameBase + FString::Printf(TEXT("_%02d"), ++i);
+			}
+
+			return NewNodeName;
 		}
 
 		UEdGraphNode* FGraphBuilder::AddInput(UObject& InMetasound, const FVector2D& Location, const FString& InName, const FName InTypeName, const FText& InToolTip, bool bInSelectNewNode)
+		{
+			Frontend::FNodeHandle NodeHandle = AddInputNodeHandle(InMetasound, InName, InTypeName, InToolTip);
+			return AddNode(InMetasound, Location, NodeHandle, bInSelectNewNode);
+		}
+
+		Frontend::FNodeHandle FGraphBuilder::AddInputNodeHandle(UObject& InMetasound, const FString& InName, const FName InTypeName, const FText& InToolTip)
 		{
 			FMetasoundInputDescription Description;
 			Description.Name = InName;
@@ -77,7 +147,7 @@ namespace Metasound
 			Frontend::FNodeHandle NodeHandle = GraphHandle.AddNewInput(Description);
 			if (!NodeHandle.IsValid())
 			{
-				return nullptr;
+				return Frontend::FNodeHandle::InvalidHandle();
 			}
 
 			IMetasoundEditorModule& EditorModule = FModuleManager::GetModuleChecked<IMetasoundEditorModule>("MetasoundEditor");
@@ -86,7 +156,7 @@ namespace Metasound
 			Metasound::FDataTypeLiteralParam LiteralParam = Frontend::GetDefaultParamForDataType(InTypeName);
 			if (!ensureAlways(LiteralParam.IsValid()))
 			{
-				return nullptr;
+				return Frontend::FNodeHandle::InvalidHandle();
 			}
 
 			switch (LiteralParam.ConstructorArgType)
@@ -125,11 +195,16 @@ namespace Metasound
 			}
 
 			GraphHandle.SetInputDisplayName(InName, FText::FromString(InName));
+			return NodeHandle;
+		}
 
+		UEdGraphNode* FGraphBuilder::AddOutput(UObject& InMetasound, const FVector2D& Location, const FString& InName, const FName InTypeName, const FText& InToolTip, bool bInSelectNewNode)
+		{
+			Frontend::FNodeHandle NodeHandle = AddOutputNodeHandle(InMetasound, InName, InTypeName, InToolTip);
 			return AddNode(InMetasound, Location, NodeHandle, bInSelectNewNode);
 		}
 
-		UEdGraphNode* FGraphBuilder::AddOutput(UObject& InMetasound, const FVector2D& Location, const FString& InName, const FName InTypeName, const FText& InToolTip, bool bInSelectNewNode)
+		Frontend::FNodeHandle FGraphBuilder::AddOutputNodeHandle(UObject& InMetasound, const FString& InName, const FName InTypeName, const FText& InToolTip)
 		{
 			FMetasoundOutputDescription Description;
 			Description.Name = InName;
@@ -141,10 +216,10 @@ namespace Metasound
 
 			Frontend::FGraphHandle GraphHandle = MetasoundAsset->GetRootGraphHandle();
 			Frontend::FNodeHandle NodeHandle = GraphHandle.AddNewOutput(Description);
-			
+
 			GraphHandle.SetOutputDisplayName(InName, FText::FromString(InName));
 
-			return AddNode(InMetasound, Location, NodeHandle, bInSelectNewNode);
+			return NodeHandle;
 		}
 
 		void FGraphBuilder::DeleteNode(UMetasoundEditorGraphNode& InNode, bool bInRecordTransaction)
@@ -284,7 +359,10 @@ namespace Metasound
 		{
 			const FScopedTransaction Transaction(LOCTEXT("RebuildMetasoundGraphNodePins", "Rebuild Metasound Pins"), bInRecordTransaction);
 
-			InGraphNode.Pins.Reset();
+			for (int32 i = InGraphNode.Pins.Num() - 1; i >= 0; i--)
+			{
+				InGraphNode.RemovePin(InGraphNode.Pins[i]);
+			}
 
 			IMetasoundEditorModule& EditorModule = FModuleManager::GetModuleChecked<IMetasoundEditorModule>("MetasoundEditor");
 
diff --git a/Engine/Plugins/Runtime/Metasound/Source/MetasoundEditor/Private/MetasoundEditorGraphBuilder.h b/Engine/Plugins/Runtime/Metasound/Source/MetasoundEditor/Private/MetasoundEditorGraphBuilder.h
--- a/Engine/Plugins/Runtime/Metasound/Source/MetasoundEditor/Private/MetasoundEditorGraphBuilder.h
+++ b/Engine/Plugins/Runtime/Metasound/Source/MetasoundEditor/Private/MetasoundEditorGraphBuilder.h
@@ -2,6 +2,7 @@
 #pragma once
 
 #include "CoreMinimal.h"
+#include "MetasoundFrontend.h"
 
 // Forward Declarations
 class UEdGraphNode;
@@ -33,10 +34,27 @@ namespace Metasound
 			// Adds a node with the given class info to both the editor and document graphs
 			static UEdGraphNode* AddNode(UObject& InMetasound, const FVector2D& Location, const Frontend::FNodeClassInfo& InClassInfo, bool bInSelectNewNode = true);
 
+			// Adds a node handle with the given class info
+			static Frontend::FNodeHandle AddNodeHandle(UObject& InMetasound, const Frontend::FNodeClassInfo& InClassInfo);
+
+			static FString GenerateUniqueInputName(const UObject& InMetasound, const FName InBaseName);
+
+			static FString GenerateUniqueOutputName(const UObject& InMetasound, const FName InBaseName);
+
+			static FString GetDataTypeDisplayName(const FName& InDataTypeName);
+
+			static TArray<FString> GetDataTypeNameCategories(const FName& InDataTypeName);
+
 			static UEdGraphNode* AddInput(UObject& InMetasound, const FVector2D& Location, const FString& InName, const FName InTypeName, const FText& InToolTip, bool bInSelectNewNode = true);
 
+			// Adds an input node handle with the given class info
+			static Frontend::FNodeHandle AddInputNodeHandle(UObject& InMetasound, const FString& InName, const FName InTypeName, const FText& InToolTip);
+
 			static UEdGraphNode* AddOutput(UObject& InMetasound, const FVector2D& Location, const FString& InName, const FName InTypeName, const FText& InToolTip, bool bInSelectNewNode = true);
 
+			// Adds an output node handle with the given class info
+			static Frontend::FNodeHandle AddOutputNodeHandle(UObject& InMetasound, const FString& InName, const FName InTypeName, const FText& InToolTip);
+
 			static void DeleteNode(UMetasoundEditorGraphNode& InNode, bool bInRecordTransaction = true);
 
 			static void RebuildGraph(UObject& InMetasound);
diff --git a/Engine/Plugins/Runtime/Metasound/Source/MetasoundEditor/Private/MetasoundEditorGraphNode.cpp b/Engine/Plugins/Runtime/Metasound/Source/MetasoundEditor/Private/MetasoundEditorGraphNode.cpp
--- a/Engine/Plugins/Runtime/Metasound/Source/MetasoundEditor/Private/MetasoundEditorGraphNode.cpp
+++ b/Engine/Plugins/Runtime/Metasound/Source/MetasoundEditor/Private/MetasoundEditorGraphNode.cpp
@@ -111,12 +111,16 @@ void UMetasoundEditorGraphNode::IteratePins(TUniqueFunction<void(UEdGraphPin* /*
 
 void UMetasoundEditorGraphNode::AllocateDefaultPins()
 {
-	ensureAlways(Pins.Num() == 0);
-
 	using namespace Metasound;
 
+	ensureAlways(Pins.Num() == 0);
+
 	Frontend::FNodeHandle NodeHandle = Frontend::FNodeHandle::InvalidHandle();
-	Editor::FGraphBuilder::RebuildNodePins(*this, NodeHandle, false /* bInRecordTransaction */);
+	if (NodeID != INDEX_NONE)
+	{
+		NodeHandle = GetNodeHandle();
+	}
+	Metasound::Editor::FGraphBuilder::RebuildNodePins(*this, NodeHandle, false /* bInRecordTransaction */);
 }
 
 void UMetasoundEditorGraphNode::ReconstructNode()
@@ -124,6 +128,10 @@ void UMetasoundEditorGraphNode::ReconstructNode()
 	using namespace Metasound;
 
 	Frontend::FNodeHandle NodeHandle = Frontend::FNodeHandle::InvalidHandle();
+	if (NodeID != INDEX_NONE)
+	{
+		NodeHandle = GetNodeHandle();
+	}
 	Editor::FGraphBuilder::RebuildNodePins(*this, NodeHandle);
 }
 
@@ -171,6 +179,23 @@ bool UMetasoundEditorGraphNode::CanCreateUnderSpecifiedSchema(const UEdGraphSche
 	return Schema->IsA(UMetasoundEditorGraphSchema::StaticClass());
 }
 
+bool UMetasoundEditorGraphNode::CanUserDeleteNode() const
+{
+	const FString& NodeName = GetNodeHandle().GetNodeName();
+	Metasound::Frontend::FGraphHandle GraphHandle = GetRootGraphHandle();
+	if (GraphHandle.IsRequiredInput(NodeName))
+	{
+		return false;
+	}
+
+	if (GraphHandle.IsRequiredOutput(NodeName))
+	{
+		return false;
+	}
+
+	return true;
+}
+
 FString UMetasoundEditorGraphNode::GetDocumentationLink() const
 {
 	return TEXT("Shared/GraphNodes/Metasound");
@@ -181,6 +206,11 @@ void UMetasoundEditorGraphNode::SetNodeID(uint32 InNodeID)
 	NodeID = InNodeID;
 }
 
+uint32 UMetasoundEditorGraphNode::GetNodeID() const
+{
+	return NodeID;
+}
+
 FText UMetasoundEditorGraphNode::GetNodeTitle(ENodeTitleType::Type TitleType) const
 {
 	using namespace Metasound::Frontend;
@@ -216,10 +246,6 @@ void UMetasoundEditorGraphNode::PrepareForCopying()
 {
 }
 
-void UMetasoundEditorGraphNode::PostCopyNode()
-{
-}
-
 void UMetasoundEditorGraphNode::PostEditImport()
 {
 }
diff --git a/Engine/Plugins/Runtime/Metasound/Source/MetasoundEditor/Private/MetasoundEditorGraphSchema.cpp b/Engine/Plugins/Runtime/Metasound/Source/MetasoundEditor/Private/MetasoundEditorGraphSchema.cpp
--- a/Engine/Plugins/Runtime/Metasound/Source/MetasoundEditor/Private/MetasoundEditorGraphSchema.cpp
+++ b/Engine/Plugins/Runtime/Metasound/Source/MetasoundEditor/Private/MetasoundEditorGraphSchema.cpp
@@ -33,31 +33,6 @@ namespace Metasound
 {
 	namespace Editor
 	{
-		FString GetDataTypeDisplayName(const FName& InDataTypeName)
-		{
-			FString CategoryString = InDataTypeName.ToString();
-			int32 Index = 0;
-			CategoryString.FindLastChar(':', Index);
-
-			return CategoryString.RightChop(Index + 1);
-		}
-
-		TArray<FString> GetDataTypeNameCategories(const FName& InDataTypeName)
-		{
-			FString CategoryString = InDataTypeName.ToString();
-
-			TArray<FString> Categories;
-			CategoryString.ParseIntoArray(Categories, TEXT(":"));
-
-			if (Categories.Num() > 0)
-			{
-				// Remove name
-				Categories.RemoveAt(Categories.Num() - 1);
-			}
-
-			return Categories;
-		}
-
 		TSharedPtr<FEditor> GetEditorForGraph(const UObject& Metasound)
 		{
 			TSharedPtr<IToolkit> FoundAssetEditor = FToolkitManager::Get().FindEditorForAsset(CastChecked<const UObject>(&Metasound));
@@ -188,7 +163,6 @@ FMetasoundGraphSchemaAction_NewInput::FMetasoundGraphSchemaAction_NewInput(FText
 UEdGraphNode* FMetasoundGraphSchemaAction_NewInput::PerformAction(UEdGraph* ParentGraph, UEdGraphPin* FromPin, const FVector2D Location, bool bSelectNewNode /* = true */)
 {
 	using namespace Metasound::Editor;
-	using namespace Metasound::Frontend;
 
 	const FScopedTransaction Transaction(LOCTEXT("MetasoundEditorNewInput", "Add New Metasound Input"));
 
@@ -197,20 +171,7 @@ UEdGraphNode* FMetasoundGraphSchemaAction_NewInput::PerformAction(UEdGraph* Pare
 	UObject& ParentMetasound = CastChecked<UMetasoundEditorGraph>(ParentGraph)->GetMetasoundChecked();
 	ParentMetasound.Modify();
 
-	FString NameBase = GetDataTypeDisplayName(NodeTypeName);
-
-	FMetasoundAssetBase* MetasoundAsset = GetObjectAsAssetBase(&ParentMetasound);
-	check(MetasoundAsset);
-
-	FGraphHandle GraphHandle = MetasoundAsset->GetRootGraphHandle();
-
-	int32 i = 1;
-	FString NewNodeName = NameBase + FString::Printf(TEXT("_%02d"), i);
-	while (GraphHandle.ContainsInputNodeWithName(NewNodeName))
-	{
-		NewNodeName = NameBase + FString::Printf(TEXT("_%02d"), ++i);
-	}
-
+	FString NewNodeName = FGraphBuilder::GenerateUniqueInputName(ParentMetasound, NodeTypeName);
 	UEdGraphNode* NewGraphNode = FGraphBuilder::AddInput(ParentMetasound, Location, NewNodeName, NodeTypeName, FText::GetEmpty());
 	return NewGraphNode;
 }
@@ -224,7 +185,6 @@ FMetasoundGraphSchemaAction_NewOutput::FMetasoundGraphSchemaAction_NewOutput(FTe
 UEdGraphNode* FMetasoundGraphSchemaAction_NewOutput::PerformAction(UEdGraph* ParentGraph, UEdGraphPin* FromPin, const FVector2D Location, bool bSelectNewNode /* = true */)
 {
 	using namespace Metasound::Editor;
-	using namespace Metasound::Frontend;
 
 	const FScopedTransaction Transaction(LOCTEXT("MetasoundEditorNewOutput", "Add New Metasound Output"));
 
@@ -233,19 +193,7 @@ UEdGraphNode* FMetasoundGraphSchemaAction_NewOutput::PerformAction(UEdGraph* Par
 	UObject& ParentMetasound = CastChecked<UMetasoundEditorGraph>(ParentGraph)->GetMetasoundChecked();
 	ParentMetasound.Modify();
 
-	FString NameBase = GetDataTypeDisplayName(NodeTypeName);
-
-	FMetasoundAssetBase* MetasoundAsset = GetObjectAsAssetBase(&ParentMetasound);
-	check(MetasoundAsset);
-
-	FGraphHandle GraphHandle = MetasoundAsset->GetRootGraphHandle();
-
-	int32 i = 1;
-	FString NewNodeName = NameBase + FString::Printf(TEXT("_%02d"), i);
-	while (GraphHandle.ContainsOutputNodeWithName(NewNodeName))
-	{
-		NewNodeName = NameBase + FString::Printf(TEXT("_%02d"), ++i);
-	}
+	FString NewNodeName = FGraphBuilder::GenerateUniqueInputName(ParentMetasound, NodeTypeName);
 
 	UEdGraphNode* NewGraphNode = FGraphBuilder::AddOutput(ParentMetasound, Location, NewNodeName, NodeTypeName, FText::GetEmpty());
 	return NewGraphNode;
@@ -289,7 +237,7 @@ UEdGraphNode* FMetasoundGraphSchemaAction_Paste::PerformAction(UEdGraph* ParentG
 	TSharedPtr<FEditor> MetasoundEditor = GetEditorForGraph(*ParentGraph);
 	if (MetasoundEditor.IsValid())
 	{
-		MetasoundEditor->PasteNodesAtLocation(Location);
+		MetasoundEditor->PasteNodes(&Location);
 	}
 
 	return nullptr;
@@ -594,11 +542,11 @@ void UMetasoundEditorGraphSchema::GetAllMetasoundActions(FGraphActionMenuBuilder
 	EditorModule.IterateDataTypes([InputMenuName, OutputMenuName, InMenuBuilder = &ActionMenuBuilder](const FEditorDataType& DataType)
 	{
 		const FName DataTypeName = DataType.RegistryInfo.DataTypeName;
-		const FText DataTypeDisplayName = FText::FromString(GetDataTypeDisplayName(DataTypeName));
+		const FText DataTypeDisplayName = FText::FromString(FGraphBuilder::GetDataTypeDisplayName(DataTypeName));
 		const FText DataTypeTextName = FText::FromName(DataTypeName);
 		const FText MenuJoinFormat = LOCTEXT("MetasoundFormatNodeSubCategory", "{0}|{1}");
 
-		const TArray<FString> Categories = GetDataTypeNameCategories(DataTypeName);
+		const TArray<FString> Categories = FGraphBuilder::GetDataTypeNameCategories(DataTypeName);
 		const FText CategoriesText = FText::FromString(FString::Join(Categories, TEXT("|")));
 
 
diff --git a/Engine/Plugins/Runtime/Metasound/Source/MetasoundEditor/Public/MetasoundEditorGraphNode.h b/Engine/Plugins/Runtime/Metasound/Source/MetasoundEditor/Public/MetasoundEditorGraphNode.h
--- a/Engine/Plugins/Runtime/Metasound/Source/MetasoundEditor/Public/MetasoundEditorGraphNode.h
+++ b/Engine/Plugins/Runtime/Metasound/Source/MetasoundEditor/Public/MetasoundEditorGraphNode.h
@@ -23,9 +23,6 @@ class UMetasoundEditorGraphNode : public UEdGraphNode
 	GENERATED_UCLASS_BODY()
 
 public:
-	/** Fix up the node's owner after being copied */
-	METASOUNDEDITOR_API void PostCopyNode();
-
 	/** Create a new input pin for this node */
 	METASOUNDEDITOR_API void CreateInputPin();
 
@@ -38,6 +35,7 @@ public:
 	virtual void AllocateDefaultPins() override;
 	virtual void AutowireNewNode(UEdGraphPin* FromPin) override;
 	virtual bool CanCreateUnderSpecifiedSchema(const UEdGraphSchema* Schema) const override;
+	virtual bool CanUserDeleteNode() const override;
 	virtual FText GetNodeTitle(ENodeTitleType::Type TitleType) const override;
 	virtual FString GetDocumentationExcerptName() const override;
 	virtual FString GetDocumentationLink() const override;
@@ -59,7 +57,9 @@ public:
 
 	Metasound::Frontend::FGraphHandle GetRootGraphHandle() const;
 	Metasound::Frontend::FNodeHandle GetNodeHandle() const;
+
 	void SetNodeID(uint32 InNodeID);
+	uint32 GetNodeID() const;
 
 protected:
 	UPROPERTY()
diff --git a/Engine/Plugins/Runtime/Metasound/Source/MetasoundEngine/Public/Metasound.h b/Engine/Plugins/Runtime/Metasound/Source/MetasoundEngine/Public/Metasound.h
--- a/Engine/Plugins/Runtime/Metasound/Source/MetasoundEngine/Public/Metasound.h
+++ b/Engine/Plugins/Runtime/Metasound/Source/MetasoundEngine/Public/Metasound.h
@@ -29,7 +29,7 @@ class METASOUNDENGINE_API UMetasound : public UObject, public FMetasoundAssetBas
 	GENERATED_BODY()
 
 protected:
-	UPROPERTY(EditAnywhere, Category = Hidden)
+	UPROPERTY(EditAnywhere, Category = CustomView)
 	FMetasoundDocument RootMetasoundDocument;
 
 #if WITH_EDITORONLY_DATA
diff --git a/Engine/Plugins/Runtime/Metasound/Source/MetasoundEngine/Public/MetasoundSource.h b/Engine/Plugins/Runtime/Metasound/Source/MetasoundEngine/Public/MetasoundSource.h
--- a/Engine/Plugins/Runtime/Metasound/Source/MetasoundEngine/Public/MetasoundSource.h
+++ b/Engine/Plugins/Runtime/Metasound/Source/MetasoundEngine/Public/MetasoundSource.h
@@ -26,7 +26,7 @@ class METASOUNDENGINE_API UMetasoundSource : public USoundWaveProcedural, public
 	GENERATED_BODY()
 
 protected:
-	UPROPERTY(EditAnywhere, Category = Hidden)
+	UPROPERTY(EditAnywhere, Category = CustomView)
 	FMetasoundDocument RootMetasoundDocument;
 
 #if WITH_EDITORONLY_DATA
diff --git a/Engine/Plugins/Runtime/Metasound/Source/MetasoundFrontend/Private/MetasoundAssetBase.cpp b/Engine/Plugins/Runtime/Metasound/Source/MetasoundFrontend/Private/MetasoundAssetBase.cpp
--- a/Engine/Plugins/Runtime/Metasound/Source/MetasoundFrontend/Private/MetasoundAssetBase.cpp
+++ b/Engine/Plugins/Runtime/Metasound/Source/MetasoundFrontend/Private/MetasoundAssetBase.cpp
@@ -91,22 +91,12 @@ TArray<Metasound::Frontend::FGraphHandle> FMetasoundAssetBase::GetAllSubgraphHan
 	return OutArray;
 }
 
-bool FMetasoundAssetBase::ImportFromJSON(const FString& InAbsolutePath)
+bool FMetasoundAssetBase::ImportFromJSON(const FString& InJSON)
 {
-	return Metasound::Frontend::ImportJSONToMetasound(InAbsolutePath, GetDocument());
+	return Metasound::Frontend::ImportJSONToMetasound(InJSON, GetDocument());
 }
 
-bool FMetasoundAssetBase::ExportToJSON(const FString& InAbsolutePath) const
+bool FMetasoundAssetBase::ImportFromJSONAsset(const FString& InAbsolutePath)
 {
-	if (TUniquePtr<FArchive> FileWriter = TUniquePtr<FArchive>(IFileManager::Get().CreateFileWriter(*InAbsolutePath)))
-	{
-		Metasound::TJsonStructSerializerBackend<Metasound::DefaultCharType> Backend(*FileWriter, EStructSerializerBackendFlags::Default);
-		FStructSerializer::Serialize<FMetasoundDocument>(GetDocument(), Backend);
-		FileWriter->Close();
-
-		return true;
-	}
-
-	ensureAlwaysMsgf(false, TEXT("Failed to create a filewriter with the given path."));
-	return false;
-}
\ No newline at end of file
+	return Metasound::Frontend::ImportJSONAssetToMetasound(InAbsolutePath, GetDocument());
+}
diff --git a/Engine/Plugins/Runtime/Metasound/Source/MetasoundFrontend/Private/MetasoundFrontend.cpp b/Engine/Plugins/Runtime/Metasound/Source/MetasoundFrontend/Private/MetasoundFrontend.cpp
--- a/Engine/Plugins/Runtime/Metasound/Source/MetasoundFrontend/Private/MetasoundFrontend.cpp
+++ b/Engine/Plugins/Runtime/Metasound/Source/MetasoundFrontend/Private/MetasoundFrontend.cpp
@@ -2,14 +2,14 @@
 
 #include "MetasoundFrontend.h"
 
-#include "MetasoundJsonBackend.h"
-#include "StructSerializer.h"
-#include "StructDeserializer.h"
 #include "HAL/FileManager.h"
 #include "HAL/IConsoleManager.h"
+#include "MetasoundJsonBackend.h"
+#include "MetasoundOperatorBuilder.h"
 #include "Modules/ModuleManager.h"
+#include "StructDeserializer.h"
+#include "StructSerializer.h"
 
-#include "MetasoundOperatorBuilder.h"
 
 static int32 MetasoundUndoRollLimitCvar = 128;
 FAutoConsoleVariableRef CVarMetasoundUndoRollLimit(
@@ -213,7 +213,21 @@ namespace Metasound
 			return FMetasoundFrontendRegistryContainer::Get()->GetInfoForDataType(InDataType, OutInfo);
 		}
 
-		bool ImportJSONToMetasound(const FString& InPath, FMetasoundDocument& OutMetasoundDocument)
+		bool ImportJSONToMetasound(const FString& InJSON, FMetasoundDocument& OutMetasoundDocument)
+		{
+			TArray<uint8> ReadBuffer;
+			ReadBuffer.SetNumUninitialized(InJSON.Len() * sizeof(ANSICHAR));
+			FMemory::Memcpy(ReadBuffer.GetData(), StringCast<ANSICHAR>(*InJSON).Get(), InJSON.Len() * sizeof(ANSICHAR));
+			FMemoryReader MemReader(ReadBuffer);
+
+			TJsonStructDeserializerBackend<DefaultCharType> Backend(MemReader);
+			bool DeserializeResult = FStructDeserializer::Deserialize(OutMetasoundDocument, Backend);
+
+			MemReader.Close();
+			return DeserializeResult && !MemReader.IsError();
+		}
+
+		bool ImportJSONAssetToMetasound(const FString& InPath, FMetasoundDocument& OutMetasoundDocument)
 		{
 			if (TUniquePtr<FArchive> FileReader = TUniquePtr<FArchive>(IFileManager::Get().CreateFileReader(*InPath)))
 			{
@@ -1307,6 +1321,63 @@ namespace Metasound
 			return GraphPtr.IsValid() && GraphsClassDeclaration.IsValid();
 		}
 
+		void FGraphHandle::CopyGraph(FGraphHandle& InOther)
+		{
+			*InOther.GraphPtr.Get() = *GraphPtr.Get();
+			*InOther.GraphsClassDeclaration.Get() = *GraphsClassDeclaration.Get();
+			*InOther.OwningDocument.Get() = *OwningDocument.Get();
+		}
+
+		bool FGraphHandle::IsRequiredInput(const FString& InInputName) const
+		{
+			const TArray<FMetasoundInputDescription>& RequiredInputs = OwningDocument->Archetype.RequiredInputs;
+			for (const FMetasoundInputDescription& RequiredInput : RequiredInputs)
+			{
+				if (InInputName == RequiredInput.Name)
+				{
+					return true;
+				}
+			}
+
+			return false;
+		}
+
+		bool FGraphHandle::IsRequiredOutput(const FString& InOutputName) const
+		{
+			const TArray<FMetasoundOutputDescription>& RequiredOutputs = OwningDocument->Archetype.RequiredOutputs;
+			for (const FMetasoundOutputDescription& RequiredOutput : RequiredOutputs)
+			{
+				if (InOutputName == RequiredOutput.Name)
+				{
+					return true;
+				}
+			}
+
+			return false;
+		}
+
+		const TArray<FMetasoundInputDescription>& FGraphHandle::GetRequiredInputs() const
+		{
+			static const TArray<FMetasoundInputDescription> InvalidDesc;
+			if (!OwningDocument.IsValid())
+			{
+				return InvalidDesc;
+			}
+
+			return OwningDocument->Archetype.RequiredInputs;
+		}
+
+		const TArray<FMetasoundOutputDescription>& FGraphHandle::GetRequiredOutputs() const
+		{
+			static const TArray<FMetasoundOutputDescription> InvalidDesc;
+			if (!OwningDocument.IsValid())
+			{
+				return InvalidDesc;
+			}
+
+			return OwningDocument->Archetype.RequiredOutputs;
+		}
+
 		TArray<FNodeHandle> FGraphHandle::GetAllNodes()
 		{
 			TArray<FNodeHandle> OutArray;
@@ -1526,7 +1597,7 @@ namespace Metasound
 			return FNodeHandle(FHandleInitParams::PrivateToken, InitParams, NewNodeDescription.ObjectTypeOfNode);
 		}
 
-		bool FGraphHandle::RemoveInput(const FString& InputName)
+		bool FGraphHandle::RemoveInput(const FString& InInputName)
 		{
 			if (!IsValid())
 			{
@@ -1538,26 +1609,26 @@ namespace Metasound
 
 			for (int32 InputIndex = 0; InputIndex < Inputs.Num(); InputIndex++)
 			{
-				if (Inputs[InputIndex].Name == InputName)
+				if (Inputs[InputIndex].Name == InInputName)
 				{
 					IndexOfInputToRemove = InputIndex;
 					break;
 				}
 			}
 
-			if (!ensureAlwaysMsgf(IndexOfInputToRemove >= 0, TEXT("Tried to remove an Input that didn't exist: %s"), *InputName))
+			if (!ensureAlwaysMsgf(IndexOfInputToRemove >= 0, TEXT("Tried to remove an Input that didn't exist: %s"), *InInputName))
 			{
 				return false;
 			}
 
 			// find the corresponding node handle to delete.
-			FNodeHandle InputNode = GetInputNodeWithName(InputName);
+			FNodeHandle InputNode = GetInputNodeWithName(InInputName);
 
 			// If we found the input declared in the class description but couldn't find it in the graph,
 			// something has gone terribly wrong. Remove the input from the description, but still ensure.
 			if (!ensureAlwaysMsgf(InputNode.IsValid(), TEXT(R"(Couldn't find an input node with name %s, even though we found the input listed as a dependency.
 				This indicates the underlying FMetasoundClassDescription is corrupted.
-				Removing the Input in the class dependency to resolve...)"), *InputName))
+				Removing the Input in the class dependency to resolve...)"), *InInputName))
 			{
 				Inputs.RemoveAt(IndexOfInputToRemove);
 				return true;
@@ -1573,6 +1644,16 @@ namespace Metasound
 			return true;
 		}
 
+		const FMetasoundEditorData& FGraphHandle::GetEditorData() const
+		{
+			return OwningDocument->EditorData;
+		}
+
+		void FGraphHandle::SetEditorData(const FMetasoundEditorData& InEditorData)
+		{
+			OwningDocument->EditorData = InEditorData;
+		}
+
 		FNodeHandle FGraphHandle::AddNewOutput(const FMetasoundOutputDescription& InDescription)
 		{
 			if (!IsValid())
@@ -1990,6 +2071,24 @@ namespace Metasound
 			return GraphsClassDeclaration->Metadata;
 		}
 
+		FString FGraphHandle::ExportToJSON() const
+		{
+			TArray<uint8> WriterBuffer;
+			FMemoryWriter MemWriter(WriterBuffer);
+
+			Metasound::TJsonStructSerializerBackend<Metasound::DefaultCharType> Backend(MemWriter, EStructSerializerBackendFlags::Default);
+			FStructSerializer::Serialize<FMetasoundDocument>(*OwningDocument, Backend);
+
+			MemWriter.Close();
+
+			// null terminator
+			WriterBuffer.AddZeroed(sizeof(ANSICHAR));
+
+			FString Output;
+			Output.AppendChars(reinterpret_cast<ANSICHAR*>(WriterBuffer.GetData()), WriterBuffer.Num() / sizeof(ANSICHAR));
+			return Output;
+		}
+
 		bool FGraphHandle::ExportToJSONAsset(const FString& InAbsolutePath) const
 		{
 			if (!IsValid())
@@ -1997,7 +2096,6 @@ namespace Metasound
 				return false;
 			}
 
-			
 			if (TUniquePtr<FArchive> FileWriter = TUniquePtr<FArchive>(IFileManager::Get().CreateFileWriter(*InAbsolutePath)))
 			{
 				
diff --git a/Engine/Plugins/Runtime/Metasound/Source/MetasoundFrontend/Public/MetasoundAssetBase.h b/Engine/Plugins/Runtime/Metasound/Source/MetasoundFrontend/Public/MetasoundAssetBase.h
--- a/Engine/Plugins/Runtime/Metasound/Source/MetasoundFrontend/Public/MetasoundAssetBase.h
+++ b/Engine/Plugins/Runtime/Metasound/Source/MetasoundFrontend/Public/MetasoundAssetBase.h
@@ -45,11 +45,11 @@ public:
 	// Returns the root class metadata
 	FMetasoundClassMetadata GetMetadata();
 
-	// Imports the asset from a JSON file at provided path
-	bool ImportFromJSON(const FString& InAbsolutePath);
+	// Imports data from a JSON string directly
+	bool ImportFromJSON(const FString& InJSON);
 
-	// Exports the asset to JSON file at provided path
-	bool ExportToJSON(const FString& InAbsolutePath) const;
+	// Imports the asset from a JSON file at provided path
+	bool ImportFromJSONAsset(const FString& InAbsolutePath);
 
 	// Returns handle for the root metasound graph of this asset.
 	Metasound::Frontend::FGraphHandle GetRootGraphHandle() const;
diff --git a/Engine/Plugins/Runtime/Metasound/Source/MetasoundFrontend/Public/MetasoundFrontend.h b/Engine/Plugins/Runtime/Metasound/Source/MetasoundFrontend/Public/MetasoundFrontend.h
--- a/Engine/Plugins/Runtime/Metasound/Source/MetasoundFrontend/Public/MetasoundFrontend.h
+++ b/Engine/Plugins/Runtime/Metasound/Source/MetasoundFrontend/Public/MetasoundFrontend.h
@@ -86,10 +86,13 @@ namespace Metasound
 		// returns false if InDataType couldn't be found.
 		METASOUNDFRONTEND_API bool GetTraitsForDataType(FName InDataType, FDataTypeRegistryInfo& OutInfo);
 
-		// This will open a json document at the given absolute path and convert it into a 
-		// metasound document struct. 
+		// Takes a JSON string and deserializes it into a Metasound document struct.
 		// @returns false if the file couldn't be found or parsed into a document.
-		METASOUNDFRONTEND_API bool ImportJSONToMetasound(const FString& InPath, FMetasoundDocument& OutMetasoundDocument);
+		METASOUNDFRONTEND_API bool ImportJSONToMetasound(const FString& InJSON, FMetasoundDocument& OutMetasoundDocument);
+
+		// Opens a json document at the given absolute path and deserializes it into a Metasound document struct.
+		// @returns false if the file couldn't be found or parsed into a document.
+		METASOUNDFRONTEND_API bool ImportJSONAssetToMetasound(const FString& InPath, FMetasoundDocument& OutMetasoundDocument);
 
 		struct METASOUNDFRONTEND_API FMetasoundArchetypeRegistryParams_Internal
 		{
@@ -117,7 +120,7 @@ namespace Metasound
 		METASOUNDFRONTEND_API TArray<FName> GetAllRegisteredArchetypes();
 
 		// Returns a new UObject, whose class corresponds to the archetype in the document.
-		// @param InDocument a fully generated metasound document, typically retrieved from ImportJSONToMetasound().
+		// @param InDocument a fully generated metasound document, typically retrieved from ImportJSONAssetToMetasound().
 		// @param InPath, path in content directory to save the generated UAsset to (ex. "/game/MyDir/MyMetasoundAsset".
 		//                if InPath is invalid, we won't save to an asset.
 		// @returns nullptr if we couldn't find the archetype.
@@ -347,8 +350,19 @@ namespace Metasound
 
 			static FGraphHandle InvalidHandle();
 
+			void CopyGraph(FGraphHandle& InOther);
+
 			bool IsValid() const;
 
+			bool IsRequiredInput(const FString& InOutputName) const;
+			bool IsRequiredOutput(const FString& InOutputName) const;
+
+			const TArray<FMetasoundInputDescription>& GetRequiredInputs() const;
+			const TArray<FMetasoundOutputDescription>& GetRequiredOutputs() const;
+
+			const FMetasoundEditorData& GetEditorData() const;
+			void SetEditorData(const FMetasoundEditorData& InEditorData);
+
 			TArray<FNodeHandle> GetAllNodes();
 			FNodeHandle GetNodeWithId(uint32 InNodeId) const;
 			TArray<FNodeHandle> GetOutputNodes();
@@ -410,6 +424,10 @@ namespace Metasound
 			// Returns the metadata for the current graph, including the name, description and author.
 			FMetasoundClassMetadata GetGraphMetadata();
 
+			// Exports this graph to a JSON string.
+			// @returns JSON string.
+			FString ExportToJSON() const;
+
 			// Exports this graph to a JSON at the given path.
 			// @returns true on success.
 			bool ExportToJSONAsset(const FString& InAbsolutePath) const;
diff --git a/Engine/Plugins/Runtime/Metasound/Source/MetasoundFrontend/Public/MetasoundFrontendDataLayout.h b/Engine/Plugins/Runtime/Metasound/Source/MetasoundFrontend/Public/MetasoundFrontendDataLayout.h
--- a/Engine/Plugins/Runtime/Metasound/Source/MetasoundFrontend/Public/MetasoundFrontendDataLayout.h
+++ b/Engine/Plugins/Runtime/Metasound/Source/MetasoundFrontend/Public/MetasoundFrontendDataLayout.h
@@ -94,7 +94,7 @@ struct FMetasoundInputDescription
 	GENERATED_BODY()
 
 	// The name of this input.
-	UPROPERTY()
+	UPROPERTY(VisibleAnywhere, Category = CustomView)
 	FString Name;
 
 	// The name of this input displayed to users.
@@ -121,7 +121,7 @@ struct FMetasoundOutputDescription
 	GENERATED_BODY()
 
 	// The descriptive name of this output.
-	UPROPERTY()
+	UPROPERTY(VisibleAnywhere, Category = CustomView)
 	FString Name;
 
 	// The name of this input displayed to users.
@@ -216,6 +216,15 @@ struct FMetasoundExternalClassLookupInfo
 	uint32 ExternalNodeClassHash;
 };
 
+USTRUCT()
+struct FMetasoundEditorData
+{
+	GENERATED_BODY()
+
+	UPROPERTY()
+	FString GraphData;
+};
+
 // Full saved out description of a Metasound object,
 // which is anything that could be used as a node in a graph.
 // It is not called "FMetasoundNodeDescription" to avoid confusion with the individual node objects that compose
@@ -232,13 +241,13 @@ struct FMetasoundClassDescription
 	UPROPERTY()
 	uint32 UniqueID;
 
-	UPROPERTY(EditAnywhere, Category = Hidden)
+	UPROPERTY(EditAnywhere, Category = CustomView)
 	FMetasoundClassMetadata Metadata;
 
-	UPROPERTY(EditAnywhere, Category = Hidden)
+	UPROPERTY(EditAnywhere, Category = CustomView)
 	TArray<FMetasoundInputDescription> Inputs;
 
-	UPROPERTY(EditAnywhere, Category = Hidden)
+	UPROPERTY(EditAnywhere, Category = CustomView)
 	TArray<FMetasoundOutputDescription> Outputs;
 
 	// If this object is itself of the Metasound type, here we list the unique IDs of other objects it depends on to be built.
@@ -291,10 +300,13 @@ struct FMetasoundDocument
 
 	// The highest level Metasound Class in this document.
 	// Typically of type MetasoundGraph, contains graph of the node classes listed in Dependencies.
-	UPROPERTY(EditAnywhere, Category = Hidden)
+	UPROPERTY(EditAnywhere, Category = CustomView)
 	FMetasoundClassDescription RootClass;
 
 	// this is the list of every dependency required by the RootClass, as well as nested dependencies.
 	UPROPERTY()
 	TArray<FMetasoundClassDescription> Dependencies;
+
+	UPROPERTY()
+	FMetasoundEditorData EditorData;
 };
\ No newline at end of file
