Add module name to Registration Macro and move Wave types to MetaaoundEngine

diff --git a/Engine/Plugins/Runtime/Metasound/Source/MetasoundEngine/MetasoundEngine.Build.cs b/Engine/Plugins/Runtime/Metasound/Source/MetasoundEngine/MetasoundEngine.Build.cs
--- a/Engine/Plugins/Runtime/Metasound/Source/MetasoundEngine/MetasoundEngine.Build.cs
+++ b/Engine/Plugins/Runtime/Metasound/Source/MetasoundEngine/MetasoundEngine.Build.cs
@@ -15,7 +15,8 @@ namespace UnrealBuildTool.Rules
 				{
 					"Core",
 					"MetasoundGraphCore",
-					"MetasoundGenerator"
+					"MetasoundGenerator",
+					"AudioExtensions"
 				}
 			);
 
diff --git a/Engine/Plugins/Runtime/Metasound/Source/MetasoundEngine/Private/MetasoundEngineModule.cpp b/Engine/Plugins/Runtime/Metasound/Source/MetasoundEngine/Private/MetasoundEngineModule.cpp
--- a/Engine/Plugins/Runtime/Metasound/Source/MetasoundEngine/Private/MetasoundEngineModule.cpp
+++ b/Engine/Plugins/Runtime/Metasound/Source/MetasoundEngine/Private/MetasoundEngineModule.cpp
@@ -5,10 +5,16 @@
 #include "MetasoundArchetypeRegistration.h"
 #include "MetasoundSource.h"
 #include "Modules/ModuleManager.h"
-
+#include "MetasoundWave.h"
+#include "MetasoundDataTypeRegistrationMacro.h"
 
 DEFINE_LOG_CATEGORY(LogMetasoundEngine);
 
+namespace Metasound
+{
+	REGISTER_METASOUND_DATATYPE(FWave, "Primitive:Wave");
+}
+
 class FMetasoundEngineModule : public IMetasoundEngineModule
 {
 	virtual void StartupModule() override
diff --git a/Engine/Plugins/Runtime/Metasound/Source/MetasoundEngine/Private/MetasoundWave.cpp b/Engine/Plugins/Runtime/Metasound/Source/MetasoundEngine/Private/MetasoundWave.cpp
new file mode 100644
--- /dev/null
+++ b/Engine/Plugins/Runtime/Metasound/Source/MetasoundEngine/Private/MetasoundWave.cpp
@@ -0,0 +1,24 @@
+// Copyright Epic Games, Inc. All Rights Reserved.
+
+#include "MetasoundWave.h"
+#include "MetasoundPrimitives.h"
+#include "IAudioCodec.h"
+
+namespace Metasound
+{	
+	// Simple for now, just wrap around the TArray<uint8>
+	class FWaveDecoderInput : public Audio::FDecoderInputArrayView
+	{
+		FWaveReadRef WaveRef;	// These are not thread safe, and will need to be
+	public:
+		FWaveDecoderInput(const FWaveReadRef& InWaveRef)
+			: Audio::FDecoderInputArrayView(MakeArrayView(InWaveRef->CompressedBytes.GetData(), InWaveRef->CompressedBytes.Num()), 0)
+			, WaveRef(InWaveRef)
+		{}
+	};
+
+	FWave::FDecoderInputPtr FWave::CreateDecoderInput(const FWaveReadRef& InWaveRef) 
+	{
+		return MakeShared<FWaveDecoderInput, ESPMode::ThreadSafe>(InWaveRef);
+	}
+}
diff --git a/Engine/Plugins/Runtime/Metasound/Source/MetasoundEngine/Private/MetasoundWavePlayerNode.cpp b/Engine/Plugins/Runtime/Metasound/Source/MetasoundEngine/Private/MetasoundWavePlayerNode.cpp
new file mode 100644
--- /dev/null
+++ b/Engine/Plugins/Runtime/Metasound/Source/MetasoundEngine/Private/MetasoundWavePlayerNode.cpp
@@ -0,0 +1,198 @@
+// Copyright Epic Games, Inc. All Rights Reserved.
+
+#include "MetasoundWavePlayerNode.h"
+#include "MetasoundExecutableOperator.h"
+#include "MetasoundPrimitives.h"
+#include "MetasoundWave.h"
+#include "MetasoundNodeRegistrationMacro.h"
+#include "MetasoundBuildError.h"
+
+#include "IAudioCodecRegistry.h"
+#include "IAudioCodec.h"
+
+#define LOCTEXT_NAMESPACE "MetasoundWaveNode"
+
+namespace Metasound
+{
+	METASOUND_REGISTER_NODE(FWavePlayerNode)
+	
+	// WavePlayer custom error 
+	class FWavePlayerError : public FBuildErrorBase
+	{
+	public:
+		FWavePlayerError(FName InName, FText InErrorDescription)
+			: FBuildErrorBase(InName, InErrorDescription)
+		{}
+		virtual ~FWavePlayerError() = default;
+	};
+	
+	class FWavePlayerOperator : public TExecutableOperator<FWavePlayerOperator>
+	{
+	public:
+
+		// Silent setup
+		FWavePlayerOperator(
+			const FOperatorSettings& InSettings, 
+			const FWaveReadRef& InWave )
+			: OperatorSettings(InSettings)
+			, Wave(InWave)
+			, AudioBuffer(FAudioBufferWriteRef::CreateNew(InSettings))
+		{
+			check(AudioBuffer->Num() == InSettings.GetNumFramesPerBlock());
+			OutputDataReferences.AddDataReadReference(TEXT("Audio"), FAudioBufferReadRef(AudioBuffer));
+		}
+		
+		FWavePlayerOperator(
+			const FOperatorSettings& InSettings,
+			const FWaveReadRef& InWave, 
+			FWave::FDecoderInputPtr&& InDecoderInput, 
+			TUniquePtr<Audio::IDecoderOutput>&& InDecoderOutput,
+			TUniquePtr<Audio::IDecoder>&& InDecoder )
+			: OperatorSettings(InSettings)
+			, Wave(InWave)
+			, AudioBuffer(FAudioBufferWriteRef::CreateNew(InSettings))
+			, Decoder(MoveTemp(InDecoder))
+			, DecoderInput(MoveTemp(InDecoderInput))
+			, DecoderOutput(MoveTemp(InDecoderOutput))
+		{
+			check(AudioBuffer->Num() == InSettings.GetNumFramesPerBlock());
+			OutputDataReferences.AddDataReadReference(TEXT("Audio"), FAudioBufferReadRef(AudioBuffer));
+		}
+
+		virtual const FDataReferenceCollection& GetInputs() const override
+		{
+			return InputDataReferences;
+		}
+
+		virtual const FDataReferenceCollection& GetOutputs() const override
+		{
+			return OutputDataReferences;
+		}
+
+		// Scheduler/decoder holds weakref, and pins the sharedptr to the output when it writes.
+		// FBufferedDecoderOutput
+		// 1 write [... .. .. ]
+		// n free
+		// m queued 
+		// 1 read [...........]
+
+		// ApuMemory 2kb compress -> 512 frames
+
+		// TCircularAudioBuffer //consider for Output object.
+		
+		void Execute()
+		{
+			int32 NumPopped = 0;
+			float* Dst = AudioBuffer->GetData();
+
+			// If we don't have a valid state, just output silence.
+			if (Decoder && DecoderOutput && DecoderInput )
+			{
+				// V1. Do the decode inline, this will sound bad.
+				Decoder->Decode();
+
+				Audio::IDecoderOutput::FPushedAudioDetails Details;
+				NumPopped = DecoderOutput->PopAudio(MakeArrayView(Dst, AudioBuffer->Num()), Details);
+			}
+
+			// Pad with Silence if we didn't pop enough
+			for ( int32 i = NumPopped; i < OperatorSettings.GetNumFramesPerBlock(); ++i)
+			{
+				Dst[i] = 0.0f;
+			}
+		}
+
+	private:
+		const FOperatorSettings OperatorSettings;
+
+		FWaveReadRef Wave;
+		FAudioBufferWriteRef AudioBuffer;
+		FDataReferenceCollection InputDataReferences;
+		FDataReferenceCollection OutputDataReferences;
+
+		// Decoder/IO. 
+		Audio::ICodec::FDecoderPtr Decoder;
+		FWave::FDecoderInputPtr DecoderInput;
+		TUniquePtr<Audio::IDecoderOutput> DecoderOutput;
+	};
+
+	const FName FWavePlayerNode::ClassName = FName(TEXT("Wave"));
+
+	TUniquePtr<IOperator> FWavePlayerNode::FOperatorFactory::CreateOperator(
+		const INode& InNode, 
+		const FOperatorSettings& InOperatorSettings,
+		const FDataReferenceCollection& InInputDataReferences, 
+		TArray<TUniquePtr<IOperatorBuildError>>& OutErrors) 
+	{
+		using namespace Audio;
+
+		const FWavePlayerNode& WaveNode = static_cast<const FWavePlayerNode&>(InNode);
+		FWaveReadRef Wave = FWaveReadRef::CreateNew();
+
+		if (InInputDataReferences.ContainsDataReadReference<FWave>(TEXT("Wave")))
+		{
+			Wave = InInputDataReferences.GetDataReadReference<FWave>(TEXT("Wave"));
+		}
+
+		FWave::FDecoderInputPtr Input = FWave::CreateDecoderInput(Wave);
+		if (Input)
+		{
+			ICodecRegistry::FCodecPtr Codec = ICodecRegistry::Get().FindCodecByFromParsingInput(Input.Get());
+			if (Codec)
+			{
+				// V1, Ask for an output buffer the size of a frame.
+				IDecoderOutput::FRequirements Reqs { Float32_Interleaved, InOperatorSettings.GetNumFramesPerBlock() };
+				TUniquePtr<IDecoderOutput> Output = IDecoderOutput::Create(Reqs);
+				TUniquePtr<IDecoder> Decoder = Codec->CreateDecoder(Input.Get(), Output.Get());
+
+				return MakeUnique<FWavePlayerOperator>(
+					InOperatorSettings, 
+					Wave, 
+					MoveTemp(Input), 
+					MoveTemp(Output), 
+					MoveTemp(Decoder)
+				);
+			}
+			else
+			{
+				OutErrors.Add(MakeUnique<FWavePlayerError>(TEXT("FailedToFindCodec"),
+					LOCTEXT("FailedToFindCodec", "Failed to find codec for opening the supplied Wave")));
+			}
+		}
+		else
+		{
+			OutErrors.Add(MakeUnique<FWavePlayerError>(TEXT("FailedToParseInput"),
+				LOCTEXT("FailedToParseInput", "Failed to parse the compressed data")));
+		}
+
+		// Create the player without any inputs, will just produce silence.
+		return MakeUnique<FWavePlayerOperator>(InOperatorSettings,Wave);
+	}
+
+	FWavePlayerNode::FWavePlayerNode(const FString& InName)
+		:	FNode(InName)
+	{
+		AddInputDataVertex<FWave>(TEXT("Wave"), LOCTEXT("WaveTooltip", "The Wave to be decoded"));
+		AddOutputDataVertex<FAudioBuffer>(TEXT("Audio"), LOCTEXT("AudioTooltip", "The output audio"));
+	}
+
+	FWavePlayerNode::FWavePlayerNode(const FNodeInitData& InInitData)
+		: FWavePlayerNode(InInitData.InstanceName)
+	{
+	}
+
+	FWavePlayerNode::~FWavePlayerNode()
+	{
+	}
+	
+	const FName& FWavePlayerNode::GetClassName() const
+	{
+		return ::Metasound::FWavePlayerNode::ClassName;
+	}
+
+	IOperatorFactory& FWavePlayerNode::GetDefaultOperatorFactory() 
+	{
+		return Factory;
+	}
+}
+#undef LOCTEXT_NAMESPACE //MetasoundWaveNode
diff --git a/Engine/Plugins/Runtime/Metasound/Source/MetasoundEngine/Public/MetasoundWave.h b/Engine/Plugins/Runtime/Metasound/Source/MetasoundEngine/Public/MetasoundWave.h
new file mode 100644
--- /dev/null
+++ b/Engine/Plugins/Runtime/Metasound/Source/MetasoundEngine/Public/MetasoundWave.h
@@ -0,0 +1,39 @@
+// Copyright Epic Games, Inc. All Rights Reserved.
+
+#pragma once
+
+#include "MetasoundDataReference.h"
+#include "MetasoundDataTypeRegistrationMacro.h"
+
+// Forward declares
+namespace Audio
+{	
+	struct IDecoderInput;
+}
+
+namespace Metasound
+{
+	// Forward declare ReadRef
+	class FWave;
+	typedef TDataReadReference<FWave> FWaveReadRef;
+		
+	class METASOUNDENGINE_API FWave
+	{
+		TArray<uint8> CompressedBytes;
+		friend class FWaveDecoderInput;
+	public:
+		FWave() = default;
+
+		// For testing only.
+		FWave(const TArray<uint8>& InBytes) 
+			: CompressedBytes(InBytes) 
+		{}
+
+		using FDecoderInputPtr = TSharedPtr<Audio::IDecoderInput,ESPMode::ThreadSafe>;
+
+		// Factory function to create a Decoder input
+		static FDecoderInputPtr CreateDecoderInput(
+			const FWaveReadRef& InWaveRef);
+	};
+	DECLARE_METASOUND_DATA_REFERENCE_TYPES(FWave, METASOUNDENGINE_API, 0x0ddba11, FWaveTypeInfo, FWaveReadRef, FWaveWriteRef)
+}
diff --git a/Engine/Plugins/Runtime/Metasound/Source/MetasoundEngine/Public/MetasoundWavePlayerNode.h b/Engine/Plugins/Runtime/Metasound/Source/MetasoundEngine/Public/MetasoundWavePlayerNode.h
new file mode 100644
--- /dev/null
+++ b/Engine/Plugins/Runtime/Metasound/Source/MetasoundEngine/Public/MetasoundWavePlayerNode.h
@@ -0,0 +1,54 @@
+// Copyright Epic Games, Inc. All Rights Reserved.
+
+#pragma once
+
+#include "CoreMinimal.h"
+#include "MetasoundNode.h"
+#include "MetasoundOperatorInterface.h"
+#include "MetasoundBuilderInterface.h"
+#include "MetasoundDataReferenceCollection.h"
+
+namespace Metasound
+{
+	class METASOUNDENGINE_API FWavePlayerNode : public FNode
+	{
+		class FOperatorFactory : public IOperatorFactory
+		{
+			virtual TUniquePtr<IOperator> CreateOperator(const INode& InNode, const FOperatorSettings& InOperatorSettings, const FDataReferenceCollection& InInputDataReferences, TArray<TUniquePtr<IOperatorBuildError>>& OutErrors) override;
+		};
+
+	public:
+		static const FName ClassName;
+
+		FWavePlayerNode(const FString& InName);
+
+		// constructor used by the Metasound Frontend.
+		FWavePlayerNode(const FNodeInitData& InInitData);
+
+		virtual ~FWavePlayerNode();
+
+		const FName& GetClassName() const override;
+
+		IOperatorFactory& GetDefaultOperatorFactory() override;
+
+		const FText& GetDescription() const override
+		{
+			static const FText StaticDescription = NSLOCTEXT("MetasoundGraphCore", "Metasound_WavePlayerNodeDescription", "Plays a supplied Wave");
+			return StaticDescription;
+		}
+
+		const FText& GetAuthorName() const override
+		{
+			return PluginAuthor;
+		}
+
+		const FText& GetPromptIfMissing() const override
+		{
+			return PluginNodeMissingPrompt;
+		}
+
+	private:
+		float DefaultFrequency;
+		FOperatorFactory Factory;
+	};
+}
diff --git a/Engine/Plugins/Runtime/Metasound/Source/MetasoundGraphCore/Public/MetasoundDataReference.h b/Engine/Plugins/Runtime/Metasound/Source/MetasoundGraphCore/Public/MetasoundDataReference.h
--- a/Engine/Plugins/Runtime/Metasound/Source/MetasoundGraphCore/Public/MetasoundDataReference.h
+++ b/Engine/Plugins/Runtime/Metasound/Source/MetasoundGraphCore/Public/MetasoundDataReference.h
@@ -10,11 +10,11 @@
 // - Adding typedefs for commonly used template types.
 // - Defining parameter type traits.
 // - TODO- add argument for module api.
-#define DECLARE_METASOUND_DATA_REFERENCE_TYPES(DataType, DataTypeMagicNumber, DataTypeInfoTypeName, DataReadReferenceTypeName, DataWriteReferenceTypeName) \
+#define DECLARE_METASOUND_DATA_REFERENCE_TYPES(DataType, ModuleApi, DataTypeMagicNumber, DataTypeInfoTypeName, DataReadReferenceTypeName, DataWriteReferenceTypeName) \
 	template<> \
 	struct ::Metasound::TDataReferenceTypeInfo<DataType> \
 	{ \
-		static METASOUNDSTANDARDNODES_API const TCHAR* TypeName; \
+		static ModuleApi const TCHAR* TypeName; \
 		static constexpr ::Metasound::FDataTypeMagicNumber MagicNumber = (DataTypeMagicNumber); \
 		static constexpr bool bIsStringParsable = TTestIfDataTypeCtorIsImplemented<DataType, const FString&>::Value; \
 		static constexpr bool bIsBoolParsable = TTestIfDataTypeCtorIsImplemented<DataType, bool>::Value; \
diff --git a/Engine/Plugins/Runtime/Metasound/Source/MetasoundStandardNodes/Private/MetasoundStandardNodesModule.cpp b/Engine/Plugins/Runtime/Metasound/Source/MetasoundStandardNodes/Private/MetasoundStandardNodesModule.cpp
--- a/Engine/Plugins/Runtime/Metasound/Source/MetasoundStandardNodes/Private/MetasoundStandardNodesModule.cpp
+++ b/Engine/Plugins/Runtime/Metasound/Source/MetasoundStandardNodes/Private/MetasoundStandardNodesModule.cpp
@@ -9,7 +9,6 @@
 #include "MetasoundDataReferenceTypes.h"
 #include "MetasoundTime.h"
 #include "MetasoundFrequency.h"
-#include "MetasoundWave.h"
 
 REGISTER_METASOUND_DATATYPE(bool, "Primitive:Bool", ::Metasound::ELiteralArgType::Boolean)
 REGISTER_METASOUND_DATATYPE(int32, "Primitive:Int32", ::Metasound::ELiteralArgType::Integer)
@@ -24,7 +23,6 @@ namespace Metasound
 	REGISTER_METASOUND_DATATYPE(FDoubleTime, "Primitive:Time:HighResolution", ::Metasound::ELiteralArgType::Float)
 	REGISTER_METASOUND_DATATYPE(FSampleTime, "Primitive:Time:SampleResolution", ::Metasound::ELiteralArgType::Float)
 	REGISTER_METASOUND_DATATYPE(FFrequency, "Primitive:Frequency", ::Metasound::ELiteralArgType::Float)
-	REGISTER_METASOUND_DATATYPE(FWave, "Primitive:Wave")
 }
 
 namespace Metasound 
diff --git a/Engine/Plugins/Runtime/Metasound/Source/MetasoundStandardNodes/Private/MetasoundWave.cpp b/Engine/Plugins/Runtime/Metasound/Source/MetasoundStandardNodes/Private/MetasoundWave.cpp
deleted file mode 100644
diff --git a/Engine/Plugins/Runtime/Metasound/Source/MetasoundStandardNodes/Private/MetasoundWavePlayerNode.cpp b/Engine/Plugins/Runtime/Metasound/Source/MetasoundStandardNodes/Private/MetasoundWavePlayerNode.cpp
deleted file mode 100644
diff --git a/Engine/Plugins/Runtime/Metasound/Source/MetasoundStandardNodes/Public/MetasoundAudioBuffer.h b/Engine/Plugins/Runtime/Metasound/Source/MetasoundStandardNodes/Public/MetasoundAudioBuffer.h
--- a/Engine/Plugins/Runtime/Metasound/Source/MetasoundStandardNodes/Public/MetasoundAudioBuffer.h
+++ b/Engine/Plugins/Runtime/Metasound/Source/MetasoundStandardNodes/Public/MetasoundAudioBuffer.h
@@ -154,5 +154,5 @@ namespace Metasound
 #endif				
 	};
 
-	DECLARE_METASOUND_DATA_REFERENCE_TYPES(FAudioBuffer, 0x1b02dac2 , FAudioBufferTypeInfo, FAudioBufferReadRef, FAudioBufferWriteRef);
+	DECLARE_METASOUND_DATA_REFERENCE_TYPES(FAudioBuffer, METASOUNDSTANDARDNODES_API, 0x1b02dac2 , FAudioBufferTypeInfo, FAudioBufferReadRef, FAudioBufferWriteRef);
 };
diff --git a/Engine/Plugins/Runtime/Metasound/Source/MetasoundStandardNodes/Public/MetasoundAudioFormats.h b/Engine/Plugins/Runtime/Metasound/Source/MetasoundStandardNodes/Public/MetasoundAudioFormats.h
--- a/Engine/Plugins/Runtime/Metasound/Source/MetasoundStandardNodes/Public/MetasoundAudioFormats.h
+++ b/Engine/Plugins/Runtime/Metasound/Source/MetasoundStandardNodes/Public/MetasoundAudioFormats.h
@@ -456,15 +456,15 @@ namespace Metasound
 	};
 
 
-	DECLARE_METASOUND_DATA_REFERENCE_TYPES(FUnformattedAudio, 0xd78a3ed1 , FUnformattedAudioTypeInfo, FUnformattedAudioReadRef, FUnformattedAudioWriteRef);
+	DECLARE_METASOUND_DATA_REFERENCE_TYPES(FUnformattedAudio, METASOUNDSTANDARDNODES_API, 0xd78a3ed1 , FUnformattedAudioTypeInfo, FUnformattedAudioReadRef, FUnformattedAudioWriteRef);
 	
 
-	DECLARE_METASOUND_DATA_REFERENCE_TYPES(FMultichannelAudioFormat, 0x56bdcbe0 , FMultichannelAudioFormatTypeInfo, FMultichannelAudioFormatReadRef, FMultichannelAudioFormatWriteRef);
+	DECLARE_METASOUND_DATA_REFERENCE_TYPES(FMultichannelAudioFormat, METASOUNDSTANDARDNODES_API, 0x56bdcbe0 , FMultichannelAudioFormatTypeInfo, FMultichannelAudioFormatReadRef, FMultichannelAudioFormatWriteRef);
 	
 
-	DECLARE_METASOUND_DATA_REFERENCE_TYPES(FMonoAudioFormat, 0x6f468c8c, FMonoAudioFormatTypeInfo, FMonoAudioFormatReadRef, FMonoAudioFormatWriteRef);
+	DECLARE_METASOUND_DATA_REFERENCE_TYPES(FMonoAudioFormat, METASOUNDSTANDARDNODES_API, 0x6f468c8c, FMonoAudioFormatTypeInfo, FMonoAudioFormatReadRef, FMonoAudioFormatWriteRef);
 
-	DECLARE_METASOUND_DATA_REFERENCE_TYPES(FStereoAudioFormat, 0xb55304e2 , FStereoAudioFormatTypeInfo, FStereoAudioFormatReadRef, FStereoAudioFormatWriteRef);
+	DECLARE_METASOUND_DATA_REFERENCE_TYPES(FStereoAudioFormat, METASOUNDSTANDARDNODES_API, 0xb55304e2 , FStereoAudioFormatTypeInfo, FStereoAudioFormatReadRef, FStereoAudioFormatWriteRef);
 	
 	
 }
diff --git a/Engine/Plugins/Runtime/Metasound/Source/MetasoundStandardNodes/Public/MetasoundDataReferenceTypes.h b/Engine/Plugins/Runtime/Metasound/Source/MetasoundStandardNodes/Public/MetasoundDataReferenceTypes.h
--- a/Engine/Plugins/Runtime/Metasound/Source/MetasoundStandardNodes/Public/MetasoundDataReferenceTypes.h
+++ b/Engine/Plugins/Runtime/Metasound/Source/MetasoundStandardNodes/Public/MetasoundDataReferenceTypes.h
@@ -85,5 +85,5 @@ namespace Metasound
 			TArray<uint32> Frames;
 	};
 
-	DECLARE_METASOUND_DATA_REFERENCE_TYPES(FBop, 2, FBopTypeInfo, FBopReadRef, FBopWriteRef);
+	DECLARE_METASOUND_DATA_REFERENCE_TYPES(FBop, METASOUNDSTANDARDNODES_API, 2, FBopTypeInfo, FBopReadRef, FBopWriteRef);
 }
diff --git a/Engine/Plugins/Runtime/Metasound/Source/MetasoundStandardNodes/Public/MetasoundFrequency.h b/Engine/Plugins/Runtime/Metasound/Source/MetasoundStandardNodes/Public/MetasoundFrequency.h
--- a/Engine/Plugins/Runtime/Metasound/Source/MetasoundStandardNodes/Public/MetasoundFrequency.h
+++ b/Engine/Plugins/Runtime/Metasound/Source/MetasoundStandardNodes/Public/MetasoundFrequency.h
@@ -150,5 +150,5 @@ namespace Metasound
 			float FrequencyInHz;
 	};
 
-	DECLARE_METASOUND_DATA_REFERENCE_TYPES(FFrequency, 0xf1c430a4, FFrequencyTypeInfo, FFrequencyReadRef, FFrequencyWriteRef)
+	DECLARE_METASOUND_DATA_REFERENCE_TYPES(FFrequency, METASOUNDSTANDARDNODES_API, 0xf1c430a4, FFrequencyTypeInfo, FFrequencyReadRef, FFrequencyWriteRef)
 }
diff --git a/Engine/Plugins/Runtime/Metasound/Source/MetasoundStandardNodes/Public/MetasoundPrimitives.h b/Engine/Plugins/Runtime/Metasound/Source/MetasoundStandardNodes/Public/MetasoundPrimitives.h
--- a/Engine/Plugins/Runtime/Metasound/Source/MetasoundStandardNodes/Public/MetasoundPrimitives.h
+++ b/Engine/Plugins/Runtime/Metasound/Source/MetasoundStandardNodes/Public/MetasoundPrimitives.h
@@ -12,14 +12,14 @@
 
 namespace Metasound
 {
-	DECLARE_METASOUND_DATA_REFERENCE_TYPES(bool, 0xcc9d715e, FBoolTypeInfo, FBoolReadRef, FBoolWriteRef);
+	DECLARE_METASOUND_DATA_REFERENCE_TYPES(bool, METASOUNDSTANDARDNODES_API, 0xcc9d715e, FBoolTypeInfo, FBoolReadRef, FBoolWriteRef);
 	
-	DECLARE_METASOUND_DATA_REFERENCE_TYPES(int32, 0xf9920034, FInt32TypeInfo, FInt32ReadRef, FInt32WriteRef);
+	DECLARE_METASOUND_DATA_REFERENCE_TYPES(int32, METASOUNDSTANDARDNODES_API, 0xf9920034, FInt32TypeInfo, FInt32ReadRef, FInt32WriteRef);
 
-	DECLARE_METASOUND_DATA_REFERENCE_TYPES(int64, 0xfc8f8ada, FInt64TypeInfo, FInt64ReadRef, FInt64WriteRef);
+	DECLARE_METASOUND_DATA_REFERENCE_TYPES(int64, METASOUNDSTANDARDNODES_API, 0xfc8f8ada, FInt64TypeInfo, FInt64ReadRef, FInt64WriteRef);
 
-	DECLARE_METASOUND_DATA_REFERENCE_TYPES(float, 0xe7adc4db, FFloatTypeInfo, FFloatReadRef, FFloatWriteRef);
+	DECLARE_METASOUND_DATA_REFERENCE_TYPES(float, METASOUNDSTANDARDNODES_API, 0xe7adc4db, FFloatTypeInfo, FFloatReadRef, FFloatWriteRef);
 
-	DECLARE_METASOUND_DATA_REFERENCE_TYPES(double, 0x8f48a053, FDoubleTypeInfo, FDoubleReadRef, FDoubleWriteRef);
+	DECLARE_METASOUND_DATA_REFERENCE_TYPES(double, METASOUNDSTANDARDNODES_API, 0x8f48a053, FDoubleTypeInfo, FDoubleReadRef, FDoubleWriteRef);
 }
 
diff --git a/Engine/Plugins/Runtime/Metasound/Source/MetasoundStandardNodes/Public/MetasoundTime.h b/Engine/Plugins/Runtime/Metasound/Source/MetasoundStandardNodes/Public/MetasoundTime.h
--- a/Engine/Plugins/Runtime/Metasound/Source/MetasoundStandardNodes/Public/MetasoundTime.h
+++ b/Engine/Plugins/Runtime/Metasound/Source/MetasoundStandardNodes/Public/MetasoundTime.h
@@ -624,19 +624,19 @@ namespace Metasound
 	 */
 	typedef TTime<float, float> FFloatTime;
 
-	DECLARE_METASOUND_DATA_REFERENCE_TYPES(FFloatTime, 0x2490848e, FFloatTimeTypeInfo, FFloatTimeReadRef, FFloatTimeWriteRef);
+	DECLARE_METASOUND_DATA_REFERENCE_TYPES(FFloatTime, METASOUNDSTANDARDNODES_API, 0x2490848e, FFloatTimeTypeInfo, FFloatTimeReadRef, FFloatTimeWriteRef);
 
 	/** FDoubleTime stores and represents all time values as 64bit floating
 	 * point values.
 	 */
 	typedef TTime<double, double> FDoubleTime;
 
-	DECLARE_METASOUND_DATA_REFERENCE_TYPES(FDoubleTime, 0x407fcff1, FDoubleTimeTypeInfo, FDoubleTimeReadRef, FDoubleTimeWriteRef);
+	DECLARE_METASOUND_DATA_REFERENCE_TYPES(FDoubleTime, METASOUNDSTANDARDNODES_API, 0x407fcff1, FDoubleTimeTypeInfo, FDoubleTimeReadRef, FDoubleTimeWriteRef);
 
 	/** FSampleTime stores time quantized to sample boundaries. It returns time
 	 * as 32bit floating point values.
 	 */
 	typedef TTime<FTimeSampleCounter, float> FSampleTime;
 
-	DECLARE_METASOUND_DATA_REFERENCE_TYPES(FSampleTime, 0xf1de8f2d, FSampleTimeTypeInfo, FSampleTimeReadRef, FSampleTimeWriteRef);
+	DECLARE_METASOUND_DATA_REFERENCE_TYPES(FSampleTime, METASOUNDSTANDARDNODES_API, 0xf1de8f2d, FSampleTimeTypeInfo, FSampleTimeReadRef, FSampleTimeWriteRef);
 }
diff --git a/Engine/Plugins/Runtime/Metasound/Source/MetasoundStandardNodes/Public/MetasoundWave.h b/Engine/Plugins/Runtime/Metasound/Source/MetasoundStandardNodes/Public/MetasoundWave.h
deleted file mode 100644
diff --git a/Engine/Plugins/Runtime/Metasound/Source/MetasoundStandardNodes/Public/MetasoundWavePlayerNode.h b/Engine/Plugins/Runtime/Metasound/Source/MetasoundStandardNodes/Public/MetasoundWavePlayerNode.h
deleted file mode 100644
