Fix for nounity error in MetasoundFrontend

diff --git a/Engine/Plugins/Runtime/Metasound/Source/MetasoundFrontend/Private/MetasoundFrontend.cpp b/Engine/Plugins/Runtime/Metasound/Source/MetasoundFrontend/Private/MetasoundFrontend.cpp
--- a/Engine/Plugins/Runtime/Metasound/Source/MetasoundFrontend/Private/MetasoundFrontend.cpp
+++ b/Engine/Plugins/Runtime/Metasound/Source/MetasoundFrontend/Private/MetasoundFrontend.cpp
@@ -9,6 +9,7 @@
 #include "Modules/ModuleManager.h"
 #include "StructDeserializer.h"
 #include "StructSerializer.h"
+#include "Serialization/MemoryReader.h"
 
 
 static int32 MetasoundUndoRollLimitCvar = 128;
