SoundWave -> WavePlayer unit tests (doing inline encode).

diff --git a/Engine/Plugins/Runtime/Metasound/Source/MetasoundEngine/Private/MetasoundWave.cpp b/Engine/Plugins/Runtime/Metasound/Source/MetasoundEngine/Private/MetasoundWave.cpp
--- a/Engine/Plugins/Runtime/Metasound/Source/MetasoundEngine/Private/MetasoundWave.cpp
+++ b/Engine/Plugins/Runtime/Metasound/Source/MetasoundEngine/Private/MetasoundWave.cpp
@@ -3,6 +3,8 @@
 #include "MetasoundWave.h"
 #include "MetasoundPrimitives.h"
 #include "IAudioCodec.h"
+#include "IAudioCodecRegistry.h"
+#include "Sound/SoundWave.h"
 
 namespace Metasound
 {	
@@ -17,8 +19,91 @@ namespace Metasound
 		{}
 	};
 
-	FWave::FDecoderInputPtr FWave::CreateDecoderInput(const FWaveReadRef& InWaveRef) 
+	// Small RAII helper.
+	struct FScopedBulkDataLock
 	{
-		return MakeShared<FWaveDecoderInput, ESPMode::ThreadSafe>(InWaveRef);
+		bool bLocked = false;
+		FByteBulkData& Bd;
+		FScopedBulkDataLock(FByteBulkData& InBd) 
+			: Bd(InBd) 
+		{}
+
+		const void* LockReadOnly()
+		{
+			if (ensure(bLocked==false))
+			{
+				if (const void* Ptr =Bd.LockReadOnly() )
+				{
+					bLocked = true;
+					return Ptr;
+				}
+			}
+			return nullptr;
+		}
+
+		~FScopedBulkDataLock() 
+		{ 
+			if (bLocked) 
+			{
+				Bd.Unlock();
+			}
+		}
+	};
+
+	// Temp. We don't have any cooking at the moment, so do the encode here.
+	bool DoInlineEncode(
+		USoundWave* InWave,
+		TArray<uint8>& OutCompressedBytes)
+	{
+#if WITH_EDITORONLY_DATA
+		using namespace Audio;		
+		FScopeLock Lock(&InWave->RawDataCriticalSection);
+		FByteBulkData& Raw = InWave->RawData;
+		FScopedBulkDataLock BdLock(Raw);
+		if (const void* RawPtr = BdLock.LockReadOnly())
+		{
+			IEncoderInput::FFormat Format { InWave->NumChannels, InWave->GetSampleRateForCurrentPlatform(), EBitRepresentation::Int16_Interleaved };
+			TUniquePtr<IEncoderInput> Input = IEncoderInput::Create(
+				MakeArrayView(static_cast<const int16*>(RawPtr), Raw.GetBulkDataSize() / sizeof(int16)),
+				Format
+			);
+
+			if (ICodecRegistry::FCodecPtr DefaultCodec = ICodecRegistry::Get().FindDefaultCodec() )
+			{
+				ICodec::FEncoderPtr Encoder = DefaultCodec->CreateEncoder(Input.Get());
+				if (Encoder)
+				{
+					// Encode default everything.
+					return Encoder->Encode(
+						Input.Get(),
+						IEncoderOutput::Create(OutCompressedBytes).Get(),
+						nullptr
+					);
+				}
+			}
+		}		
+		return false;
+#else //WITH_EDITORONLY_DATA
+		return true;
+#endif //WITH_EDITORONLY_DATA
+	}
+	
+	FWave::FWave(USoundWave* InWave) 
+	{
+		DoInlineEncode(InWave,CompressedBytes);
+	}
+
+	FWave::FWave(const TArray<uint8>& InBytes) 
+		: CompressedBytes(InBytes)
+	{
+	}
+
+	FWave::FDecoderInputPtr FWave::CreateDecoderInput(const FWaveReadRef& InWaveRef)
+	{
+		if( InWaveRef->CompressedBytes.Num() > 0)
+		{
+			return MakeShared<FWaveDecoderInput, ESPMode::ThreadSafe>(InWaveRef);
+		}
+		return nullptr;
 	}
 }
diff --git a/Engine/Plugins/Runtime/Metasound/Source/MetasoundEngine/Public/MetasoundWave.h b/Engine/Plugins/Runtime/Metasound/Source/MetasoundEngine/Public/MetasoundWave.h
--- a/Engine/Plugins/Runtime/Metasound/Source/MetasoundEngine/Public/MetasoundWave.h
+++ b/Engine/Plugins/Runtime/Metasound/Source/MetasoundEngine/Public/MetasoundWave.h
@@ -11,6 +11,8 @@ namespace Audio
 	struct IDecoderInput;
 }
 
+class USoundWave;
+
 namespace Metasound
 {
 	// Forward declare ReadRef
@@ -21,13 +23,14 @@ namespace Metasound
 	{
 		TArray<uint8> CompressedBytes;
 		friend class FWaveDecoderInput;
+
 	public:
 		FWave() = default;
 
 		// For testing only.
-		FWave(const TArray<uint8>& InBytes) 
-			: CompressedBytes(InBytes) 
-		{}
+		FWave(const TArray<uint8>& InBytes);
+		
+		FWave(USoundWave* InWave);
 
 		using FDecoderInputPtr = TSharedPtr<Audio::IDecoderInput,ESPMode::ThreadSafe>;
 
diff --git a/Engine/Plugins/Runtime/Metasound/Source/MetasoundEngine/Public/MetasoundWavePlayerNode.h b/Engine/Plugins/Runtime/Metasound/Source/MetasoundEngine/Public/MetasoundWavePlayerNode.h
--- a/Engine/Plugins/Runtime/Metasound/Source/MetasoundEngine/Public/MetasoundWavePlayerNode.h
+++ b/Engine/Plugins/Runtime/Metasound/Source/MetasoundEngine/Public/MetasoundWavePlayerNode.h
@@ -48,7 +48,6 @@ namespace Metasound
 		}
 
 	private:
-		float DefaultFrequency;
 		FOperatorFactory Factory;
 	};
 }
