Fix c++11-narrowing warning as error on Clang.

diff --git a/Engine/Plugins/Runtime/Metasound/Source/MetasoundEngine/Private/MetasoundWave.cpp b/Engine/Plugins/Runtime/Metasound/Source/MetasoundEngine/Private/MetasoundWave.cpp
--- a/Engine/Plugins/Runtime/Metasound/Source/MetasoundEngine/Private/MetasoundWave.cpp
+++ b/Engine/Plugins/Runtime/Metasound/Source/MetasoundEngine/Private/MetasoundWave.cpp
@@ -62,7 +62,7 @@ namespace Metasound
 		FScopedBulkDataLock BdLock(Raw);
 		if (const void* RawPtr = BdLock.LockReadOnly())
 		{
-			IEncoderInput::FFormat Format { InWave->NumChannels, InWave->GetSampleRateForCurrentPlatform(), EBitRepresentation::Int16_Interleaved };
+			IEncoderInput::FFormat Format { static_cast<uint32>(InWave->NumChannels), static_cast<uint32>(InWave->GetSampleRateForCurrentPlatform()), EBitRepresentation::Int16_Interleaved };
 			TUniquePtr<IEncoderInput> Input = IEncoderInput::Create(
 				MakeArrayView(static_cast<const int16*>(RawPtr), Raw.GetBulkDataSize() / sizeof(int16)),
 				Format
