Adding validation checks and better error reporting to FOperatorBuidler

diff --git a/Engine/Plugins/Runtime/Metasound/Source/MetasoundFrontend/Private/MetasoundFrontend.cpp b/Engine/Plugins/Runtime/Metasound/Source/MetasoundFrontend/Private/MetasoundFrontend.cpp
--- a/Engine/Plugins/Runtime/Metasound/Source/MetasoundFrontend/Private/MetasoundFrontend.cpp
+++ b/Engine/Plugins/Runtime/Metasound/Source/MetasoundFrontend/Private/MetasoundFrontend.cpp
@@ -2475,7 +2475,8 @@ namespace Metasound
 			/** NOTE: In the future- we should split steps 1-4 above, and 5 below. GraphToBuild and the node map can be cached on the graph handle. */
 
 			// Step 5: Invoke Operator Builder
-			FOperatorBuilder Builder(InSettings);
+			FOperatorBuilder Builder(InSettings, FOperatorBuilderSettings::GetDefaultSettings());
+
 			return Builder.BuildGraphOperator(GraphToBuild, OutBuildErrors);
 		}
 
diff --git a/Engine/Plugins/Runtime/Metasound/Source/MetasoundGenerator/Private/SynthComponentMetasoundExample.cpp b/Engine/Plugins/Runtime/Metasound/Source/MetasoundGenerator/Private/SynthComponentMetasoundExample.cpp
--- a/Engine/Plugins/Runtime/Metasound/Source/MetasoundGenerator/Private/SynthComponentMetasoundExample.cpp
+++ b/Engine/Plugins/Runtime/Metasound/Source/MetasoundGenerator/Private/SynthComponentMetasoundExample.cpp
@@ -220,7 +220,7 @@ ISoundGeneratorPtr USynthComponentMetasoundExample::CreateSoundGenerator(const F
 	TUniquePtr<FSynthComponentMetasoundExampleGraph> ExampleGraph = CreateSynthComponentMetasoundExampleGraph();
 
 	FOperatorSettings Settings(InParams.SampleRate, InParams.NumFramesPerCallback);
-	FOperatorBuilder Builder(Settings);
+	FOperatorBuilder Builder(Settings, FOperatorBuilderSettings::GetDefaultSettings());
 
 	TArray<IOperatorBuilder::FBuildErrorPtr> Errors;
 	FOperatorUniquePtr Operator = Builder.BuildGraphOperator(*ExampleGraph->Graph, Errors);
diff --git a/Engine/Plugins/Runtime/Metasound/Source/MetasoundGraphCore/Private/MetasoundBuildError.cpp b/Engine/Plugins/Runtime/Metasound/Source/MetasoundGraphCore/Private/MetasoundBuildError.cpp
--- a/Engine/Plugins/Runtime/Metasound/Source/MetasoundGraphCore/Private/MetasoundBuildError.cpp
+++ b/Engine/Plugins/Runtime/Metasound/Source/MetasoundGraphCore/Private/MetasoundBuildError.cpp
@@ -30,6 +30,18 @@ namespace Metasound
 		return ErrorDescription;
 	}
 
+	/** Returns an array of destinations associated with the error. */
+	const TArray<FInputDataDestination>& FBuildErrorBase::GetDestinations() const
+	{
+		return Destinations;
+	}
+	
+	/** Returns an array of sources associated with the error. */
+	const TArray<FOutputDataSource>& FBuildErrorBase::GetSources() const
+	{
+		return Sources;
+	}
+
 	/** Returns an array of Nodes associated with the error. */
 	const TArray<const INode*>& FBuildErrorBase::GetNodes() const 
 	{
@@ -42,6 +54,26 @@ namespace Metasound
 		return Edges;
 	}
 
+	void FBuildErrorBase::AddInputDataDestination(const FInputDataDestination& InDestination)
+	{
+		Destinations.Add(InDestination);
+	}
+
+	void FBuildErrorBase::AddInputDataDestinations(TArrayView<const FInputDataDestination> InDestinations)
+	{
+		Destinations.Append(InDestinations.GetData(), InDestinations.Num());
+	}
+
+	void FBuildErrorBase::AddOutputDataSource(const FOutputDataSource& InSource)
+	{
+		Sources.Add(InSource);
+	}
+
+	void FBuildErrorBase::AddOutputDataSources(TArrayView<const FOutputDataSource> InSources)
+	{
+		Sources.Append(InSources.GetData(), InSources.Num());
+	}
+
 	void FBuildErrorBase::AddEdge(const FDataEdge& InEdge)
 	{
 		Edges.Add(InEdge);
@@ -62,14 +94,52 @@ namespace Metasound
 		Nodes.Append(InNodes.GetData(), InNodes.Num());
 	}
 
-	const FName FDanglingEdgeError::ErrorType = FName(TEXT("MetasoundDanglingEdge"));
+	const FName FDanglingVertexError::ErrorType = FName(TEXT("MetasoundDanglingVertex"));
 
-	FDanglingEdgeError::FDanglingEdgeError(const FDataEdge& InEdge)
-	:	FBuildErrorBase(ErrorType, LOCTEXT("DanglingEdgeError", "Edge is not connected"))
+	FDanglingVertexError::FDanglingVertexError()
+	:	FBuildErrorBase(ErrorType, LOCTEXT("DanglingVertexError", "Edge is not connected"))
+	{}
+
+	FDanglingVertexError::FDanglingVertexError(const FInputDataDestination& InDestination)
+	:	FDanglingVertexError()
+	{
+		AddInputDataDestination(InDestination);
+	}
+
+	FDanglingVertexError::FDanglingVertexError(const FOutputDataSource& InSource)
+	:	FDanglingVertexError()
+	{
+		AddOutputDataSource(InSource);
+	}
+
+	FDanglingVertexError::FDanglingVertexError(const FDataEdge& InEdge)
+	:	FDanglingVertexError()
 	{
 		AddEdge(InEdge);
 	}
 
+	const FName FMissingVertexError::ErrorType = FName(TEXT("MetasoundMissingVertex"));
+
+	FMissingVertexError::FMissingVertexError(const FInputDataDestination& InDest)
+	:	FBuildErrorBase(ErrorType, LOCTEXT("MissingVertexError", "Node is missing a vertex"))
+	{
+		AddInputDataDestination(InDest);
+	}
+
+	FMissingVertexError::FMissingVertexError(const FOutputDataSource& InSource)
+	:	FBuildErrorBase(ErrorType, LOCTEXT("MissingVertexError", "Node is missing a vertex"))
+	{
+		AddOutputDataSource(InSource);
+	}
+
+	const FName FDuplicateInputError::ErrorType = FName(TEXT("MetasoundDuplicateInput"));
+
+	FDuplicateInputError::FDuplicateInputError(const TArrayView<FDataEdge> InEdges)
+	:	FBuildErrorBase(ErrorType, LOCTEXT("DuplicateInputError", "Node has duplicate connections to input"))
+	{
+		AddEdges(InEdges);
+	}
+
 	const FName FGraphCycleError::ErrorType = FName(TEXT("MetasoundGraphCycleError"));
 
 	FGraphCycleError::FGraphCycleError(TArrayView<INode const* const> InNodes, const TArray<FDataEdge>& InEdges)
@@ -79,6 +149,17 @@ namespace Metasound
 		AddEdges(InEdges);
 	}
 
+	const FName FNodePrunedError::ErrorType = FName(TEXT("MetasoundNodePrunedError"));
+
+	FNodePrunedError::FNodePrunedError(const INode* InNode)
+	:	FBuildErrorBase(ErrorType, LOCTEXT("NodePrunedError", "Unreachable node pruned from graph."))
+	{
+		if (ensure(nullptr != InNode))
+		{
+			AddNode(*InNode);
+		}
+	}
+
 	const FName FInternalError::ErrorType = FName(TEXT("MetasoundInternalError"));
 
 	FInternalError::FInternalError(const FString& InFileName, int32 InLineNumber)
@@ -102,17 +183,8 @@ namespace Metasound
 	
 	FMissingInputDataReferenceError::FMissingInputDataReferenceError(const FInputDataDestination& InInputDataDestination)
 	:	FBuildErrorBase(ErrorType, LOCTEXT("MissingInputDataReferenceError", "Missing input data reference."))
-	,	InputDataDestination(InInputDataDestination)
 	{
-		if (nullptr != InInputDataDestination.Node)
-		{
-			AddNode(*InInputDataDestination.Node);
-		}
-	}
-
-	const FInputDataDestination& FMissingInputDataReferenceError::GetInputDataDestination() const
-	{
-		return InputDataDestination;
+		AddInputDataDestination(InInputDataDestination);
 	}
 
 
@@ -120,17 +192,8 @@ namespace Metasound
 	
 	FMissingOutputDataReferenceError::FMissingOutputDataReferenceError(const FOutputDataSource& InOutputDataSource)
 	:	FBuildErrorBase(ErrorType, LOCTEXT("MissingOutputDataReferenceError", "Missing output data reference."))
-	,	OutputDataSource(InOutputDataSource)
-	{
-		if (nullptr != InOutputDataSource.Node)
-		{
-			AddNode(*InOutputDataSource.Node);
-		}
-	}
-
-	const FOutputDataSource& FMissingOutputDataReferenceError::GetOutputDataSource() const
 	{
-		return OutputDataSource;
+		AddOutputDataSource(InOutputDataSource);
 	}
 
 	const FName FInvalidConnectionDataTypeError::ErrorType = FName(TEXT("MetasoundInvalidConnectionDataTypeError"));
diff --git a/Engine/Plugins/Runtime/Metasound/Source/MetasoundGraphCore/Private/MetasoundGraphLinter.cpp b/Engine/Plugins/Runtime/Metasound/Source/MetasoundGraphCore/Private/MetasoundGraphLinter.cpp
new file mode 100644
--- /dev/null
+++ b/Engine/Plugins/Runtime/Metasound/Source/MetasoundGraphCore/Private/MetasoundGraphLinter.cpp
@@ -0,0 +1,246 @@
+// Copyright Epic Games, Inc. All Rights Reserved.
+
+#include "MetasoundGraphLinter.h"
+
+#include "Algo/Count.h"
+#include "Algo/Transform.h"
+#include "CoreMinimal.h"
+
+#include "MetasoundBuildError.h"
+#include "MetasoundBuilderInterface.h"
+#include "MetasoundGraphAlgo.h"
+
+namespace Metasound
+{
+	namespace GraphLinterPrivate
+	{
+		using FBuildErrorPtr = TUniquePtr<IOperatorBuildError>;
+
+		template<typename ErrorType, typename... ArgTypes>
+		void AddBuildError(TArray<FBuildErrorPtr>& OutErrors, ArgTypes&&... Args)
+		{
+			OutErrors.Add(MakeUnique<ErrorType>(Forward<ArgTypes>(Args)...));
+		}
+
+		// run validation on input destination
+		bool IsInputDestinationValid(const FInputDataDestination& InDest)
+		{
+			if (InDest.Node == nullptr)
+			{
+				return false;
+			}
+
+			const FInputDataVertexCollection& InputCollection = InDest.Node->GetInputDataVertices();
+
+			FDataVertexKey VertexKey = MakeDataVertexKey(InDest.Vertex);
+
+			if (!InputCollection.Contains(VertexKey))
+			{
+				return false;
+			}
+
+			if (InputCollection[VertexKey] != InDest.Vertex)
+			{
+				return false;
+			}
+
+			return true;
+		}
+		
+		// run validation on output source
+		bool IsOutputSourceValid(const FOutputDataSource& InSource)
+		{
+			if (InSource.Node == nullptr)
+			{
+				return false;
+			}
+
+			const FOutputDataVertexCollection& OutputCollection = InSource.Node->GetOutputDataVertices();
+
+			FDataVertexKey VertexKey = MakeDataVertexKey(InSource.Vertex);
+
+			if (!OutputCollection.Contains(VertexKey))
+			{
+				return false;
+			}
+
+			if (OutputCollection[VertexKey] != InSource.Vertex)
+			{
+				return false;
+			}
+
+			return true;
+		};
+
+		const FDataEdge* GetEdgePointer(const FDataEdge& InEdge) 
+		{ 
+			return &InEdge; 
+		};
+
+		// Raw pointer comparison for sorting.
+		bool CompareEdgeDestinationNode(const FDataEdge& InLHS, const FDataEdge& InRHS)
+		{
+			return InLHS.To.Node < InRHS.To.Node;
+		}
+	}
+
+	bool FGraphLinter::ValidateEdgeDataTypesMatch(const IGraph& InGraph, TArray<FBuildErrorPtr>& OutErrors)
+	{
+		using namespace GraphLinterPrivate;
+
+		bool bIsValid = true;
+
+		for (const FDataEdge& Edge : InGraph.GetDataEdges())
+		{
+			if (Edge.From.Vertex.DataReferenceTypeName != Edge.To.Vertex.DataReferenceTypeName)
+			{
+				AddBuildError<FInvalidConnectionDataTypeError>(OutErrors, Edge);
+
+				bIsValid = false;
+			}
+		}
+
+		return bIsValid;
+	}
+
+	bool FGraphLinter::ValidateVerticesExist(const IGraph& InGraph, TArray<FBuildErrorPtr>& OutErrors) 
+	{
+		using namespace GraphLinterPrivate;
+
+		bool bAllVerticesAreValid = true;
+
+		for (const FDataEdge& Edge : InGraph.GetDataEdges())
+		{
+			if (nullptr == Edge.From.Node)
+			{
+				bAllVerticesAreValid = false;
+				AddBuildError<FDanglingVertexError>(OutErrors, Edge.From);
+			}
+
+			if (!IsOutputSourceValid(Edge.From))
+			{
+				bAllVerticesAreValid = false;
+				AddBuildError<FMissingVertexError>(OutErrors, Edge.From);
+			}
+
+			if (nullptr == Edge.To.Node)
+			{
+				bAllVerticesAreValid = false;
+				AddBuildError<FDanglingVertexError>(OutErrors, Edge.To);
+			}
+
+			if (!IsInputDestinationValid(Edge.To))
+			{
+				bAllVerticesAreValid = false;
+				AddBuildError<FMissingVertexError>(OutErrors, Edge.To);
+			}
+		}
+
+		for (auto Element : InGraph.GetInputDataDestinations())
+		{
+			const FInputDataDestination& Dest = Element.Value;
+			if (!IsInputDestinationValid(Dest))
+			{
+				bAllVerticesAreValid = false;
+				AddBuildError<FMissingVertexError>(OutErrors, Dest);
+			}
+		}
+
+		for (auto Element : InGraph.GetOutputDataSources())
+		{
+			const FOutputDataSource& Source = Element.Value;
+			if (!IsOutputSourceValid(Source))
+			{
+				bAllVerticesAreValid = false;
+				AddBuildError<FMissingVertexError>(OutErrors, Source);
+			}
+		}
+		
+		return bAllVerticesAreValid;
+	}
+
+	bool FGraphLinter::ValidateNoDuplicateInputs(const IGraph& InGraph, TArray<FBuildErrorPtr>& OutErrors)
+	{
+		using namespace GraphLinterPrivate;
+
+		bool bIsValid = true;
+
+		TArray<const FDataEdge*> Edges;
+
+		Algo::Transform(InGraph.GetDataEdges(), Edges, GetEdgePointer);
+
+		Edges.Sort(CompareEdgeDestinationNode);
+
+		if (Edges.Num() < 1)
+		{
+			return bIsValid;
+		}
+
+		int32 GroupStartIndex = 0;
+		int32 GroupEndIndex = 0;
+		const INode* CurrentNode = Edges[GroupStartIndex]->To.Node;
+
+		for (int32 i = 1; i < Edges.Num(); i++)
+		{
+			const INode* Node = Edges[i]->To.Node;
+			if (Node == CurrentNode)
+			{
+				GroupEndIndex = i;
+			}
+			else
+			{
+				if (GroupEndIndex > GroupStartIndex)
+				{
+					bIsValid = false;
+
+					TArray<FDataEdge> Duplicates;
+					for (int32 j = GroupStartIndex; j <= GroupEndIndex; j++)
+					{
+						Duplicates.Add(*Edges[j]);
+						AddBuildError<FDuplicateInputError>(OutErrors, Duplicates);
+					}
+				}
+
+				GroupStartIndex = GroupEndIndex = i;
+				CurrentNode = Node;
+			}
+		}
+
+		return bIsValid;
+	}
+
+	bool FGraphLinter::ValidateNoCyclesInGraph(const IGraph& InGraph, TArray<FBuildErrorPtr>& OutErrors) 
+	{
+		TPimplPtr<FDirectedGraphAlgoAdapter> Adapter = FDirectedGraphAlgo::CreateDirectedGraphAlgoAdapter(InGraph);
+
+		check(Adapter.IsValid());
+
+		return ValidateNoCyclesInGraph(*Adapter, OutErrors);
+	}
+
+	bool FGraphLinter::ValidateNoCyclesInGraph(const FDirectedGraphAlgoAdapter& InAdapter, TArray<FBuildErrorPtr>& OutErrors) 
+	{
+		using namespace GraphLinterPrivate;
+
+		bool bIsValid = true;
+
+		TArray<FStronglyConnectedComponent> Cycles;
+
+		// In graph theory, a single vertex is technically a strongly connected
+		// component. The graph linter is only interested in strongly connected
+		// components of more than one vertex since this denotes a cycle.
+		bool bExcludeSingleVertex = true;
+
+		if(FDirectedGraphAlgo::TarjanStronglyConnectedComponents(InAdapter, Cycles, bExcludeSingleVertex))
+		{
+			bIsValid = false;
+
+			for (const FStronglyConnectedComponent& Cycle : Cycles)
+			{
+				AddBuildError<FGraphCycleError>(OutErrors, Cycle.Nodes, Cycle.Edges);
+			}
+		}
+
+		return bIsValid;
+	}
+}
diff --git a/Engine/Plugins/Runtime/Metasound/Source/MetasoundGraphCore/Private/MetasoundOperatorBuilder.cpp b/Engine/Plugins/Runtime/Metasound/Source/MetasoundGraphCore/Private/MetasoundOperatorBuilder.cpp
--- a/Engine/Plugins/Runtime/Metasound/Source/MetasoundGraphCore/Private/MetasoundOperatorBuilder.cpp
+++ b/Engine/Plugins/Runtime/Metasound/Source/MetasoundGraphCore/Private/MetasoundOperatorBuilder.cpp
@@ -2,42 +2,126 @@
 
 #include "MetasoundOperatorBuilder.h"
 
+#include "Algo/Count.h"
+#include "Algo/Transform.h"
+#include "CoreMinimal.h"
+
+#include "MetasoundBuildError.h"
 #include "MetasoundBuilderInterface.h"
 #include "MetasoundBuildError.h"
 #include "MetasoundGraphAlgo.h"
 #include "MetasoundGraphOperator.h"
 #include "MetasoundOperatorInterface.h"
+#include "MetasoundGraphAlgo.h"
+#include "MetasoundGraphLinter.h"
+#include "Templates/PimplPtr.h"
 
 
 namespace Metasound
 {
 	namespace OperatorBuilderPrivate
 	{
-		using FBuildErrorPointer = TUniquePtr<IOperatorBuildError>;
+		using FBuildErrorPtr = TUniquePtr<IOperatorBuildError>;
 
+		// Convenience template for adding build errors
 		template<typename ErrorType, typename... ArgTypes>
-		void AddBuildError(TArray<FBuildErrorPointer>& OutErrors, ArgTypes&&... Args)
+		void AddBuildError(TArray<FBuildErrorPtr>& OutErrors, ArgTypes&&... Args)
 		{
 			OutErrors.Add(MakeUnique<ErrorType>(Forward<ArgTypes>(Args)...));
 		}
 
-		void GetEdgesBetweenNodes(const TArray<INode*>& InNodes, const TArray<FDataEdge>& InEdges, TArray<FDataEdge>& OutEdges)
+		// Convenience function for adding graph cycle build errors
+		void AddBuildErrorsForCycles(const FDirectedGraphAlgoAdapter& InAdapter, TArray<FBuildErrorPtr>& OutErrors)
 		{
-			OutEdges = InEdges.FilterByPredicate(
-				[&](const FDataEdge& InEdge) 
-				{
-					bool bEdgeInNodes = InNodes.Contains(InEdge.To.Node);
-					bEdgeInNodes &= InNodes.Contains(InEdge.From.Node);
-
-					return bEdgeInNodes;
-				}
-			);
+			if (FGraphLinter::ValidateNoCyclesInGraph(InAdapter, OutErrors))
+			{
+				AddBuildError<FInternalError>(OutErrors, __FILE__, __LINE__);
+			}
 		}
 	}
 
-	FOperatorBuilder::FOperatorBuilder(const FOperatorSettings& InSettings)
+	// Default settings
+	FOperatorBuilderSettings FOperatorBuilderSettings::GetDefaultSettings()
+	{
+#if UE_BUILD_SHIPPING
+		return GetDefaultShippingSettings();
+#elif UE_BUILD_TEST
+		return GetDefaultTestSettings();
+#elif UE_BUILD_DEVELOPMENT
+		return GetDefaultDevelopementSettings();
+#elif UE_BUILD_DEBUG
+		return GetDefaultDebugSettings();
+#else
+		return GetDefaultShippingSettings();
+#endif
+	}
+
+	// Debug settings
+	FOperatorBuilderSettings FOperatorBuilderSettings::GetDefaultDebugSettings()
+	{
+		FOperatorBuilderSettings Settings;
+
+		Settings.PruningMode = EOperatorBuilderNodePruning::None;
+		Settings.bValidateNoCyclesInGraph = true;
+		Settings.bValidateNoDuplicateInputs = true;
+		Settings.bValidateVerticesExist = true;
+		Settings.bValidateEdgeDataTypesMatch = true;
+		Settings.bFailOnAnyError = false;
+
+		return Settings;
+	}
+
+	// Development settings
+	FOperatorBuilderSettings FOperatorBuilderSettings::GetDefaultDevelopementSettings()
+	{
+		FOperatorBuilderSettings Settings;
+
+		Settings.PruningMode = EOperatorBuilderNodePruning::PruneNodesWithoutOutputDependency;
+		Settings.bValidateNoCyclesInGraph = true;
+		Settings.bValidateNoDuplicateInputs = true;
+		Settings.bValidateVerticesExist = true;
+		Settings.bValidateEdgeDataTypesMatch = true;
+		Settings.bFailOnAnyError = false;
+
+		return Settings;
+	}
+
+	// Testing settings
+	FOperatorBuilderSettings FOperatorBuilderSettings::GetDefaultTestSettings()
+	{
+		FOperatorBuilderSettings Settings;
+
+		Settings.PruningMode = EOperatorBuilderNodePruning::PruneNodesWithoutOutputDependency;
+		Settings.bValidateNoCyclesInGraph = false;
+		Settings.bValidateNoDuplicateInputs = false;
+		Settings.bValidateVerticesExist = false;
+		Settings.bValidateEdgeDataTypesMatch = false;
+		Settings.bFailOnAnyError = false;
+
+		return Settings;
+	}
+
+	// Shipping settings
+	FOperatorBuilderSettings FOperatorBuilderSettings::GetDefaultShippingSettings()
+	{
+		FOperatorBuilderSettings Settings;
+
+		Settings.PruningMode = EOperatorBuilderNodePruning::PruneNodesWithoutOutputDependency;
+		Settings.bValidateNoCyclesInGraph = false;
+		Settings.bValidateNoDuplicateInputs = false;
+		Settings.bValidateVerticesExist = false;
+		Settings.bValidateEdgeDataTypesMatch = false;
+		Settings.bFailOnAnyError = false;
+
+		return Settings;
+	}
+
+
+	FOperatorBuilder::FOperatorBuilder(const FOperatorSettings& InSettings, const FOperatorBuilderSettings& InBuilderSettings)
 	:	OperatorSettings(InSettings)
+	,	BuilderSettings(InBuilderSettings)
 	{
+		MaxBuildStatusErrorLevel = BuilderSettings.bFailOnAnyError ? FBuildStatus::NoError : FBuildStatus::NonFatalError;
 	}
 
 	FOperatorBuilder::~FOperatorBuilder()
@@ -46,147 +130,198 @@ namespace Metasound
 
 	TUniquePtr<IOperator> FOperatorBuilder::BuildGraphOperator(const IGraph& InGraph, TArray<FBuildErrorPtr>& OutErrors)
 	{
-		TArray<INode*> SortedNodes;
+		using namespace OperatorBuilderPrivate;
 
-		FNodeEdgeMultiMap NodeInputEdges;
+		FBuildStatus BuildStatus = FBuildStatus::NoError;
 
-		bool bSuccess = true;
+		// Validate that the sources and destinations declared in an edge actually
+		// exist in the node.
+		if (BuilderSettings.bValidateVerticesExist)
+		{
+			if (!FGraphLinter::ValidateVerticesExist(InGraph, OutErrors))
+			{
+				BuildStatus |= FBuildStatus::FatalError;
+			}
+		}
 
+		// Validate that the data types for a source and destination match.
+		if (BuilderSettings.bValidateEdgeDataTypesMatch)
+		{
+			if (!FGraphLinter::ValidateEdgeDataTypesMatch(InGraph, OutErrors))
+			{
+				BuildStatus |= FBuildStatus::FatalError;
+			}
+		}
+		
+		// Validate that node inputs only have one source
+		if (BuilderSettings.bValidateNoDuplicateInputs)
+		{
+			if (!FGraphLinter::ValidateNoDuplicateInputs(InGraph, OutErrors))
+			{
+				BuildStatus |= FBuildStatus::FatalError;
+			}
+		}
 
-		bSuccess = GroupInputEdges(InGraph.GetDataEdges(), NodeInputEdges, OutErrors);
-		if (!bSuccess)
+		// Possible early exit if edge validation fails.
+		if (BuildStatus > MaxBuildStatusErrorLevel)
 		{
 			return TUniquePtr<IOperator>(nullptr);
 		}
 
-		// TODO: Validate edges.
-
-		bSuccess = TopologicalSort(InGraph, SortedNodes, OutErrors);
-		if (!bSuccess)
+		// Create algo adapter view of graph to cache graph operations.
+		TPimplPtr<FDirectedGraphAlgoAdapter> AlgoAdapter = FDirectedGraphAlgo::CreateDirectedGraphAlgoAdapter(InGraph);
+		
+		if (!AlgoAdapter.IsValid())
 		{
+			AddBuildError<FInternalError>(OutErrors, __FILE__, __LINE__);
 			return TUniquePtr<IOperator>(nullptr);
 		}
 
-		// TODO: Add graph pruning step.
+		TArray<const INode*> SortedNodes;
+
+		// Sort the nodes in a valid execution order
+		BuildStatus |= DepthFirstTopologicalSort(*AlgoAdapter, SortedNodes, OutErrors);
+
+		// Prune unneeded nodes. 
+		BuildStatus |= PruneNodes(*AlgoAdapter, SortedNodes, OutErrors);
+
+		// Check build status in case build routine should be exited early.
+		if (BuildStatus > MaxBuildStatusErrorLevel)
+		{
+			return TUniquePtr<IOperator>(nullptr);
+		}
 
 		TArray<TUniquePtr<IOperator>> Operators;
 		FNodeDataReferenceMap NodeData;
 
-		bSuccess = CreateOperators(SortedNodes, NodeInputEdges, Operators, NodeData, OutErrors);
-		if (!bSuccess)
+		// Create node operators from factories. 
+		BuildStatus |= CreateOperators(InGraph, SortedNodes, Operators, NodeData, OutErrors);
+
+		if (BuildStatus > MaxBuildStatusErrorLevel)
 		{
 			return TUniquePtr<IOperator>(nullptr);
 		}
 
+		// Create graph operator from collection of node operators.
 		return CreateGraphOperator(InGraph, Operators, NodeData, OutErrors);
 	}
 
-	bool FOperatorBuilder::GroupInputEdges(const TArray<FDataEdge>& InEdges, FNodeEdgeMultiMap& OutNodeInputs, TArray<FBuildErrorPtr>& OutErrors) const
+	FOperatorBuilder::FBuildStatus FOperatorBuilder::DepthFirstTopologicalSort(const FDirectedGraphAlgoAdapter& InAdapter, TArray<const INode*>& OutNodes, TArray<FBuildErrorPtr>& OutErrors) const
 	{
-		for (const FDataEdge& Edge : InEdges)
+		using namespace OperatorBuilderPrivate;
+
+		bool bSuccess = FDirectedGraphAlgo::DepthFirstTopologicalSort(InAdapter, OutNodes);
+
+		if (!bSuccess)
 		{
-			OutNodeInputs.Add(Edge.To.Node, &Edge);
+			// If there was an error, there is likely a cycle in the graph.
+			AddBuildErrorsForCycles(InAdapter, OutErrors);
+
+			return FBuildStatus::FatalError;
 		}
 
-		// TODO: error on duplicate inputs.
-		
-		return true;
+		return FBuildStatus::NoError;
 	}
 
-	bool FOperatorBuilder::TopologicalSort(const IGraph& InGraph, TArray<INode*>& OutNodes, TArray<FBuildErrorPtr>& OutErrors) const
+	FOperatorBuilder::FBuildStatus FOperatorBuilder::KahnsTopologicalSort(const FDirectedGraphAlgoAdapter& InAdapter, TArray<const INode*>& OutNodes, TArray<FBuildErrorPtr>& OutErrors) const
 	{
 		using namespace OperatorBuilderPrivate;
-		using FNodePair = TPair<INode*, INode*>;
-		using FNodeMultiMap = TMultiMap<INode*, INode*>;
-		// TODO: Implement depth first topological sort to improve cache coherency.
-		
-		TArray<INode*> UniqueNodes;
-		FNodeMultiMap Dependencies;
 
-		for (const FDataEdge& Edge : InGraph.GetDataEdges())
+		bool bSuccess = FDirectedGraphAlgo::KahnTopologicalSort(InAdapter, OutNodes);
+
+		if (!bSuccess)
 		{
-			if (nullptr == Edge.To.Node)
-			{
-				// Edges are required to be connected to a destination node.
-				AddBuildError<FDanglingEdgeError>(OutErrors, Edge);
+			// If there was an error, there is likely a cycle in the graph.
+			AddBuildErrorsForCycles(InAdapter, OutErrors);
 
-				continue;
-			}
+			return FBuildStatus::FatalError;
+		}
 
-			if (nullptr == Edge.From.Node)
-			{
-				// Edges are required to be connected to a source node.
-				AddBuildError<FDanglingEdgeError>(OutErrors, Edge);
+		return FBuildStatus::NoError;
+	}
 
-				continue;
-			}
+	FOperatorBuilder::FBuildStatus FOperatorBuilder::PruneNodes(const FDirectedGraphAlgoAdapter& InAdapter, TArray<const INode*>& InOutNodes, TArray<FBuildErrorPtr>& OutErrors) const
+	{
+		using namespace OperatorBuilderPrivate;
 
-			Dependencies.AddUnique(Edge.To.Node, Edge.From.Node);
-			UniqueNodes.AddUnique(Edge.To.Node);
-			UniqueNodes.AddUnique(Edge.From.Node);
-		}
+		FBuildStatus BuildStatus = FBuildStatus::NoError;
 
+		TSet<const INode*> ReachableNodes;
 
-		// Sort graph so that nodes with no dependencies always go first.
-		while (UniqueNodes.Num() > 0)
+		switch (BuilderSettings.PruningMode)
 		{
-			TArray<INode*> IndependentNodes = UniqueNodes.FilterByPredicate([&](INode* InNode) { 
-					return (Dependencies.Num(InNode) == 0);
-			});
+			case EOperatorBuilderNodePruning::PruneNodesWithoutExternalDependency:
+				FDirectedGraphAlgo::FindReachableNodes(InAdapter, ReachableNodes);
+				break;
 
-			if (0 == IndependentNodes.Num())
-			{
-				// likely a cycle in the graph.
-				
-				bool bExcludeSingleVertex = true;
-
-				// Try to find cycles
-				TArray<FStronglyConnectedComponent> Cycles;
-
-				if(FDirectedGraphAlgo::TarjanStronglyConnectedComponents(InGraph, Cycles, bExcludeSingleVertex))
-				{
-					for (const FStronglyConnectedComponent& Cycle : Cycles)
-					{
-						AddBuildError<FGraphCycleError>(OutErrors, Cycle.Nodes, Cycle.Edges);
-					}
-				}
-				else
-				{
-					AddBuildError<FInternalError>(OutErrors, TEXT(__FILE__), __LINE__);
-				}
-
-				return false;
-			}
+			case EOperatorBuilderNodePruning::PruneNodesWithoutOutputDependency:
+				FDirectedGraphAlgo::FindReachableNodesFromOutput(InAdapter, ReachableNodes);
+				break;
 
-			FNodeMultiMap UpdatedDependencies;
-			
-			// Remove independent nodes from dependency map.
-			for (FNodePair& Element : Dependencies)
+			case EOperatorBuilderNodePruning::PruneNodesWithoutInputDependency:
+				FDirectedGraphAlgo::FindReachableNodesFromInput(InAdapter, ReachableNodes);
+				break;
+
+			case EOperatorBuilderNodePruning::None:
+			default:
+				return FBuildStatus::NoError;
+		}
+
+		if (InOutNodes.Num() == ReachableNodes.Num())
+		{
+			// Nothing to remove since all nodes are reachable. It's assumed that
+			// InOutNodes has a unique set of nodes. 
+			return FBuildStatus::NoError;
+		}
+
+		if (0 == ReachableNodes.Num())
+		{
+			// Pruning all nodes. 
+			for (const INode* Node : InOutNodes)
 			{
-				if (!IndependentNodes.Contains(Element.Value))
-				{
-					UpdatedDependencies.AddUnique(Element.Key, Element.Value);
-				}
+				AddBuildError<FNodePrunedError>(OutErrors, Node);
 			}
-			Dependencies = UpdatedDependencies;
 
-			// Remove independent nodes from node list
-			for (INode* Node : IndependentNodes)
+			InOutNodes.Reset();
+
+			// This is non fatal, but results in an IOperator which is a No-op.
+			return FBuildStatus::NonFatalError;
+		}
+
+		// Split the nodes-to-keep and the nodes-to-prune into two arrays. Need 
+		// to ensure that kept nodes are still in same relative order. 
+		TArray<const INode*> SortedNodesToKeep;
+
+		SortedNodesToKeep.Reserve(ReachableNodes.Num());
+
+		for (const INode* Node : InOutNodes)
+		{
+			if (ReachableNodes.Contains(Node))
 			{
-				UniqueNodes.RemoveSwap(Node);
+				SortedNodesToKeep.Add(Node);
 			}
+			else
+			{
+				AddBuildError<FNodePrunedError>(OutErrors, Node);
 
-			// Add independent nodes to output.
-			OutNodes.Append(IndependentNodes);
+				// Denote a pruned node as a non-fatal error. In the future this
+				// may be simply a warning as some nodes are required to conform
+				// to metasound archetypes even if they are unused.
+				BuildStatus |= FBuildStatus::NonFatalError;
+			}
 		}
 
-		return true;
+		InOutNodes = SortedNodesToKeep;
+
+		return BuildStatus;
 	}
 
-	bool FOperatorBuilder::GatherInputDataReferences(const INode* InNode, const FNodeEdgeMultiMap& InEdgeMap, const FNodeDataReferenceMap& InDataReferenceMap, FDataReferenceCollection& OutCollection, TArray<FBuildErrorPtr>& OutErrors) const
+	FOperatorBuilder::FBuildStatus FOperatorBuilder::GatherInputDataReferences(const INode* InNode, const FNodeEdgeMultiMap& InEdgeMap, const FNodeDataReferenceMap& InDataReferenceMap, FDataReferenceCollection& OutCollection, TArray<FBuildErrorPtr>& OutErrors) const
 	{
 		using namespace OperatorBuilderPrivate;
 
+		FBuildStatus BuildStatus;
+
 		// Find all input edges associated with this node. 
 		TArray<const FDataEdge*> Edges;
 		InEdgeMap.MultiFind(InNode, Edges);
@@ -200,7 +335,7 @@ namespace Metasound
 				// than a user error.
 				AddBuildError<FInternalError>(OutErrors, TEXT(__FILE__), __LINE__);
 
-				return false;
+				return FBuildStatus::NonFatalError;
 			}
 
 			
@@ -212,51 +347,59 @@ namespace Metasound
 				// did not match the readable parameter refs created by the operators outputs. Or, the edge description is invalid.
 
 				// TODO: consider checking that outputs of node operators match output descriptions of nodes. 
-				// TODO: consider checking the edges are supported by nodes existing parameter descriptions.
 				AddBuildError<FMissingOutputDataReferenceError>(OutErrors, Edge->From);
 
-				return false;
+				BuildStatus |= FBuildStatus::NonFatalError;
+				continue;
 			}
 
-			if (Edge->From.Vertex.DataReferenceTypeName != Edge->To.Vertex.DataReferenceTypeName)
-			{
-				// TODO: may want to have this check higher up in the stack. 
-				AddBuildError<FInvalidConnectionDataTypeError>(OutErrors, *Edge);
-
-				return false;
-			}
 
 			bool bSuccess = OutCollection.AddDataReadReferenceFrom(Edge->To.Vertex.VertexName, FromDataReferenceCollection, Edge->From.Vertex.VertexName, Edge->From.Vertex.DataReferenceTypeName);
 
 			if (!bSuccess)
 			{
 				AddBuildError<FMissingOutputDataReferenceError>(OutErrors, Edge->From);
-				return false;
+				BuildStatus |= FBuildStatus::NonFatalError;
 			}
 		}
 
-		return true;
+		return BuildStatus;
 	}
 
-	bool FOperatorBuilder::CreateOperators(const TArray<INode*>& InSortedNodes, FNodeEdgeMultiMap& InNodeInputEdges, TArray<TUniquePtr<IOperator>>& OutOperators, FNodeDataReferenceMap& OutDataReferences, TArray<FBuildErrorPtr>& OutErrors)
+	FOperatorBuilder::FBuildStatus FOperatorBuilder::CreateOperators(const IGraph& InGraph, const TArray<const INode*>& InSortedNodes, TArray<TUniquePtr<IOperator>>& OutOperators, FNodeDataReferenceMap& OutDataReferences, TArray<FBuildErrorPtr>& OutErrors) const
 	{
-		for (INode* Node : InSortedNodes)
+		FBuildStatus BuildStatus;
+
+		FNodeEdgeMultiMap NodeInputEdges;
+
+		// Gather input edges for each node.
+		for (const FDataEdge& Edge : InGraph.GetDataEdges())
+		{
+			NodeInputEdges.Add(Edge.To.Node, &Edge);
+		}
+
+		// Call operator factory for each node.
+		for (const INode* Node : InSortedNodes)
 		{
+
 			FDataReferenceCollection InputCollection;
 
 			// Gather the input parameters for this IOperator from the output parameters of already created IOperators. 
-			if (!GatherInputDataReferences(Node, InNodeInputEdges, OutDataReferences, InputCollection, OutErrors))
+			BuildStatus |= GatherInputDataReferences(Node, NodeInputEdges, OutDataReferences, InputCollection, OutErrors);
+
+			if (BuildStatus >= FBuildStatus::FatalError)
 			{
-				return false;
+				return BuildStatus;
 			}
 
-			IOperatorFactory& Factory = Node->GetDefaultOperatorFactory();
+			// TODO: make factory retrieval const 
+			IOperatorFactory& Factory = const_cast<INode*>(Node)->GetDefaultOperatorFactory();
 
 			FOperatorPtr Operator = Factory.CreateOperator(*Node, OperatorSettings, InputCollection, OutErrors);
 
 			if (!Operator.IsValid())
 			{
-				return false;
+				return FBuildStatus::FatalError;
 			}
 
 			// Save input and output collection for future use 
@@ -266,24 +409,28 @@ namespace Metasound
 			OutOperators.Add(MoveTemp(Operator));
 		}
 
-		return true;
+		return BuildStatus;
 	}
 
-	bool FOperatorBuilder::GatherGraphDataReferences(const IGraph& InGraph, FNodeDataReferenceMap& InNodeDataReferences, FDataReferenceCollection& OutGraphInputs, FDataReferenceCollection& OutGraphOutputs, TArray<FBuildErrorPtr>& OutErrors) const
+	FOperatorBuilder::FBuildStatus FOperatorBuilder::GatherGraphDataReferences(const IGraph& InGraph, FNodeDataReferenceMap& InNodeDataReferences, FDataReferenceCollection& OutGraphInputs, FDataReferenceCollection& OutGraphOutputs, TArray<FBuildErrorPtr>& OutErrors) const
 	{
 		using namespace OperatorBuilderPrivate;
 		using FDestinationElement = FInputDataDestinationCollection::ElementType;
 		using FSourceElement = FOutputDataSourceCollection::ElementType;
 
-		// Gather inputs
+		FBuildStatus BuildStatus;
+
+		// Gather graph inputs
 		for (const FDestinationElement& Element : InGraph.GetInputDataDestinations())
 		{
 			const FInputDataDestination& InputDestination = Element.Value;
 
 			if (!InNodeDataReferences.Contains(InputDestination.Node))
 			{
+				// An input node was likely pruned.
 				AddBuildError<FMissingInputDataReferenceError>(OutErrors, InputDestination);
-				// This indicates an input node wasn't connected in some way to the output, and was culled earlier.
+
+				BuildStatus |= FBuildStatus::NonFatalError;
 				continue;
 			}
 
@@ -292,7 +439,9 @@ namespace Metasound
 			if (!Collection.ContainsDataWriteReference(InputDestination.Vertex.VertexName, InputDestination.Vertex.DataReferenceTypeName))
 			{
 				AddBuildError<FMissingInputDataReferenceError>(OutErrors, InputDestination);
-				return false;
+
+				BuildStatus |= FBuildStatus::NonFatalError;
+				continue;
 			}
 
 			bool bSuccess = OutGraphInputs.AddDataWriteReferenceFrom(InputDestination.Vertex.VertexName, Collection, InputDestination.Vertex.VertexName, InputDestination.Vertex.DataReferenceTypeName);
@@ -300,18 +449,24 @@ namespace Metasound
 			if (!bSuccess)
 			{
 				AddBuildError<FMissingInputDataReferenceError>(OutErrors, InputDestination);
-				return false;
+
+				BuildStatus |= FBuildStatus::NonFatalError;
+				continue;
 			}
 		}
 
-		// Gather outputs
+		// Gather graph outputs.
 		for (const FSourceElement& Element : InGraph.GetOutputDataSources())
 		{
 			const FOutputDataSource& OutputSource = Element.Value;
 
 			if (!InNodeDataReferences.Contains(OutputSource.Node))
 			{
+				// An output node was likely pruned.
 				AddBuildError<FMissingOutputDataReferenceError>(OutErrors, OutputSource);
+
+				BuildStatus |= FBuildStatus::NonFatalError;
+
 				continue;
 			}
 
@@ -319,20 +474,30 @@ namespace Metasound
 
 			if (!Collection.ContainsDataReadReference(OutputSource.Vertex.VertexName, OutputSource.Vertex.DataReferenceTypeName))
 			{
+				// This will likely produce an IOperator which does not work as
+				// expected.
 				AddBuildError<FMissingOutputDataReferenceError>(OutErrors, OutputSource);
-				return false;
+
+				BuildStatus |= FBuildStatus::NonFatalError;
+
+				continue;
 			}
 
 			bool bSuccess = OutGraphOutputs.AddDataReadReferenceFrom(OutputSource.Vertex.VertexName, Collection, OutputSource.Vertex.VertexName, OutputSource.Vertex.DataReferenceTypeName);
 
 			if (!bSuccess)
 			{
+				// This will likely produce an IOperator which does not work as
+				// expected.
 				AddBuildError<FMissingOutputDataReferenceError>(OutErrors, OutputSource);
-				return false;
+
+				BuildStatus |= FBuildStatus::NonFatalError;
+
+				continue;
 			}
 		}
 	
-		return true;
+		return BuildStatus;
 	}
 
 	TUniquePtr<IOperator> FOperatorBuilder::CreateGraphOperator(const IGraph& InGraph, TArray<FOperatorPtr>& InOperators, FNodeDataReferenceMap& InNodeDataReferences, TArray<FBuildErrorPtr>& OutErrors) const
@@ -340,9 +505,9 @@ namespace Metasound
 		FDataReferenceCollection GraphInputs;
 		FDataReferenceCollection GraphOutputs;
 
-		bool bSuccess = GatherGraphDataReferences(InGraph, InNodeDataReferences, GraphInputs, GraphOutputs, OutErrors);
+		FBuildStatus BuildStatus = GatherGraphDataReferences(InGraph, InNodeDataReferences, GraphInputs, GraphOutputs, OutErrors);
 
-		if (!bSuccess)
+		if (BuildStatus > MaxBuildStatusErrorLevel)
 		{
 			return TUniquePtr<IOperator>(nullptr);
 		}
diff --git a/Engine/Plugins/Runtime/Metasound/Source/MetasoundGraphCore/Private/MetasoundOperatorSettings.cpp b/Engine/Plugins/Runtime/Metasound/Source/MetasoundGraphCore/Private/MetasoundOperatorSettings.cpp
--- a/Engine/Plugins/Runtime/Metasound/Source/MetasoundGraphCore/Private/MetasoundOperatorSettings.cpp
+++ b/Engine/Plugins/Runtime/Metasound/Source/MetasoundGraphCore/Private/MetasoundOperatorSettings.cpp
@@ -47,7 +47,7 @@ namespace Metasound
 	{
 		using namespace MetasoundOperatorSettingsPrivate;
 
-		if (ensureMsgf(InTargetBlockRate > MinimumTargetBlockRate, TEXT("Invalid block rate %f. Defaulting to %f"), InTargetBlockRate, MinimumTargetBlockRate))
+		if (ensureMsgf(InTargetBlockRate >= MinimumTargetBlockRate, TEXT("Invalid block rate %f. Defaulting to %f"), InTargetBlockRate, MinimumTargetBlockRate))
 		{
 			TargetBlockRate = InTargetBlockRate;
 		}
diff --git a/Engine/Plugins/Runtime/Metasound/Source/MetasoundGraphCore/Public/MetasoundBuildError.h b/Engine/Plugins/Runtime/Metasound/Source/MetasoundGraphCore/Public/MetasoundBuildError.h
--- a/Engine/Plugins/Runtime/Metasound/Source/MetasoundGraphCore/Public/MetasoundBuildError.h
+++ b/Engine/Plugins/Runtime/Metasound/Source/MetasoundGraphCore/Public/MetasoundBuildError.h
@@ -29,13 +29,27 @@ namespace Metasound
 			/** Returns a human readable error description. */
 			virtual const FText& GetErrorDescription() const override;
 
+			/** Returns an array of destinations associated with the error. */
+			virtual const TArray<FInputDataDestination>& GetDestinations() const override;
+			
+			/** Returns an array of sources associated with the error. */
+			virtual const TArray<FOutputDataSource>& GetSources() const override;
+
 			/** Returns an array of edges associated with the error. */
 			virtual const TArray<FDataEdge>& GetEdges() const override;
 
 			/** Returns an array of Nodes associated with the error. */
 			virtual const TArray<const INode*>& GetNodes() const override;
 
+
 		protected:
+			// Add input destinations to be associated with error.
+			void AddInputDataDestination(const FInputDataDestination& InInputDataDestination);
+			void AddInputDataDestinations(TArrayView<const FInputDataDestination> InInputDataDestinations);
+
+			// Add input destinations to be associated with error.
+			void AddOutputDataSource(const FOutputDataSource& InOutputDataSource);
+			void AddOutputDataSources(TArrayView<const FOutputDataSource> InOutputDataSources);
 
 			// Add edges to be associated with error.
 			void AddEdge(const FDataEdge& InEdge);
@@ -52,20 +66,62 @@ namespace Metasound
 
 			TArray<const INode*> Nodes;
 			TArray<FDataEdge> Edges;
+			TArray<FInputDataDestination> Destinations;
+			TArray<FOutputDataSource> Sources;
 	};
 
-	/** FDanglingEdgeError
+	/** FDanglingVertexError
 	 *
-	 * Caused by FDataEdges pointing to invalid nodes.
+	 * Caused by FDataEdges, FInputDataDestinations or FOutputDataSources 
+	 * pointing to null nodes.
 	 */
-	class METASOUNDGRAPHCORE_API FDanglingEdgeError : public FBuildErrorBase
+	class METASOUNDGRAPHCORE_API FDanglingVertexError : public FBuildErrorBase
 	{
 		public:
 			static const FName ErrorType;
 
-			FDanglingEdgeError(const FDataEdge& InEdge);
+			FDanglingVertexError(const FInputDataDestination& InDest);
+			FDanglingVertexError(const FOutputDataSource& InSource);
+			FDanglingVertexError(const FDataEdge& InEdge);
+
+			virtual ~FDanglingVertexError() = default;
+		protected:
 
-			virtual ~FDanglingEdgeError() = default;
+			FDanglingVertexError();
+
+		private:
+	};
+
+	/** FMissingVertexError
+	 *
+	 * Caused by a referenced FDataVertex which does not exist on a node.
+	 */
+	class METASOUNDGRAPHCORE_API FMissingVertexError : public FBuildErrorBase
+	{
+		public:
+			static const FName ErrorType;
+
+			FMissingVertexError(const FInputDataDestination& InDestination);
+			FMissingVertexError(const FOutputDataSource& InSource);
+
+			virtual ~FMissingVertexError() = default;
+
+		private:
+	};
+
+
+	/** FDuplicateInputError
+	 *
+	 * Caused by multiple FDataEdges pointing same FInputDataDestination
+	 */
+	class METASOUNDGRAPHCORE_API FDuplicateInputError : public FBuildErrorBase
+	{
+		public:
+			static const FName ErrorType;
+
+			FDuplicateInputError(const TArrayView<FDataEdge> InEdges);
+
+			virtual ~FDuplicateInputError() = default;
 
 		private:
 	};
@@ -86,6 +142,23 @@ namespace Metasound
 		private:
 	};
 
+	/** FNodePrunedError
+	 *
+	 * Caused by nodes which are in the graph but unreachable from the graph's
+	 * inputs and/or outputs.
+	 */
+	class METASOUNDGRAPHCORE_API FNodePrunedError : public FBuildErrorBase
+	{
+		public:
+			static const FName ErrorType;
+
+			FNodePrunedError(const INode* InNode);
+
+			virtual ~FNodePrunedError() = default;
+
+		private:
+	};
+
 	/** FInternalError
 	 *
 	 * Caused by internal state or logic errors. 
@@ -121,11 +194,6 @@ namespace Metasound
 			FMissingInputDataReferenceError(const FInputDataDestination& InInputDataDestination);
 
 			virtual ~FMissingInputDataReferenceError() = default;
-
-			const FInputDataDestination& GetInputDataDestination() const;
-
-		private:
-			FInputDataDestination InputDataDestination;
 	};
 
 	/** FMissingOutputDataReferenceError
@@ -143,10 +211,6 @@ namespace Metasound
 
 			virtual ~FMissingOutputDataReferenceError() = default;
 
-			const FOutputDataSource& GetOutputDataSource() const;
-
-		private:
-			FOutputDataSource OutputDataSource;
 	};
 
 	/** FInvalidConnectionDataTypeError
diff --git a/Engine/Plugins/Runtime/Metasound/Source/MetasoundGraphCore/Public/MetasoundBuilderInterface.h b/Engine/Plugins/Runtime/Metasound/Source/MetasoundGraphCore/Public/MetasoundBuilderInterface.h
--- a/Engine/Plugins/Runtime/Metasound/Source/MetasoundGraphCore/Public/MetasoundBuilderInterface.h
+++ b/Engine/Plugins/Runtime/Metasound/Source/MetasoundGraphCore/Public/MetasoundBuilderInterface.h
@@ -23,6 +23,12 @@ namespace Metasound
 			/** Returns a human readable error description. */
 			virtual const FText& GetErrorDescription() const = 0;
 
+			/** Returns an array of destinations associated with the error. */
+			virtual const TArray<FInputDataDestination>& GetDestinations() const = 0;
+			
+			/** Returns an array of sources associated with the error. */
+			virtual const TArray<FOutputDataSource>& GetSources() const = 0;
+
 			/** Returns an array of Nodes associated with the error. */
 			virtual const TArray<const INode*>& GetNodes() const = 0;
 
diff --git a/Engine/Plugins/Runtime/Metasound/Source/MetasoundGraphCore/Public/MetasoundGraphLinter.h b/Engine/Plugins/Runtime/Metasound/Source/MetasoundGraphCore/Public/MetasoundGraphLinter.h
new file mode 100644
--- /dev/null
+++ b/Engine/Plugins/Runtime/Metasound/Source/MetasoundGraphCore/Public/MetasoundGraphLinter.h
@@ -0,0 +1,83 @@
+// Copyright Epic Games, Inc. All Rights Reserved.
+
+#pragma once
+
+#include "CoreMinimal.h"
+
+#include "MetasoundNodeInterface.h"
+#include "MetasoundBuilderInterface.h"
+
+namespace Metasound
+{
+	// Forward declaration
+	class FDirectedGraphAlgoAdapter;
+
+	/** FGraphLinter
+	 *
+	 * Functions for detecting issues and errors in an IGraph.
+	 */
+	class METASOUNDGRAPHCORE_API FGraphLinter
+	{
+		public:
+			using FBuildErrorPtr = TUniquePtr<IOperatorBuildError>;
+
+
+			/** 
+			 * Validate each FDataEdge in the IGraph by checking that the FInputDataSource 
+			 * and FOutputDataDestination data types are equal. Errors will be generated 
+			 * if unequal data types are detected.
+			 *
+			 * @param InGraph - Graph to analyze.
+			 * @param OutErrors - Array to hold detected errors.
+			 *
+			 * @return True if no errors detected. False otherwise.
+			 */
+			static bool ValidateEdgeDataTypesMatch(const IGraph& InGraph, TArray<FBuildErrorPtr>& OutErrors);
+
+			/** 
+			 * Validate each FDataEdge in the IGraph by checking that the corresponding 
+			 * INodes contain matching FDataVertex information as described by the 
+			 * FDataEdge. Errors will be generated if inconsistencies are detected.
+			 *
+			 * @param InGraph - Graph to analyze.
+			 * @param OutErrors - Array to hold detected errors.
+			 *
+			 * @return True if no errors detected. False otherwise.
+			 */
+			static bool ValidateVerticesExist(const IGraph& InGraph, TArray<FBuildErrorPtr>& OutErrors);
+
+			/** 
+			 * Validate the IGraph by checking for duplicate inputs connected to 
+			 * an individual vertex on a given node. Errors will be generated if 
+			 * duplicates are detected. 
+			 *
+			 * @param InGraph - Graph to analyze.
+			 * @param OutErrors - Array to hold detected errors.
+			 *
+			 * @return True if no errors detected. False otherwise.
+			 */
+			static bool ValidateNoDuplicateInputs(const IGraph& InGraph, TArray<FBuildErrorPtr>& OutErrors);
+
+			/** 
+			 * Validate the IGraph by checking for cycles. Errors will be generated
+			 * if a cycle is detected in the graph.
+			 *
+			 * @param InGraph - Graph to analyze.
+			 * @param OutErrors - Array to hold detected errors.
+			 *
+			 * @return True if no errors detected. False otherwise.
+			 */
+			static bool ValidateNoCyclesInGraph(const IGraph& InGraph, TArray<FBuildErrorPtr>& OutErrors);
+
+			/** 
+			 * Validate the IGraph by checking for cycles. Errors will be generated
+			 * if a cycle is detected in the graph.
+			 *
+			 * @param InAdapter - Graph adapter to analyze.
+			 * @param OutErrors - Array to hold detected errors.
+			 *
+			 * @return True if no errors detected. False otherwise.
+			 */
+			static bool ValidateNoCyclesInGraph(const FDirectedGraphAlgoAdapter& InAdapter, TArray<FBuildErrorPtr>& OutErrors);
+	};
+}
diff --git a/Engine/Plugins/Runtime/Metasound/Source/MetasoundGraphCore/Public/MetasoundOperatorBuilder.h b/Engine/Plugins/Runtime/Metasound/Source/MetasoundGraphCore/Public/MetasoundOperatorBuilder.h
--- a/Engine/Plugins/Runtime/Metasound/Source/MetasoundGraphCore/Public/MetasoundOperatorBuilder.h
+++ b/Engine/Plugins/Runtime/Metasound/Source/MetasoundGraphCore/Public/MetasoundOperatorBuilder.h
@@ -9,17 +9,116 @@
 
 namespace Metasound
 {
+	/** EOperatorBuildNodePruning expresses the desired pruning behavior during
+	 * the node pruning step.  
+	 *
+	 * Some nodes are unreachable in the graph either by traversing from the
+	 * input nodes to from the output nodes. Because they are not dependent, 
+	 * they have no impact on the produced output and can be pruned without 
+	 * causing any change to the declared behavior of the graph.
+	 */
+	enum class EOperatorBuilderNodePruning : uint8
+	{
+		/** Do not prune any nodes. */
+		None, 
+
+		/** Prune nodes which cannot be reached from the output nodes. */
+		PruneNodesWithoutOutputDependency, 
+
+		/** Prune nodes which cannot be reached from the input nodes. */
+		PruneNodesWithoutInputDependency, 
+
+		/** Prune nodes which cannot be reached from the input nodes or output nodes. */
+		PruneNodesWithoutExternalDependency, 
+	};
+
+	/** FOperatorBuilderSettings
+	 *
+	 * Settings for building IGraphs into IOperators.
+	 */
+	struct METASOUNDGRAPHCORE_API FOperatorBuilderSettings
+	{
+		/** Desired node pruning behavior. */
+		EOperatorBuilderNodePruning PruningMode = EOperatorBuilderNodePruning::None;
+
+		/** If true, the IGraph will be analyzed to detect cycles. Errors will be
+		 * generated if a cycle is detected in the graph.
+		 */
+		bool bValidateNoCyclesInGraph = true;
+
+		/** If true, the inputs to each node in the IGraph will be analyzed to
+		 * detect duplicate inputs connected to an individual vertex on a given 
+		 * node.  Errors will be generated if duplicates are detected. */
+		bool bValidateNoDuplicateInputs = true;
+
+		/** If true, each FDataEdge in the IGraph will be validated by checking
+		 * that the corresponding INodes contain matching FDataVertex information
+		 * as described by the FDataEdge. Errors will be generated if 
+		 * inconsistencies are detected.
+		 */
+		bool bValidateVerticesExist = true;
+
+		/** If true, each FDataEdge in the IGraph will be validated by checking
+		 * that the FInputDataSource and FOutputDataDestination data types are
+		 * equal. Errors will be generated if unequal data types are detected.
+		 */
+		bool bValidateEdgeDataTypesMatch = true;
+
+		/** If true, the builder will return an invalid IOperator if any errors
+		 * are detected. If false, the builder will return an invalid IOperator
+		 * only if fatal errors are detected.
+		 */
+		bool bFailOnAnyError = false;
+
+		/** Return the default settings for the current build environment. */
+		static FOperatorBuilderSettings GetDefaultSettings();
+
+		/** Return the default settings for a debug build environment. */
+		static FOperatorBuilderSettings GetDefaultDebugSettings();
+
+		/** Return the default settings for a development build environment. */
+		static FOperatorBuilderSettings GetDefaultDevelopementSettings();
+
+		/** Return the default settings for a test build environment. */
+		static FOperatorBuilderSettings GetDefaultTestSettings();
+
+		/** Return the default settings for a shipping build environment. */
+		static FOperatorBuilderSettings GetDefaultShippingSettings();
+	};
+
+	// Forward declare.
+	class FDirectedGraphAlgoAdapter;
+
+	/** FOperatorBuilder builds an IOperator from an IGraph. */
 	class METASOUNDGRAPHCORE_API FOperatorBuilder : public IOperatorBuilder
 	{
 		public:
-			FOperatorBuilder(const FOperatorSettings& InSettings);
+
+			/** FOperatorBuilder constructor.
+			 *
+			 * @param InOperatorSettings - Settings to be passed to all operators
+			 *                             on creation.
+			 * @param InBuilderSettings  - Settings to configure builder options.
+			 */
+			FOperatorBuilder(const FOperatorSettings& InOperatorSettings, const FOperatorBuilderSettings& InBuilderSettings);
 
 			virtual ~FOperatorBuilder();
 
+			/** Create an IOperator from an IGraph.
+			 *
+			 * @param InGraph   - The graph containing input, output and edge information.
+			 * @param OutErrors - An array of build errors that will be populated 
+			 *                    with any issues encountered during the build process.
+			 *
+			 * @return A TUniquePtr to an IOperator. If the processes was unsuccessful, 
+			 *         the returned pointer will contain a nullptr and be invalid.
+			 */
 			virtual TUniquePtr<IOperator> BuildGraphOperator(const IGraph& InGraph, TArray<FBuildErrorPtr>& OutErrors) override;
 			
 
 		private:
+			// Collection of existing inputs and outputs associated with a given
+			// IOperator.
 			struct FOperatorDataReferences
 			{
 				FOperatorDataReferences()
@@ -35,23 +134,83 @@ namespace Metasound
 				FDataReferenceCollection Inputs;
 				FDataReferenceCollection Outputs;
 			};
+
 			using FNodeEdgeMultiMap = TMultiMap<const INode*, const FDataEdge*>;
-			using FNodeDataReferenceMap = TMap<INode*, FOperatorDataReferences>;
+			using FNodeDataReferenceMap = TMap<const INode*, FOperatorDataReferences>;
 			using FOperatorPtr = TUniquePtr<IOperator>;
 
-			bool GroupInputEdges(const TArray<FDataEdge>& InEdges, FNodeEdgeMultiMap& OutNodeInputs, TArray<FBuildErrorPtr>& OutErrors) const;
+			// Handles build status of current build operation.
+			struct FBuildStatus
+			{
+				// Enumeration of build status states. 
+				//
+				// Note: plain enum used here instead of enum class so that implicit 
+				// conversion to int32 can be utilized. It is assumed that the 
+				// build status int32 values increase as the build status deteriorates.
+				// Build statuses are merged by taking the maximum int32 value of
+				// the EStatus. 
+				enum EStatus 
+				{
+					// No error has been encountered.
+					NoError = 0,
+					
+					// A non fatal error has been encountered.
+					NonFatalError = 1,
+
+					// A fatal error has been encountered.
+					FatalError = 2
+				};
+
+				FBuildStatus() = default;
 
-			bool TopologicalSort(const IGraph& InGraph, TArray<INode*>& OutNodes, TArray<FBuildErrorPtr>& OutErrors) const;
+				FBuildStatus(FBuildStatus::EStatus InStatus)
+				:	Value(InStatus)
+				{
+				}
+
+				// Merge build statuses by taking the maximum of EStatus.
+				FBuildStatus& operator |= (FBuildStatus RHS)
+				{
+					Value = Value > RHS.Value ? Value : RHS.Value;
+					return *this;
+				}
+
+				operator EStatus() const
+				{
+					return Value;
+				}
 
-			bool GatherInputDataReferences(const INode* InNode, const FNodeEdgeMultiMap& InEdgeMap, const FNodeDataReferenceMap& InDataReferenceMap, FDataReferenceCollection& OutCollection, TArray<FBuildErrorPtr>& OutErrors) const;
+				private:
 
-			bool GatherGraphDataReferences(const IGraph& InGraph, FNodeDataReferenceMap& InNodeDataReferences, FDataReferenceCollection& OutGraphInputs, FDataReferenceCollection& OutGraphOutputs, TArray<FBuildErrorPtr>& OutErrors) const;
+				EStatus Value = NoError;
+			};
+
+			// Perform topological sort using depth first algorithm.
+			FBuildStatus DepthFirstTopologicalSort(const FDirectedGraphAlgoAdapter& InAdapter, TArray<const INode*>& OutNodes, TArray<FBuildErrorPtr>& OutErrors) const;
+
+			// Perform topological sort using kahns algorithm.
+			FBuildStatus KahnsTopologicalSort(const FDirectedGraphAlgoAdapter& InAdapter, TArray<const INode*>& OutNodes, TArray<FBuildErrorPtr>& OutErrors) const;
+
+			// Prune unreachable nodes from InOutNodes
+			FBuildStatus PruneNodes(const FDirectedGraphAlgoAdapter& InAdapter, TArray<const INode*>& InOutNodes, TArray<FBuildErrorPtr>& OutErrors) const;
 
-			bool CreateOperators(const TArray<INode*>& InSortedNodes, FNodeEdgeMultiMap& InNodeInputEdges, TArray<FOperatorPtr>& OutOperators, FNodeDataReferenceMap& OutDataReferences, TArray<FBuildErrorPtr>& OutErrors);
+			// Get all input data references for a given node.
+			FBuildStatus GatherInputDataReferences(const INode* InNode, const FNodeEdgeMultiMap& InEdgeMap, const FNodeDataReferenceMap& InDataReferenceMap, FDataReferenceCollection& OutCollection, TArray<FBuildErrorPtr>& OutErrors) const;
 
+			// Get all input/output data references for a given graph.
+			FBuildStatus GatherGraphDataReferences(const IGraph& InGraph, FNodeDataReferenceMap& InNodeDataReferences, FDataReferenceCollection& OutGraphInputs, FDataReferenceCollection& OutGraphOutputs, TArray<FBuildErrorPtr>& OutErrors) const;
+
+			// Call the operator factories for the nodes
+			FBuildStatus CreateOperators(const IGraph& InGraph, const TArray<const INode*>& InSortedNodes, TArray<FOperatorPtr>& OutOperators, FNodeDataReferenceMap& OutDataReferences, TArray<FBuildErrorPtr>& OutErrors) const;
+
+			// Create the final graph operator from the individual operators.
 			TUniquePtr<IOperator> CreateGraphOperator(const IGraph& InGraph, TArray<FOperatorPtr>& InOperators, FNodeDataReferenceMap& InNodeDataReferences, TArray<FBuildErrorPtr>& OutErrors) const;
 
 			FOperatorSettings OperatorSettings;
+
+			FOperatorBuilderSettings BuilderSettings;
+
+			FBuildStatus MaxBuildStatusErrorLevel;
 	};
 }
 
