Support for UObject literals in metasounds.

diff --git a/Engine/Plugins/Runtime/Metasound/Source/MetasoundEditor/Private/MetasoundEditorGraphBuilder.cpp b/Engine/Plugins/Runtime/Metasound/Source/MetasoundEditor/Private/MetasoundEditorGraphBuilder.cpp
--- a/Engine/Plugins/Runtime/Metasound/Source/MetasoundEditor/Private/MetasoundEditorGraphBuilder.cpp
+++ b/Engine/Plugins/Runtime/Metasound/Source/MetasoundEditor/Private/MetasoundEditorGraphBuilder.cpp
@@ -10,6 +10,7 @@
 #include "MetasoundEditorGraphNode.h"
 #include "MetasoundEditorModule.h"
 #include "MetasoundFrontend.h"
+#include "MetasoundFrontendRegistries.h"
 #include "Modules/ModuleManager.h"
 #include "ScopedTransaction.h"
 #include "Templates/Tuple.h"
@@ -22,10 +23,12 @@ namespace Metasound
 {
 	namespace Editor
 	{
-		const FName FGraphBuilder::PinPrimitiveBoolean = "Boolean";
-		const FName FGraphBuilder::PinPrimitiveFloat = "Float";
-		const FName FGraphBuilder::PinPrimitiveInteger = "Int";
-		const FName FGraphBuilder::PinPrimitiveString = "String";
+		const FName FGraphBuilder::PinPrimitiveBoolean = TEXT("Boolean");
+		const FName FGraphBuilder::PinPrimitiveFloat = TEXT("Float");
+		const FName FGraphBuilder::PinPrimitiveInteger = TEXT("Int");
+		const FName FGraphBuilder::PinPrimitiveString = TEXT("String");
+		const FName FGraphBuilder::PinPrimitiveUObject = TEXT("UObject");
+		const FName FGraphBuilder::PinPrimitiveUObjectArray = TEXT("UObjectArray");
 
 		UEdGraphNode* FGraphBuilder::AddNode(UObject& InMetasound, const FVector2D& Location, Frontend::FNodeHandle& InNodeHandle, bool bInSelectNewNode)
 		{
@@ -185,11 +188,34 @@ namespace Metasound
 				}
 				break;
 
+				case ELiteralArgType::UObjectProxy:
+				{
+					UClass* ClassToUse = FMetasoundFrontendRegistryContainer::Get()->GetLiteralUClassForDataType(InTypeName);
+					if (ClassToUse)
+					{
+						GraphHandle.SetInputToLiteral(InName, ClassToUse->ClassDefaultObject);
+					}
+				}
+				break;
+
+				case ELiteralArgType::UObjectProxyArray:
+				{
+					UClass* ClassToUse = FMetasoundFrontendRegistryContainer::Get()->GetLiteralUClassForDataType(InTypeName);
+					if (ClassToUse)
+					{
+						TArray<UObject*> ObjectArray;
+						ObjectArray.Add(ClassToUse->ClassDefaultObject);
+						GraphHandle.SetInputToLiteral(InName, ObjectArray);
+					}
+					GraphHandle.SetInputToLiteral(InName, FString(TEXT("")));
+				}
+				break;
+
 				case ELiteralArgType::Invalid:
 				case ELiteralArgType::None:
 				default:
 				{
-					static_assert(static_cast<int32>(ELiteralArgType::Invalid) == 5, "Possible missing ELiteralArgType case coverage");
+					static_assert(static_cast<int32>(ELiteralArgType::Invalid) == 7, "Possible missing ELiteralArgType case coverage");
 				}
 				break;
 			}
diff --git a/Engine/Plugins/Runtime/Metasound/Source/MetasoundEditor/Private/MetasoundEditorGraphBuilder.h b/Engine/Plugins/Runtime/Metasound/Source/MetasoundEditor/Private/MetasoundEditorGraphBuilder.h
--- a/Engine/Plugins/Runtime/Metasound/Source/MetasoundEditor/Private/MetasoundEditorGraphBuilder.h
+++ b/Engine/Plugins/Runtime/Metasound/Source/MetasoundEditor/Private/MetasoundEditorGraphBuilder.h
@@ -27,6 +27,8 @@ namespace Metasound
 			static const FName PinPrimitiveFloat;
 			static const FName PinPrimitiveInteger;
 			static const FName PinPrimitiveString;
+			static const FName PinPrimitiveUObject;
+			static const FName PinPrimitiveUObjectArray;
 
 			// Adds a node to the editor graph that corresponds to the provided node handle.
 			static UEdGraphNode* AddNode(UObject& InMetasound, const FVector2D& Location, Frontend::FNodeHandle& InNodeHandle, bool bInSelectNewNode = true);
diff --git a/Engine/Plugins/Runtime/Metasound/Source/MetasoundEditor/Private/MetasoundEditorModule.cpp b/Engine/Plugins/Runtime/Metasound/Source/MetasoundEditor/Private/MetasoundEditorModule.cpp
--- a/Engine/Plugins/Runtime/Metasound/Source/MetasoundEditor/Private/MetasoundEditorModule.cpp
+++ b/Engine/Plugins/Runtime/Metasound/Source/MetasoundEditor/Private/MetasoundEditorModule.cpp
@@ -133,12 +133,24 @@ namespace Metasound
 						}
 						break;
 
+						case ELiteralArgType::UObjectProxy:
+						{
+							PinType = FGraphBuilder::PinPrimitiveUObject;
+						}
+						break;
+
+						case ELiteralArgType::UObjectProxyArray:
+						{
+							PinType = FGraphBuilder::PinPrimitiveUObjectArray;
+						}
+						break;
+
 						// Register atypical primitives
 						default:
 						case ELiteralArgType::None:
 						case ELiteralArgType::Invalid:
 						{
-							static_assert(static_cast<int32>(ELiteralArgType::Invalid) == 5, "Possible missing binding of pin category to primitive type");
+							static_assert(static_cast<int32>(ELiteralArgType::Invalid) == 7, "Possible missing binding of pin category to primitive type");
 						}
 						break;
 					}
diff --git a/Engine/Plugins/Runtime/Metasound/Source/MetasoundEditor/Private/MetasoundEditorModule.h b/Engine/Plugins/Runtime/Metasound/Source/MetasoundEditor/Private/MetasoundEditorModule.h
--- a/Engine/Plugins/Runtime/Metasound/Source/MetasoundEditor/Private/MetasoundEditorModule.h
+++ b/Engine/Plugins/Runtime/Metasound/Source/MetasoundEditor/Private/MetasoundEditorModule.h
@@ -5,7 +5,7 @@
 #include "EdGraph/EdGraphNode.h"
 #include "EdGraph/EdGraphPin.h"
 #include "Modules/ModuleInterface.h"
-#include "MetasoundDataReference.h"
+#include "MetasoundFrontendRegistries.h"
 #include "Templates/Function.h"
 
 DECLARE_LOG_CATEGORY_EXTERN(LogMetasoundEditor, Log, All);
diff --git a/Engine/Plugins/Runtime/Metasound/Source/MetasoundEditor/Private/MetasoundLiteralDescriptionDetailCustomization.cpp b/Engine/Plugins/Runtime/Metasound/Source/MetasoundEditor/Private/MetasoundLiteralDescriptionDetailCustomization.cpp
--- a/Engine/Plugins/Runtime/Metasound/Source/MetasoundEditor/Private/MetasoundLiteralDescriptionDetailCustomization.cpp
+++ b/Engine/Plugins/Runtime/Metasound/Source/MetasoundEditor/Private/MetasoundLiteralDescriptionDetailCustomization.cpp
@@ -10,6 +10,8 @@
 #include "IDetailGroup.h"
 #include "Layout/Visibility.h"
 #include "MetasoundFrontendDataLayout.h"
+#include "MetasoundFrontend.h"
+#include "MetasoundFrontendRegistries.h"
 #include "Misc/AssertionMacros.h"
 #include "Misc/Attribute.h"
 #include "PropertyEditorDelegates.h"
@@ -17,7 +19,7 @@
 #include "Templates/Casts.h"
 #include "Templates/SharedPointer.h"
 #include "Widgets/Text/STextBlock.h"
-
+#include "PropertyCustomizationHelpers.h"
 
 #define LOCTEXT_NAMESPACE "MetasoundEditor"
 
@@ -49,9 +51,19 @@ namespace Metasound
 				return StructPropertyHandle->GetChildHandle(GET_MEMBER_NAME_CHECKED(FMetasoundLiteralDescription, AsString));
 			}
 
+			case EMetasoundLiteralType::UObject:
+			{
+				return StructPropertyHandle->GetChildHandle(GET_MEMBER_NAME_CHECKED(FMetasoundLiteralDescription, AsUObject));
+			}
+
+			case EMetasoundLiteralType::UObjectArray:
+			{
+				return StructPropertyHandle->GetChildHandle(GET_MEMBER_NAME_CHECKED(FMetasoundLiteralDescription, AsUObjectArray));
+			}
+
 			case EMetasoundLiteralType::None:
 			default:
-				static_assert(static_cast<int32>(EMetasoundLiteralType::Invalid) == 5, "Possible missing case coverage for EMetasoundLiteralType");
+				static_assert(static_cast<int32>(EMetasoundLiteralType::Invalid) == 7, "Possible missing case coverage for EMetasoundLiteralType");
 				return TSharedPtr<IPropertyHandle>();
 			}
 		}
@@ -87,42 +99,156 @@ void FMetasoundLiteralDescriptionDetailCustomization::CustomizeChildren(TSharedR
 {
 	TSharedRef<IPropertyHandle> TypeEnumHandle = StructPropertyHandle->GetChildHandle(GET_MEMBER_NAME_CHECKED(FMetasoundLiteralDescription, LiteralType)).ToSharedRef();
 
+
+	FName OwningDataTypeName;
+	UClass* OwningDataTypeClass = nullptr;
+	EMetasoundLiteralType PreferredLiteralType = EMetasoundLiteralType::None;
+
+	// Grab the preferred uclass for the owning data type.
+	TSharedPtr<IPropertyHandle> InputDescription = StructPropertyHandle->GetParentHandle();
+
+	// TODO: Ensure that StructPropertyHandle is an FMetasoundInputDescription.
+	if (InputDescription)
+	{
+		TSharedPtr<IPropertyHandle> DataTypeNameHandle = InputDescription->GetChildHandle(GET_MEMBER_NAME_CHECKED(FMetasoundInputDescription, TypeName));
+		if (DataTypeNameHandle)
+		{
+			
+			DataTypeNameHandle->GetValue(OwningDataTypeName);
+
+			FMetasoundFrontendRegistryContainer* Registry = FMetasoundFrontendRegistryContainer::Get();
+			OwningDataTypeClass = Registry->GetLiteralUClassForDataType(OwningDataTypeName);
+			PreferredLiteralType = Metasound::Frontend::GetMetasoundLiteralType(Registry->GetDesiredLiteralTypeForDataType(OwningDataTypeName));
+		}
+	}
+
+	
+
 	struct FTypeEditorInfo
 	{
 		EMetasoundLiteralType LiteralType;
+		TSharedRef<IPropertyHandle> TypeEnumHandle;
 		TSharedPtr<IPropertyHandle> PropertyHandle;
 		FText DisplayName;
+		UClass* OwningDataTypeUClass;
+
+		TSharedRef<SWidget> CreatePropertyValueWidget(bool bDisplayDefaultPropertyButtons = true) const
+		{
+			if (LiteralType == EMetasoundLiteralType::UObject)
+			{
+				auto ValidateAsset = [CachedDataTypeUClass = OwningDataTypeUClass](const UObject* InObject, FText& OutReason)  -> bool
+				{
+					UClass* SelectedObjectClass = InObject->GetClass();
+					return CachedDataTypeUClass && SelectedObjectClass && SelectedObjectClass->IsChildOf(CachedDataTypeUClass);
+				};
+
+				auto CommitAsset = [CachedLiteralType = LiteralType, CachedPropertyHandle = PropertyHandle, CachedEnumHandle = TypeEnumHandle, CachedDataTypeUClass = OwningDataTypeUClass](const FAssetData& InAssetData) -> void
+				{
+					// if we've hit this code, the presumption is that the datatype for this parameter has already defined a corresponding UClass that can be used to set it.
+					ensureAlways(CachedDataTypeUClass && CachedDataTypeUClass != UObject::StaticClass());
+
+					UObject* InObject = InAssetData.GetAsset();
+					CachedPropertyHandle->SetValue(InObject);
+
+					// If this asset selector was set to no asset, clear out the literal description.
+					if (InObject)
+					{
+						CachedEnumHandle->SetValue(static_cast<uint8>(EMetasoundLiteralType::UObject));
+					}
+					else
+					{
+						CachedEnumHandle->SetValue(static_cast<uint8>(EMetasoundLiteralType::None));
+					}
+				};
+
+				auto GetAssetPath = [CachedPropertyHandle = PropertyHandle]() -> FString
+				{
+					UObject* Obj = nullptr;
+					CachedPropertyHandle->GetValue(Obj);
+					return  Obj != nullptr ? Obj->GetPathName() : FString();
+				};
+
+				// TODO: get the UFactory corresponding to this datatype's UClass.
+				TArray<UFactory*> FactoriesToUse;
+
+				return SNew(SObjectPropertyEntryBox)
+					.ObjectPath_Lambda(GetAssetPath)
+					.AllowedClass(OwningDataTypeUClass)
+					// .OnShouldSetAsset_Lambda(ValidateAsset)
+					.OnObjectChanged_Lambda(CommitAsset)
+					.AllowClear(false)
+					.DisplayUseSelected(true)
+					.DisplayBrowse(true)
+					.DisplayThumbnail(true)
+					.NewAssetFactories(FactoriesToUse);
+			}
+			else if (LiteralType == EMetasoundLiteralType::UObjectArray)
+			{
+				// TODO: Implement.
+				return SNullWidget::NullWidget;
+			}
+			else
+			{
+				return PropertyHandle->CreatePropertyValueWidget();
+			}
+		}
 	};
 
-	static_assert(static_cast<int32>(EMetasoundLiteralType::Invalid) == 5, "Possible missing property descriptor for literal customization display");
+	static_assert(static_cast<int32>(EMetasoundLiteralType::Invalid) == 7, "Possible missing property descriptor for literal customization display");
 	const TArray<FTypeEditorInfo> TypePropertyInfo =
 	{
 		FTypeEditorInfo
 		{
 			EMetasoundLiteralType::Bool,
+			TypeEnumHandle,
 			StructPropertyHandle->GetChildHandle(GET_MEMBER_NAME_CHECKED(FMetasoundLiteralDescription, AsBool)),
-			LOCTEXT("Metasound_LiteralDisplayNameBool", "Bool")
+			LOCTEXT("Metasound_LiteralDisplayNameBool", "Bool"),
+			nullptr
 		},
 
 		FTypeEditorInfo
 		{
 			EMetasoundLiteralType::Integer,
+			TypeEnumHandle,
 			StructPropertyHandle->GetChildHandle(GET_MEMBER_NAME_CHECKED(FMetasoundLiteralDescription, AsInteger)),
-			LOCTEXT("Metasound_LiteralDisplayNameInteger", "Int32")
+			LOCTEXT("Metasound_LiteralDisplayNameInteger", "Int32"),
+			nullptr
 		},
 
 		FTypeEditorInfo
 		{
 			EMetasoundLiteralType::Float,
+			TypeEnumHandle,
 			StructPropertyHandle->GetChildHandle(GET_MEMBER_NAME_CHECKED(FMetasoundLiteralDescription, AsFloat)),
-			LOCTEXT("Metasound_LiteralDisplayNameFloat", "Float")
+			LOCTEXT("Metasound_LiteralDisplayNameFloat", "Float"),
+			nullptr
 		},
 
 		FTypeEditorInfo
 		{
 			EMetasoundLiteralType::String,
+			TypeEnumHandle,
 			StructPropertyHandle->GetChildHandle(GET_MEMBER_NAME_CHECKED(FMetasoundLiteralDescription, AsString)),
-			LOCTEXT("Metasound_LiteralDisplayNameString", "String")
+			LOCTEXT("Metasound_LiteralDisplayNameString", "String"),
+			nullptr
+		},
+
+		FTypeEditorInfo
+		{
+			EMetasoundLiteralType::UObject,
+			TypeEnumHandle,
+			StructPropertyHandle->GetChildHandle(GET_MEMBER_NAME_CHECKED(FMetasoundLiteralDescription, AsUObject)),
+			LOCTEXT("Metasound_LiteralDisplayNameUObject", "UObject"),
+			OwningDataTypeClass
+		},
+
+		FTypeEditorInfo
+		{
+			EMetasoundLiteralType::UObjectArray,
+			TypeEnumHandle,
+			StructPropertyHandle->GetChildHandle(GET_MEMBER_NAME_CHECKED(FMetasoundLiteralDescription, AsUObjectArray)),
+			LOCTEXT("Metasound_LiteralDisplayNameUObjectArray", "UObjectArray"),
+			OwningDataTypeClass
 		},
 	};
 
@@ -161,12 +287,12 @@ void FMetasoundLiteralDescriptionDetailCustomization::CustomizeChildren(TSharedR
 					.Padding(1.0f, 0.0f, 0.0f, 0.0f)
 					.VAlign(VAlign_Center)
 					[
-						Info.PropertyHandle->CreatePropertyValueWidget()
+						Info.CreatePropertyValueWidget()
 					]
 			]
-			.Visibility(TAttribute<EVisibility>::Create([TypeEnumHandle, LiteralType = Info.LiteralType]()
+			.Visibility(TAttribute<EVisibility>::Create([PreferredLiteralType, LiteralType = Info.LiteralType]()
 			{
-				const bool bIsLiteralActive = Metasound::Editor::IsLiteralTypeActive(TypeEnumHandle, LiteralType);
+				const bool bIsLiteralActive = LiteralType == PreferredLiteralType;
 				return bIsLiteralActive ? EVisibility::Visible : EVisibility::Hidden;
 			}));
 	}
diff --git a/Engine/Plugins/Runtime/Metasound/Source/MetasoundEngine/Private/MetasoundEngineModule.cpp b/Engine/Plugins/Runtime/Metasound/Source/MetasoundEngine/Private/MetasoundEngineModule.cpp
--- a/Engine/Plugins/Runtime/Metasound/Source/MetasoundEngine/Private/MetasoundEngineModule.cpp
+++ b/Engine/Plugins/Runtime/Metasound/Source/MetasoundEngine/Private/MetasoundEngineModule.cpp
@@ -13,6 +13,7 @@ DEFINE_LOG_CATEGORY(LogMetasoundEngine);
 namespace Metasound
 {
 	REGISTER_METASOUND_DATATYPE(FWave, "Primitive:Wave");
+	REGISTER_METASOUND_DATATYPE(FWaveAsset, "Primitive:WaveAsset", ::Metasound::ELiteralArgType::UObjectProxy, USoundWave)
 }
 
 class FMetasoundEngineModule : public IMetasoundEngineModule
diff --git a/Engine/Plugins/Runtime/Metasound/Source/MetasoundEngine/Public/MetasoundWave.h b/Engine/Plugins/Runtime/Metasound/Source/MetasoundEngine/Public/MetasoundWave.h
--- a/Engine/Plugins/Runtime/Metasound/Source/MetasoundEngine/Public/MetasoundWave.h
+++ b/Engine/Plugins/Runtime/Metasound/Source/MetasoundEngine/Public/MetasoundWave.h
@@ -4,6 +4,8 @@
 
 #include "MetasoundDataReference.h"
 #include "MetasoundDataTypeRegistrationMacro.h"
+#include "IAudioProxyInitializer.h"
+#include "Sound/SoundWave.h"
 
 // Forward declares
 namespace Audio
@@ -38,5 +40,32 @@ namespace Metasound
 		static FDecoderInputPtr CreateDecoderInput(
 			const FWaveReadRef& InWaveRef);
 	};
+	// Metasound datatype that holds onto a weak ptr. Mostly used as a placeholder until we have a proper proxy type.
+	class FWaveAsset
+	{
+	private:
+		TWeakObjectPtr<USoundWave> SoundWave;
+	public:
+
+		FWaveAsset() = default;
+
+		FWaveAsset(const Audio::IProxyData& InInitData)
+		{
+			const FSoundWaveProxy& SoundWaveProxy = InInitData.GetAs<FSoundWaveProxy>();
+			SoundWave = SoundWaveProxy.SoundWavePtr;
+		}
+
+		USoundWave* GetSoundWave()
+		{
+			return SoundWave.Get();
+		}
+
+		const USoundWave* GetSoundWave() const
+		{
+			return SoundWave.Get();
+		}
+	};
+
 	DECLARE_METASOUND_DATA_REFERENCE_TYPES(FWave, METASOUNDENGINE_API, 0x0ddba11, FWaveTypeInfo, FWaveReadRef, FWaveWriteRef)
+	DECLARE_METASOUND_DATA_REFERENCE_TYPES(FWaveAsset, METASOUNDENGINE_API, 0x0ddba42, FWaveAssetTypeInfo, FWaveAssetReadRef, FWaveAssetWriteRef)
 }
diff --git a/Engine/Plugins/Runtime/Metasound/Source/MetasoundFrontend/MetasoundFrontend.Build.cs b/Engine/Plugins/Runtime/Metasound/Source/MetasoundFrontend/MetasoundFrontend.Build.cs
--- a/Engine/Plugins/Runtime/Metasound/Source/MetasoundFrontend/MetasoundFrontend.Build.cs
+++ b/Engine/Plugins/Runtime/Metasound/Source/MetasoundFrontend/MetasoundFrontend.Build.cs
@@ -16,6 +16,7 @@ namespace UnrealBuildTool.Rules
 				new string[]
 				{
 					"Core",
+					"CoreUObject",
 					"Serialization",
 					"SignalProcessing"
 				}
@@ -26,7 +27,6 @@ namespace UnrealBuildTool.Rules
 				new string[]
 				{
 					"MetasoundGraphCore",
-					"CoreUObject"
 				}
 			);
 
diff --git a/Engine/Plugins/Runtime/Metasound/Source/MetasoundFrontend/Private/MetasoundFrontend.cpp b/Engine/Plugins/Runtime/Metasound/Source/MetasoundFrontend/Private/MetasoundFrontend.cpp
--- a/Engine/Plugins/Runtime/Metasound/Source/MetasoundFrontend/Private/MetasoundFrontend.cpp
+++ b/Engine/Plugins/Runtime/Metasound/Source/MetasoundFrontend/Private/MetasoundFrontend.cpp
@@ -10,6 +10,7 @@
 #include "StructDeserializer.h"
 #include "StructSerializer.h"
 #include "Serialization/MemoryReader.h"
+#include "MetasoundFrontendRegistries.h"
 
 
 static int32 MetasoundUndoRollLimitCvar = 128;
@@ -204,6 +205,42 @@ namespace Metasound
 			}
 		}
 
+		EMetasoundLiteralType GetMetasoundLiteralType(Metasound::ELiteralArgType InLiteralType)
+		{
+			switch (InLiteralType)
+			{
+				case Metasound::ELiteralArgType::Boolean:
+				{
+					return EMetasoundLiteralType::Bool;
+				}
+				case Metasound::ELiteralArgType::Integer:
+				{
+					return EMetasoundLiteralType::Integer;
+				}
+				case Metasound::ELiteralArgType::Float:
+				{
+					return EMetasoundLiteralType::Float;
+				}
+				case Metasound::ELiteralArgType::String:
+				{
+					return EMetasoundLiteralType::String;
+				}
+				case Metasound::ELiteralArgType::UObjectProxy:
+				{
+					return EMetasoundLiteralType::UObject;
+				}
+				case Metasound::ELiteralArgType::UObjectProxyArray:
+				{
+					return EMetasoundLiteralType::UObjectArray;
+				}
+				case Metasound::ELiteralArgType::None:
+				default:
+				{
+					return EMetasoundLiteralType::None;
+				}
+			}
+		}
+
 		TArray<FName> GetAllAvailableDataTypes()
 		{
 			return FMetasoundFrontendRegistryContainer::Get()->GetAllValidDataTypes();
@@ -1763,6 +1800,21 @@ namespace Metasound
 			}
 		}
 
+		UClass* FGraphHandle::GetSupportedClassForInput(const FString& InInputName)
+		{
+			FName DataType;
+			if (GetDataTypeForInput(InInputName, DataType))
+			{
+				FMetasoundFrontendRegistryContainer* Registry = FMetasoundFrontendRegistryContainer::Get();
+
+				return Registry->GetLiteralUClassForDataType(DataType);
+			}
+			else
+			{
+				return nullptr;
+			}
+		}
+
 		bool FGraphHandle::SetInputToLiteral(const FString& InInputName, bool bInValue)
 		{
 			FName DataType;
@@ -1831,6 +1883,42 @@ namespace Metasound
 			return false;
 		}
 
+		bool FGraphHandle::SetInputToLiteral(const FString& InInputName, UObject* InValue)
+		{
+			FName DataType;
+			if (FMetasoundLiteralDescription* Literal = GetLiteralDescriptionForInput(InInputName, DataType))
+			{
+				if (!ensureAlwaysMsgf(DoesDataTypeSupportLiteralType(DataType, ELiteralArgType::UObjectProxy), TEXT("Tried to set Data Type %s to an unsupported literal type (String)")))
+				{
+					return false;
+				}
+
+				SetLiteralDescription(*Literal, InValue);
+				return true;
+			}
+
+			return false;
+		
+		
+		}
+
+		bool FGraphHandle::SetInputToLiteral(const FString& InInputName, TArray<UObject*> InValue)
+		{
+			FName DataType;
+			if (FMetasoundLiteralDescription* Literal = GetLiteralDescriptionForInput(InInputName, DataType))
+			{
+				if (!ensureAlwaysMsgf(DoesDataTypeSupportLiteralType(DataType, ELiteralArgType::UObjectProxyArray), TEXT("Tried to set Data Type %s to an unsupported literal type (String)")))
+				{
+					return false;
+				}
+
+				SetLiteralDescription(*Literal, InValue);
+				return true;
+			}
+
+			return false;
+		}
+
 		bool FGraphHandle::ClearLiteralForInput(const FString& InInputName)
 		{
 			FName DataType;
@@ -2248,7 +2336,7 @@ namespace Metasound
 								UE_LOG(LogTemp, Display, TEXT("Data Type %s supports the following literal types:"), *InputDescription.TypeName.ToString());
 
 								FDataTypeRegistryInfo DataTypeInfo;
-
+								
 								// this shouldn't hit at all, because this should be a registered data type.
 								ensure(GetTraitsForDataType(InputDescription.TypeName, DataTypeInfo));
 
@@ -2272,6 +2360,18 @@ namespace Metasound
 									UE_LOG(LogTemp, Display, TEXT("    String"));
 								}
 
+								if (DataTypeInfo.bIsProxyParsable)
+								{
+									check(DataTypeInfo.ProxyGeneratorClass);
+									UE_LOG(LogTemp, Display, TEXT("    U%s*"), *DataTypeInfo.ProxyGeneratorClass->GetName());
+								}
+
+								if (DataTypeInfo.bIsProxyArrayParsable)
+								{
+									check(DataTypeInfo.ProxyGeneratorClass);
+									UE_LOG(LogTemp, Display, TEXT("    TArray<U%s*>"), *DataTypeInfo.ProxyGeneratorClass->GetName());
+								}
+
 								return nullptr;
 							}
 
@@ -2280,10 +2380,10 @@ namespace Metasound
 								InputDescription.Name,
 								InputDescription.Name,
 								InSettings,
-								LiteralParam
+								MoveTemp(LiteralParam)
 							};
 
-							INodePtr InputNode = ConstructInputNode(InputDescription.TypeName, InitParams);
+							INodePtr InputNode = ConstructInputNode(InputDescription.TypeName, MoveTemp(InitParams));
 							if (!ensureAlwaysMsgf(InputNode.IsValid(), TEXT("Failed to construct a valid input node for Data Type %s!"), *InputDescription.TypeName.ToString()))
 							{
 								return nullptr;
@@ -2372,7 +2472,7 @@ namespace Metasound
 
 								if (LiteralParam.IsValid())
 								{
-									InitData.ParamMap.Add(StaticParamTuple.Key, LiteralParam);
+									InitData.ParamMap.Add(StaticParamTuple.Key, MoveTemp(LiteralParam));
 								}
 							}
 
diff --git a/Engine/Plugins/Runtime/Metasound/Source/MetasoundFrontend/Private/MetasoundFrontendBaseClasses.cpp b/Engine/Plugins/Runtime/Metasound/Source/MetasoundFrontend/Private/MetasoundFrontendBaseClasses.cpp
--- a/Engine/Plugins/Runtime/Metasound/Source/MetasoundFrontend/Private/MetasoundFrontendBaseClasses.cpp
+++ b/Engine/Plugins/Runtime/Metasound/Source/MetasoundFrontend/Private/MetasoundFrontendBaseClasses.cpp
@@ -18,9 +18,9 @@ namespace Metasound
 			return FMetasoundFrontendRegistryContainer::Get()->GetExternalNodeRegistry();
 		}
 
-		TUniquePtr<INode> ConstructInputNode(const FName& InInputType, const FInputNodeConstructorParams& InParams)
+		TUniquePtr<INode> ConstructInputNode(const FName& InInputType, FInputNodeConstructorParams&& InParams)
 		{
-			return FMetasoundFrontendRegistryContainer::Get()->ConstructInputNode(InInputType, InParams);
+			return FMetasoundFrontendRegistryContainer::Get()->ConstructInputNode(InInputType, MoveTemp(InParams));
 		}
 
 		TUniquePtr<INode> ConstructOutputNode(const FName& InOutputType, const FOutputNodeConstrutorParams& InParams)
@@ -40,6 +40,8 @@ namespace Metasound
 			OutDescription.AsInteger = 0;
 			OutDescription.AsFloat = 0.0f;
 			OutDescription.AsString.Empty();
+			OutDescription.AsUObject = nullptr;
+			OutDescription.AsUObjectArray.Empty();
 		}
 
 		void SetLiteralDescription(FMetasoundLiteralDescription& OutDescription, int32 InValue)
@@ -49,6 +51,8 @@ namespace Metasound
 			OutDescription.AsInteger = InValue;
 			OutDescription.AsFloat = 0.0f;
 			OutDescription.AsString.Empty();
+			OutDescription.AsUObject = nullptr;
+			OutDescription.AsUObjectArray.Empty();
 		}
 
 		void SetLiteralDescription(FMetasoundLiteralDescription& OutDescription, float InValue)
@@ -58,6 +62,8 @@ namespace Metasound
 			OutDescription.AsInteger = 0;
 			OutDescription.AsFloat = InValue;
 			OutDescription.AsString.Empty();
+			OutDescription.AsUObject = nullptr;
+			OutDescription.AsUObjectArray.Empty();
 		}
 
 		void SetLiteralDescription(FMetasoundLiteralDescription& OutDescription, const FString& InValue)
@@ -67,6 +73,31 @@ namespace Metasound
 			OutDescription.AsInteger = 0;
 			OutDescription.AsFloat = 0.0f;
 			OutDescription.AsString = InValue;
+			OutDescription.AsUObject = nullptr;
+			OutDescription.AsUObjectArray.Empty();
+		}
+
+		void SetLiteralDescription(FMetasoundLiteralDescription& OutDescription, UObject* InValue)
+		{
+			OutDescription.LiteralType = EMetasoundLiteralType::UObject;
+			OutDescription.AsBool = false;
+			OutDescription.AsInteger = 0;
+			OutDescription.AsFloat = 0.0f;
+			OutDescription.AsString.Empty();
+			OutDescription.AsUObject = InValue;
+			OutDescription.AsUObjectArray.Empty();
+
+		}
+
+		void SetLiteralDescription(FMetasoundLiteralDescription& OutDescription, const TArray<UObject*>& InValue)
+		{
+			OutDescription.LiteralType = EMetasoundLiteralType::UObjectArray;
+			OutDescription.AsBool = false;
+			OutDescription.AsInteger = 0;
+			OutDescription.AsFloat = 0.0f;
+			OutDescription.AsString.Empty();
+			OutDescription.AsUObject = nullptr;
+			OutDescription.AsUObjectArray = InValue;
 		}
 
 		void ClearLiteralDescription(FMetasoundLiteralDescription& OutDescription)
@@ -76,6 +107,8 @@ namespace Metasound
 			OutDescription.AsInteger = 0;
 			OutDescription.AsFloat = 0.0f;
 			OutDescription.AsString.Empty();
+			OutDescription.AsUObject = nullptr;
+			OutDescription.AsUObjectArray.Empty();
 		}
 
 		FDataTypeLiteralParam GetLiteralParamForDataType(FName InDataType, const FMetasoundLiteralDescription& InDescription)
@@ -133,6 +166,28 @@ namespace Metasound
 						return FDataTypeLiteralParam(InDescription.AsString);
 					}
 				}
+				case EMetasoundLiteralType::UObject:
+				{
+					if (!Registry->DoesDataTypeSupportLiteralType(InDataType, ELiteralArgType::UObjectProxy))
+					{
+						return FDataTypeLiteralParam::InvalidParam();
+					}
+					else
+					{
+						return Registry->GenerateLiteralForUObject(InDataType, InDescription.AsUObject);
+					}
+				}
+				case EMetasoundLiteralType::UObjectArray:
+				{
+					if (!Registry->DoesDataTypeSupportLiteralType(InDataType, ELiteralArgType::UObjectProxyArray))
+					{
+						return FDataTypeLiteralParam::InvalidParam();
+					}
+					else
+					{
+						return Registry->GenerateLiteralForUObjectArray(InDataType, InDescription.AsUObjectArray);
+					}
+				}
 				case EMetasoundLiteralType::Invalid:
 				default:
 				{
@@ -152,9 +207,13 @@ namespace Metasound
 			case EMetasoundLiteralType::Float:
 				return DoesDataTypeSupportLiteralType(InDataType, ELiteralArgType::Float);
 			case EMetasoundLiteralType::Integer:
-				return DoesDataTypeSupportLiteralType(InDataType, EMetasoundLiteralType::Integer);
+				return DoesDataTypeSupportLiteralType(InDataType, ELiteralArgType::Integer);
 			case EMetasoundLiteralType::String:
-				return DoesDataTypeSupportLiteralType(InDataType, EMetasoundLiteralType::String);
+				return DoesDataTypeSupportLiteralType(InDataType, ELiteralArgType::String);
+			case EMetasoundLiteralType::UObject:
+				return DoesDataTypeSupportLiteralType(InDataType, ELiteralArgType::UObjectProxy);
+			case EMetasoundLiteralType::UObjectArray:
+				return DoesDataTypeSupportLiteralType(InDataType, ELiteralArgType::UObjectProxyArray);
 			case EMetasoundLiteralType::Invalid:
 			default:
 				return false;
@@ -190,6 +249,8 @@ namespace Metasound
 				{
 					return FDataTypeLiteralParam(InDescription.AsString);
 				}
+				case EMetasoundLiteralType::UObject:
+				case EMetasoundLiteralType::UObjectArray:
 				case EMetasoundLiteralType::None:
 				case EMetasoundLiteralType::Invalid:
 				default:
@@ -223,6 +284,8 @@ namespace Metasound
 				{
 					return FDataTypeLiteralParam(FString());
 				}
+				case Metasound::ELiteralArgType::UObjectProxy:
+				case Metasound::ELiteralArgType::UObjectProxyArray:
 				case Metasound::ELiteralArgType::None:
 				{
 					return FDataTypeLiteralParam();
diff --git a/Engine/Plugins/Runtime/Metasound/Source/MetasoundFrontend/Private/MetasoundFrontendRegistries.cpp b/Engine/Plugins/Runtime/Metasound/Source/MetasoundFrontend/Private/MetasoundFrontendRegistries.cpp
--- a/Engine/Plugins/Runtime/Metasound/Source/MetasoundFrontend/Private/MetasoundFrontendRegistries.cpp
+++ b/Engine/Plugins/Runtime/Metasound/Source/MetasoundFrontend/Private/MetasoundFrontendRegistries.cpp
@@ -84,11 +84,11 @@ TMap<FMetasoundFrontendRegistryContainer::FNodeRegistryKey, FMetasoundFrontendRe
 	return ExternalNodeRegistry;
 }
 
-TUniquePtr<Metasound::INode> FMetasoundFrontendRegistryContainer::ConstructInputNode(const FName& InInputType, const Metasound::FInputNodeConstructorParams& InParams)
+TUniquePtr<Metasound::INode> FMetasoundFrontendRegistryContainer::ConstructInputNode(const FName& InInputType, Metasound::FInputNodeConstructorParams&& InParams)
 {
 	if (ensureAlwaysMsgf(DataTypeRegistry.Contains(InInputType), TEXT("Couldn't find data type %s!"), *InInputType.ToString()))
 	{
-		return DataTypeRegistry[InInputType].InputNodeConstructor(InParams);
+		return DataTypeRegistry[InInputType].Callbacks.InputNodeConstructor(MoveTemp(InParams));
 	}
 	else
 	{
@@ -100,7 +100,7 @@ TUniquePtr<Metasound::INode> FMetasoundFrontendRegistryContainer::ConstructOutpu
 {
 	if (ensureAlwaysMsgf(DataTypeRegistry.Contains(InOutputType), TEXT("Couldn't find data type %s!"), *InOutputType.ToString()))
 	{
-		return DataTypeRegistry[InOutputType].OutputNodeConstructor(InParams);
+		return DataTypeRegistry[InOutputType].Callbacks.OutputNodeConstructor(InParams);
 	}
 	else
 	{
@@ -108,6 +108,50 @@ TUniquePtr<Metasound::INode> FMetasoundFrontendRegistryContainer::ConstructOutpu
 	}
 }
 
+Metasound::FDataTypeLiteralParam FMetasoundFrontendRegistryContainer::GenerateLiteralForUObject(const FName& InDataType, UObject* InObject)
+{
+	if (ensureAlwaysMsgf(DataTypeRegistry.Contains(InDataType), TEXT("Couldn't find data type %s!"), *InDataType.ToString()))
+	{
+		 Audio::IProxyDataPtr ProxyPtr = DataTypeRegistry[InDataType].Callbacks.ProxyConstructor(InObject);
+		 if (ensureAlwaysMsgf(ProxyPtr.IsValid(), TEXT("UObject failed to create a valid proxy!")))
+		 {
+			 return Metasound::FDataTypeLiteralParam(MoveTemp(ProxyPtr));
+		 }
+		 else
+		 {
+			 return Metasound::FDataTypeLiteralParam();
+		 }
+	}
+	else
+	{
+		return Metasound::FDataTypeLiteralParam();
+	}
+}
+
+Metasound::FDataTypeLiteralParam FMetasoundFrontendRegistryContainer::GenerateLiteralForUObjectArray(const FName& InDataType, TArray<UObject*> InObjectArray)
+{
+	if (ensureAlwaysMsgf(DataTypeRegistry.Contains(InDataType), TEXT("Couldn't find data type %s!"), *InDataType.ToString()))
+	{
+		TArray<Audio::IProxyDataPtr> ProxyArray;
+
+		for (UObject* InObject : InObjectArray)
+		{
+			if (InObject)
+			{
+				Audio::IProxyDataPtr ProxyPtr = DataTypeRegistry[InDataType].Callbacks.ProxyConstructor(InObject);
+				ensureAlwaysMsgf(ProxyPtr.IsValid(), TEXT("UObject failed to create a valid proxy!"));
+				ProxyArray.Add(MoveTemp(ProxyPtr));
+			}
+		}
+
+		return Metasound::FDataTypeLiteralParam(MoveTemp(ProxyArray));
+	}
+	else
+	{
+		return Metasound::FDataTypeLiteralParam();
+	}
+}
+
 TUniquePtr<Metasound::INode> FMetasoundFrontendRegistryContainer::ConstructExternalNode(const FName& InNodeType, uint32 InNodeHash, const Metasound::FNodeInitData& InInitData)
 {
 	Metasound::Frontend::FNodeRegistryKey RegistryKey;
@@ -169,6 +213,19 @@ Metasound::ELiteralArgType FMetasoundFrontendRegistryContainer::GetDesiredLitera
 	}
 }
 
+UClass* FMetasoundFrontendRegistryContainer::GetLiteralUClassForDataType(FName InDataType) const
+{
+	if (!DataTypeRegistry.Contains(InDataType))
+	{
+		ensureAlwaysMsgf(false, TEXT("couldn't find DataType %s in the registry."), *InDataType.ToString());
+		return nullptr;
+	}
+	else
+	{
+		return DataTypeRegistry[InDataType].Info.ProxyGeneratorClass;
+	}
+}
+
 bool FMetasoundFrontendRegistryContainer::DoesDataTypeSupportLiteralType(FName InDataType, Metasound::ELiteralArgType InLiteralType) const
 {
 	if (!DataTypeRegistry.Contains(InDataType))
@@ -197,6 +254,14 @@ bool FMetasoundFrontendRegistryContainer::DoesDataTypeSupportLiteralType(FName I
 		{
 			return DataTypeInfo.Info.bIsStringParsable;
 		}
+		case Metasound::ELiteralArgType::UObjectProxy:
+		{
+			return DataTypeInfo.Info.bIsProxyParsable;
+		}
+		case Metasound::ELiteralArgType::UObjectProxyArray:
+		{
+			return DataTypeInfo.Info.bIsProxyArrayParsable;
+		}
 		case Metasound::ELiteralArgType::None:
 		{
 			return DataTypeInfo.Info.bIsConstructableWithSettings || DataTypeInfo.Info.bIsDefaultConstructible;
@@ -209,7 +274,7 @@ bool FMetasoundFrontendRegistryContainer::DoesDataTypeSupportLiteralType(FName I
 	}
 }
 
-bool FMetasoundFrontendRegistryContainer::RegisterDataType(const ::Metasound::FDataTypeRegistryInfo& InDataInfo, FInputNodeConstructorCallback&& InputNodeConstructor, FOutputNodeConstructorCallback&& OutputNodeConstructor)
+bool FMetasoundFrontendRegistryContainer::RegisterDataType(const ::Metasound::FDataTypeRegistryInfo& InDataInfo, ::Metasound::FDataTypeConstructorCallbacks&& InCallbacks)
 {
 	if (!ensureAlwaysMsgf(!DataTypeRegistry.Contains(InDataInfo.DataTypeName), TEXT("Name collision when trying to register Metasound Data Type %s! Make sure that you created a unique name for your data type, and that REGISTER_METASOUND_DATATYPE isn't called in a public header."), *InDataInfo.DataTypeName.ToString()))
 	{
@@ -218,7 +283,7 @@ bool FMetasoundFrontendRegistryContainer::RegisterDataType(const ::Metasound::FD
 	}
 	else
 	{
-		FDataTypeRegistryElement InElement = { MoveTemp(InputNodeConstructor), MoveTemp(OutputNodeConstructor), InDataInfo };
+		FDataTypeRegistryElement InElement = { MoveTemp(InCallbacks), InDataInfo };
 		DataTypeRegistry.Add(InDataInfo.DataTypeName, MoveTemp(InElement));
 		UE_LOG(LogTemp, Display, TEXT("Registered Metasound Datatype %s."), *InDataInfo.DataTypeName.ToString());
 		return true;
diff --git a/Engine/Plugins/Runtime/Metasound/Source/MetasoundFrontend/Public/MetasoundDataTypeRegistrationMacro.h b/Engine/Plugins/Runtime/Metasound/Source/MetasoundFrontend/Public/MetasoundDataTypeRegistrationMacro.h
--- a/Engine/Plugins/Runtime/Metasound/Source/MetasoundFrontend/Public/MetasoundDataTypeRegistrationMacro.h
+++ b/Engine/Plugins/Runtime/Metasound/Source/MetasoundFrontend/Public/MetasoundDataTypeRegistrationMacro.h
@@ -3,6 +3,7 @@
 #pragma once
 
 #include "CoreMinimal.h"
+#include "Templates/Casts.h"
 
 #include "MetasoundDataReference.h"
 #include "MetasoundOutputNode.h"
@@ -10,11 +11,30 @@
 #include "MetasoundOperatorInterface.h"
 #include "MetasoundFrontendRegistries.h"
 
-typedef TUniqueFunction<TUniquePtr<Metasound::INode>(const ::Metasound::FInputNodeConstructorParams&)> FInputNodeConstructorCallback;
-typedef TUniqueFunction<TUniquePtr<Metasound::INode>(const ::Metasound::FOutputNodeConstrutorParams&)> FOutputNodeConstructorCallback;
 
-template<typename TDataType>
-bool RegisterDataTypeWithFrontend(::Metasound::ELiteralArgType PreferredArgType = ::Metasound::ELiteralArgType::None)
+// SFINAE used to optionally invoke subclasses of IAudioProxyDataFactory when we can.
+template<typename UClassToUse, typename TEnableIf<TIsDerivedFrom<UClassToUse, IAudioProxyDataFactory>::Value, bool>::Type = true>
+IAudioProxyDataFactory* CastToAudioProxyDataFactory(UObject* InObject)
+{
+	UClassToUse* DowncastObject = Cast<UClassToUse>(InObject);
+	if (ensureAlways(DowncastObject))
+	{
+		return static_cast<IAudioProxyDataFactory*>(DowncastObject);
+	}
+	else
+	{
+		return nullptr;
+	}
+}
+
+template<typename UClassToUse, typename TEnableIf<!TIsDerivedFrom<UClassToUse, IAudioProxyDataFactory>::Value, bool>::Type = true>
+IAudioProxyDataFactory* CastToAudioProxyDataFactory(UObject* InObject)
+{
+	return nullptr;
+}
+
+template<typename TDataType, ::Metasound::ELiteralArgType PreferredArgType = ::Metasound::ELiteralArgType::None, typename UClassToUse = UObject>
+bool RegisterDataTypeWithFrontend()
 {
 	// if we reenter this code (because DECLARE_METASOUND_DATA_REFERENCE_TYPES was called twice with the same type),
 	// we catch it here.
@@ -28,9 +48,9 @@ bool RegisterDataTypeWithFrontend(::Metasound::ELiteralArgType PreferredArgType
 	bAlreadyRegisteredThisDataType = true;
 
 	// Lambdas that generate our template-instantiated input and output nodes:
-	FInputNodeConstructorCallback InputNodeConstructor = [](const ::Metasound::FInputNodeConstructorParams& InParams) -> TUniquePtr<::Metasound::INode>
+	FInputNodeConstructorCallback InputNodeConstructor = [](::Metasound::FInputNodeConstructorParams&& InParams) -> TUniquePtr<::Metasound::INode>
 	{
-		return TUniquePtr<::Metasound::INode>(new ::Metasound::TInputNode<TDataType>(InParams.InNodeName, InParams.InVertexName, InParams.InitParam, InParams.InSettings));
+		return TUniquePtr<::Metasound::INode>(new ::Metasound::TInputNode<TDataType>(InParams.InNodeName, InParams.InVertexName, MoveTemp(InParams.InitParam), InParams.InSettings));
 	};
 
 	FOutputNodeConstructorCallback OutputNodeConstructor = [](const ::Metasound::FOutputNodeConstrutorParams& InParams) -> TUniquePtr<::Metasound::INode>
@@ -38,8 +58,43 @@ bool RegisterDataTypeWithFrontend(::Metasound::ELiteralArgType PreferredArgType
 		return TUniquePtr<::Metasound::INode>(new ::Metasound::TOutputNode<TDataType>(InParams.InNodeName, InParams.InVertexName));
 	};
 	
+	// By default, this function should not be used, unless the preferred arg type is UObjectProxy or UObjectProxyArray, and UClassToUse should be specified.
+	FProxyGetterCallback ProxyGenerator = [](UObject* InObject) -> Audio::IProxyDataPtr
+	{
+		checkNoEntry();
+		return Audio::IProxyDataPtr(nullptr);
+	};
+
+	// If this datatype uses a UObject or UObject array literal, we generate a lambda to build a proxy here:
+	constexpr bool bSpecifiedUClassForProxy = !TIsSame<UClassToUse, UObject>::Value;
+	if (bSpecifiedUClassForProxy)
+	{
+		static_assert(!bSpecifiedUClassForProxy || std::is_base_of<IAudioProxyDataFactory, UClassToUse>::value, "If a Metasound Datatype uses a UObject as a literal, the UClass of that object needs to also derive from Audio::IProxyDataFactory. See USoundWave as an example.");
+		ProxyGenerator = [](UObject* InObject) -> Audio::IProxyDataPtr
+		{
+			IAudioProxyDataFactory* ObjectAsFactory = CastToAudioProxyDataFactory<UClassToUse>(InObject);
+
+			if (ensureAlways(ObjectAsFactory))
+			{
+				static FName ProxySubsystemName = TEXT("Metasound");
+				
+				Audio::FProxyDataInitParams ProxyInitParams;
+				ProxyInitParams.NameOfFeatureRequestingProxy = ProxySubsystemName;
+
+				return ObjectAsFactory->CreateNewProxyData(ProxyInitParams);
+			}
+			else
+			{
+				return Audio::IProxyDataPtr(nullptr);
+			}
+		};
+	}
+
 	static FName DataTypeName = FName(::Metasound::TDataReferenceTypeInfo<TDataType>::TypeName);
 
+	// Pack all of our various constructor lambdas to a single struct.
+	::Metasound::FDataTypeConstructorCallbacks Callbacks = { MoveTemp(InputNodeConstructor), MoveTemp(OutputNodeConstructor), MoveTemp(ProxyGenerator) };
+
 	::Metasound::FDataTypeRegistryInfo RegistryInfo;
 	RegistryInfo.DataTypeName = DataTypeName;
 	RegistryInfo.PreferredLiteralType = PreferredArgType;
@@ -47,25 +102,31 @@ bool RegisterDataTypeWithFrontend(::Metasound::ELiteralArgType PreferredArgType
 	RegistryInfo.bIsIntParsable = ::Metasound::TDataReferenceTypeInfo<TDataType>::bIsIntParsable;
 	RegistryInfo.bIsFloatParsable = ::Metasound::TDataReferenceTypeInfo<TDataType>::bIsFloatParsable;
 	RegistryInfo.bIsStringParsable = ::Metasound::TDataReferenceTypeInfo<TDataType>::bIsStringParsable;
+	RegistryInfo.bIsProxyParsable = ::Metasound::TDataReferenceTypeInfo<TDataType>::bIsProxyParsable;
+	RegistryInfo.bIsProxyArrayParsable = ::Metasound::TDataReferenceTypeInfo<TDataType>::bIsProxyArrayParsable;
 	RegistryInfo.bIsConstructableWithSettings = ::Metasound::TDataReferenceTypeInfo<TDataType>::bIsConstructableWithSettings;
 	RegistryInfo.bIsDefaultConstructible = ::Metasound::TDataReferenceTypeInfo<TDataType>::bCanUseDefaultConstructor;
+	
+	RegistryInfo.ProxyGeneratorClass = UClassToUse::StaticClass();
 
-	bool bSucceeded = FMetasoundFrontendRegistryContainer::Get()->RegisterDataType(RegistryInfo, MoveTemp(InputNodeConstructor), MoveTemp(OutputNodeConstructor));
+	bool bSucceeded = FMetasoundFrontendRegistryContainer::Get()->RegisterDataType(RegistryInfo, MoveTemp(Callbacks));
 	ensureAlwaysMsgf(bSucceeded, TEXT("Failed to register data type %s in the node registry!"), ::Metasound::TDataReferenceTypeInfo<TDataType>::TypeName);
 	return bSucceeded;
 }
 
 // This should be used to expose a datatype as a potential input or output for a metasound graph.
 // The first argument to the macro is the class to expose.
+// the second argument is the display name of that type in the Metasound editor.
 // Optionally, a Metasound::ELiteralArgType can be passed in to designate a preferred literal type-
 // For example, if Metasound::ELiteralArgType::Float is passed in, we will default to using a float parameter to create this datatype.
 // If no argument is passed in, we will infer a literal type to use.
+// If ::Metasound::
 // Metasound::ELiteralArgType::Invalid can be used to enforce that we don't provide space for a literal, in which case you should have a default constructor or a constructor that takes [const FOperatorSettings&] implemented.
 // If you pass in a preferred arg type, please make sure that the passed in datatype has a matching constructor, since we won't check this until runtime.
 #define REGISTER_METASOUND_DATATYPE(DataType, DataTypeName, ...) \
 	IMPL_METASOUND_DATA_TYPE(DataType, DataTypeName) \
 	static_assert(::Metasound::TDataReferenceTypeInfo<DataType>::bIsValidSpecialization, "Please call DECLARE_METASOUND_DATA_REFERENCE_TYPES(" #DataType "...) before calling REGISTER_METASOUND_DATATYPE(" #DataType")."); \
-	static constexpr bool bCanRegister##DataType = ::Metasound::TDataReferenceTypeInfo<DataType>::bIsStringParsable || ::Metasound::TDataReferenceTypeInfo<DataType>::bIsBoolParsable || ::Metasound::TDataReferenceTypeInfo<DataType>::bIsIntParsable || ::Metasound::TDataReferenceTypeInfo<DataType>::bIsFloatParsable || ::Metasound::TDataReferenceTypeInfo<DataType>::bIsConstructableWithSettings || ::Metasound::TDataReferenceTypeInfo<DataType>::bCanUseDefaultConstructor; \
-	static_assert(bCanRegister##DataType , "To register " #DataType " to be used as a Metasounds input or output type, it needs a default constructor or one of the following constructors must be implemented:  " #DataType "(const ::Metasound::FOperatorSettings& InSettings), " #DataType "(int32 InValue, const ::Metasound::FOperatorSettings& InSettings), " #DataType "(float InValue, const ::Metasound::FOperatorSettings& InSettings), " #DataType "(bool InValue, const ::Metasound::FOperatorSettings& InSettings), or " #DataType "(const FString& InString, const ::Metasound::FOperatorSettings& InSettings)."); \
-	static const bool bSuccessfullyRegistered##DataType = FMetasoundFrontendRegistryContainer::Get()->EnqueueInitCommand([](){ ::RegisterDataTypeWithFrontend<DataType>(__VA_ARGS__); }); // This static bool is useful for debugging, but also is the only way the compiler will let us call this function outside of an expression.
+	static constexpr bool bCanRegister##DataType = ::Metasound::TDataReferenceTypeInfo<DataType>::bIsStringParsable || ::Metasound::TDataReferenceTypeInfo<DataType>::bIsBoolParsable || ::Metasound::TDataReferenceTypeInfo<DataType>::bIsIntParsable || ::Metasound::TDataReferenceTypeInfo<DataType>::bIsFloatParsable || ::Metasound::TDataReferenceTypeInfo<DataType>::bIsProxyParsable || ::Metasound::TDataReferenceTypeInfo<DataType>::bIsProxyArrayParsable || ::Metasound::TDataReferenceTypeInfo<DataType>::bIsConstructableWithSettings || ::Metasound::TDataReferenceTypeInfo<DataType>::bCanUseDefaultConstructor; \
+	static_assert(bCanRegister##DataType , "To register " #DataType " to be used as a Metasounds input or output type, it needs a default constructor or one of the following constructors must be implemented:  " #DataType "(const ::Metasound::FOperatorSettings& InSettings), " #DataType "(int32 InValue, const ::Metasound::FOperatorSettings& InSettings), " #DataType "(float InValue, const ::Metasound::FOperatorSettings& InSettings), " #DataType "(bool InValue, const ::Metasound::FOperatorSettings& InSettings), " #DataType "(const Audio::IProxyData& InData, const ::Metasound::FOperatorSettings& InSettings), " #DataType "(const TArray<Audio::IProxyDataPtr>& InProxyArray, const ::Metasound::FOperatorSettings& InSettings) or " #DataType "(const FString& InString, const ::Metasound::FOperatorSettings& InSettings)."); \
+	static const bool bSuccessfullyRegistered##DataType = FMetasoundFrontendRegistryContainer::Get()->EnqueueInitCommand([](){ ::RegisterDataTypeWithFrontend<DataType, __VA_ARGS__>(); }); // This static bool is useful for debugging, but also is the only way the compiler will let us call this function outside of an expression.
 
diff --git a/Engine/Plugins/Runtime/Metasound/Source/MetasoundFrontend/Public/MetasoundFrontend.h b/Engine/Plugins/Runtime/Metasound/Source/MetasoundFrontend/Public/MetasoundFrontend.h
--- a/Engine/Plugins/Runtime/Metasound/Source/MetasoundFrontend/Public/MetasoundFrontend.h
+++ b/Engine/Plugins/Runtime/Metasound/Source/MetasoundFrontend/Public/MetasoundFrontend.h
@@ -72,6 +72,9 @@ namespace Metasound
 		// Generates a new FMetasoundClassDescription for a given node class. Only used by classes that manipulate Metasound Description data directly.
 		METASOUNDFRONTEND_API FMetasoundClassDescription GenerateClassDescriptionForNode(const FNodeClassInfo& InInfo);
 
+		// Convenience function to convert Metasound::ELiteralArgType to EMetasoundLiteralType, the UEnum used for metasound documents.
+		METASOUNDFRONTEND_API EMetasoundLiteralType GetMetasoundLiteralType(Metasound::ELiteralArgType InLiteralType);
+
 		template<typename DataType>
 		FName GetDataTypeName()
 		{
@@ -384,6 +387,10 @@ namespace Metasound
 			// Will return Invalid if the input couldn't be found, or if the input doesn't support any kind of literals.
 			ELiteralArgType GetPreferredLiteralTypeForInput(const FString& InInputName);
 
+			// For inputs whose preferred literal type is UObject or UObjectArray, this can be used to determine the UObject corresponding to that input's datatype.
+			UClass* GetSupportedClassForInput(const FString& InInputName);
+
+
 			// These can be used to set the default value for a given input on this graph.
 			// @returns false if the input name couldn't be found, or if the literal type was incompatible with the Data Type of this input.
 			bool SetInputToLiteral(const FString& InInputName, bool bInValue);
diff --git a/Engine/Plugins/Runtime/Metasound/Source/MetasoundFrontend/Public/MetasoundFrontendBaseClasses.h b/Engine/Plugins/Runtime/Metasound/Source/MetasoundFrontend/Public/MetasoundFrontendBaseClasses.h
--- a/Engine/Plugins/Runtime/Metasound/Source/MetasoundFrontend/Public/MetasoundFrontendBaseClasses.h
+++ b/Engine/Plugins/Runtime/Metasound/Source/MetasoundFrontend/Public/MetasoundFrontendBaseClasses.h
@@ -21,7 +21,7 @@ namespace Metasound
 
 		// Convenience functions to create an INode corresponding to a specific input or output for a metasound graph.
 		// @returns nullptr if the type given wasn't found.
-		METASOUNDFRONTEND_API TUniquePtr<INode> ConstructInputNode(const FName& InInputType, const FInputNodeConstructorParams& InParams);
+		METASOUNDFRONTEND_API TUniquePtr<INode> ConstructInputNode(const FName& InInputType, FInputNodeConstructorParams&& InParams);
 		METASOUNDFRONTEND_API TUniquePtr<INode> ConstructOutputNode(const FName& InOutputType, const FOutputNodeConstrutorParams& InParams);
 
 		// Convenience functions to create an INodeB corresponding to a specific externally declared node type.
@@ -34,6 +34,9 @@ namespace Metasound
 		METASOUNDFRONTEND_API void SetLiteralDescription(FMetasoundLiteralDescription& OutDescription, int32 InValue);
 		METASOUNDFRONTEND_API void SetLiteralDescription(FMetasoundLiteralDescription& OutDescription, float InValue);
 		METASOUNDFRONTEND_API void SetLiteralDescription(FMetasoundLiteralDescription& OutDescription, const FString& InValue);
+		METASOUNDFRONTEND_API void SetLiteralDescription(FMetasoundLiteralDescription& OutDescription, UObject* InValue);
+		METASOUNDFRONTEND_API void SetLiteralDescription(FMetasoundLiteralDescription& OutDescription, const TArray<UObject*>& InValue);
+		
 		METASOUNDFRONTEND_API void ClearLiteralDescription(FMetasoundLiteralDescription& OutDescription);
 
 		// Utility functions for building a ::Metasound::FDataInitParam corresponding to a literal.
diff --git a/Engine/Plugins/Runtime/Metasound/Source/MetasoundFrontend/Public/MetasoundFrontendDataLayout.h b/Engine/Plugins/Runtime/Metasound/Source/MetasoundFrontend/Public/MetasoundFrontendDataLayout.h
--- a/Engine/Plugins/Runtime/Metasound/Source/MetasoundFrontend/Public/MetasoundFrontendDataLayout.h
+++ b/Engine/Plugins/Runtime/Metasound/Source/MetasoundFrontend/Public/MetasoundFrontendDataLayout.h
@@ -31,6 +31,8 @@ enum class EMetasoundLiteralType : uint8
 	Float,
 	Integer,
 	String,
+	UObject,
+	UObjectArray,
 	Invalid UMETA(Hidden)
 };
 
@@ -69,8 +71,6 @@ struct FMetasoundLiteralDescription
 {
 	GENERATED_BODY()
 
-	// HACK: Once we make a property inspector for these we should hide them.
-
 	// The actual type of this literal.
 	UPROPERTY(EditAnywhere, Category = Customized)
 	EMetasoundLiteralType LiteralType;
@@ -86,6 +86,12 @@ struct FMetasoundLiteralDescription
 
 	UPROPERTY(EditAnywhere, Category = Customized)
 	FString AsString;
+
+	UPROPERTY(EditAnywhere, Category = Customized)
+	UObject* AsUObject = nullptr;
+
+	UPROPERTY(EditAnywhere, Category = Customized)
+	TArray<UObject*> AsUObjectArray;
 };
 
 USTRUCT()
diff --git a/Engine/Plugins/Runtime/Metasound/Source/MetasoundFrontend/Public/MetasoundFrontendRegistries.h b/Engine/Plugins/Runtime/Metasound/Source/MetasoundFrontend/Public/MetasoundFrontendRegistries.h
--- a/Engine/Plugins/Runtime/Metasound/Source/MetasoundFrontend/Public/MetasoundFrontendRegistries.h
+++ b/Engine/Plugins/Runtime/Metasound/Source/MetasoundFrontend/Public/MetasoundFrontendRegistries.h
@@ -6,14 +6,59 @@
 #include "MetasoundNodeInterface.h"
 #include "MetasoundDataReference.h"
 #include "MetasoundOperatorInterface.h"
+#include "IAudioProxyInitializer.h"
 
-typedef TUniqueFunction<TUniquePtr<Metasound::INode>(const ::Metasound::FInputNodeConstructorParams&)> FInputNodeConstructorCallback;
+typedef TUniqueFunction<TUniquePtr<Metasound::INode>(::Metasound::FInputNodeConstructorParams&&)> FInputNodeConstructorCallback;
 typedef TUniqueFunction<TUniquePtr<Metasound::INode>(const ::Metasound::FOutputNodeConstrutorParams&)> FOutputNodeConstructorCallback;
 
+// This function is used to create a proxy from a datatype's base uclass.
+typedef TUniqueFunction<Audio::IProxyDataPtr(UObject*)> FProxyGetterCallback;
+
 typedef TUniqueFunction<TUniquePtr<Metasound::INode>(const Metasound::FNodeInitData&)> FNodeGetterCallback;
 
 namespace Metasound
 {
+	// Various elements that we pass to the frontend registry based on templated type traits.
+	struct FDataTypeRegistryInfo
+	{
+		// The name of the data type itself.
+		FName DataTypeName;
+
+		// What type we should default to using for literals.
+		ELiteralArgType PreferredLiteralType;
+
+		// These bools signify what basic
+		// UProperty primitives we can use to describe this data type as a literal in a document.
+		bool bIsBoolParsable;
+		bool bIsIntParsable;
+		bool bIsFloatParsable;
+		bool bIsStringParsable;
+
+		// these are used for using UObjects, or arrays of UObjects.
+		bool bIsProxyParsable;
+		bool bIsProxyArrayParsable;
+
+		// If this datatype was registered with a specific UClass to use to filter with, that will be used here:
+		UClass* ProxyGeneratorClass;
+
+		// This indicates the type can only be constructed with FOperatorSettings and no other args.
+		// TODO: these can be consolidated to a single bool, since FDataTypeLiteralParam automatically falls back to not using the settings.
+		bool bIsConstructableWithSettings;
+		bool bIsDefaultConstructible;
+
+		FDataTypeRegistryInfo()
+			: bIsBoolParsable(false)
+			, bIsIntParsable(false)
+			, bIsFloatParsable(false)
+			, bIsStringParsable(false)
+			, bIsProxyParsable(false)
+			, bIsProxyArrayParsable(false)
+			, ProxyGeneratorClass(nullptr)
+			, bIsConstructableWithSettings(false)
+			, bIsDefaultConstructible(false)
+		{}
+	};
+
 	namespace Frontend
 	{
 		struct METASOUNDFRONTEND_API FNodeRegistryKey
@@ -51,14 +96,19 @@ namespace Metasound
 			}
 		};
 	}
-}
 
-/*
-static uint32 METASOUNDGRAPHCORE_API GetTypeHash(const Metasound::Frontend::FNodeRegistryKey& InKey)
-{
-	return InKey.NodeHash;
+	struct FDataTypeConstructorCallbacks
+	{
+		// This constructs a TInputNode<> with the corresponding datatype.
+		FInputNodeConstructorCallback InputNodeConstructor;
+
+		// This constructs a TOutputNode<> with the corresponding datatype.
+		FOutputNodeConstructorCallback OutputNodeConstructor;
+
+		// For datatypes that use a UObject literal or a UObject literal array, this lambda generates a literal from the corresponding UObject.
+		FProxyGetterCallback ProxyConstructor;
+	};
 }
-*/
 
 /**
  * Singleton registry for all types and nodes.
@@ -79,18 +129,24 @@ public:
 
 	TMap<FNodeRegistryKey, FNodeRegistryElement>& GetExternalNodeRegistry();
 
-	TUniquePtr<Metasound::INode> ConstructInputNode(const FName& InInputType, const Metasound::FInputNodeConstructorParams& InParams);
+	TUniquePtr<Metasound::INode> ConstructInputNode(const FName& InInputType, Metasound::FInputNodeConstructorParams&& InParams);
 	TUniquePtr<Metasound::INode> ConstructOutputNode(const FName& InOutputType, const Metasound::FOutputNodeConstrutorParams& InParams);
 
+	Metasound::FDataTypeLiteralParam GenerateLiteralForUObject(const FName& InDataType, UObject* InObject);
+	Metasound::FDataTypeLiteralParam GenerateLiteralForUObjectArray(const FName& InDataType, TArray<UObject*> InObjectArray);
+
 	TUniquePtr<Metasound::INode> ConstructExternalNode(const FName& InNodeType, uint32 InNodeHash, const Metasound::FNodeInitData& InInitData);
 
 	// Get the desired kind of literal for a given data type. Returns EConstructorArgType::Invalid if the data type couldn't be found.
 	Metasound::ELiteralArgType GetDesiredLiteralTypeForDataType(FName InDataType) const;
 
+	UClass* GetLiteralUClassForDataType(FName InDataType) const;
+
+
 	// Get whether we can build a literal of this specific type for InDataType.
 	bool DoesDataTypeSupportLiteralType(FName InDataType, Metasound::ELiteralArgType InLiteralType) const;
 
-	bool RegisterDataType(const ::Metasound::FDataTypeRegistryInfo& InDataInfo, FInputNodeConstructorCallback&& InputNodeConstructor, FOutputNodeConstructorCallback&& OutputNodeConstructor);
+	bool RegisterDataType(const ::Metasound::FDataTypeRegistryInfo& InDataInfo, ::Metasound::FDataTypeConstructorCallbacks&& InCallbacks);
 	bool RegisterExternalNode(FNodeGetterCallback&& InCallback);
 
 	// Return any data types that can be used as a metasound input type or output type.
@@ -120,8 +176,8 @@ private:
 
 	struct FDataTypeRegistryElement
 	{
-		FInputNodeConstructorCallback InputNodeConstructor;
-		FOutputNodeConstructorCallback  OutputNodeConstructor;
+		Metasound::FDataTypeConstructorCallbacks Callbacks;
+
 		Metasound::FDataTypeRegistryInfo Info;
 	};
 
diff --git a/Engine/Plugins/Runtime/Metasound/Source/MetasoundFrontend/Public/MetasoundInputNode.h b/Engine/Plugins/Runtime/Metasound/Source/MetasoundFrontend/Public/MetasoundInputNode.h
--- a/Engine/Plugins/Runtime/Metasound/Source/MetasoundFrontend/Public/MetasoundInputNode.h
+++ b/Engine/Plugins/Runtime/Metasound/Source/MetasoundFrontend/Public/MetasoundInputNode.h
@@ -139,10 +139,10 @@ namespace Metasound
 
 			}
 
-			explicit TInputNode(const FString& InNodeDescription, const FString& InVertexName, FDataTypeLiteralParam InParam, const FOperatorSettings& InSettings)
+			explicit TInputNode(const FString& InNodeDescription, const FString& InVertexName, FDataTypeLiteralParam&& InParam, const FOperatorSettings& InSettings)
 				: NodeDescription(InNodeDescription)
 				, VertexName(InVertexName)
-				, Factory(InParam, InSettings)
+				, Factory(MoveTemp(InParam), InSettings)
 			{
 				FOutputDataVertex OutputVertex = MakeOutputDataVertex<DataType>(VertexName, FText::GetEmpty());
 				FDataVertexKey OutputVertexKey = MakeDataVertexKey(OutputVertex);
diff --git a/Engine/Plugins/Runtime/Metasound/Source/MetasoundGraphCore/MetasoundGraphCore.Build.cs b/Engine/Plugins/Runtime/Metasound/Source/MetasoundGraphCore/MetasoundGraphCore.Build.cs
--- a/Engine/Plugins/Runtime/Metasound/Source/MetasoundGraphCore/MetasoundGraphCore.Build.cs
+++ b/Engine/Plugins/Runtime/Metasound/Source/MetasoundGraphCore/MetasoundGraphCore.Build.cs
@@ -11,17 +11,12 @@ namespace UnrealBuildTool.Rules
 		{
             //OptimizeCode = CodeOptimization.Never;
 
-            PrivateDependencyModuleNames.AddRange(
-				new string[] {
-                    "AudioExtensions"
-                }
-			);
-
             PublicDependencyModuleNames.AddRange(
 				new string[] {
                     "Core",
-					"SignalProcessing"
-                }
+					"SignalProcessing",
+					"AudioExtensions"
+				}
             );
 		}
 	}
diff --git a/Engine/Plugins/Runtime/Metasound/Source/MetasoundGraphCore/Public/MetasoundDataReference.h b/Engine/Plugins/Runtime/Metasound/Source/MetasoundGraphCore/Public/MetasoundDataReference.h
--- a/Engine/Plugins/Runtime/Metasound/Source/MetasoundGraphCore/Public/MetasoundDataReference.h
+++ b/Engine/Plugins/Runtime/Metasound/Source/MetasoundGraphCore/Public/MetasoundDataReference.h
@@ -20,6 +20,8 @@
 		static constexpr bool bIsBoolParsable = TTestIfDataTypeCtorIsImplemented<DataType, bool>::Value; \
 		static constexpr bool bIsIntParsable = TTestIfDataTypeCtorIsImplemented<DataType, int32>::Value; \
 		static constexpr bool bIsFloatParsable = TTestIfDataTypeCtorIsImplemented<DataType, float>::Value; \
+		static constexpr bool bIsProxyParsable =  TTestIfDataTypeCtorIsImplemented<DataType, const Audio::IProxyData&>::Value; \
+		static constexpr bool bIsProxyArrayParsable =  TTestIfDataTypeCtorIsImplemented<DataType, const Audio::IProxyData&>::Value; \
 		static constexpr bool bIsConstructableWithSettings = TTestIfDataTypeDefaultCtorIsImplemented<DataType>::Value; \
 		static constexpr bool bCanUseDefaultConstructor = TIsConstructible<DataType>::Value; \
 		static constexpr bool bIsValidSpecialization = true; \
@@ -91,6 +93,8 @@ namespace Metasound
 		static constexpr bool bIsBoolParsable = false;
 		static constexpr bool bIsIntParsable = false;
 		static constexpr bool bIsFloatParsable = false;
+		static constexpr bool bIsProxyParsable = false;
+		static constexpr bool bIsProxyArrayParsable = false;
 		static constexpr bool bIsConstructableWithSettings = false;
 		static constexpr bool bCanUseDefaultConstructor = false;
 		static constexpr bool bIsValidSpecialization = false;
@@ -451,43 +455,9 @@ namespace Metasound
 		Integer,
 		Float,
 		String,
+		UObjectProxy,
+		UObjectProxyArray,
 		None, // If this is set, we will invoke TType(const FOperatorSettings&) If that constructor exists, or the default constructor if not.
 		Invalid,
 	};
-
-	// Various elements that we pass to the frontend registry based on templated type traits.
-	struct FDataTypeRegistryInfo
-	{
-		// The name of the data type itself.
-		FName DataTypeName;
-
-		// What type we should default to using for literals.
-		ELiteralArgType PreferredLiteralType;
-
-		// These bools signify what basic
-		// UProperty primitives we can use to describe this data type as a literal in a document.
-		bool bIsBoolParsable;
-		bool bIsIntParsable;
-		bool bIsFloatParsable;
-		bool bIsStringParsable;
-
-		// This indicates the type can only be constructed with FOperatorSettings and no other args.
-		// TODO: these can be consolidated to a single bool, since FDataTypeLiteralParam automatically falls back to not using the settings.
-		bool bIsConstructableWithSettings;
-		bool bIsDefaultConstructible;
-
-		FDataTypeRegistryInfo()
-			: bIsBoolParsable(false)
-			, bIsIntParsable(false)
-			, bIsFloatParsable(false)
-			, bIsStringParsable(false)
-			, bIsConstructableWithSettings(false)
-			, bIsDefaultConstructible(false)
-		{}
-	};
 }
-
-struct FMyRandomDataType
-{
-	explicit FMyRandomDataType(bool bInValue);
-};
\ No newline at end of file
diff --git a/Engine/Plugins/Runtime/Metasound/Source/MetasoundGraphCore/Public/MetasoundOperatorInterface.h b/Engine/Plugins/Runtime/Metasound/Source/MetasoundGraphCore/Public/MetasoundOperatorInterface.h
--- a/Engine/Plugins/Runtime/Metasound/Source/MetasoundGraphCore/Public/MetasoundOperatorInterface.h
+++ b/Engine/Plugins/Runtime/Metasound/Source/MetasoundGraphCore/Public/MetasoundOperatorInterface.h
@@ -7,6 +7,7 @@
 #include "MetasoundDataReference.h"
 #include "MetasoundDataReferenceCollection.h"
 #include "MetasoundOperatorSettings.h"
+#include "IAudioProxyInitializer.h"
 
 namespace Metasound
 {
@@ -143,15 +144,18 @@ namespace Metasound
 	 */
 	struct FDataTypeLiteralParam
 	{
+		FDataTypeLiteralParam(const FDataTypeLiteralParam& Other) = delete;
+		FDataTypeLiteralParam(FDataTypeLiteralParam&& Other) = default;
+
 		ELiteralArgType ConstructorArgType;
-		TVariant<bool, int, float, FString> ConstructorArg;
+		TVariant<bool, int, float, FString, Audio::IProxyDataPtr, TArray<Audio::IProxyDataPtr>> ConstructorArg;
 
 		// builds an invalid FDataTypeLiteralParam.
 		static FDataTypeLiteralParam InvalidParam()
 		{
 			FDataTypeLiteralParam InitParam;
 			InitParam.ConstructorArgType = ELiteralArgType::Invalid;
-			return InitParam;
+			return MoveTemp(InitParam);
 		}
 
 		bool IsValid() const
@@ -187,6 +191,18 @@ namespace Metasound
 			ConstructorArg.Set<FString>(InString);
 		}
 
+		FDataTypeLiteralParam(Audio::IProxyDataPtr&& InProxy)
+			: ConstructorArgType(ELiteralArgType::UObjectProxy)
+		{
+			ConstructorArg.Set<Audio::IProxyDataPtr>(MoveTemp(InProxy));
+		}
+
+		FDataTypeLiteralParam(TArray<Audio::IProxyDataPtr>&& InProxyArray)
+			: ConstructorArgType(ELiteralArgType::UObjectProxyArray)
+		{
+			ConstructorArg.Set<TArray<Audio::IProxyDataPtr>>(MoveTemp(InProxyArray));
+		}
+
 		template<typename DataType>
 		bool IsCompatibleWithType()
 		{
@@ -208,6 +224,14 @@ namespace Metasound
 				{
 					return TTestIfDataTypeCtorIsImplemented<DataType, const FString&>::Value;
 				}
+				case ELiteralArgType::UObjectProxy:
+				{
+					return TTestIfDataTypeCtorIsImplemented<DataType, const Audio::IProxyData&>::Value;
+				}
+				case ELiteralArgType::UObjectProxyArray:
+				{
+					return TTestIfDataTypeCtorIsImplemented<DataType, const TArray<Audio::IProxyDataPtr>&>::Value;
+				}
 				case ELiteralArgType::None:
 				{
 					return TTestIfDataTypeDefaultCtorIsImplemented<DataType>::Value || TIsConstructible<DataType>::Value;
@@ -240,6 +264,59 @@ namespace Metasound
 		// 2: Otherwise, if there's a constructor that just takes the literal type, we'll use that.
 		// 3: If neither constructor exists, we crash. Ideally this has already been caught by the static assert in the REGISTER_METASOUND_DATATYPE macro.
 
+		template<typename DataType, typename TEnableIf<
+			TIsConstructible<DataType, const TArray<Audio::IProxyDataPtr>&, const FOperatorSettings&>::Value
+			, bool >::Type = true >
+			DataType ParseAsAudioProxyArray(const FOperatorSettings& WithOperatorSettings)
+		{
+			check(ConstructorArgType == ELiteralArgType::UObjectProxyArray);
+			const TArray<Audio::IProxyDataPtr>& ProxyArrayToParse = ConstructorArg.Get<TArray<Audio::IProxyDataPtr>>();
+			return DataType(ProxyArrayToParse, WithOperatorSettings);
+		}
+
+		template<typename DataType, typename TEnableIf<
+			TIsConstructible<DataType, const TArray<Audio::IProxyDataPtr>&>::Value && !TIsConstructible<DataType, const TArray<Audio::IProxyDataPtr>&, const FOperatorSettings&>::Value
+			, bool >::Type = true >
+			DataType ParseAsAudioProxyArray(const FOperatorSettings& WithOperatorSettings)
+		{
+			check(ConstructorArgType == ELiteralArgType::UObjectProxyArray);
+			const TArray<Audio::IProxyDataPtr>& ProxyArrayToParse = ConstructorArg.Get<TArray<Audio::IProxyDataPtr>>();
+			return DataType(ProxyArrayToParse);
+		}
+
+		template<typename DataType>
+		DataType ParseAsAudioProxyArray(...)
+		{
+			return InvalidConstructor<DataType>();
+		}
+
+		template<typename DataType, typename TEnableIf<
+			TIsConstructible<DataType, const Audio::IProxyData&, const FOperatorSettings&>::Value
+			, bool >::Type = true >
+		DataType ParseAsAudioProxy(const FOperatorSettings& WithOperatorSettings)
+		{
+			check(ConstructorArgType == ELiteralArgType::UObjectProxy);
+			const Audio::IProxyDataPtr& ProxyToParse = ConstructorArg.Get<Audio::IProxyDataPtr>();
+			return DataType(ProxyToParse, WithOperatorSettings);
+		}
+
+		template<typename DataType, typename TEnableIf<
+			TIsConstructible<DataType, const Audio::IProxyData&>::Value && !TIsConstructible<DataType, const Audio::IProxyData&, const FOperatorSettings&>::Value
+			, bool >::Type = true >
+		DataType ParseAsAudioProxy(const FOperatorSettings& WithOperatorSettings)
+		{
+			check(ConstructorArgType == ELiteralArgType::UObjectProxy);
+			const Audio::IProxyDataPtr& ProxyToParse = ConstructorArg.Get<Audio::IProxyDataPtr>();
+			check(ProxyToParse.IsValid());
+			return DataType(*ProxyToParse);
+		}
+
+		template<typename DataType>
+		DataType ParseAsAudioProxy(...)
+		{
+			return InvalidConstructor<DataType>();
+		}
+
 		template<typename DataType, typename TEnableIf< 
 			TIsConstructible<DataType, const FString&, const FOperatorSettings&>::Value
 			, bool >::Type = true >
@@ -393,6 +470,14 @@ namespace Metasound
 				{
 					return ParseAsString<DataType>(WithOperatorSettings);
 				}
+				case ELiteralArgType::UObjectProxy:
+				{
+					return ParseAsAudioProxy<DataType>(WithOperatorSettings);
+				}
+				case ELiteralArgType::UObjectProxyArray:
+				{
+					return ParseAsAudioProxyArray<DataType>(WithOperatorSettings);
+				}
 				case ELiteralArgType::None:
 				{
 					return ParseWithDefaultConstructor<DataType>(WithOperatorSettings);
diff --git a/Engine/Source/Runtime/AudioExtensions/Public/IAudioProxyInitializer.h b/Engine/Source/Runtime/AudioExtensions/Public/IAudioProxyInitializer.h
new file mode 100644
--- /dev/null
+++ b/Engine/Source/Runtime/AudioExtensions/Public/IAudioProxyInitializer.h
@@ -0,0 +1,131 @@
+// Copyright Epic Games, Inc. All Rights Reserved.
+
+#pragma once
+
+#include "CoreMinimal.h"
+#include "Templates/UnrealTypeTraits.h"
+
+/**
+ * Interfaces for Audio Proxy Objects 
+ * These are used to spawn thread safe instances of UObjects that may be garbage collected on the game thread.
+ * In shipping builds, these are effectively abstract pointers, but CHECK_AUDIOPROXY_TYPES can optionally be used
+ * to check downcasts.
+ */
+
+// This can be set to 1 to double check proxies.
+#define CHECK_AUDIOPROXY_TYPES !(UE_BUILD_SHIPPING || UE_BUILD_TEST || UE_BUILD_DEVELOPMENT)
+
+#define  IMPL_AUDIOPROXY_CLASS(FClassName) \
+	static FName GetAudioProxyTypeName() \
+	{ \
+		static FName MyClassName = TEXT( #FClassName ); \
+		return MyClassName; \
+	} \
+	static constexpr bool bWasAudioProxyClassImplemented = true; \
+	friend class ::Audio::IProxyData; \
+	friend class ::Audio::TProxyData<FClassName>;
+
+
+
+namespace Audio
+{
+	constexpr bool bShouldCheckAudioProxyTypes = CHECK_AUDIOPROXY_TYPES;
+
+	/*
+	 * Base class that allows us to typecheck proxy data before downcasting it in debug builds.
+	*/
+	class IProxyData
+	{
+	private:
+		FName ProxyTypeName;
+	public:
+		virtual ~IProxyData() = default;
+
+		template<typename ProxyType>
+		bool CheckTypeCast() const
+		{
+			if (bShouldCheckAudioProxyTypes)
+			{
+				FName DestinationTypeName = ProxyType::GetAudioProxyTypeName();
+				return ensureAlwaysMsgf(ProxyTypeName == DestinationTypeName, TEXT("Tried to downcast type %s to %s!"), *ProxyTypeName.ToString(), *DestinationTypeName.ToString());
+			}
+			else
+			{
+				return true;
+			}
+		}
+
+		template<typename ProxyType>
+		ProxyType& GetAs()
+		{
+			static_assert(TIsDerivedFrom<ProxyType, IProxyData>::Value, "Tried to downcast IProxyInitData to an unrelated type!");
+			if (CheckTypeCast<ProxyType>())
+			{
+				return static_cast<ProxyType&>(*this);
+			}
+			else
+			{
+				// This is an illegal cast, and is considered a fatal error.
+				checkNoEntry();
+				return *((ProxyType*)0x1);
+			}
+		}
+
+		template<typename ProxyType>
+		const ProxyType& GetAs() const
+		{
+			static_assert(TIsDerivedFrom<ProxyType, IProxyData>::Value, "Tried to downcast IProxyInitData to an unrelated type!");
+			if (CheckTypeCast<ProxyType>())
+			{
+				return static_cast<const ProxyType&>(*this);
+			}
+			else
+			{
+				// This is an illegal cast, and is considered a fatal error.
+				checkNoEntry();
+				return *((ProxyType*)0x1);
+			}
+		}
+
+		IProxyData(FName InProxyTypeName)
+			: ProxyTypeName(InProxyTypeName)
+		{}
+	};
+
+	/**
+	 * This class can be implemented to create a custom, threadsafe instance of a given UObject.
+	 * This is a CRTP class, and should always be subclassed with the name of the subclass.
+	 */
+	template <typename Type>
+	class TProxyData : public IProxyData
+	{
+	protected:
+		static constexpr bool bWasAudioProxyClassImplemented = false;
+
+	public:
+		TProxyData()
+			: IProxyData(Type::GetAudioProxyTypeName())
+		{
+			static_assert(Type::bWasAudioProxyClassImplemented, "Make sure to include IMPL_AUDIOPROXY_CLASS(ClassName) in your implementation of TProxyData.");
+		}
+	};
+
+
+	struct FProxyDataInitParams
+	{
+		FName NameOfFeatureRequestingProxy;
+	};
+
+	typedef TUniquePtr<IProxyData> IProxyDataPtr;
+
+	
+}
+
+/*
+* This can be subclassed to make a UClass an audio proxy factory.
+*/
+class IAudioProxyDataFactory
+{
+public:
+	virtual TUniquePtr<Audio::IProxyData> CreateNewProxyData(const Audio::FProxyDataInitParams& InitParams) = 0;
+};
diff --git a/Engine/Source/Runtime/Engine/Classes/Sound/SoundWave.h b/Engine/Source/Runtime/Engine/Classes/Sound/SoundWave.h
--- a/Engine/Source/Runtime/Engine/Classes/Sound/SoundWave.h
+++ b/Engine/Source/Runtime/Engine/Classes/Sound/SoundWave.h
@@ -20,6 +20,7 @@
 #include "AudioCompressionSettings.h"
 #include "PerPlatformProperties.h"
 #include "ContentStreaming.h"
+#include "IAudioProxyInitializer.h"
 #include "SoundWave.generated.h"
 
 class ITargetPlatform;
@@ -302,10 +303,26 @@ struct ISoundWaveClient
 	virtual void OnFinishDestroy(class USoundWave* Wave) = 0;
 };
 
+class FSoundWaveProxy : public Audio::TProxyData<FSoundWaveProxy>
+{
+	IMPL_AUDIOPROXY_CLASS(FSoundWaveProxy)
+
+public:
+	// For now, this is just a weak object pointer.
+	// However, in the future this should be a thread-safe handle to whatever data is needed for this wave.
+	TWeakObjectPtr<USoundWave> SoundWavePtr;
+
+	FSoundWaveProxy(USoundWave* InWave)
+		: SoundWavePtr(InWave)
+	{
+	}
+};
+
 UCLASS(hidecategories=Object, editinlinenew, BlueprintType)
-class ENGINE_API USoundWave : public USoundBase
+class ENGINE_API USoundWave : public USoundBase, public IAudioProxyDataFactory
 {
 	GENERATED_UCLASS_BODY()
+
 public:
 	/** Platform agnostic compression quality. 1..100 with 1 being best compression and 100 being best quality. */
 	UPROPERTY(EditAnywhere, Category="Format|Quality", meta=(DisplayName = "Compression", ClampMin = "1", ClampMax = "100"), AssetRegistrySearchable)
@@ -725,12 +742,16 @@ public:
 	virtual bool HasCookedAmplitudeEnvelopeData() const override;
 	//~ End USoundBase Interface.
 
+	//~Begin IAudioProxyDataFactory Interface.
+	virtual TUniquePtr<Audio::IProxyData> CreateNewProxyData(const Audio::FProxyDataInitParams& InitParams) override;
+	//~ End IAudioProxyDataFactory Interface.
+
 	// Called  when the procedural sound wave begins on the render thread. Only used in the audio mixer and when bProcedural is true.
 	virtual void OnBeginGenerate() {}
 
 	// Called when the procedural sound wave is done generating on the render thread. Only used in the audio mixer and when bProcedural is true..
 	virtual void OnEndGenerate() {};
-
+	
 	void AddPlayingSource(const FSoundWaveClientPtr& Source);
 	void RemovePlayingSource(const FSoundWaveClientPtr& Source);
 
diff --git a/Engine/Source/Runtime/Engine/Private/SoundWave.cpp b/Engine/Source/Runtime/Engine/Private/SoundWave.cpp
--- a/Engine/Source/Runtime/Engine/Private/SoundWave.cpp
+++ b/Engine/Source/Runtime/Engine/Private/SoundWave.cpp
@@ -2166,6 +2166,11 @@ bool USoundWave::HasCookedAmplitudeEnvelopeData() const
 	return CookedEnvelopeTimeData.Num() > 0;
 }
 
+TUniquePtr<Audio::IProxyData> USoundWave::CreateNewProxyData(const Audio::FProxyDataInitParams& InitParams)
+{
+	return TUniquePtr<Audio::IProxyData>(new FSoundWaveProxy(this));
+}
+
 void USoundWave::AddPlayingSource(const FSoundWaveClientPtr& Source)
 {
 	check(Source);
