Android NDK fix- fix __VA_ARGS__ syntax error when REGISTER_METASOUND_DATA_TYPE is called with just two args.

diff --git a/Engine/Plugins/Runtime/Metasound/Source/MetasoundFrontend/Public/MetasoundDataTypeRegistrationMacro.h b/Engine/Plugins/Runtime/Metasound/Source/MetasoundFrontend/Public/MetasoundDataTypeRegistrationMacro.h
--- a/Engine/Plugins/Runtime/Metasound/Source/MetasoundFrontend/Public/MetasoundDataTypeRegistrationMacro.h
+++ b/Engine/Plugins/Runtime/Metasound/Source/MetasoundFrontend/Public/MetasoundDataTypeRegistrationMacro.h
@@ -128,5 +128,5 @@ bool RegisterDataTypeWithFrontend()
 	static_assert(::Metasound::TDataReferenceTypeInfo<DataType>::bIsValidSpecialization, "Please call DECLARE_METASOUND_DATA_REFERENCE_TYPES(" #DataType "...) before calling REGISTER_METASOUND_DATATYPE(" #DataType")."); \
 	static constexpr bool bCanRegister##DataType = ::Metasound::TDataReferenceTypeInfo<DataType>::bIsStringParsable || ::Metasound::TDataReferenceTypeInfo<DataType>::bIsBoolParsable || ::Metasound::TDataReferenceTypeInfo<DataType>::bIsIntParsable || ::Metasound::TDataReferenceTypeInfo<DataType>::bIsFloatParsable || ::Metasound::TDataReferenceTypeInfo<DataType>::bIsProxyParsable || ::Metasound::TDataReferenceTypeInfo<DataType>::bIsProxyArrayParsable || ::Metasound::TDataReferenceTypeInfo<DataType>::bIsConstructableWithSettings || ::Metasound::TDataReferenceTypeInfo<DataType>::bCanUseDefaultConstructor; \
 	static_assert(bCanRegister##DataType , "To register " #DataType " to be used as a Metasounds input or output type, it needs a default constructor or one of the following constructors must be implemented:  " #DataType "(const ::Metasound::FOperatorSettings& InSettings), " #DataType "(int32 InValue, const ::Metasound::FOperatorSettings& InSettings), " #DataType "(float InValue, const ::Metasound::FOperatorSettings& InSettings), " #DataType "(bool InValue, const ::Metasound::FOperatorSettings& InSettings), " #DataType "(const Audio::IProxyData& InData, const ::Metasound::FOperatorSettings& InSettings), " #DataType "(const TArray<Audio::IProxyDataPtr>& InProxyArray, const ::Metasound::FOperatorSettings& InSettings) or " #DataType "(const FString& InString, const ::Metasound::FOperatorSettings& InSettings)."); \
-	static const bool bSuccessfullyRegistered##DataType = FMetasoundFrontendRegistryContainer::Get()->EnqueueInitCommand([](){ ::RegisterDataTypeWithFrontend<DataType, __VA_ARGS__>(); }); // This static bool is useful for debugging, but also is the only way the compiler will let us call this function outside of an expression.
+	static const bool bSuccessfullyRegistered##DataType = FMetasoundFrontendRegistryContainer::Get()->EnqueueInitCommand([](){ ::RegisterDataTypeWithFrontend<DataType, ##__VA_ARGS__>(); }); // This static bool is useful for debugging, but also is the only way the compiler will let us call this function outside of an expression.
 
