Fix for static analysis warnings.

diff --git a/Engine/Plugins/Runtime/Metasound/Source/MetasoundFrontend/Private/MetasoundFrontend.cpp b/Engine/Plugins/Runtime/Metasound/Source/MetasoundFrontend/Private/MetasoundFrontend.cpp
--- a/Engine/Plugins/Runtime/Metasound/Source/MetasoundFrontend/Private/MetasoundFrontend.cpp
+++ b/Engine/Plugins/Runtime/Metasound/Source/MetasoundFrontend/Private/MetasoundFrontend.cpp
@@ -2097,7 +2097,7 @@ namespace Metasound
 					}
 				}
 
-				if (!ensureAlwaysMsgf(IndexOfDependencyToRemove >= 0,
+				if (!ensureAlwaysMsgf(IndexOfDependencyInDocument >= 0,
 					TEXT("Couldn't find node class %s in the list of dependencies for this document, but found it in the nodes list. This likely means that the underlying class description is corrupted.)"), *NodeClassName))
 				{
 					return false;
diff --git a/Engine/Plugins/Runtime/Metasound/Source/MetasoundFrontend/Private/MetasoundFrontendRegistries.cpp b/Engine/Plugins/Runtime/Metasound/Source/MetasoundFrontend/Private/MetasoundFrontendRegistries.cpp
--- a/Engine/Plugins/Runtime/Metasound/Source/MetasoundFrontend/Private/MetasoundFrontendRegistries.cpp
+++ b/Engine/Plugins/Runtime/Metasound/Source/MetasoundFrontend/Private/MetasoundFrontendRegistries.cpp
@@ -317,13 +317,13 @@ bool FMetasoundFrontendRegistryContainer::RegisterExternalNode(FNodeGetterCallba
 
 	for (auto& InputTuple : Inputs)
 	{
-		HashCombine(NodeHash, FCrc::StrCrc32(*InputTuple.Value.VertexName));
+		NodeHash = HashCombine(NodeHash, FCrc::StrCrc32(*InputTuple.Value.VertexName));
 		InputTypes.Add(InputTuple.Value.DataReferenceTypeName);
 	}
 
 	for (auto& OutputTuple : Outputs)
 	{
-		HashCombine(NodeHash, FCrc::StrCrc32(*OutputTuple.Value.VertexName));
+		NodeHash = HashCombine(NodeHash, FCrc::StrCrc32(*OutputTuple.Value.VertexName));
 		OutputTypes.Add(OutputTuple.Value.DataReferenceTypeName);
 	}
 
diff --git a/Engine/Plugins/Runtime/Metasound/Source/MetasoundGenerator/Private/MetasoundGenerator.cpp b/Engine/Plugins/Runtime/Metasound/Source/MetasoundGenerator/Private/MetasoundGenerator.cpp
--- a/Engine/Plugins/Runtime/Metasound/Source/MetasoundGenerator/Private/MetasoundGenerator.cpp
+++ b/Engine/Plugins/Runtime/Metasound/Source/MetasoundGenerator/Private/MetasoundGenerator.cpp
@@ -69,23 +69,20 @@ namespace Metasound
 		// Single channel version:
 		NumChannels = 1;
 
-		if (NumChannels > 0)
-		{
-			// All buffers have same number of frames, so only need to query
-			// first buffer to know number of frames.
-			
-			// Multichannel version:
-			//NumFramesPerExecute = GraphOutputAudioRef->GetBuffers()[0]->Num();
-			// single channel version:
-			NumFramesPerExecute = GraphOutputAudioRef->Num();
+		// All buffers have same number of frames, so only need to query
+		// first buffer to know number of frames.
 
-			NumSamplesPerExecute = NumFramesPerExecute * NumChannels;
+		// Multichannel version:
+		//NumFramesPerExecute = GraphOutputAudioRef->GetBuffers()[0]->Num();
+		// single channel version:
+		NumFramesPerExecute = GraphOutputAudioRef->Num();
 
-			if (NumSamplesPerExecute > 0)
-			{
-				// Preallocate interleaved buffer as it is necessary for any audio generation calls.
-				InterleavedAudioBuffer.AddUninitialized(NumSamplesPerExecute);
-			}
+		NumSamplesPerExecute = NumFramesPerExecute * NumChannels;
+
+		if (NumSamplesPerExecute > 0)
+		{
+			// Preallocate interleaved buffer as it is necessary for any audio generation calls.
+			InterleavedAudioBuffer.AddUninitialized(NumSamplesPerExecute);
 		}
 	}
 
