Added FVertexInterface and reworked existing INode interfaces
Added FNodeFacade to simplify writing nodes and reduce coding errors.

diff --git a/Engine/Plugins/Runtime/Metasound/Source/MetasoundEngine/Private/MetasoundSource.cpp b/Engine/Plugins/Runtime/Metasound/Source/MetasoundEngine/Private/MetasoundSource.cpp
--- a/Engine/Plugins/Runtime/Metasound/Source/MetasoundEngine/Private/MetasoundSource.cpp
+++ b/Engine/Plugins/Runtime/Metasound/Source/MetasoundEngine/Private/MetasoundSource.cpp
@@ -64,8 +64,9 @@ ISoundGeneratorPtr UMetasoundSource::CreateSoundGenerator(const FSoundGeneratorI
 	Duration = INDEFINITELY_LOOPING_DURATION;
 	bLooping = true;
 	SampleRate = InParams.SampleRate;
+	const float BlockRate = 100.f; // Metasound graph gets evaluated 100 times per second.
 
-	Metasound::FOperatorSettings InSettings(InParams.SampleRate, InParams.NumFramesPerCallback);
+	Metasound::FOperatorSettings InSettings(InParams.SampleRate, BlockRate);
 
 	TArray<Metasound::IOperatorBuilder::FBuildErrorPtr> BuildErrors;
 
@@ -80,7 +81,7 @@ ISoundGeneratorPtr UMetasoundSource::CreateSoundGenerator(const FSoundGeneratorI
 		Metasound::FMetasoundGeneratorInitParams InitParams =
 		{
 			MoveTemp(Operator),
-			Outputs.GetDataReadReferenceOrConstruct<Metasound::FAudioBuffer>(GetAudioOutputName(), InParams.NumFramesPerCallback),
+			Outputs.GetDataReadReferenceOrConstruct<Metasound::FAudioBuffer>(GetAudioOutputName(), InSettings.GetNumFramesPerBlock()),
 			Outputs.GetDataReadReferenceOrConstruct<Metasound::FBop>(GetIsFinishedOutputName(), false)
 		};
 
diff --git a/Engine/Plugins/Runtime/Metasound/Source/MetasoundEngine/Private/MetasoundWavePlayerNode.cpp b/Engine/Plugins/Runtime/Metasound/Source/MetasoundEngine/Private/MetasoundWavePlayerNode.cpp
--- a/Engine/Plugins/Runtime/Metasound/Source/MetasoundEngine/Private/MetasoundWavePlayerNode.cpp
+++ b/Engine/Plugins/Runtime/Metasound/Source/MetasoundEngine/Private/MetasoundWavePlayerNode.cpp
@@ -118,23 +118,26 @@ namespace Metasound
 		TUniquePtr<Audio::IDecoderOutput> DecoderOutput;
 	};
 
-	const FName FWavePlayerNode::ClassName = FName(TEXT("Wave"));
+	const FNodeInfo FWavePlayerNode::Info = FNodeInfo(
+		{
+			FName(TEXT("Wave")),
+			LOCTEXT("Metasound_WavePlayerNodeDescription", "Plays a supplied Wave"),
+			PluginAuthor,
+			PluginNodeMissingPrompt
+		}
+	);
 
 	TUniquePtr<IOperator> FWavePlayerNode::FOperatorFactory::CreateOperator(
-		const INode& InNode, 
-		const FOperatorSettings& InOperatorSettings,
-		const FDataReferenceCollection& InInputDataReferences, 
-		TArray<TUniquePtr<IOperatorBuildError>>& OutErrors) 
+		const FCreateOperatorParams& InParams, 
+		FBuildErrorArray& OutErrors) 
 	{
 		using namespace Audio;
 
-		const FWavePlayerNode& WaveNode = static_cast<const FWavePlayerNode&>(InNode);
-		FWaveReadRef Wave = FWaveReadRef::CreateNew();
+		const FWavePlayerNode& WaveNode = static_cast<const FWavePlayerNode&>(InParams.Node);
 
-		if (InInputDataReferences.ContainsDataReadReference<FWave>(TEXT("Wave")))
-		{
-			Wave = InInputDataReferences.GetDataReadReference<FWave>(TEXT("Wave"));
-		}
+		const FDataReferenceCollection& InputCol = InParams.InputDataReferences;
+
+		FWaveReadRef Wave = InputCol.GetDataReadReferenceOrConstruct<FWave>(TEXT("Wave"));
 
 		FWave::FDecoderInputPtr Input = FWave::CreateDecoderInput(Wave);
 		if (Input)
@@ -143,12 +146,12 @@ namespace Metasound
 			if (Codec)
 			{
 				// V1, Ask for an output buffer the size of a frame.
-				IDecoderOutput::FRequirements Reqs { Float32_Interleaved, InOperatorSettings.GetNumFramesPerBlock() };
+				IDecoderOutput::FRequirements Reqs { Float32_Interleaved, InParams.OperatorSettings.GetNumFramesPerBlock() };
 				TUniquePtr<IDecoderOutput> Output = IDecoderOutput::Create(Reqs);
 				TUniquePtr<IDecoder> Decoder = Codec->CreateDecoder(Input.Get(), Output.Get());
 
 				return MakeUnique<FWavePlayerOperator>(
-					InOperatorSettings, 
+					InParams.OperatorSettings, 
 					Wave, 
 					MoveTemp(Input), 
 					MoveTemp(Output), 
@@ -157,25 +160,26 @@ namespace Metasound
 			}
 			else
 			{
-				OutErrors.Add(MakeUnique<FWavePlayerError>(TEXT("FailedToFindCodec"),
-					LOCTEXT("FailedToFindCodec", "Failed to find codec for opening the supplied Wave")));
+				AddBuildError<FWavePlayerError>(OutErrors, TEXT("FailedToFindCodec"),
+					LOCTEXT("FailedToFindCodec", "Failed to find codec for opening the supplied Wave"));
 			}
 		}
 		else
 		{
-			OutErrors.Add(MakeUnique<FWavePlayerError>(TEXT("FailedToParseInput"),
-				LOCTEXT("FailedToParseInput", "Failed to parse the compressed data")));
+			AddBuildError<FWavePlayerError>(OutErrors, TEXT("FailedToParseInput"),
+				LOCTEXT("FailedToParseInput", "Failed to parse the compressed data"));
 		}
 
 		// Create the player without any inputs, will just produce silence.
-		return MakeUnique<FWavePlayerOperator>(InOperatorSettings,Wave);
+		return MakeUnique<FWavePlayerOperator>(InParams.OperatorSettings, Wave);
 	}
 
 	FWavePlayerNode::FWavePlayerNode(const FString& InName)
-		:	FNode(InName)
+		:	FNode(InName, FWavePlayerNode::Info)
+		,	Factory(MakeOperatorFactoryRef<FWavePlayerNode::FOperatorFactory>())
 	{
-		AddInputDataVertex<FWave>(TEXT("Wave"), LOCTEXT("WaveTooltip", "The Wave to be decoded"));
-		AddOutputDataVertex<FAudioBuffer>(TEXT("Audio"), LOCTEXT("AudioTooltip", "The output audio"));
+		Interface.GetInputInterface().Add(TInputDataVertexModel<FWave>(TEXT("Wave"), LOCTEXT("WaveTooltip", "The Wave to be decoded")));
+		Interface.GetOutputInterface().Add(TOutputDataVertexModel<FAudioBuffer>(TEXT("Audio"), LOCTEXT("AudioTooltip", "The output audio")));
 	}
 
 	FWavePlayerNode::FWavePlayerNode(const FNodeInitData& InInitData)
@@ -183,18 +187,29 @@ namespace Metasound
 	{
 	}
 
-	FWavePlayerNode::~FWavePlayerNode()
+	FOperatorFactorySharedRef FWavePlayerNode::GetDefaultOperatorFactory() const
 	{
+		return Factory;
 	}
-	
-	const FName& FWavePlayerNode::GetClassName() const
+
+	const FVertexInterface& FWavePlayerNode::GetVertexInterface() const
 	{
-		return ::Metasound::FWavePlayerNode::ClassName;
+		return Interface;
 	}
 
-	IOperatorFactory& FWavePlayerNode::GetDefaultOperatorFactory() 
+	const FVertexInterface& FWavePlayerNode::GetDefaultVertexInterface() const
 	{
-		return Factory;
+		return Interface;
+	}
+
+	bool FWavePlayerNode::SetVertexInterface(const FVertexInterface& InInterface)
+	{
+		return InInterface == Interface;
+	}
+
+	bool FWavePlayerNode::IsVertexInterfaceSupported(const FVertexInterface& InInterface) const
+	{
+		return InInterface == Interface;
 	}
 }
 #undef LOCTEXT_NAMESPACE //MetasoundWaveNode
diff --git a/Engine/Plugins/Runtime/Metasound/Source/MetasoundEngine/Public/MetasoundWavePlayerNode.h b/Engine/Plugins/Runtime/Metasound/Source/MetasoundEngine/Public/MetasoundWavePlayerNode.h
--- a/Engine/Plugins/Runtime/Metasound/Source/MetasoundEngine/Public/MetasoundWavePlayerNode.h
+++ b/Engine/Plugins/Runtime/Metasound/Source/MetasoundEngine/Public/MetasoundWavePlayerNode.h
@@ -14,40 +14,47 @@ namespace Metasound
 	{
 		class FOperatorFactory : public IOperatorFactory
 		{
-			virtual TUniquePtr<IOperator> CreateOperator(const INode& InNode, const FOperatorSettings& InOperatorSettings, const FDataReferenceCollection& InInputDataReferences, TArray<TUniquePtr<IOperatorBuildError>>& OutErrors) override;
+			virtual TUniquePtr<IOperator> CreateOperator(const FCreateOperatorParams& InParams, FBuildErrorArray& OutErrors) override;
 		};
 
 	public:
-		static const FName ClassName;
+		static const FNodeInfo Info;
 
 		FWavePlayerNode(const FString& InName);
 
 		// constructor used by the Metasound Frontend.
 		FWavePlayerNode(const FNodeInitData& InInitData);
 
-		virtual ~FWavePlayerNode();
+		virtual ~FWavePlayerNode() = default;
 
-		const FName& GetClassName() const override;
+		virtual FOperatorFactorySharedRef GetDefaultOperatorFactory() const override;
 
-		IOperatorFactory& GetDefaultOperatorFactory() override;
+		/** Return the current vertex interface. */
+		virtual const FVertexInterface& GetVertexInterface() const override;
 
-		const FText& GetDescription() const override
-		{
-			static const FText StaticDescription = NSLOCTEXT("MetasoundGraphCore", "Metasound_WavePlayerNodeDescription", "Plays a supplied Wave");
-			return StaticDescription;
-		}
+		/** Return the default vertex interface. */
+		virtual const FVertexInterface& GetDefaultVertexInterface() const override;
 
-		const FText& GetAuthorName() const override
-		{
-			return PluginAuthor;
-		}
+		/** Set the vertex interface. If the vertex was successfully changed, returns true. 
+		 *
+		 * @param InInterface - New interface for node. 
+		 *
+		 * @return True on success, false otherwise.
+		 */
+		virtual bool SetVertexInterface(const FVertexInterface& InInterface) override;
 
-		const FText& GetPromptIfMissing() const override
-		{
-			return PluginNodeMissingPrompt;
-		}
+		/** Expresses whether a specific vertex interface is supported.
+		 *
+		 * @param InInterface - New interface. 
+		 *
+		 * @return True if the interface is supported, false otherwise. 
+		 */
+		virtual bool IsVertexInterfaceSupported(const FVertexInterface& InInterface) const override;
 
 	private:
-		FOperatorFactory Factory;
+
+		FOperatorFactorySharedRef Factory;
+
+		FVertexInterface Interface;
 	};
 }
diff --git a/Engine/Plugins/Runtime/Metasound/Source/MetasoundFrontend/Private/MetasoundFrontend.cpp b/Engine/Plugins/Runtime/Metasound/Source/MetasoundFrontend/Private/MetasoundFrontend.cpp
--- a/Engine/Plugins/Runtime/Metasound/Source/MetasoundFrontend/Private/MetasoundFrontend.cpp
+++ b/Engine/Plugins/Runtime/Metasound/Source/MetasoundFrontend/Private/MetasoundFrontend.cpp
@@ -157,8 +157,9 @@ namespace Metasound
 				}
 
 				const FName NodeName = DummyNode->GetClassName();
-				const FInputDataVertexCollection& Inputs = DummyNode->GetInputDataVertices();
-				const FOutputDataVertexCollection& Outputs = DummyNode->GetOutputDataVertices();
+				const FVertexInterface& DummyNodeInterface = DummyNode->GetDefaultVertexInterface();
+				const FInputVertexInterface& DummyInputInterface = DummyNodeInterface.GetInputInterface();
+				const FOutputVertexInterface& DummyOutputInterface = DummyNodeInterface.GetOutputInterface();
 
 				FMetasoundClassMetadata NodeMetadata;
 				NodeMetadata.NodeName = NodeName.ToString();
@@ -172,22 +173,22 @@ namespace Metasound
 
 				// External metasounds aren't dependent on any other nodes by definition, so all we need to do
 				// is populate the Input and Output sets.
-				for (auto& InputTuple : Inputs)
+				for (auto& InputTuple : DummyInputInterface)
 				{
 					FMetasoundInputDescription InputDescription;
-					InputDescription.Name = InputTuple.Value.VertexName;
-					InputDescription.TypeName = InputTuple.Value.DataReferenceTypeName;
-					InputDescription.ToolTip = InputTuple.Value.Description;
+					InputDescription.Name = InputTuple.Value.GetVertexName();
+					InputDescription.TypeName = InputTuple.Value.GetDataTypeName();
+					InputDescription.ToolTip = InputTuple.Value.GetDescription();
 
 					ClassDescription.Inputs.Add(InputDescription);
 				}
 
-				for (auto& OutputTuple : Outputs)
+				for (auto& OutputTuple : DummyOutputInterface)
 				{
 					FMetasoundOutputDescription OutputDescription;
-					OutputDescription.Name = OutputTuple.Value.VertexName;
-					OutputDescription.TypeName = OutputTuple.Value.DataReferenceTypeName;
-					OutputDescription.ToolTip = OutputTuple.Value.Description;
+					OutputDescription.Name = OutputTuple.Value.GetVertexName();
+					OutputDescription.TypeName = OutputTuple.Value.GetDataTypeName();
+					OutputDescription.ToolTip = OutputTuple.Value.GetDescription();
 
 					ClassDescription.Outputs.Add(OutputDescription);
 				}
diff --git a/Engine/Plugins/Runtime/Metasound/Source/MetasoundFrontend/Private/MetasoundFrontendRegistries.cpp b/Engine/Plugins/Runtime/Metasound/Source/MetasoundFrontend/Private/MetasoundFrontendRegistries.cpp
--- a/Engine/Plugins/Runtime/Metasound/Source/MetasoundFrontend/Private/MetasoundFrontendRegistries.cpp
+++ b/Engine/Plugins/Runtime/Metasound/Source/MetasoundFrontend/Private/MetasoundFrontendRegistries.cpp
@@ -292,8 +292,9 @@ bool FMetasoundFrontendRegistryContainer::RegisterDataType(const ::Metasound::FD
 
 bool FMetasoundFrontendRegistryContainer::RegisterExternalNode(FNodeGetterCallback&& InCallback)
 {
-	using FInputDataVertexCollection = Metasound::FInputDataVertexCollection;
-	using FOutputDataVertexCollection = Metasound::FOutputDataVertexCollection;
+	using FVertexInterface = Metasound::FVertexInterface;
+	using FInputVertexInterface = Metasound::FInputVertexInterface;
+	using FOutputVertexInterface = Metasound::FOutputVertexInterface;
 
 	Metasound::FNodeInitData DummyInitData;
 	TUniquePtr<Metasound::INode> DummyNodePtr = InCallback(DummyInitData);
@@ -306,8 +307,8 @@ bool FMetasoundFrontendRegistryContainer::RegisterExternalNode(FNodeGetterCallba
 
 	// First, build the key.
 	const FName NodeName = DummyNode.GetClassName();
-	const FInputDataVertexCollection& Inputs = DummyNode.GetInputDataVertices();
-	const FOutputDataVertexCollection& Outputs = DummyNode.GetOutputDataVertices();
+	const FInputVertexInterface& Inputs = DummyNode.GetDefaultVertexInterface().GetInputInterface();
+	const FOutputVertexInterface& Outputs = DummyNode.GetDefaultVertexInterface().GetOutputInterface();
 
 	// Construct a hash using a combination of the class name, input names and output names.
 	uint32 NodeHash = FCrc::StrCrc32(*NodeName.ToString());
@@ -317,14 +318,14 @@ bool FMetasoundFrontendRegistryContainer::RegisterExternalNode(FNodeGetterCallba
 
 	for (auto& InputTuple : Inputs)
 	{
-		NodeHash = HashCombine(NodeHash, FCrc::StrCrc32(*InputTuple.Value.VertexName));
-		InputTypes.Add(InputTuple.Value.DataReferenceTypeName);
+		NodeHash = HashCombine(NodeHash, FCrc::StrCrc32(*InputTuple.Value.GetVertexName()));
+		InputTypes.Add(InputTuple.Value.GetDataTypeName());
 	}
 
 	for (auto& OutputTuple : Outputs)
 	{
-		NodeHash = HashCombine(NodeHash, FCrc::StrCrc32(*OutputTuple.Value.VertexName));
-		OutputTypes.Add(OutputTuple.Value.DataReferenceTypeName);
+		NodeHash = HashCombine(NodeHash, FCrc::StrCrc32(*OutputTuple.Value.GetVertexName()));
+		OutputTypes.Add(OutputTuple.Value.GetDataTypeName());
 	}
 
 	FNodeRegistryKey InKey = { NodeName, NodeHash };
@@ -340,13 +341,13 @@ bool FMetasoundFrontendRegistryContainer::RegisterExternalNode(FNodeGetterCallba
 		UE_LOG(LogTemp, Display, TEXT("    %d inputs:"), Inputs.Num());
 		for (auto& InputTuple : Inputs)
 		{
-			UE_LOG(LogTemp, Display, TEXT("      %s (of type %s)"), *InputTuple.Value.VertexName, *InputTuple.Value.DataReferenceTypeName.ToString());
+			UE_LOG(LogTemp, Display, TEXT("      %s (of type %s)"), *InputTuple.Value.GetVertexName(), *InputTuple.Value.GetDataTypeName().ToString());
 		}
 
 		UE_LOG(LogTemp, Display, TEXT("    %d outputs:"), Outputs.Num());
 		for (auto& OutputTuple : Outputs)
 		{
-			UE_LOG(LogTemp, Display, TEXT("      %s (of type %s)"), *OutputTuple.Value.VertexName, *OutputTuple.Value.DataReferenceTypeName.ToString());
+			UE_LOG(LogTemp, Display, TEXT("      %s (of type %s)"), *OutputTuple.Value.GetVertexName(), *OutputTuple.Value.GetDataTypeName().ToString());
 		}
 	}
 
diff --git a/Engine/Plugins/Runtime/Metasound/Source/MetasoundFrontend/Public/MetasoundInputNode.h b/Engine/Plugins/Runtime/Metasound/Source/MetasoundFrontend/Public/MetasoundInputNode.h
--- a/Engine/Plugins/Runtime/Metasound/Source/MetasoundFrontend/Public/MetasoundInputNode.h
+++ b/Engine/Plugins/Runtime/Metasound/Source/MetasoundFrontend/Public/MetasoundInputNode.h
@@ -12,7 +12,7 @@
 namespace Metasound
 {
 	template<typename DataType>
-	class TInputNode : public INode
+	class TInputNode : public FNode
 	{
 		static_assert(TDataReferenceTypeInfo<DataType>::bIsValidSpecialization, "Please use DECLARE_METASOUND_DATA_REFERENCE_TYPES with this class before trying to create an input node with it.");
 
@@ -49,6 +49,8 @@ namespace Metasound
 				FDataReferenceCollection Inputs;
 		};
 
+		public:
+
 		// TODO: Would be nice if we didn't utilize the copy constructor by default here. But it's a nice catch-all
 		// for our use case.  Could store constructor parameters of the paramtype, but run the risk of holding rvalues
 		// or pointers to objects which are no longer valid. 
@@ -81,14 +83,14 @@ namespace Metasound
 				{
 				}
 
-				virtual TUniquePtr<IOperator> CreateOperator(const INode& InNode, const FOperatorSettings& InOperatorSettings, const FDataReferenceCollection& InInputDataReferences, TArray<TUniquePtr<IOperatorBuildError>>& OutErrors) override
+				virtual TUniquePtr<IOperator> CreateOperator(const FCreateOperatorParams& InParams, FBuildErrorArray& OutErrors) override
 				{
 					// Use copy constructor to create a new parameter reference.
 					FDataWriteReference DataRef = FDataWriteReference::CreateNew(Data);
 					// TODO: Write special version of this for audio types since they will need to be constructed based upon the operator settings. 
 					// TODO: Need to ponder how inputs are initialized, but template specialization might do the trick. 
 
-					const FInputNodeType& InputNode = static_cast<const FInputNodeType&>(InNode);
+					const FInputNodeType& InputNode = static_cast<const FInputNodeType&>(InParams.Node);
 
 					return MakeUnique<FInputOperator>(InputNode.GetVertexName(), DataRef);
 				}
@@ -96,8 +98,6 @@ namespace Metasound
 			private:
 				DataType Data;
 		};
-
-		public:
 			// Constructors that need to get working on linux:
 			/*
 			template<typename = typename TEnableIf< TIsConstructible<DataType>::Value >::Type>
@@ -119,99 +119,88 @@ namespace Metasound
 			}
 			*/
 
-			template<typename... ArgTypes>
-			TInputNode(const FString& InNodeDescription, const FString& InVertexName, ArgTypes&&... Args)
-			:	NodeDescription(InNodeDescription)
-			,	VertexName(InVertexName)
-			,	Factory(Forward<ArgTypes>(Args)...)
+			static const FNodeInfo& GetNodeInfo()
 			{
-				static_assert(TIsConstructible<DataType, ArgTypes...>::Value, "Tried to construct TInputNode<DataType> with arguments that don't match any constructor for DataType.");
-
-				FOutputDataVertex OutputVertex = MakeOutputDataVertex<DataType>(VertexName, FText::GetEmpty());
-				FDataVertexKey OutputVertexKey = MakeDataVertexKey(OutputVertex);
-
-				Outputs.Add(OutputVertexKey, OutputVertex);
-
-				FInputDataVertex InputVertex = MakeInputDataVertex<DataType>(VertexName, FText::GetEmpty());
-				FDataVertexKey InputVertexKey = MakeDataVertexKey(InputVertex);
+				static const FString ClassNameString = FString(TEXT("Input_")) + GetMetasoundDataTypeName<DataType>().ToString();
+				static const FName ClassName(*ClassNameString);
 
-				Inputs.Add(InputVertexKey, InputVertex);
+				static const FNodeInfo Info = {
+					ClassName,
+					LOCTEXT("Metasound_InputNodeDescription", "Input into the parent Metasound graph."),
+					PluginAuthor,
+					PluginNodeMissingPrompt
+				};
 
+				return Info;
 			}
 
-			explicit TInputNode(const FString& InNodeDescription, const FString& InVertexName, FDataTypeLiteralParam&& InParam, const FOperatorSettings& InSettings)
-				: NodeDescription(InNodeDescription)
-				, VertexName(InVertexName)
-				, Factory(MoveTemp(InParam), InSettings)
+			static FVertexInterface DeclareVertexInterface(const FString& InVertexName)
 			{
-				FOutputDataVertex OutputVertex = MakeOutputDataVertex<DataType>(VertexName, FText::GetEmpty());
-				FDataVertexKey OutputVertexKey = MakeDataVertexKey(OutputVertex);
-
-				Outputs.Add(OutputVertexKey, OutputVertex);
-
-				FInputDataVertex InputVertex = MakeInputDataVertex<DataType>(VertexName, FText::GetEmpty());
-				FDataVertexKey InputVertexKey = MakeDataVertexKey(InputVertex);
-
-				Inputs.Add(InputVertexKey, InputVertex);
+				return FVertexInterface(
+					FInputVertexInterface(
+						TInputDataVertexModel<DataType>(InVertexName, FText::GetEmpty())
+					),
+					FOutputVertexInterface(
+						TOutputDataVertexModel<DataType>(InVertexName, FText::GetEmpty())
+					)
+				);
 			}
 
-			virtual const FString& GetInstanceName() const override
+			template<typename... ArgTypes>
+			TInputNode(const FString& InNodeDescription, const FString& InVertexName, ArgTypes&&... Args)
+			:	FNode(InNodeDescription, GetNodeInfo())
+			,	VertexName(InVertexName)
+			,	Interface(DeclareVertexInterface(InVertexName))
+			,	Factory(MakeOperatorFactoryRef<FCopyOperatorFactory>(Forward<ArgTypes>(Args)...))
 			{
-				return NodeDescription;
+				static_assert(TIsConstructible<DataType, ArgTypes...>::Value, "Tried to construct TInputNode<DataType> with arguments that don't match any constructor for DataType.");
 			}
 
-			virtual const FName& GetClassName() const override
+			explicit TInputNode(const FString& InNodeDescription, const FString& InVertexName, FDataTypeLiteralParam&& InParam, const FOperatorSettings& InSettings)
+			:	FNode(InNodeDescription, GetNodeInfo())
+			,	VertexName(InVertexName)
+			,	Interface(DeclareVertexInterface(InVertexName))
+			, 	Factory(MakeOperatorFactoryRef<FCopyOperatorFactory>(MoveTemp(InParam), InSettings))
 			{
-				// TODO: although this is ok with MSVC's lax template instantiation, every other compiler will complain about TDataReferenceTypeInfo.
-				//static const FName ClassName = FName(FString(TEXT("Input_")) + FString(TDataReferenceTypeInfo<DataType>::TypeName));
-				static const FName ClassName("InputNode");
-				return ClassName;
 			}
 
-			virtual const FText& GetDescription() const override
-			{
-				static const FText Description = LOCTEXT("Metasound_InputNodeDescription", "Input into the parent Metasound graph.");
-				return Description;
-			}
 
-			virtual const FText& GetAuthorName() const override
+			const FString& GetVertexName() const
 			{
-				return PluginAuthor;
+				return VertexName;
 			}
 
-			virtual const FText& GetPromptIfMissing() const override
+			virtual const FVertexInterface& GetVertexInterface() const override
 			{
-				return PluginNodeMissingPrompt;
+				return Interface;
 			}
 
-			const FString& GetVertexName() const
+			virtual const FVertexInterface& GetDefaultVertexInterface() const override
 			{
-				return VertexName;
+				return Interface;
 			}
 
-			virtual const FInputDataVertexCollection& GetInputDataVertices() const override
+			virtual bool SetVertexInterface(const FVertexInterface& InInterface) override
 			{
-				return Inputs;
+				return Interface == InInterface;
 			}
 
-			virtual const FOutputDataVertexCollection& GetOutputDataVertices() const override
+			virtual bool IsVertexInterfaceSupported(const FVertexInterface& InInterface) const override
 			{
-				return Outputs;
+				return Interface == InInterface;
 			}
 
-			virtual IOperatorFactory& GetDefaultOperatorFactory() override
+			virtual TSharedRef<IOperatorFactory, ESPMode::ThreadSafe> GetDefaultOperatorFactory() const override
 			{
 				return Factory;
 			}
 
 		private:
-			FString NodeDescription;
 			FString VertexName;
 
-			FCopyOperatorFactory Factory;
+			FVertexInterface Interface;
+			FOperatorFactorySharedRef Factory;
 
-			FInputDataVertexCollection Inputs;
-			FOutputDataVertexCollection Outputs;
 	};
 } // namespace Metasound
 #undef LOCTEXT_NAMESPACE //MetasoundOutputNode
diff --git a/Engine/Plugins/Runtime/Metasound/Source/MetasoundFrontend/Public/MetasoundOutputNode.h b/Engine/Plugins/Runtime/Metasound/Source/MetasoundFrontend/Public/MetasoundOutputNode.h
--- a/Engine/Plugins/Runtime/Metasound/Source/MetasoundFrontend/Public/MetasoundOutputNode.h
+++ b/Engine/Plugins/Runtime/Metasound/Source/MetasoundFrontend/Public/MetasoundOutputNode.h
@@ -5,14 +5,14 @@
 #include "MetasoundBuilderInterface.h"
 #include "MetasoundNodeInterface.h"
 #include "MetasoundOperatorInterface.h"
+#include "MetasoundNode.h"
 
 #define LOCTEXT_NAMESPACE "MetasoundGraphCore"
 
-
 namespace Metasound
 {
 	template<typename DataType>
-	class TOutputNode : public INode
+	class TOutputNode : public FNode
 	{
 		static_assert(TDataReferenceTypeInfo<DataType>::bIsValidSpecialization, "Please use DECLARE_METASOUND_DATA_REFERENCE_TYPES with this class before trying to create an output node with it.");
 			
@@ -56,92 +56,99 @@ namespace Metasound
 				{
 				}
 
-				virtual TUniquePtr<IOperator> CreateOperator(const INode& InNode, const FOperatorSettings& InOperatorSettings, const FDataReferenceCollection& InInputDataReferences, TArray<TUniquePtr<IOperatorBuildError>>& OutErrors) override
+				virtual TUniquePtr<IOperator> CreateOperator(const FCreateOperatorParams& InParams, FBuildErrorArray& OutErrors) override
 				{
-					if (!InInputDataReferences.ContainsDataReadReference<DataType>(DataReferenceName))
+					if (!InParams.InputDataReferences.ContainsDataReadReference<DataType>(DataReferenceName))
 					{
 						// TODO: Add build error.
 						return TUniquePtr<IOperator>(nullptr);
 					}
 
-					return MakeUnique<FOutputOperator>(DataReferenceName, InInputDataReferences.GetDataReadReference<DataType>(DataReferenceName));
+					return MakeUnique<FOutputOperator>(DataReferenceName, InParams.InputDataReferences.GetDataReadReference<DataType>(DataReferenceName));
 				}
 
 			private:
 				FString DataReferenceName;
 		};
 
-		public:
-			TOutputNode(const FString& InNodeDescription, const FString& InVertexName)
-			:	NodeDescription(InNodeDescription)
-			,	VertexName(InVertexName)
-			,	Factory(InVertexName)
-			{
-				FInputDataVertex InputVertex = MakeInputDataVertex<DataType>(VertexName, FText::GetEmpty());
-				FDataVertexKey InputKey = MakeDataVertexKey(InputVertex);
-
-				Inputs.Add(InputKey, InputVertex);
-
-				FOutputDataVertex OutputVertex = MakeOutputDataVertex<DataType>(VertexName, FText::GetEmpty());
-				FDataVertexKey OutputKey = MakeDataVertexKey(OutputVertex);
+		static const FNodeInfo& GetNodeInfo()
+		{
+			static const FString ClassNameString = FString(TEXT("Output_")) + GetMetasoundDataTypeName<DataType>().ToString();
 
-				Outputs.Add(OutputKey, OutputVertex);
-			}
+			static const FNodeInfo NodeInfo = {
+				FName(*ClassNameString),
+				LOCTEXT("Metasound_OutputNodeDescription", "Output from the parent Metasound graph."),
+				PluginAuthor,
+				PluginNodeMissingPrompt
+			};
 
-			virtual const FName& GetClassName() const override
-			{
-				// TODO: Any special formatting for these node type names?
-				// TODO: although this is ok with MSVC's lax template instantiation, every other compiler will complain about TDataReferenceTypeInfo.
-				//static const FName ClassName = FName(FString(TEXT("Input_")) + FString(TDataReferenceTypeInfo<DataType>::TypeName));
-				static const FName ClassName("OutputNode");
-				return ClassName;
-			}
+			return NodeInfo;
+		}
 
-			virtual const FString& GetInstanceName() const override
-			{
-				return NodeDescription;
-			}
+		static FVertexInterface GetVertexInterface(const FString& InVertexName)
+		{
+			return FVertexInterface(
+				FInputVertexInterface(
+					TInputDataVertexModel<DataType>(InVertexName, LOCTEXT("Metasound_OutputVertexDescription", "Output from the parent Metasound graph."))
+				),
+				FOutputVertexInterface(
+					TOutputDataVertexModel<DataType>(InVertexName, LOCTEXT("Metasound_OutputVertexDescription", "Output from the parent Metasound graph."))
+				)
+			);
+		}
 
-			virtual const FText& GetDescription() const override
+		public:
+			TOutputNode(const FString& InInstanceName, const FString& InVertexName)
+			:	FNode(InInstanceName, GetNodeInfo())
+			,	VertexInterface(GetVertexInterface(InVertexName))
+			,	Factory(MakeShared<FOutputOperatorFactory, ESPMode::ThreadSafe>(InVertexName))
 			{
-				static const FText Description = LOCTEXT("Metasound_InputNodeDescription", "Output from the parent Metasound graph.");
-				return Description;
 			}
 
-			virtual const FText& GetAuthorName() const override
+			/** Return the current vertex interface. */
+			virtual const FVertexInterface& GetVertexInterface() const override
 			{
-				return PluginAuthor;
+				return VertexInterface;
 			}
 
-			virtual const FText& GetPromptIfMissing() const override
+			/** Return the default vertex interface. */
+			virtual const FVertexInterface& GetDefaultVertexInterface() const override
 			{
-				return PluginNodeMissingPrompt;
+				return VertexInterface;
 			}
 
-			virtual const FInputDataVertexCollection& GetInputDataVertices() const override
+			/** Set the vertex interface. If the vertex was successfully changed, returns true. 
+			 *
+			 * @param InInterface - New interface for node. 
+			 *
+			 * @return True on success, false otherwise.
+			 */
+			virtual bool SetVertexInterface(const FVertexInterface& InInterface) override
 			{
-				return Inputs;
+				return VertexInterface == InInterface;
 			}
 
-			virtual const FOutputDataVertexCollection& GetOutputDataVertices() const override
+			/** Expresses whether a specific vertex interface is supported.
+			 *
+			 * @param InInterface - New interface. 
+			 *
+			 * @return True if the interface is supported, false otherwise. 
+			 */
+			virtual bool IsVertexInterfaceSupported(const FVertexInterface& InInterface) const 
 			{
-				return Outputs;
+				return VertexInterface == InInterface;
 			}
 
-			virtual IOperatorFactory& GetDefaultOperatorFactory() override
+			virtual TSharedRef<IOperatorFactory, ESPMode::ThreadSafe> GetDefaultOperatorFactory() const override
 			{
 				return Factory;
 			}
 
 		private:
-			FString NodeDescription;
-			FString VertexName;
-
-			FOutputOperatorFactory Factory;
+			FVertexInterface VertexInterface;
 
+			TSharedRef<FOutputOperatorFactory, ESPMode::ThreadSafe> Factory;
 
-			FInputDataVertexCollection Inputs;
-			FOutputDataVertexCollection Outputs;
 	};
 }
 #undef LOCTEXT_NAMESPACE // MetasoundOutputNode
diff --git a/Engine/Plugins/Runtime/Metasound/Source/MetasoundGraphCore/Private/MetasoundDataReference.cpp b/Engine/Plugins/Runtime/Metasound/Source/MetasoundGraphCore/Private/MetasoundDataReference.cpp
new file mode 100644
--- /dev/null
+++ b/Engine/Plugins/Runtime/Metasound/Source/MetasoundGraphCore/Private/MetasoundDataReference.cpp
@@ -0,0 +1,8 @@
+// Copyright Epic Games, Inc. All Rights Reserved.
+
+#include "MetasoundDataReference.h"
+
+namespace Metasound
+{
+	const TCHAR* TDataReferenceTypeInfo<void>::TypeName = TEXT("void");
+}
diff --git a/Engine/Plugins/Runtime/Metasound/Source/MetasoundGraphCore/Private/MetasoundFacade.cpp b/Engine/Plugins/Runtime/Metasound/Source/MetasoundGraphCore/Private/MetasoundFacade.cpp
new file mode 100644
--- /dev/null
+++ b/Engine/Plugins/Runtime/Metasound/Source/MetasoundGraphCore/Private/MetasoundFacade.cpp
@@ -0,0 +1,42 @@
+// Copyright Epic Games, Inc. All Rights Reserved.
+
+#include "MetasoundFacade.h"
+
+namespace Metasound
+{
+	FNodeFacade::FFactory::FFactory(FCreateOperatorFunction InCreateFunc)
+	:	CreateFunc(InCreateFunc)
+	{
+	}
+
+	TUniquePtr<IOperator> FNodeFacade::FFactory::CreateOperator(const FCreateOperatorParams& InParams, FBuildErrorArray& OutErrors)
+	{
+		return CreateFunc(InParams, OutErrors);
+	}
+
+	const FVertexInterface& FNodeFacade::GetVertexInterface() const
+	{
+		return VertexInterface;
+	}
+
+	const FVertexInterface& FNodeFacade::GetDefaultVertexInterface() const
+	{
+		return VertexInterface;
+	}
+
+	bool FNodeFacade::SetVertexInterface(const FVertexInterface& InInterface)
+	{
+		return InInterface == VertexInterface;
+	}
+
+	bool FNodeFacade::IsVertexInterfaceSupported(const FVertexInterface& InInterface) const
+	{
+		return InInterface == VertexInterface;
+	}
+
+	/** Return a reference to the default operator factory. */
+	FOperatorFactorySharedRef FNodeFacade::GetDefaultOperatorFactory() const
+	{
+		return Factory;
+	}
+}
diff --git a/Engine/Plugins/Runtime/Metasound/Source/MetasoundGraphCore/Private/MetasoundGraph.cpp b/Engine/Plugins/Runtime/Metasound/Source/MetasoundGraphCore/Private/MetasoundGraph.cpp
--- a/Engine/Plugins/Runtime/Metasound/Source/MetasoundGraphCore/Private/MetasoundGraph.cpp
+++ b/Engine/Plugins/Runtime/Metasound/Source/MetasoundGraphCore/Private/MetasoundGraph.cpp
@@ -37,33 +37,23 @@ namespace Metasound
 		return PluginNodeMissingPrompt;
 	}
 
-	const FInputDataVertexCollection& FGraph::GetInputDataVertices() const
+	bool FGraph::AddInputDataDestination(const INode& InNode, const FDataVertexKey& InVertexKey)
 	{
-		return InputVertices;
-	}
-
-	const FOutputDataVertexCollection& FGraph::GetOutputDataVertices() const
-	{
-		return OutputVertices;
-	}
-
-	bool FGraph::AddInputDataDestination(INode& InNode, const FString& InVertexName)
-	{
-		if (!InNode.GetInputDataVertices().Contains(InVertexName))
+		if (!InNode.GetVertexInterface().ContainsInputVertex(InVertexKey))
 		{
 			return false;
 		}
 
-		FInputDataDestination Destination(InNode, InNode.GetInputDataVertices()[InVertexName]);
+		FInputDataDestination Destination(InNode, InNode.GetVertexInterface().GetInputVertex(InVertexKey));
 
 		AddInputDataDestination(Destination);
-
+		
 		return true;
 	}
 
 	void FGraph::AddInputDataDestination(const FInputDataDestination& InDestination)
 	{
-		InputVertices.Add(MakeDataVertexKey(InDestination.Vertex), InDestination.Vertex);
+		VertexInterface.GetInputInterface().Add(InDestination.Vertex);
 		InputDestinations.Add(MakeDestinationDataVertexKey(InDestination), InDestination);
 	}
 
@@ -72,14 +62,14 @@ namespace Metasound
 		return InputDestinations;
 	}
 
-	bool FGraph::AddOutputDataSource(INode& InNode, const FString& InVertexName)
+	bool FGraph::AddOutputDataSource(const INode& InNode, const FDataVertexKey& InVertexKey)
 	{
-		if (!InNode.GetOutputDataVertices().Contains(InVertexName))
+		if (!InNode.GetVertexInterface().ContainsOutputVertex(InVertexKey))
 		{
 			return false;
 		}
 
-		FOutputDataSource Source(InNode, InNode.GetOutputDataVertices()[InVertexName]);
+		FOutputDataSource Source(InNode, InNode.GetVertexInterface().GetOutputVertex(InVertexKey));
 
 		AddOutputDataSource(Source);
 
@@ -88,7 +78,7 @@ namespace Metasound
 
 	void FGraph::AddOutputDataSource(const FOutputDataSource& InSource)
 	{
-		OutputVertices.Add(MakeDataVertexKey(InSource.Vertex), InSource.Vertex);
+		VertexInterface.GetOutputInterface().Add(InSource.Vertex);
 		OutputSources.Add(MakeSourceDataVertexKey(InSource), InSource);
 	}
 
@@ -102,24 +92,26 @@ namespace Metasound
 		Edges.Add(InEdge);
 	}
 
-	// TODO: can we make node references const?  Would need to make operator factory accessing const as well.
-	bool FGraph::AddDataEdge(INode& FromNode, const FString& FromName, INode& ToNode, const FString ToName)
+	bool FGraph::AddDataEdge(const INode& FromNode, const FDataVertexKey& FromKey, const INode& ToNode, const FDataVertexKey& ToKey)
 	{
-		
-		if (!FromNode.GetOutputDataVertices().Contains(FromName))
+		const FVertexInterface& FromVertexInterface = FromNode.GetVertexInterface();
+		const FVertexInterface& ToVertexInterface = ToNode.GetVertexInterface();
+
+		if (!FromVertexInterface.ContainsOutputVertex(FromKey))
 		{
 			return false;
 		}
 
-		if (!ToNode.GetInputDataVertices().Contains(ToName))
+		if (!ToVertexInterface.ContainsInputVertex(ToKey))
 		{
 			return false;
 		}
 
-		const FOutputDataVertex& FromVertex = FromNode.GetOutputDataVertices()[FromName];
-		const FInputDataVertex& ToVertex = ToNode.GetInputDataVertices()[ToName];
+		const FOutputDataVertex& FromVertex = FromVertexInterface.GetOutputVertex(FromKey);
+		const FInputDataVertex& ToVertex = ToVertexInterface.GetInputVertex(ToKey);
+
 
-		if (FromVertex.DataReferenceTypeName != ToVertex.DataReferenceTypeName)
+		if (FromVertex.GetDataTypeName() != ToVertex.GetDataTypeName())
 		{
 			return false;
 		}
@@ -135,4 +127,26 @@ namespace Metasound
 	{
 		return Edges;
 	}
+
+	const FVertexInterface& FGraph::GetVertexInterface() const 
+	{
+		return VertexInterface;
+	}
+
+	const FVertexInterface& FGraph::GetDefaultVertexInterface() const 
+	{
+		static const FVertexInterface EmptyInterface;
+
+		return EmptyInterface;
+	}
+
+	bool FGraph::SetVertexInterface(const FVertexInterface& InInterface)
+	{
+		return InInterface == VertexInterface;
+	}
+
+	bool FGraph::IsVertexInterfaceSupported(const FVertexInterface& InInterface) const
+	{
+		return InInterface == VertexInterface;
+	}
 }
diff --git a/Engine/Plugins/Runtime/Metasound/Source/MetasoundGraphCore/Private/MetasoundGraphLinter.cpp b/Engine/Plugins/Runtime/Metasound/Source/MetasoundGraphCore/Private/MetasoundGraphLinter.cpp
--- a/Engine/Plugins/Runtime/Metasound/Source/MetasoundGraphCore/Private/MetasoundGraphLinter.cpp
+++ b/Engine/Plugins/Runtime/Metasound/Source/MetasoundGraphCore/Private/MetasoundGraphLinter.cpp
@@ -16,12 +16,6 @@ namespace Metasound
 	{
 		using FBuildErrorPtr = TUniquePtr<IOperatorBuildError>;
 
-		template<typename ErrorType, typename... ArgTypes>
-		void AddBuildError(TArray<FBuildErrorPtr>& OutErrors, ArgTypes&&... Args)
-		{
-			OutErrors.Add(MakeUnique<ErrorType>(Forward<ArgTypes>(Args)...));
-		}
-
 		// run validation on input destination
 		bool IsInputDestinationValid(const FInputDataDestination& InDest)
 		{
@@ -30,16 +24,17 @@ namespace Metasound
 				return false;
 			}
 
-			const FInputDataVertexCollection& InputCollection = InDest.Node->GetInputDataVertices();
+
+			const FInputVertexInterface& InputInterface = InDest.Node->GetVertexInterface().GetInputInterface();
 
 			FDataVertexKey VertexKey = MakeDataVertexKey(InDest.Vertex);
 
-			if (!InputCollection.Contains(VertexKey))
+			if (!InputInterface.Contains(VertexKey))
 			{
 				return false;
 			}
 
-			if (InputCollection[VertexKey] != InDest.Vertex)
+			if (InputInterface[VertexKey] != InDest.Vertex)
 			{
 				return false;
 			}
@@ -55,16 +50,16 @@ namespace Metasound
 				return false;
 			}
 
-			const FOutputDataVertexCollection& OutputCollection = InSource.Node->GetOutputDataVertices();
+			const FOutputVertexInterface& OutputInterface = InSource.Node->GetVertexInterface().GetOutputInterface();
 
 			FDataVertexKey VertexKey = MakeDataVertexKey(InSource.Vertex);
 
-			if (!OutputCollection.Contains(VertexKey))
+			if (!OutputInterface.Contains(VertexKey))
 			{
 				return false;
 			}
 
-			if (OutputCollection[VertexKey] != InSource.Vertex)
+			if (OutputInterface[VertexKey] != InSource.Vertex)
 			{
 				return false;
 			}
@@ -92,7 +87,7 @@ namespace Metasound
 
 		for (const FDataEdge& Edge : InGraph.GetDataEdges())
 		{
-			if (Edge.From.Vertex.DataReferenceTypeName != Edge.To.Vertex.DataReferenceTypeName)
+			if (Edge.From.Vertex.GetDataTypeName() != Edge.To.Vertex.GetDataTypeName())
 			{
 				AddBuildError<FInvalidConnectionDataTypeError>(OutErrors, Edge);
 
diff --git a/Engine/Plugins/Runtime/Metasound/Source/MetasoundGraphCore/Private/MetasoundNode.cpp b/Engine/Plugins/Runtime/Metasound/Source/MetasoundGraphCore/Private/MetasoundNode.cpp
--- a/Engine/Plugins/Runtime/Metasound/Source/MetasoundGraphCore/Private/MetasoundNode.cpp
+++ b/Engine/Plugins/Runtime/Metasound/Source/MetasoundGraphCore/Private/MetasoundNode.cpp
@@ -6,47 +6,42 @@
 
 namespace Metasound
 {
-	FNode::FNode(const FString& InDescription)
-	:	Description(InDescription)
-	{
-	}
-
-	FNode::~FNode()
+	FNode::FNode(const FString& InInstanceName, const FNodeInfo& InInfo)
+	:	InstanceName(InInstanceName)
+	,	Info(InInfo)
 	{
 	}
 
+	/** Return the name of this specific instance of the node class. */
 	const FString& FNode::GetInstanceName() const
 	{
-		return Description;
-	}
-
-	const FInputDataVertexCollection& FNode::GetInputDataVertices() const
-	{
-		return Inputs;
-	}
-
-	const FOutputDataVertexCollection& FNode::GetOutputDataVertices() const
-	{
-		return Outputs;
+		return InstanceName;
 	}
 
-	void FNode::AddInputDataVertex(const FInputDataVertex& InVertex)
+	/** Return the type name of this node. */
+	const FName& FNode::GetClassName() const
 	{
-		Inputs.Add(MakeDataVertexKey(InVertex), InVertex);
+		return Info.ClassName;
 	}
 
-	void FNode::RemoveInputDataVertex(const FInputDataVertex& InVertex)
+	/** Return a longer text description describing how this node is used. */
+	const FText& FNode::GetDescription() const
 	{
-		Inputs.Remove(MakeDataVertexKey(InVertex));
+		return Info.Description;
 	}
 
-	void FNode::AddOutputDataVertex(const FOutputDataVertex& InVertex)
+	/** Return the original author of this node class. */
+	const FText& FNode::GetAuthorName() const
 	{
-		Outputs.Add(MakeDataVertexKey(InVertex), InVertex);
+		return Info.AuthorName;
 	}
 
-	void FNode::RemoveOutputDataVertex(const FOutputDataVertex& InVertex)
+	/** 
+	 *  Return an optional prompt on how users can get the plugin this node is in,
+	 *  if they have found a metasound that uses this node but don't have this plugin downloaded or enabled.
+	 */
+	const FText& FNode::GetPromptIfMissing() const
 	{
-		Outputs.Remove(MakeDataVertexKey(InVertex));
+		return Info.PromptIfMissing;
 	}
 }
diff --git a/Engine/Plugins/Runtime/Metasound/Source/MetasoundGraphCore/Private/MetasoundNodeInterface.cpp b/Engine/Plugins/Runtime/Metasound/Source/MetasoundGraphCore/Private/MetasoundNodeInterface.cpp
--- a/Engine/Plugins/Runtime/Metasound/Source/MetasoundGraphCore/Private/MetasoundNodeInterface.cpp
+++ b/Engine/Plugins/Runtime/Metasound/Source/MetasoundGraphCore/Private/MetasoundNodeInterface.cpp
@@ -4,21 +4,6 @@
 
 namespace Metasound
 {
-	bool operator==(const FDataVertex& LHS, const FDataVertex& RHS)
-	{
-		bool bIsEqual = /*(LHS.Source == RHS.Source) && */
-			(LHS.VertexName == RHS.VertexName) &&
-			(LHS.DataReferenceTypeName == RHS.DataReferenceTypeName) &&
-			LHS.Description.EqualTo(RHS.Description);
-
-		return bIsEqual;
-	}
-
-	bool operator!=(const FDataVertex& LHS, const FDataVertex& RHS)
-	{
-		return !(LHS == RHS);
-	}
-
 	bool operator==(const FOutputDataSource& InLeft, const FOutputDataSource& InRight)
 	{
 		return (InLeft.Node == InRight.Node) && (InLeft.Vertex == InRight.Vertex);
diff --git a/Engine/Plugins/Runtime/Metasound/Source/MetasoundGraphCore/Private/MetasoundOperatorBuilder.cpp b/Engine/Plugins/Runtime/Metasound/Source/MetasoundGraphCore/Private/MetasoundOperatorBuilder.cpp
--- a/Engine/Plugins/Runtime/Metasound/Source/MetasoundGraphCore/Private/MetasoundOperatorBuilder.cpp
+++ b/Engine/Plugins/Runtime/Metasound/Source/MetasoundGraphCore/Private/MetasoundOperatorBuilder.cpp
@@ -23,13 +23,6 @@ namespace Metasound
 	{
 		using FBuildErrorPtr = TUniquePtr<IOperatorBuildError>;
 
-		// Convenience template for adding build errors
-		template<typename ErrorType, typename... ArgTypes>
-		void AddBuildError(TArray<FBuildErrorPtr>& OutErrors, ArgTypes&&... Args)
-		{
-			OutErrors.Add(MakeUnique<ErrorType>(Forward<ArgTypes>(Args)...));
-		}
-
 		// Convenience function for adding graph cycle build errors
 		void AddBuildErrorsForCycles(const FDirectedGraphAlgoAdapter& InAdapter, TArray<FBuildErrorPtr>& OutErrors)
 		{
@@ -338,10 +331,9 @@ namespace Metasound
 				return FBuildStatus::NonFatalError;
 			}
 
-			
 			const FDataReferenceCollection& FromDataReferenceCollection = InDataReferenceMap[Edge->From.Node].Outputs;
 
-			if (!FromDataReferenceCollection.ContainsDataReadReference(Edge->From.Vertex.VertexName, Edge->From.Vertex.DataReferenceTypeName))
+			if (!FromDataReferenceCollection.ContainsDataReadReference(Edge->From.Vertex.GetVertexName(), Edge->From.Vertex.GetDataTypeName()))
 			{
 				// This is likely a node programming error where the edges reported by the INode interface
 				// did not match the readable parameter refs created by the operators outputs. Or, the edge description is invalid.
@@ -354,7 +346,7 @@ namespace Metasound
 			}
 
 
-			bool bSuccess = OutCollection.AddDataReadReferenceFrom(Edge->To.Vertex.VertexName, FromDataReferenceCollection, Edge->From.Vertex.VertexName, Edge->From.Vertex.DataReferenceTypeName);
+			bool bSuccess = OutCollection.AddDataReadReferenceFrom(Edge->To.Vertex.GetVertexName(), FromDataReferenceCollection, Edge->From.Vertex.GetVertexName(), Edge->From.Vertex.GetDataTypeName());
 
 			if (!bSuccess)
 			{
@@ -393,9 +385,11 @@ namespace Metasound
 			}
 
 			// TODO: make factory retrieval const 
-			IOperatorFactory& Factory = const_cast<INode*>(Node)->GetDefaultOperatorFactory();
+			FOperatorFactorySharedRef Factory = Node->GetDefaultOperatorFactory();
 
-			FOperatorPtr Operator = Factory.CreateOperator(*Node, OperatorSettings, InputCollection, OutErrors);
+			FCreateOperatorParams CreateParams(*Node, OperatorSettings, InputCollection);
+
+			FOperatorPtr Operator = Factory->CreateOperator(CreateParams, OutErrors);
 
 			if (!Operator.IsValid())
 			{
@@ -436,7 +430,8 @@ namespace Metasound
 
 			FDataReferenceCollection& Collection = InNodeDataReferences[InputDestination.Node].Inputs;
 
-			if (!Collection.ContainsDataWriteReference(InputDestination.Vertex.VertexName, InputDestination.Vertex.DataReferenceTypeName))
+
+			if (!Collection.ContainsDataWriteReference(InputDestination.Vertex.GetVertexName(), InputDestination.Vertex.GetDataTypeName()))
 			{
 				AddBuildError<FMissingInputDataReferenceError>(OutErrors, InputDestination);
 
@@ -444,7 +439,7 @@ namespace Metasound
 				continue;
 			}
 
-			bool bSuccess = OutGraphInputs.AddDataWriteReferenceFrom(InputDestination.Vertex.VertexName, Collection, InputDestination.Vertex.VertexName, InputDestination.Vertex.DataReferenceTypeName);
+			bool bSuccess = OutGraphInputs.AddDataWriteReferenceFrom(InputDestination.Vertex.GetVertexName(), Collection, InputDestination.Vertex.GetVertexName(), InputDestination.Vertex.GetDataTypeName());
 
 			if (!bSuccess)
 			{
@@ -472,7 +467,7 @@ namespace Metasound
 
 			FDataReferenceCollection& Collection = InNodeDataReferences[OutputSource.Node].Outputs;
 
-			if (!Collection.ContainsDataReadReference(OutputSource.Vertex.VertexName, OutputSource.Vertex.DataReferenceTypeName))
+			if (!Collection.ContainsDataReadReference(OutputSource.Vertex.GetVertexName(), OutputSource.Vertex.GetDataTypeName()))
 			{
 				// This will likely produce an IOperator which does not work as
 				// expected.
@@ -483,7 +478,7 @@ namespace Metasound
 				continue;
 			}
 
-			bool bSuccess = OutGraphOutputs.AddDataReadReferenceFrom(OutputSource.Vertex.VertexName, Collection, OutputSource.Vertex.VertexName, OutputSource.Vertex.DataReferenceTypeName);
+			bool bSuccess = OutGraphOutputs.AddDataReadReferenceFrom(OutputSource.Vertex.GetVertexName(), Collection, OutputSource.Vertex.GetVertexName(), OutputSource.Vertex.GetDataTypeName());
 
 			if (!bSuccess)
 			{
diff --git a/Engine/Plugins/Runtime/Metasound/Source/MetasoundGraphCore/Private/MetasoundVertex.cpp b/Engine/Plugins/Runtime/Metasound/Source/MetasoundGraphCore/Private/MetasoundVertex.cpp
new file mode 100644
--- /dev/null
+++ b/Engine/Plugins/Runtime/Metasound/Source/MetasoundGraphCore/Private/MetasoundVertex.cpp
@@ -0,0 +1,210 @@
+// Copyright Epic Games, Inc. All Rights Reserved.
+
+#include "MetasoundVertex.h"
+
+#include "CoreMinimal.h"
+
+namespace Metasound
+{
+	FInputDataVertex::FInputDataVertex()
+	:	VertexModel(MakeUnique<FEmptyVertexModel>(TEXT(""), FText::GetEmpty()))
+	{}
+
+	FInputDataVertex::FInputDataVertex(const FInputDataVertex& InOther)
+	:	VertexModel()
+	{
+		// Copy constructor has to call underlying model's clone.
+		if (InOther.VertexModel.IsValid())
+		{
+			VertexModel = InOther.VertexModel->Clone();
+		}
+
+		if (!VertexModel.IsValid())
+		{
+			// Default to empty vertex if there was a failure. 
+			VertexModel = MakeUnique<FEmptyVertexModel>(TEXT(""), FText::GetEmpty());
+		}
+	}
+
+	FInputDataVertex& FInputDataVertex::operator=(const FInputDataVertex& InOther)
+	{
+		// Assignment operator has to call underlying model's clone.
+		if (InOther.VertexModel.IsValid())
+		{
+			VertexModel = InOther.VertexModel->Clone();
+		}
+
+		if (!VertexModel.IsValid())
+		{
+			// Default to empty vertex if there was a failure. 
+			VertexModel = MakeUnique<FEmptyVertexModel>(TEXT(""), FText::GetEmpty());
+		}
+
+		return *this;
+	}
+
+	const FString& FInputDataVertex::GetVertexName() const
+	{
+		return VertexModel->VertexName;
+	}
+
+	const FName& FInputDataVertex::GetDataTypeName() const
+	{
+		return VertexModel->DataTypeName;
+	}
+
+	const FText& FInputDataVertex::GetDescription() const
+	{
+		return VertexModel->Description;
+	}
+
+	bool FInputDataVertex::IsReferenceOfSameType(const IDataReference& InReference) const 
+	{
+		return VertexModel->IsReferenceOfSameType(InReference);
+	}
+
+	const FName& FInputDataVertex::GetVertexTypeName() const 
+	{
+		return VertexModel->GetVertexTypeName();
+	}
+
+	bool operator==(const FInputDataVertex& InLHS, const FInputDataVertex& InRHS)
+	{
+		return (InLHS.VertexModel->IsEqual(*InRHS.VertexModel) && (InRHS.VertexModel->IsEqual(*InLHS.VertexModel)));
+	}
+
+	bool operator!=(const FInputDataVertex& InLHS, const FInputDataVertex& InRHS)
+	{
+		return !(InLHS == InRHS);
+	}
+
+	FOutputDataVertex::FOutputDataVertex()
+	:	VertexModel(MakeUnique<FEmptyVertexModel>(TEXT(""), FText::GetEmpty()))
+	{}
+
+	FOutputDataVertex::FOutputDataVertex(const FOutputDataVertex& InOther)
+	:	VertexModel()
+	{
+		// Copy constructor has to call underlying model's clone.
+		if (InOther.VertexModel.IsValid())
+		{
+			VertexModel = InOther.VertexModel->Clone();
+		}
+
+		if (!VertexModel.IsValid())
+		{
+			// Default to empty vertex if there was a failure. 
+			VertexModel = MakeUnique<FEmptyVertexModel>(TEXT(""), FText::GetEmpty());
+		}
+	}
+
+	FOutputDataVertex& FOutputDataVertex::operator=(const FOutputDataVertex& InOther)
+	{
+		// assignment operator has to call underlying model's clone.
+		if (InOther.VertexModel.IsValid())
+		{
+			VertexModel = InOther.VertexModel->Clone();
+		}
+
+		if (!VertexModel.IsValid())
+		{
+			// Default to empty vertex if there was a failure. 
+			VertexModel = MakeUnique<FEmptyVertexModel>(TEXT(""), FText::GetEmpty());
+		}
+
+		return *this;
+	}
+
+	const FString& FOutputDataVertex::GetVertexName() const
+	{
+		return VertexModel->VertexName;
+	}
+
+	const FName& FOutputDataVertex::GetDataTypeName() const
+	{
+		return VertexModel->DataTypeName;
+	}
+
+	const FText& FOutputDataVertex::GetDescription() const
+	{
+		return VertexModel->Description;
+	}
+
+	bool FOutputDataVertex::IsReferenceOfSameType(const IDataReference& InReference) const 
+	{
+		return VertexModel->IsReferenceOfSameType(InReference);
+	}
+
+	const FName& FOutputDataVertex::GetVertexTypeName() const 
+	{
+		return VertexModel->GetVertexTypeName();
+	}
+
+
+	bool operator==(const FOutputDataVertex& InLHS, const FOutputDataVertex& InRHS)
+	{
+		return (InLHS.VertexModel->IsEqual(*InRHS.VertexModel) && (InRHS.VertexModel->IsEqual(*InLHS.VertexModel)));
+	}
+
+	bool operator!=(const FOutputDataVertex& InLHS, const FOutputDataVertex& InRHS)
+	{
+		return !(InLHS == InRHS);
+	}
+
+			/** Construct with an input and output interface. */
+	FVertexInterface::FVertexInterface(const FInputVertexInterface& InInputs, const FOutputVertexInterface& InOutputs)
+	:	InputInterface(InInputs)
+	,	OutputInterface(InOutputs)
+	{
+	}
+
+	const FInputVertexInterface& FVertexInterface::GetInputInterface() const
+	{
+		return InputInterface;
+	}
+
+	FInputVertexInterface& FVertexInterface::GetInputInterface()
+	{
+		return InputInterface;
+	}
+
+	const FInputDataVertex& FVertexInterface::GetInputVertex(const FDataVertexKey& InKey) const
+	{
+		return InputInterface[InKey];
+	}
+
+	bool FVertexInterface::ContainsInputVertex(const FDataVertexKey& InKey) const
+	{
+		return InputInterface.Contains(InKey);
+	}
+
+	const FOutputVertexInterface& FVertexInterface::GetOutputInterface() const
+	{
+		return OutputInterface;
+	}
+
+	FOutputVertexInterface& FVertexInterface::GetOutputInterface()
+	{
+		return OutputInterface;
+	}
+
+	const FOutputDataVertex& FVertexInterface::GetOutputVertex(const FDataVertexKey& InKey) const
+	{
+		return OutputInterface[InKey];
+	}
+
+	bool FVertexInterface::ContainsOutputVertex(const FDataVertexKey& InKey) const
+	{
+		return OutputInterface.Contains(InKey);
+	}
+
+	bool operator==(const FVertexInterface& InLHS, const FVertexInterface& InRHS)
+	{
+		return (InLHS.InputInterface == InRHS.InputInterface) && (InLHS.OutputInterface == InRHS.OutputInterface);
+	}
+
+	bool operator!=(const FVertexInterface& InLHS, const FVertexInterface& InRHS)
+	{
+		return !(InLHS == InRHS);
+	}
+}
diff --git a/Engine/Plugins/Runtime/Metasound/Source/MetasoundGraphCore/Public/MetasoundBuilderInterface.h b/Engine/Plugins/Runtime/Metasound/Source/MetasoundGraphCore/Public/MetasoundBuilderInterface.h
--- a/Engine/Plugins/Runtime/Metasound/Source/MetasoundGraphCore/Public/MetasoundBuilderInterface.h
+++ b/Engine/Plugins/Runtime/Metasound/Source/MetasoundGraphCore/Public/MetasoundBuilderInterface.h
@@ -36,6 +36,45 @@ namespace Metasound
 			virtual const TArray<FDataEdge>& GetEdges() const = 0;
 	};
 
+	/** FCreateOperatorParams holds the parameters provided to operator factories
+	 * during the creation of an IOperator
+	 */
+	struct FCreateOperatorParams
+	{
+		/** The node associated with this factory and the desired IOperator. */
+		const INode& Node;
+
+		/** General operator settings for the graph. */
+		const FOperatorSettings& OperatorSettings;
+
+		/** Collection of input parameters available for to an IOperator. */
+		const FDataReferenceCollection& InputDataReferences;
+
+		FCreateOperatorParams(const INode& InNode, const FOperatorSettings& InOperatorSettings, const FDataReferenceCollection& InInputDataReferences)
+		:	Node(InNode)
+		,	OperatorSettings(InOperatorSettings)
+		,	InputDataReferences(InInputDataReferences)
+		{
+		}
+	};
+
+	typedef TArray<TUniquePtr<IOperatorBuildError>> FBuildErrorArray;
+
+	/** Convenience template for adding build errors.
+	 *
+	 * The function can be used in the following way:
+	 * 
+	 * FBuildErrorArray MyErrorArray;
+	 * AddBuildError<FMyBuildErrorType>(MyErrorArray, MyBuildErrorConstructorArgs...);
+	 *
+	 * @param OutErrors - Array which holds the errors.
+	 * @param Args - Constructor arguments for the error.
+	 */
+	template<typename ErrorType, typename... ArgTypes>
+	void AddBuildError(FBuildErrorArray& OutErrors, ArgTypes&&... Args)
+	{
+		OutErrors.Add(MakeUnique<ErrorType>(Forward<ArgTypes>(Args)...));
+	}
 
 	/** IOperatorFactory
 	 *
@@ -50,14 +89,12 @@ namespace Metasound
 
 			/** Create a new IOperator.
 			 *
-			 * @param InNode - The node associated with this factory and the desired IOperator.
-			 * @param FOperatorSettings - General operator settings for the graph.
-			 * @param InInputCollection - Collection of input parameters available for to this IOperator.
+			 * @param InParams - The parameters available for building an IOperator.
 			 * @param OutErrors - An array of errors. Errors can be added if issues occur while creating the IOperator.
 			 *
 			 * @return A unique pointer to an IOperator. 
 			 */
-			virtual TUniquePtr<IOperator> CreateOperator(const INode& InNode, const FOperatorSettings& InOperatorSettings, const FDataReferenceCollection& InInputDataReferences, TArray<TUniquePtr<IOperatorBuildError>>& OutErrors) = 0;
+			virtual TUniquePtr<IOperator> CreateOperator(const FCreateOperatorParams& InParams, FBuildErrorArray& OutErrors) = 0;
 	};
 
 
diff --git a/Engine/Plugins/Runtime/Metasound/Source/MetasoundGraphCore/Public/MetasoundDataReference.h b/Engine/Plugins/Runtime/Metasound/Source/MetasoundGraphCore/Public/MetasoundDataReference.h
--- a/Engine/Plugins/Runtime/Metasound/Source/MetasoundGraphCore/Public/MetasoundDataReference.h
+++ b/Engine/Plugins/Runtime/Metasound/Source/MetasoundGraphCore/Public/MetasoundDataReference.h
@@ -9,7 +9,6 @@
 // Declares a metasound parameter type by
 // - Adding typedefs for commonly used template types.
 // - Defining parameter type traits.
-// - TODO- add argument for module api.
 #define DECLARE_METASOUND_DATA_REFERENCE_TYPES(DataType, ModuleApi, DataTypeMagicNumber, DataTypeInfoTypeName, DataReadReferenceTypeName, DataWriteReferenceTypeName) \
 	template<> \
 	struct ::Metasound::TDataReferenceTypeInfo<DataType> \
@@ -100,6 +99,23 @@ namespace Metasound
 		static constexpr bool bIsValidSpecialization = false;
 	};
 
+	template<>
+	struct TDataReferenceTypeInfo<void>
+	{
+		static METASOUNDGRAPHCORE_API const TCHAR* TypeName;
+		static constexpr const FDataTypeMagicNumber MagicNumber = -1;
+		static constexpr bool bIsStringParsable = false;
+		static constexpr bool bIsBoolParsable = false;
+		static constexpr bool bIsIntParsable = false;
+		static constexpr bool bIsFloatParsable = false;
+		static constexpr bool bIsProxyParsable = false;
+		static constexpr bool bIsProxyArrayParsable = false;
+		static constexpr bool bIsConstructableWithSettings = false;
+		static constexpr bool bCanUseDefaultConstructor = false;
+		static constexpr bool bIsValidSpecialization = false;
+	};
+
+	// TODO: comments
 	template<typename DataType>
 	const FName GetMetasoundDataTypeName() 
 	{
@@ -108,6 +124,15 @@ namespace Metasound
 		return TypeName;
 	}
 	
+	// TODO: comments
+	template<typename DataType>
+	bool IsReferenceOfType(const IDataReference& InReference)
+	{
+		static const FName TypeName = GetMetasoundDataTypeName<DataType>();
+		static const FDataTypeMagicNumber MagicNumber = TDataReferenceTypeInfo<typename TDecay<DataType>::Type >::MagicNumber;
+
+		return (InReference.GetDataTypeName() == TypeName) && (InReference.GetDataTypeMagicNumber() == MagicNumber);
+	}
 
 	// This enum is used as a token to explicitly delineate when we should create a new object for the reference,
 	// or use a different constructor.
@@ -461,3 +486,4 @@ namespace Metasound
 		Invalid,
 	};
 }
+
diff --git a/Engine/Plugins/Runtime/Metasound/Source/MetasoundGraphCore/Public/MetasoundFacade.h b/Engine/Plugins/Runtime/Metasound/Source/MetasoundGraphCore/Public/MetasoundFacade.h
new file mode 100644
--- /dev/null
+++ b/Engine/Plugins/Runtime/Metasound/Source/MetasoundGraphCore/Public/MetasoundFacade.h
@@ -0,0 +1,182 @@
+// Copyright Epic Games, Inc. All Rights Reserved.
+
+#pragma once
+
+#include "MetasoundBuilderInterface.h"
+#include "MetasoundNode.h"
+#include "MetasoundNodeInterface.h"
+#include "MetasoundVertex.h"
+#include <type_traits>
+
+namespace Metasound
+{
+	// Helper template to determine whether a static class function is declared
+	// for a given template class.
+	template <typename U>
+	class TIsVertexInterfaceDeclared
+	{
+		private:
+			template<typename T, T> 
+			struct Helper;
+
+			// Check for "static FVertexInterface U::DeclareVertexInterface()"
+			template<typename T>
+			static uint8 Check(Helper<FVertexInterface(*)(), &T::DeclareVertexInterface>*);
+
+			template<typename T> static uint16 Check(...);
+		public:
+			// If the function exists, then "Value" is true. Otherwise "Value" is false.
+			static constexpr bool Value = sizeof(Check<U>(0)) == sizeof(uint8_t);
+	};
+
+	// Helper template to determine whether a static class function is declared
+	// for a given template class.
+	template <typename U>
+	class TIsFactoryMethodDeclared
+	{
+		private:
+			template<typename T, T> 
+			struct Helper;
+
+			// Check for "static TUniquePtr<IOperator> U::CreateOperator(const FCreateOperatorParams& Inparams, FBuildErrorArray& OutErrors)"
+			template<typename T>
+			static uint8 Check(Helper<TUniquePtr<IOperator>(*)(const FCreateOperatorParams& InParams, FBuildErrorArray& OutErrors), &T::CreateOperator>*);
+
+			template<typename T> static uint16 Check(...);
+
+		public:
+
+			// If the function exists, then "Value" is true. Otherwise "Value" is false.
+			static constexpr bool Value = sizeof(Check<U>(0)) == sizeof(uint8_t);
+	};
+
+	// Helper template to determine whether a static class function is declared
+	// for a given template class.
+	template <typename U>
+	class TIsNodeInfoDeclared
+	{
+		private:
+			template<typename T, T> 
+			struct Helper;
+
+			// Check for "static const FNodeInfo& U::CreateOperator(const FCreateOperatorParams& Inparams, FBuildErrorArray& OutErrors)"
+			template<typename T>
+			static uint8 Check(Helper<const FNodeInfo&(*)(), &T::GetNodeInfo>*);
+
+			template<typename T> static uint16 Check(...);
+
+		public:
+
+			// If the function exists, then "Value" is true. Otherwise "Value" is false.
+			static constexpr bool Value = sizeof(Check<U>(0)) == sizeof(uint8_t);
+	};
+
+	/** TFacadeOperatorClass encapsulates an operator type and checks that the
+	 * required static functions exist to build the facade operator class.  It 
+	 * is required to call the FNodeFacade template constructor.
+	 */
+	template<typename OperatorType>
+	struct TFacadeOperatorClass
+	{
+		// Require that OperatorType is subclass of IOperator
+		static_assert(std::is_base_of<IOperator, OperatorType>::value, "To use the FNodeFacade consturctor, the OperatorType must be derived from IOperator");
+
+		// Require static FVertexInterface OperatorType::DeclareVertexInterface() exists.
+		static_assert(TIsVertexInterfaceDeclared<OperatorType>::Value, "To use the FNodeFacade constructor, the OperatorType must have the static function \"static FVertexInterface OperatorType::DeclareVertexInterface()\"");
+
+		// Require static TUniquePtr<IOperator> OperatorType::CreateOperator(const FCreateOperatorParams&, TArray<TUniquePtr<IOperatorBuildError>>&) exists.
+		static_assert(TIsFactoryMethodDeclared<OperatorType>::Value, "To use the FNodeFacade constructor, the OperatorType must have the static function \"static TUniquePtr<IOperator> OperatorType::CreateOperator(const FCreateOperatorParams&, TArray<TUniquePtr<IOperatorBuildError>>&)\"");
+
+
+		// Require static const FNodeInfo& OperatorType::GetNodeInfo() exists.
+		static_assert(TIsNodeInfoDeclared<OperatorType>::Value, "To use the FNodeFacade constructor, the OperatorType must have the static function \"static const FNodeInfo& OperatorType::GetNodeInfo()\"");
+
+		static const FNodeInfo& GetNodeInfo() 
+		{ 
+			return OperatorType::GetNodeInfo(); 
+		}
+
+		static FVertexInterface DeclareVertexInterface() 
+		{ 
+			return OperatorType::DeclareVertexInterface(); 
+		}
+
+		static TUniquePtr<TUniquePtr<IOperator>> CreateOperator(const FCreateOperatorParams& InParams, FBuildErrorArray& OutErrors) 
+		{
+			return OperatorType::CreateOperator(InParams, OutErrors);
+		}
+
+		typedef OperatorType Type;
+	};
+
+	/** FNodeFacade implements a significant amount of boilerplate code required
+	 * to build a Metasound INode. FNodeFacade is particularly useful for an INode
+	 * which has a static FVertexInterface, and always creates the same IOperator type.
+	 * 
+	 * The type of the concrete IOperator class to instantiate is defined in the
+	 * FNodeFacade constructors TFacadeOperatorClass<>.
+	 */
+	class METASOUNDGRAPHCORE_API FNodeFacade : public FNode
+	{
+			// Factory class for create an IOperator by using a TFunction matching
+			// the signature of the CreateOperator function.
+			class METASOUNDGRAPHCORE_API FFactory : public IOperatorFactory
+			{
+				using FCreateOperatorFunction = TFunction<TUniquePtr<IOperator>(const FCreateOperatorParams& InParams, FBuildErrorArray& OutErrors)>;
+
+				public:
+					FFactory(FCreateOperatorFunction InCreateFunc);
+
+					virtual TUniquePtr<IOperator> CreateOperator(const FCreateOperatorParams& InParams, FBuildErrorArray& OutErrors) override;
+
+				private:
+					FCreateOperatorFunction CreateFunc;
+			};
+
+		public: 
+
+			FNodeFacade() = delete;
+
+			/** FNodeFacade constructor using the OperatorType template parameter
+			 * to get the node info, operator factory method and vertex interface. 
+			 *
+			 * @param InInstanceName - Instance name for the node.
+			 * @param OperatorClass - Template class wrapper for the underlying 
+			 * 						  IOperator which is created by this node.
+			 */
+			template<typename OperatorType>
+			FNodeFacade(const FString& InInstanceName, TFacadeOperatorClass<OperatorType> OperatorClass)
+			:	FNode(InInstanceName, OperatorType::GetNodeInfo())
+			, 	Factory(MakeShared<FFactory, ESPMode::ThreadSafe>(&OperatorType::CreateOperator))
+			,	VertexInterface(OperatorType::DeclareVertexInterface())
+			{
+			}
+
+			virtual ~FNodeFacade() = default;
+
+			virtual const FVertexInterface& GetVertexInterface() const override;
+
+			virtual const FVertexInterface& GetDefaultVertexInterface() const override;
+
+			/** FNodeFacade has a static vertex interface. This will return true
+			 * only if the input interface is equal to the default interface.
+			 */
+			virtual bool SetVertexInterface(const FVertexInterface& InInterface) override;
+
+			/** FNodeFacade has a static vertex interface. This will return true
+			 * only if the input interface is equal to the default interface.
+			 */
+			virtual bool IsVertexInterfaceSupported(const FVertexInterface& InInterface) const override;
+
+			/** Return a reference to the default operator factory. */
+			virtual FOperatorFactorySharedRef GetDefaultOperatorFactory() const override;
+
+		private:
+
+			TSharedRef<IOperatorFactory, ESPMode::ThreadSafe> Factory;
+
+			FVertexInterface VertexInterface;
+	};
+
+
+}
diff --git a/Engine/Plugins/Runtime/Metasound/Source/MetasoundGraphCore/Public/MetasoundGraph.h b/Engine/Plugins/Runtime/Metasound/Source/MetasoundGraphCore/Public/MetasoundGraph.h
--- a/Engine/Plugins/Runtime/Metasound/Source/MetasoundGraphCore/Public/MetasoundGraph.h
+++ b/Engine/Plugins/Runtime/Metasound/Source/MetasoundGraphCore/Public/MetasoundGraph.h
@@ -10,7 +10,6 @@ class FText;
 
 namespace Metasound
 {
-	// TODO: consider making nodes TSharedPtr<INode>
 	/** FGraph contains the edges between nodes as well as input and output 
 	 * vertices.  FGraph does not maintain ownership over any node. Nodes used
 	 * within the graph must be valid for the lifetime of the graph. 
@@ -35,15 +34,31 @@ namespace Metasound
 
 			virtual const FText& GetPromptIfMissing() const override;
 
-			/** Return a collection of input parameter descriptions for this graph. */
-			virtual const FInputDataVertexCollection& GetInputDataVertices() const override;
-
-			/** Return a collection of output parameter descriptions for this graph. */
-			virtual const FOutputDataVertexCollection& GetOutputDataVertices() const override;
-
 			/** Retrieve all the edges associated with a graph. */
 			virtual const TArray<FDataEdge>& GetDataEdges() const override;
 
+			/** Return the current vertex interface. */
+			virtual const FVertexInterface& GetVertexInterface() const override;
+
+			/** Return the default vertex interface. */
+			virtual const FVertexInterface& GetDefaultVertexInterface() const override;
+
+			/** Set the vertex interface. If the vertex was successfully changed, returns true. 
+			 *
+			 * @param InInterface - New interface for node. 
+			 *
+			 * @return True on success, false otherwise.
+			 */
+			virtual bool SetVertexInterface(const FVertexInterface& InInterface) override;
+
+			/** Expresses whether a specific vertex interface is supported.
+			 *
+			 * @param InInterface - New interface. 
+			 *
+			 * @return True if the interface is supported, false otherwise. 
+			 */
+			virtual bool IsVertexInterfaceSupported(const FVertexInterface& InInterface) const override;
+
 			/** Get vertices which contain input parameters. */
 			virtual const FInputDataDestinationCollection& GetInputDataDestinations() const override;
 
@@ -57,24 +72,24 @@ namespace Metasound
 			 * nodes. 
 			 *
 			 * @param FromNode - Node which contains the output vertex.
-			 * @param FromVertexName - Name of the vertex in the FromNode.
+			 * @param FromVertexKey - Key of the vertex in the FromNode.
 			 * @param ToNode - Node which contains the input vertex.
-			 * @param ToVertexName - Name of the vertex in the ToNode.
+			 * @param ToVertexKey - Key of the vertex in the ToNode.
 			 *
 			 * @return True if the edge was successfully added. False otherwise.
 			 */
-			bool AddDataEdge(INode& FromNode, const FString& FromVertexName, INode& ToNode, const FString ToVertexName);
+			bool AddDataEdge(const INode& FromNode, const FDataVertexKey& FromVertexKey, const INode& ToNode, const FDataVertexKey& ToVertexKey);
 
 			/** Add an input data destination to describe how data provided 
 			 * outside this graph should be routed internally.
 			 *
 			 * @param InNode - Node which receives the data.
-			 * @param InVertexName - Name of input vertex on InNode.
+			 * @param InVertexKey - Key for input vertex on InNode.
 			 *
 			 * @return True if the destination was successfully added. False 
 			 * otherwise.
 			 */
-			bool AddInputDataDestination(INode& InNode, const FString& InVertexName);
+			bool AddInputDataDestination(const INode& InNode, const FDataVertexKey& InVertexKey);
 
 
 			/** Add an input data destination to describe how data provided 
@@ -86,12 +101,12 @@ namespace Metasound
 			 * owned this graph and exposed externally.
 			 *
 			 * @param InNode - Node which produces the data.
-			 * @param InVertexName - Name of output vertex on InNode.
+			 * @param InVertexKey - Key for output vertex on InNode.
 			 *
 			 * @return True if the source was successfully added. False 
 			 * otherwise.
 			 */
-			bool AddOutputDataSource(INode& InNode, const FString& InVertexName);
+			bool AddOutputDataSource(const INode& InNode, const FDataVertexKey& InVertexKey);
 
 			/** Add an output data source which describes routing of data which is 
 			 * owned this graph and exposed externally.
@@ -103,11 +118,10 @@ namespace Metasound
 		private:
 			FString Description;
 
-			FInputDataVertexCollection InputVertices;
-			FOutputDataVertexCollection OutputVertices;
 
 			TArray<FDataEdge> Edges;
 
+			FVertexInterface VertexInterface;
 			FInputDataDestinationCollection InputDestinations;
 			FOutputDataSourceCollection OutputSources;
 	};
diff --git a/Engine/Plugins/Runtime/Metasound/Source/MetasoundGraphCore/Public/MetasoundNode.h b/Engine/Plugins/Runtime/Metasound/Source/MetasoundGraphCore/Public/MetasoundNode.h
--- a/Engine/Plugins/Runtime/Metasound/Source/MetasoundGraphCore/Public/MetasoundNode.h
+++ b/Engine/Plugins/Runtime/Metasound/Source/MetasoundGraphCore/Public/MetasoundNode.h
@@ -7,43 +7,43 @@
 
 namespace Metasound
 {
+	struct FNodeInfo
+	{
+		FName ClassName;
+		FText Description;
+		FText AuthorName;
+		FText PromptIfMissing;
+	};
+
 	class METASOUNDGRAPHCORE_API FNode : public INode
 	{
 		public:
-			FNode(const FString& InDescription);
-			virtual ~FNode();
+			FNode(const FString& InInstanceName, const FNodeInfo& InInfo);
 
-			virtual const FString& GetInstanceName() const override;
-
-			virtual const FInputDataVertexCollection& GetInputDataVertices() const override;
-			virtual const FOutputDataVertexCollection& GetOutputDataVertices() const override;
-
-		protected:
+			virtual ~FNode() = default;
 
+			/** Return the name of this specific instance of the node class. */
+			virtual const FString& GetInstanceName() const override;
 
-			template<typename DataType>
-			void AddInputDataVertex(const FString& InVertexName, const FText& InVertexDescription)
-			{
-				AddInputDataVertex(MakeInputDataVertex<DataType>(InVertexName, InVertexDescription));
-			}
+			/** Return the type name of this node. */
+			virtual const FName& GetClassName() const override;
 
-			void AddInputDataVertex(const FInputDataVertex& InVertex);
-			void RemoveInputDataVertex(const FInputDataVertex& InVertex);
+			/** Return a longer text description describing how this node is used. */
+			virtual const FText& GetDescription() const override;
 
-			template<typename DataType>
-			void AddOutputDataVertex(const FString& InVertexName, const FText& InVertexDescription)
-			{
-				AddOutputDataVertex(MakeOutputDataVertex<DataType>(InVertexName, InVertexDescription));
-			}
+			/** Return the original author of this node class. */
+			virtual const FText& GetAuthorName() const override;
 
-			void AddOutputDataVertex(const FOutputDataVertex& InVertex);
-			void RemoveOutputDataVertex(const FOutputDataVertex& InVertex);
+			/** 
+			 *  Return an optional prompt on how users can get the plugin this node is in,
+			 *  if they have found a metasound that uses this node but don't have this plugin downloaded or enabled.
+			 */
+			virtual const FText& GetPromptIfMissing() const override;
 
 		private:
 
-			FString Description;
+			FString InstanceName;
+			FNodeInfo Info;
 
-			FInputDataVertexCollection Inputs;
-			FOutputDataVertexCollection Outputs;
 	};
 }
diff --git a/Engine/Plugins/Runtime/Metasound/Source/MetasoundGraphCore/Public/MetasoundNodeInterface.h b/Engine/Plugins/Runtime/Metasound/Source/MetasoundGraphCore/Public/MetasoundNodeInterface.h
--- a/Engine/Plugins/Runtime/Metasound/Source/MetasoundGraphCore/Public/MetasoundNodeInterface.h
+++ b/Engine/Plugins/Runtime/Metasound/Source/MetasoundGraphCore/Public/MetasoundNodeInterface.h
@@ -4,117 +4,15 @@
 
 #include "CoreMinimal.h"
 #include "Internationalization/Text.h"
-#include "MetasoundOperatorInterface.h"
 #include "MetasoundDataReference.h"
+#include "MetasoundOperatorInterface.h"
+#include "MetasoundVertex.h"
 
 namespace Metasound
 {
 	static const FText PluginAuthor = NSLOCTEXT("MetasoundGraphCore", "Metasound_DefaultAuthor", "Epic Games, Inc.");
 	static const FText PluginNodeMissingPrompt = NSLOCTEXT("MetasoundGraphCore", "Metasound_DefaultMissingPrompt", "Make sure that the Metasound plugin is loaded.");
 
-	/** FDataVertex
-	 *
-	 *  This describes a data vertex on a Node.
-	 */
-	struct FDataVertex
-	{
-		/** Name of vertex. */
-		FString VertexName;
-
-		/** Type name of data reference. */
-		FName DataReferenceTypeName;
-
-		/** Description of the vertex. */
-		FText Description;
-
-		FDataVertex()
-		:	Description(FText::GetEmpty())
-		{
-		}
-
-		FDataVertex(const FString& InVertexName, const FName& InDataReferenceTypeName, const FText& InDescription)
-		:	VertexName(InVertexName)
-		,	DataReferenceTypeName(InDataReferenceTypeName)
-		,	Description(InDescription)
-		{
-		}
-
-		/** Check if two data vertex specs are equal. */
-		friend bool METASOUNDGRAPHCORE_API operator==(const FDataVertex& LHS, const FDataVertex& RHS);
-
-		/** Check if two data vertex specs are unequal. */
-		friend bool METASOUNDGRAPHCORE_API operator!=(const FDataVertex& LHS, const FDataVertex& RHS);
-	};
-
-	struct FInputDataVertex: FDataVertex
-	{
-		using FDataVertex::FDataVertex;
-	};
-
-	struct FOutputDataVertex : FDataVertex
-	{
-		using FDataVertex::FDataVertex;
-	};
-
-
-	/** Create an input data vertex for a DataType. DataType refers to the C++
-	 * type used to specialize the TMetasoundDataTypeInfo<> struct.
-	 *
-	 * @param InVertexName - Name of the input data vertex.
-	 * @param InDescription - Description of the input data vertex.
-	 *
-	 * @return A new FInputDataVertex.
-	 */
-	template<typename DataType>
-	FInputDataVertex MakeInputDataVertex(const FString& InVertexName, const FText& InDescription)
-	{
-		static const FName DataTypeName = FName(TDataReferenceTypeInfo<DataType>::TypeName);
-
-		return FInputDataVertex(InVertexName, DataTypeName, InDescription);
-	}
-
-	/** Create an output data vertex for a DataType. DataType refers to the C++
-	 * type used to specialize the TMetasoundDataTypeInfo<> struct.
-	 *
-	 * @param InVertexName - Name of the output data vertex.
-	 * @param InDescription - Description of the output data vertex.
-	 *
-	 * @return A new FOutputDataVertex.
-	 */
-	template<typename DataType>
-	FOutputDataVertex MakeOutputDataVertex(const FString& InVertexName, const FText& InDescription)
-	{
-		static const FName DataTypeName = FName(TDataReferenceTypeInfo<DataType>::TypeName);
-
-		return FOutputDataVertex(InVertexName, DataTypeName, InDescription);
-	}
-
-	/** Key type for an FInputDataVertexColletion or 
-	 * FOutputDataVertexCollection. 
-	 */
-	typedef FString FDataVertexKey;
-
-	/** FInputDataVertexCollection contains multiple FInputDataVertexes mapped
-	 * by FDataVertexKeys.
-	 */
-	typedef TMap<FDataVertexKey, FInputDataVertex> FInputDataVertexCollection;
-
-	/** FOutputDataVertexCollection contains multiple FOutputDataVertexes mapped
-	 * by FDataVertexKeys.
-	 */
-	typedef TMap<FDataVertexKey, FOutputDataVertex> FOutputDataVertexCollection;
-
-	/** Create an FDataVertexKey from a FInputDataVertex. */
-	FORCEINLINE FDataVertexKey MakeDataVertexKey(const FInputDataVertex& InVertex)
-	{
-		return InVertex.VertexName;
-	}
-
-	/** Create an FDataVertexKey from a FOutputDataVertex. */
-	FORCEINLINE FDataVertexKey MakeDataVertexKey(const FOutputDataVertex& InVertex)
-	{
-		return InVertex.VertexName;
-	}
 
 	/**
 	 * This struct is used to pass in any arguments required for constructing a single node instance.
@@ -136,26 +34,6 @@ namespace Metasound
 		}
 	};
 
-	typedef FInputDataVertexCollection::ElementType FInputDataVertexElement;
-	typedef FOutputDataVertexCollection::ElementType FOutputDataVertexElement;
-
-	template<typename DataType>
-	FORCEINLINE FInputDataVertexElement MakeInputDataVertexElement(const FString& InName, const FText& InDescription)
-	{
-		FInputDataVertex Vertex = MakeInputDataVertex<DataType>(InName, InDescription);
-		return FInputDataVertexElement(MakeDataVertexKey(Vertex), Vertex);
-	}
-
-	template<typename DataType>
-	FORCEINLINE FOutputDataVertexElement MakeOutputDataVertexElement(const FString& InName, const FText& InDescription)
-	{
-		FOutputDataVertex Vertex = MakeOutputDataVertex<DataType>(InName, InDescription);
-		return FOutputDataVertexElement(MakeDataVertexKey(Vertex), Vertex);
-	}
-
-
-	//typedef uint32 FNodeId;
-
 	/** INodeBase
 	 * 
 	 * Interface for all nodes that can describe their name, type, inputs and outputs.
@@ -165,9 +43,6 @@ namespace Metasound
 		public:
 			virtual ~INodeBase() {}
 
-			// TODO: talk to ethan about nodeids. Maybe this should be handled by frontend. 
-			//FNodeId GetNodeId() const = 0;
-
 			/** Return the name of this specific instance of the node class. */
 			virtual const FString& GetInstanceName() const = 0;
 
@@ -186,16 +61,43 @@ namespace Metasound
 			 */
 			virtual const FText& GetPromptIfMissing() const = 0;
 
-			/** Return a collection of input parameter descriptions for this node. */
-			virtual const FInputDataVertexCollection& GetInputDataVertices() const = 0;
+			/** Return the current vertex interface. */
+			virtual const FVertexInterface& GetVertexInterface() const = 0;
+
+			/** Return the default vertex interface. */
+			virtual const FVertexInterface& GetDefaultVertexInterface() const = 0;
+
+			/** Set the vertex interface. If the vertex was successfully changed, returns true. 
+			 *
+			 * @param InInterface - New interface for node. 
+			 *
+			 * @return True on success, false otherwise.
+			 */
+			virtual bool SetVertexInterface(const FVertexInterface& InInterface) = 0;
 
-			/** Return a collection of output parameter descriptions for this node. */
-			virtual const FOutputDataVertexCollection& GetOutputDataVertices() const = 0;
+			/** Expresses whether a specific vertex interface is supported.
+			 *
+			 * @param InInterface - New interface. 
+			 *
+			 * @return True if the interface is supported, false otherwise. 
+			 */
+			virtual bool IsVertexInterfaceSupported(const FVertexInterface& InInterface) const = 0;
 	};
 
 	// Forward declare
 	class IOperatorFactory;
 
+	/** Shared ref type of operator factory. */
+	typedef TSharedRef<IOperatorFactory, ESPMode::ThreadSafe> FOperatorFactorySharedRef;
+
+	/** Convenience function for making operator factory references */
+	template<typename FactoryType, typename... ArgTypes>
+	TSharedRef<FactoryType, ESPMode::ThreadSafe> MakeOperatorFactoryRef(ArgTypes&&... Args)
+	{
+		return MakeShared<FactoryType, ESPMode::ThreadSafe>(Forward<ArgTypes>(Args)...);
+	}
+
+
 	/** INode 
 	 * 
 	 * Interface for all nodes that can create operators. 
@@ -205,9 +107,8 @@ namespace Metasound
 		public:
 			virtual ~INode() {}
 
-			// TODO: consider making this return a sharedref to allow for more flexible management of factories. 
 			/** Return a reference to the default operator factory. */
-			virtual IOperatorFactory& GetDefaultOperatorFactory() = 0;
+			virtual FOperatorFactorySharedRef GetDefaultOperatorFactory() const = 0;
 
 			/* Future implementations may support additional factory types and interfaces
 				virtual bool DoesSupportOperatorFactory(const FString& InFractoryName) const = 0;
@@ -216,14 +117,13 @@ namespace Metasound
 			*/
 	};
 
-
 	/** FOutputDataSource describes the source of data which is produced within
 	 * a graph and exposed external to the graph. 
 	 */
 	struct FOutputDataSource
 	{
 		/** Node containing the output data vertex. */
-		INode* Node = nullptr;
+		const INode* Node = nullptr;
 
 		/** Output data vertex. */
 		FOutputDataVertex Vertex;
@@ -232,14 +132,15 @@ namespace Metasound
 		{
 		}
 
-		FOutputDataSource(INode& InNode, const FOutputDataVertex& InVertex)
+		FOutputDataSource(const INode& InNode, const FOutputDataVertex& InVertex)
 		:	Node(&InNode)
 		,	Vertex(InVertex)
 		{
 		}
 
+
 		/** Check if two FOutputDataSources are equal. */
-		friend bool operator==(const FOutputDataSource& InLeft, const FOutputDataSource& InRight);
+		friend bool METASOUNDGRAPHCORE_API operator==(const FOutputDataSource& InLeft, const FOutputDataSource& InRight);
 	};
 
 	/** FInputDataSource describes the destination of data which produced 
@@ -248,7 +149,7 @@ namespace Metasound
 	struct FInputDataDestination
 	{
 		/** Node containing the input data vertex. */
-		INode* Node = nullptr;
+		const INode* Node = nullptr;
 
 		/** Input data vertex of edge. */
 		FInputDataVertex Vertex;
@@ -257,14 +158,14 @@ namespace Metasound
 		{
 		}
 
-		FInputDataDestination(INode& InNode, const FInputDataVertex& InVertex)
+		FInputDataDestination(const INode& InNode, const FInputDataVertex& InVertex)
 		:	Node(&InNode)
 		,	Vertex(InVertex)
 		{
 		}
 
 		/** Check if two FInputDataDestinations are equal. */
-		friend bool operator==(const FInputDataDestination& InLeft, const FInputDataDestination& InRight);
+		friend bool METASOUNDGRAPHCORE_API operator==(const FInputDataDestination& InLeft, const FInputDataDestination& InRight);
 	};
 
 	/** Key type for an FOutputDataSource or FInputDataDestination. */
@@ -283,12 +184,12 @@ namespace Metasound
 	/** Make a FNodeDataVertexKey from an FOutputDataSource. */
 	FORCEINLINE FNodeDataVertexKey MakeSourceDataVertexKey(const FOutputDataSource& InSource)
 	{
-		return FNodeDataVertexKey(InSource.Node, InSource.Vertex.VertexName);
+		return FNodeDataVertexKey(InSource.Node, InSource.Vertex.GetVertexName());
 	}
 
 	FORCEINLINE FNodeDataVertexKey MakeDestinationDataVertexKey(const FInputDataDestination& InDestination)
 	{
-		return FNodeDataVertexKey(InDestination.Node, InDestination.Vertex.VertexName);
+		return FNodeDataVertexKey(InDestination.Node, InDestination.Vertex.GetVertexName());
 	}
 
 	/** FDataEdge
@@ -312,7 +213,7 @@ namespace Metasound
 		}
 
 		/** Check if two FDataEdges are equal. */
-		friend bool operator==(const FDataEdge& InLeft, const FDataEdge& InRight);
+		friend bool METASOUNDGRAPHCORE_API operator==(const FDataEdge& InLeft, const FDataEdge& InRight);
 	};
 
 
@@ -331,13 +232,4 @@ namespace Metasound
 			/** Get vertices which contain output parameters. */
 			virtual const FOutputDataSourceCollection& GetOutputDataSources() const = 0;
 	};
-
-
-	// This one is not yet implemented. It's design will likely be driven by UX and object management requirements. 
-	class INodeFactory
-	{
-		public:
-			virtual ~INodeFactory() {}
-			virtual TUniquePtr<INode> CreateNode(const FString& InJsonString) = 0;
-	};
 }
diff --git a/Engine/Plugins/Runtime/Metasound/Source/MetasoundGraphCore/Public/MetasoundOperatorInterface.h b/Engine/Plugins/Runtime/Metasound/Source/MetasoundGraphCore/Public/MetasoundOperatorInterface.h
--- a/Engine/Plugins/Runtime/Metasound/Source/MetasoundGraphCore/Public/MetasoundOperatorInterface.h
+++ b/Engine/Plugins/Runtime/Metasound/Source/MetasoundGraphCore/Public/MetasoundOperatorInterface.h
@@ -146,6 +146,7 @@ namespace Metasound
 	{
 		FDataTypeLiteralParam(const FDataTypeLiteralParam& Other) = delete;
 		FDataTypeLiteralParam(FDataTypeLiteralParam&& Other) = default;
+		FDataTypeLiteralParam& operator=(FDataTypeLiteralParam&& Other) = default;
 
 		ELiteralArgType ConstructorArgType;
 		TVariant<bool, int, float, FString, Audio::IProxyDataPtr, TArray<Audio::IProxyDataPtr>> ConstructorArg;
diff --git a/Engine/Plugins/Runtime/Metasound/Source/MetasoundGraphCore/Public/MetasoundVertex.h b/Engine/Plugins/Runtime/Metasound/Source/MetasoundGraphCore/Public/MetasoundVertex.h
new file mode 100644
--- /dev/null
+++ b/Engine/Plugins/Runtime/Metasound/Source/MetasoundGraphCore/Public/MetasoundVertex.h
@@ -0,0 +1,531 @@
+// Copyright Epic Games, Inc. All Rights Reserved.
+
+#pragma once
+
+#include "CoreMinimal.h"
+#include "Containers/SortedMap.h"
+#include "MetasoundDataReference.h"
+
+#include <type_traits>
+
+namespace Metasound
+{
+	/** Key type for an FInputDataVertexColletion or 
+	 * FOutputDataVertexCollection. 
+	 */
+	typedef FString FDataVertexKey;
+
+	// Forward declare.
+	class FInputDataVertex;
+
+	// Forward declare.
+	class FOutputDataVertex;
+
+	/** FDataVertexModel
+	 *
+	 * An FDataVertexModel implements various vertex behaviors. 
+	 *
+	 * Essentially, an FDataVertexModel implements the underlying behavior of a
+	 * FInputDataVertex or FOutputDataVertex.  
+	 *
+	 * The FInputDataVertex and FOutputDataVertex supply copy constructors and 
+	 * assignment operators by utilizing the Clone() operator of the 
+	 * FDataVertexModel. This allows FDataVertexModel behavior to be passed 
+	 * to other objects without having to know the concrete implementation of the
+	 * FDataVertexModel.
+	 */
+	struct FDataVertexModel
+	{
+		/** FDataVertexModel Construtor
+		 *
+		 * @InVertexName - Name of vertex.
+		 * @InDataTypeName - Name of data type.
+		 * @InDescription - Human readable vertex description.
+		 */
+		FDataVertexModel(const FString& InVertexName, const FName& InDataTypeName, const FText& InDescription)
+		:	VertexName(InVertexName)
+		,	DataTypeName(InDataTypeName)
+		,	Description(InDescription)
+		{
+		}
+
+		virtual ~FDataVertexModel() = default;
+
+		/** Name of vertex. */
+		const FString VertexName;
+
+		/** Type name of data. */
+		const FName DataTypeName;
+
+		/** Description of the vertex. */
+		const FText Description;
+
+		/** Test if a IDataReference matches the DataType of this vertex model.
+		 *
+		 * @param InReference - Data reference object.
+		 *
+		 * @return True if the types are equal, false otherwise.
+		 */
+		virtual bool IsReferenceOfSameType(const IDataReference& InReference) const = 0;
+
+		/** Return the vertex type name (not to be confused with the data type name).
+		 *
+		 * @return Name of this vertex type.
+		 */
+		virtual const FName& GetVertexTypeName() const = 0;
+	};
+
+
+	/** FInputDataVertexModel
+	 *
+	 * Vertex model for inputs.
+	 */
+	struct FInputDataVertexModel: FDataVertexModel
+	{
+		// Inherit constructor
+		using FDataVertexModel::FDataVertexModel;
+
+		/** Create a clone of this FInputDataVertexModel */
+		virtual TUniquePtr<FInputDataVertexModel> Clone() const = 0;
+
+		friend bool METASOUNDGRAPHCORE_API operator==(const FInputDataVertex& InLHS, const FInputDataVertex& InRHS);
+		friend bool METASOUNDGRAPHCORE_API operator!=(const FInputDataVertex& InLHS, const FInputDataVertex& InRHS);
+
+		private:
+
+		/** Compare another FInputDataVertexModel is equal to this
+		 *
+		 * @return True if models are equal, false otherwise.
+		 */
+		virtual bool IsEqual(const FInputDataVertexModel& InOther) const = 0;
+	};
+
+	/** FOutputDataVertexModel
+	 *
+	 * Vertex model for outputs.
+	 */
+	struct FOutputDataVertexModel : FDataVertexModel
+	{
+		using FDataVertexModel::FDataVertexModel;
+
+		/** Create a clone of this FOutputDataVertexModel */
+		virtual TUniquePtr<FOutputDataVertexModel> Clone() const = 0;
+
+		friend bool METASOUNDGRAPHCORE_API operator==(const FOutputDataVertex& InLHS, const FOutputDataVertex& InRHS);
+		friend bool METASOUNDGRAPHCORE_API operator!=(const FOutputDataVertex& InLHS, const FOutputDataVertex& InRHS);
+
+		private:
+
+		/** Compare another FOutputDataVertexModel is equal to this
+		 *
+		 * @return True if models are equal, false otherwise.
+		 */
+		virtual bool IsEqual(const FOutputDataVertexModel& InOther) const = 0;
+
+	};
+
+	/** TBaseVertexModel provides basic functionality of vertex models. */
+	template<typename DataType, typename VertexModelType>
+	struct TBaseVertexModel : VertexModelType
+	{
+		/** TBaseVertexModel Construtor
+		 *
+		 * @InVertexName - Name of vertex.
+		 * @InDescription - Human readable vertex description.
+		 */
+		TBaseVertexModel(const FString& InVertexName, const FText& InDescription)
+		:	VertexModelType(InVertexName, GetMetasoundDataTypeName<DataType>(), InDescription)
+		{
+		}
+
+		/** Test if a IDataReference matches the DataType of this vertex model.
+		 *
+		 * @param InReference - Data reference object.
+		 *
+		 * @return True if the types are equal, false otherwise.
+		 */
+		virtual bool IsReferenceOfSameType(const IDataReference& InReference) const override
+		{
+			return IsReferenceOfType<DataType>(InReference);
+		}
+
+
+		private:
+
+		/** Compare if another VertexModelType is equal to this.
+		 *
+		 * @return True if models are equal, false otherwise.
+		 */
+		virtual bool IsEqual(const VertexModelType& InOther) const override
+		{
+			bool bIsEqual = this->GetVertexTypeName() == InOther.GetVertexTypeName();
+
+			bIsEqual &= this->VertexName == InOther.VertexName;
+			bIsEqual &= this->DataTypeName == InOther.DataTypeName;
+
+			return bIsEqual;
+		}
+	};
+
+	/** TOuputDataVertexModel creates a simple, unchanging, input vertex. */
+	template<typename DataType>
+	struct TInputDataVertexModel : TBaseVertexModel<DataType, FInputDataVertexModel>
+	{
+		using TBaseVertexModel<DataType, FInputDataVertexModel>::TBaseVertexModel;
+
+		/** Return the vertex type name (not to be confused with the data type name).
+		 *
+		 * @return Name of this vertex type.
+		 */
+		virtual const FName& GetVertexTypeName() const override
+		{
+			static const FString VertexTypeString = TEXT("InputDataVertex_") + GetMetasoundDataTypeName<DataType>().ToString();
+			static const FName VertexTypeName(*VertexTypeString);
+
+			return VertexTypeName;
+		}
+
+		/** Create a clone of this VertexModelType */
+		virtual TUniquePtr<FInputDataVertexModel> Clone() const override
+		{
+			return MakeUnique<TInputDataVertexModel<DataType>>(this->VertexName, this->Description);
+		}
+	};
+
+
+	/** FInputDataVertex represents an input vertex to an interface. It uses
+	 * a FInputDataVertexModel to determine its behavior.
+	 */
+	class METASOUNDGRAPHCORE_API FInputDataVertex
+	{
+			using FEmptyVertexModel = TInputDataVertexModel<void>;
+
+		public:
+
+			/** Construct an FInputDataVertex with a given vertex model.
+			 *
+			 * @param InVertexModel - A model subclass of FInputDataVertexModel.
+			 */
+			template<typename VertexModelType, typename = typename TEnableIf<TIsDerivedFrom<VertexModelType, FInputDataVertexModel>::Value>::Type>
+			FInputDataVertex(const VertexModelType& InVertexModel)
+			:	VertexModel(InVertexModel.Clone())
+			{
+				static_assert(TIsDerivedFrom<VertexModelType, FInputDataVertexModel>::Value, "Vertex implementation must be a subclass of FInputDataVertexModel");
+
+				if (!VertexModel.IsValid())
+				{
+					VertexModel = MakeUnique<FEmptyVertexModel>(TEXT(""), FText::GetEmpty());
+				}
+			}
+
+			/** Construct an empty FInputDataVertex. */
+			FInputDataVertex();
+
+			/** Copy constructor */
+			FInputDataVertex(const FInputDataVertex& InOther);
+
+			/** Assignment operator. */
+			FInputDataVertex& operator=(const FInputDataVertex& InOther);
+
+			/** Name of vertex. */
+			const FString& GetVertexName() const;
+
+			/** Type name of data reference. */
+			const FName& GetDataTypeName() const;
+
+			/** Description of the vertex. */
+			const FText& GetDescription() const;
+
+			/** Determine if vertex refers to same data type. 
+			 *
+			 * @param InReference - Data reference object.
+			 *
+			 * @return True if the types are equal, false otherwise.
+			 */
+			bool IsReferenceOfSameType(const IDataReference& InReference) const;
+
+			/** Return the vertex type name (not to be confused with the data type name).
+			 *
+			 * @return Name of this vertex type.
+			 */
+			const FName& GetVertexTypeName() const;
+
+			friend bool METASOUNDGRAPHCORE_API operator==(const FInputDataVertex& InLHS, const FInputDataVertex& InRHS);
+			friend bool METASOUNDGRAPHCORE_API operator!=(const FInputDataVertex& InLHS, const FInputDataVertex& InRHS);
+
+		private:
+
+			TUniquePtr<FInputDataVertexModel> VertexModel;
+	};
+
+	/** TOuputDataVertexModel creates a simple, unchanging, output vertex. */
+	template<typename DataType>
+	struct TOutputDataVertexModel : TBaseVertexModel<DataType, FOutputDataVertexModel>
+	{
+		using TBaseVertexModel<DataType, FOutputDataVertexModel>::TBaseVertexModel;
+
+		/** Return the vertex type name (not to be confused with the data type name).
+		 *
+		 * @return Name of this vertex type.
+		 */
+		virtual const FName& GetVertexTypeName() const override
+		{
+			static const FString VertexTypeString = TEXT("OutputDataVertex_") + GetMetasoundDataTypeName<DataType>().ToString();
+			static const FName VertexTypeName(*VertexTypeString);
+
+			return VertexTypeName;
+		}
+
+		/** Create a clone of this TOutputDataVertexModel*/
+		virtual TUniquePtr<FOutputDataVertexModel> Clone() const override
+		{
+			return MakeUnique<TOutputDataVertexModel<DataType>>(this->VertexName, this->Description);
+		}
+	};
+
+	/** FOutputDataVertex represents an input vertex to an interface. It uses
+	 * a FOutputDataVertexModel to determine its behavior.
+	 */
+	class METASOUNDGRAPHCORE_API FOutputDataVertex
+	{
+		using FEmptyVertexModel = TOutputDataVertexModel<void>;
+
+		public:
+
+			/** Construct an FOutputDataVertex with a given vertex model.
+			 *
+			 * @param InVertexModel - A model subclass of FOutputDataVertexModel.
+			 */
+			template<typename VertexModelType>
+			FOutputDataVertex(const VertexModelType& InVertexModel)
+			:	VertexModel(InVertexModel.Clone())
+			{
+				static_assert(TIsDerivedFrom<VertexModelType, FOutputDataVertexModel>::Value, "Vertex implementation must be a subclass of FOutputDataVertexModel");
+
+				if (!VertexModel.IsValid())
+				{
+					VertexModel = MakeUnique<FEmptyVertexModel>(TEXT(""), FText::GetEmpty());
+				}
+			}
+
+			/** Construct an empty FInputDataVertex. */
+			FOutputDataVertex();
+
+			/** Copy constructor */
+			FOutputDataVertex(const FOutputDataVertex& InOther);
+
+			/** Assignment operator. */
+			FOutputDataVertex& operator=(const FOutputDataVertex& InOther);
+
+			/** Name of vertex. */
+			const FString& GetVertexName() const;
+
+			/** Type name of data reference. */
+			const FName& GetDataTypeName() const;
+
+			/** Description of the vertex. */
+			const FText& GetDescription() const;
+
+			/** Test if a IDataReference matches the DataType of this vertex.
+			 *
+			 * @param InReference - Data reference object.
+			 *
+			 * @return True if the types are equal, false otherwise.
+			 */
+			bool IsReferenceOfSameType(const IDataReference& InReference) const;
+
+			/** Return the vertex type name (not to be confused with the data type name).
+			 *
+			 * @return Name of this vertex type.
+			 */
+			const FName& GetVertexTypeName() const;
+
+			friend bool METASOUNDGRAPHCORE_API operator==(const FOutputDataVertex& InLHS, const FOutputDataVertex& InRHS);
+			friend bool METASOUNDGRAPHCORE_API operator!=(const FOutputDataVertex& InLHS, const FOutputDataVertex& InRHS);
+
+		private:
+
+			TUniquePtr<FOutputDataVertexModel> VertexModel;
+	};
+
+
+	/** Create an FDataVertexKey from a FDataVertex. */
+	FORCEINLINE FDataVertexKey MakeDataVertexKey(const FInputDataVertex& InVertex)
+	{
+		return InVertex.GetVertexName();
+	}
+
+	/** Create an FDataVertexKey from a FDataVertex. */
+	FORCEINLINE FDataVertexKey MakeDataVertexKey(const FOutputDataVertex& InVertex)
+	{
+		return InVertex.GetVertexName();
+	}
+
+	/** TVertexInterfaceGroups encapsulates multiple related data vertices. It 
+	 * requires that each vertex in the group have a unique FDataVertexKey.
+	 */
+	template<typename VertexType>
+	class TVertexInterfaceGroup
+	{
+			static_assert(std::is_base_of<FInputDataVertex, VertexType>::value || std::is_base_of<FOutputDataVertex, VertexType>::value, "VertexType must be derived from either FInputDataVertex or FOutputDataVertex");
+
+			using FContainerType = TSortedMap<FDataVertexKey, VertexType>;
+
+			// Required for end of recursion.
+			static void CopyInputs(FContainerType& InStorage)
+			{
+			}
+
+			// Recursive call used to unpack a template parameter pack. This stores
+			// each object in the template parameter pack into the groups internal
+			// storage.
+			//
+			// Assume that InputType is either a sublcass of FDataVertexModel
+			// and that VertexType is either FInputDataVertex or FOutputDataVertex
+			template<typename VertexModelType, typename... RemainingVertexModelTypes>
+			static void CopyInputs(FContainerType& InStorage, const VertexModelType& InInput, const RemainingVertexModelTypes&... InRemainingInputs)
+			{
+				// Create vertex out of vertex model
+				VertexType Vertex(InInput);
+
+				InStorage.Add(MakeDataVertexKey(Vertex), Vertex);
+				CopyInputs(InStorage, InRemainingInputs...);
+			}
+
+		public:
+
+			using RangedForConstIteratorType = typename TSortedMap<FDataVertexKey, VertexType>::RangedForConstIteratorType;
+
+			/** TVertexInterfaceGroup constructor with variadic list of vertex
+			 * models.
+			 */
+			template<typename... VertexModelTypes>
+			TVertexInterfaceGroup(VertexModelTypes&&... InVertexModels)
+			{
+				CopyInputs(Vertices, Forward<VertexModelTypes>(InVertexModels)...);
+			}
+
+			/** Add a vertex model to the group. */
+			template<typename VertexModelType, typename = typename TEnableIf<TIsDerivedFrom<VertexModelType, FDataVertexModel>::Value>::Type>
+			void Add(const VertexModelType& InVertexModel)
+			{
+				VertexType Vertex(InVertexModel);
+				Add(Vertex);
+			}
+
+			/** Add a vertex to the group. */
+			void Add(const VertexType& InVertex)
+			{
+				Vertices.Add(MakeDataVertexKey(InVertex), InVertex);
+			}
+
+			/** Remove a vertex by key. */
+			bool Remove(const FDataVertexKey& InKey)
+			{
+				int32 NumRemoved = Vertices.Remove();
+				return (NumRemoved > 0);
+			}
+
+			/** Returns true if the group contains a vertex with a matching key. */
+			bool Contains(const FDataVertexKey& InKey) const
+			{
+				return Vertices.Contains(InKey);
+			}
+
+			/** Return the vertex for a given vertex key. */
+			const VertexType& operator[](const FDataVertexKey& InKey) const
+			{
+				return Vertices[InKey];
+			}
+
+			/** Iterator for ranged for loops. */
+			RangedForConstIteratorType begin() const
+			{
+				return Vertices.begin();
+			}
+
+			/** Iterator for ranged for loops. */
+			RangedForConstIteratorType end() const
+			{
+				return Vertices.end();
+			}
+
+			/** Returns the number of vertices in the group. */
+			int32 Num() const
+			{
+				return Vertices.Num();
+			}
+
+			/** Compare whether two vertex groups are equal. */
+			friend bool operator==(const TVertexInterfaceGroup<VertexType>& InLHS, const TVertexInterfaceGroup<VertexType>& InRHS)
+			{
+				return InLHS.Vertices == InRHS.Vertices;
+			}
+
+			/** Compare whether two vertex groups are unequal. */
+			friend bool operator!=(const TVertexInterfaceGroup<VertexType>& InLHS, const TVertexInterfaceGroup<VertexType>& InRHS)
+			{
+				return !(InLHS == InRHS);
+			}
+
+		private:
+
+			FContainerType Vertices;
+	};
+
+	
+	/** FInputVertexInterface is a TVertexInterfaceGroup which holds FInputDataVertexes. */
+	typedef TVertexInterfaceGroup<FInputDataVertex> FInputVertexInterface;
+
+	/** FOutputVertexInterface is a TVertexInterfaceGroup which holds FOutputDataVertexes. */
+	typedef TVertexInterfaceGroup<FOutputDataVertex> FOutputVertexInterface;
+
+	/** FVertexInterface provides access to a collection of input and output vertex
+	 * interfaces. 
+	 */
+	class METASOUNDGRAPHCORE_API FVertexInterface
+	{
+		public:
+
+			/** Default constructor. */
+			FVertexInterface() = default;
+
+			/** Construct with an input and output interface. */
+			FVertexInterface(const FInputVertexInterface& InInputs, const FOutputVertexInterface& InOutputs);
+
+			/** Return the input interface. */
+			const FInputVertexInterface& GetInputInterface() const;
+
+			/** Return the input interface. */
+			FInputVertexInterface& GetInputInterface();
+
+			/** Return an input vertex. */
+			const FInputDataVertex& GetInputVertex(const FDataVertexKey& InKey) const;
+
+			/** Returns true if an input vertex with the given key exists. */
+			bool ContainsInputVertex(const FDataVertexKey& InKey) const;
+
+			/** Return the output interface. */
+			const FOutputVertexInterface& GetOutputInterface() const;
+
+			/** Return the output interface. */
+			FOutputVertexInterface& GetOutputInterface();
+
+			/** Return an output vertex. */
+			const FOutputDataVertex& GetOutputVertex(const FDataVertexKey& InKey) const;
+
+			/** Returns true if an outptu vertex with the given key exists. */
+			bool ContainsOutputVertex(const FDataVertexKey& InKey) const;
+
+			/** Test for equality between two interfaces. */
+			friend bool METASOUNDGRAPHCORE_API operator==(const FVertexInterface& InLHS, const FVertexInterface& InRHS);
+
+			/** Test for inequality between two interfaces. */
+			friend bool METASOUNDGRAPHCORE_API operator!=(const FVertexInterface& InLHS, const FVertexInterface& InRHS);
+
+		private:
+
+			FInputVertexInterface InputInterface;
+			FOutputVertexInterface OutputInterface;
+	};
+}
diff --git a/Engine/Plugins/Runtime/Metasound/Source/MetasoundStandardNodes/Private/MetasoundADSRNode.cpp b/Engine/Plugins/Runtime/Metasound/Source/MetasoundStandardNodes/Private/MetasoundADSRNode.cpp
--- a/Engine/Plugins/Runtime/Metasound/Source/MetasoundStandardNodes/Private/MetasoundADSRNode.cpp
+++ b/Engine/Plugins/Runtime/Metasound/Source/MetasoundStandardNodes/Private/MetasoundADSRNode.cpp
@@ -7,110 +7,33 @@
 #include "MetasoundNodeRegistrationMacro.h"
 #include "DSP/Envelope.h"
 
-#define LOCTEXT_NAMESPACE "MetasoundADSRNode"
-
-
+#define LOCTEXT_NAMESPACE "MetasoundStandardNodes"
 
 namespace Metasound
 {
 	METASOUND_REGISTER_NODE(FADSRNode)
 
-	struct FADSRDataReferences
-	{
-		FFloatTimeReadRef Attack;
-		FFloatTimeReadRef Decay;
-		FFloatTimeReadRef Sustain;
-		FFloatTimeReadRef Release;
-	};
-
 	class FADSROperator : public TExecutableOperator<FADSROperator>
 	{
 		public:
-			FADSROperator(const FOperatorSettings& InSettings, const FBopReadRef& InBop, const FADSRDataReferences& InADSRData)
-			:	OperatorSettings(InSettings)
-			,	Bop(InBop)
-			,	StopEnvelopePos(-1)
-			,	ADSRDataReferences(InADSRData)
-			,	EnvelopeBuffer(FAudioBufferWriteRef::CreateNew(InSettings.GetNumFramesPerBlock()))
+			struct FADSRReferences
 			{
-				check(EnvelopeBuffer->Num() == InSettings.GetNumFramesPerBlock());
-
-				OutputDataReferences.AddDataReadReference(TEXT("Envelope"), FAudioBufferReadRef(EnvelopeBuffer));
-
-				Envelope.Init(OperatorSettings.GetSampleRate());
-			}
+				FFloatTimeReadRef Attack;
+				FFloatTimeReadRef Decay;
+				FFloatTimeReadRef Sustain;
+				FFloatTimeReadRef Release;
+			};
 
-			virtual const FDataReferenceCollection& GetInputs() const override
-			{
-				return InputDataReferences;
-			}
+			static const FNodeInfo& GetNodeInfo();
+			static FVertexInterface DeclareVertexInterface();
+			static TUniquePtr<IOperator> CreateOperator(const FCreateOperatorParams& InParams, FBuildErrorArray& OutErrors);
 
-			virtual const FDataReferenceCollection& GetOutputs() const override
-			{
-				return OutputDataReferences;
-			}
 
-			void Execute()
-			{
-				Envelope.SetAttackTime(ADSRDataReferences.Attack->GetMilliseconds());
-				Envelope.SetDecayTime(ADSRDataReferences.Decay->GetMilliseconds());
-				Envelope.SetReleaseTime(ADSRDataReferences.Release->GetMilliseconds());
+			FADSROperator(const FOperatorSettings& InSettings, const FBopReadRef& InBop, const FADSRReferences& InADSRData);
 
-				float ADSMilliseconds = ADSRDataReferences.Attack->GetMilliseconds();
-				ADSMilliseconds += ADSRDataReferences.Decay->GetMilliseconds();
-				ADSMilliseconds += ADSRDataReferences.Sustain->GetMilliseconds();
-
-				// TODO: Bops need to be sorted.
-				// TODO: Make convenient way to do bop loops like this since it might happen alot.
-				// TODO: Bops need to be done on 4-sample boundaries to allow for SIMD.
-			
-
-				int32 BopIndex = 0;
-				int32 StartPos = 0;
-				int32 NextBop = BopIndex < Bop->Num() ? (*Bop)[BopIndex] : OperatorSettings.GetNumFramesPerBlock() + 1;
-				int32 EndPos = FMath::Min(OperatorSettings.GetNumFramesPerBlock(), NextBop);
-				
-				if (StopEnvelopePos > 0)
-				{
-					StopEnvelopePos -= OperatorSettings.GetNumFramesPerBlock();
-					EndPos = FMath::Min(StopEnvelopePos, EndPos);
-				}
-
-				float* EnvelopeData = EnvelopeBuffer->GetData();
-
-				// TODO: GetNumFramesPerBlock() need to be on 4 sample boundaries to allow for SIMD.
-				// TODO: just make GetNumFramesPerBlock() an int32 instead of a uint32
-				while (StartPos < (int32)OperatorSettings.GetNumFramesPerBlock())
-				{
-					int32 i = StartPos;
-					for (; i < EndPos; i++)
-					{
-						EnvelopeData[i] = Envelope.Generate();
-					}
-
-					StartPos = EndPos;
-					
-					if (EndPos == StopEnvelopePos)
-					{
-						// Stop sustain, start release.
-						Envelope.Stop();
-						StopEnvelopePos = -1;
-						EndPos = FMath::Min((int32)OperatorSettings.GetNumFramesPerBlock(), NextBop);
-					}
-
-					if (EndPos == NextBop)
-					{
-						// Trigger next envelope. 
-						Envelope.Start();
-						StopEnvelopePos = EndPos + FMath::RoundToInt(OperatorSettings.GetSampleRate() * 0.001f * ADSMilliseconds);
-
-						BopIndex++;
-
-						NextBop = BopIndex < Bop->Num() ? (*Bop)[BopIndex] : OperatorSettings.GetNumFramesPerBlock() + 1;
-						EndPos = FMath::Min(StopEnvelopePos, FMath::Min((int32)OperatorSettings.GetNumFramesPerBlock(), NextBop));
-					}
-				}
-			}
+			virtual const FDataReferenceCollection& GetInputs() const override;
+			virtual const FDataReferenceCollection& GetOutputs() const override;
+			void Execute();
 
 		private:
 			const FOperatorSettings OperatorSettings;
@@ -120,55 +43,158 @@ namespace Metasound
 
 			FBopReadRef Bop;
 			int32 StopEnvelopePos;
-			FADSRDataReferences ADSRDataReferences;
+			FADSRReferences ADSRReferences;
 			FAudioBufferWriteRef EnvelopeBuffer;
 
 			FDataReferenceCollection OutputDataReferences;
 			FDataReferenceCollection InputDataReferences;
 	};
 
-	const FName FADSRNode::ClassName = FName(TEXT("ADSR"));
+	FADSROperator::FADSROperator(const FOperatorSettings& InSettings, const FBopReadRef& InBop, const FADSRReferences& InADSRData)
+	:	OperatorSettings(InSettings)
+	,	Bop(InBop)
+	,	StopEnvelopePos(-1)
+	,	ADSRReferences(InADSRData)
+	,	EnvelopeBuffer(FAudioBufferWriteRef::CreateNew(InSettings.GetNumFramesPerBlock()))
+	{
+		check(EnvelopeBuffer->Num() == InSettings.GetNumFramesPerBlock());
+
+		OutputDataReferences.AddDataReadReference(TEXT("Envelope"), FAudioBufferReadRef(EnvelopeBuffer));
+
+		Envelope.Init(OperatorSettings.GetSampleRate());
+	}
 
-	TUniquePtr<IOperator> FADSRNode::FOperatorFactory::CreateOperator(const INode& InNode, const FOperatorSettings& InOperatorSettings, const FDataReferenceCollection& InInputDataReferences, TArray<TUniquePtr<IOperatorBuildError>>& OutErrors) 
+	const FDataReferenceCollection& FADSROperator::GetInputs() const
 	{
-		const FADSRNode& ADSRNode = static_cast<const FADSRNode&>(InNode);
+		return InputDataReferences;
+	}
 
-		FBopReadRef Bop = FBopReadRef::CreateNew();
+	const FDataReferenceCollection& FADSROperator::GetOutputs() const
+	{
+		return OutputDataReferences;
+	}
 
-		// TODO: Could no-op this if the bop is not connected.
-		SetReadableRefIfInCollection(TEXT("Bop"), InInputDataReferences, Bop);
+	void FADSROperator::Execute()
+	{
+		Envelope.SetAttackTime(ADSRReferences.Attack->GetMilliseconds());
+		Envelope.SetDecayTime(ADSRReferences.Decay->GetMilliseconds());
+		Envelope.SetReleaseTime(ADSRReferences.Release->GetMilliseconds());
+
+		float ADSRMilliseconds = ADSRReferences.Attack->GetMilliseconds();
+		ADSRMilliseconds += ADSRReferences.Decay->GetMilliseconds();
+		ADSRMilliseconds += ADSRReferences.Sustain->GetMilliseconds();
+
+		// TODO: Bops need to be sorted.
+		// TODO: Make convenient way to do bop loops like this since it might happen alot.
+		// TODO: Bops need to be done on 4-sample boundaries to allow for SIMD.
+	
+
+		int32 BopIndex = 0;
+		int32 StartPos = 0;
+		int32 NextBop = BopIndex < Bop->Num() ? (*Bop)[BopIndex] : OperatorSettings.GetNumFramesPerBlock() + 1;
+		int32 EndPos = FMath::Min(OperatorSettings.GetNumFramesPerBlock(), NextBop);
+		
+		if (StopEnvelopePos > 0)
+		{
+			StopEnvelopePos -= OperatorSettings.GetNumFramesPerBlock();
+			EndPos = FMath::Min(StopEnvelopePos, EndPos);
+		}
+
+		float* EnvelopeData = EnvelopeBuffer->GetData();
+
+		// TODO: GetNumFramesPerBlock() need to be on 4 sample boundaries to allow for SIMD.
+		// TODO: just make GetNumFramesPerBlock() an int32 instead of a uint32
+		while (StartPos < (int32)OperatorSettings.GetNumFramesPerBlock())
+		{
+			int32 i = StartPos;
+			for (; i < EndPos; i++)
+			{
+				EnvelopeData[i] = Envelope.Generate();
+			}
 
-		FADSRDataReferences ADSRDataReferences = 
+			StartPos = EndPos;
+			
+			if (EndPos == StopEnvelopePos)
 			{
-				FFloatTimeReadRef::CreateNew(ADSRNode.GetDefaultAttackMs(), ETimeResolution::Milliseconds),
-				FFloatTimeReadRef::CreateNew(ADSRNode.GetDefaultDecayMs(), ETimeResolution::Milliseconds),
-				FFloatTimeReadRef::CreateNew(ADSRNode.GetDefaultSustainMs(), ETimeResolution::Milliseconds),
-				FFloatTimeReadRef::CreateNew(ADSRNode.GetDefaultReleaseMs(), ETimeResolution::Milliseconds)
-			};
+				// Stop sustain, start release.
+				Envelope.Stop();
+				StopEnvelopePos = -1;
+				EndPos = FMath::Min((int32)OperatorSettings.GetNumFramesPerBlock(), NextBop);
+			}
+
+			if (EndPos == NextBop)
+			{
+				// Trigger next envelope. 
+				Envelope.Start();
+				StopEnvelopePos = EndPos + FMath::RoundToInt(OperatorSettings.GetSampleRate() * 0.001f * ADSRMilliseconds);
 
-		// TODO: If none of these are connected, could pregenerate ADSR envelope and return a different operator. 
-		SetReadableRefIfInCollection(TEXT("Attack"), InInputDataReferences, ADSRDataReferences.Attack);
-		SetReadableRefIfInCollection(TEXT("Decay"), InInputDataReferences, ADSRDataReferences.Decay);
-		SetReadableRefIfInCollection(TEXT("Sustain"), InInputDataReferences, ADSRDataReferences.Sustain);
-		SetReadableRefIfInCollection(TEXT("Release"), InInputDataReferences, ADSRDataReferences.Release);
+				BopIndex++;
 
-		return MakeUnique<FADSROperator>(InOperatorSettings, Bop, ADSRDataReferences);
+				NextBop = BopIndex < Bop->Num() ? (*Bop)[BopIndex] : OperatorSettings.GetNumFramesPerBlock() + 1;
+				EndPos = FMath::Min(StopEnvelopePos, FMath::Min((int32)OperatorSettings.GetNumFramesPerBlock(), NextBop));
+			}
+		}
+	}
+
+	const FNodeInfo& FADSROperator::GetNodeInfo()
+	{
+		static const FNodeInfo Info = {
+			FName(TEXT("ADSR")),
+			LOCTEXT("Metasound_ADSRNodeDescription", "Emits an ADSR (Attack, decay, sustain, & release) envelope when bopped."),
+			PluginAuthor,
+			PluginNodeMissingPrompt
+		};
+
+		return Info;
+	}
+
+	FVertexInterface FADSROperator::DeclareVertexInterface()
+	{
+		static const FVertexInterface Interface(
+			FInputVertexInterface(
+				TInputDataVertexModel<FBop>(TEXT("Bop"), LOCTEXT("BopTooltip", "Trigger for envelope.")),
+				TInputDataVertexModel<FFloatTime>(TEXT("Attack"), LOCTEXT("AttackTooltip", "Attack time in milliseconds.")),
+				TInputDataVertexModel<FFloatTime>(TEXT("Decay"), LOCTEXT("DecayTooltip", "Decay time in milliseconds.")),
+				TInputDataVertexModel<FFloatTime>(TEXT("Sustain"), LOCTEXT("SustainTooltip", "Sustain time in milliseconds.")),
+				TInputDataVertexModel<FFloatTime>(TEXT("Release"), LOCTEXT("ReleaseTooltip", "Release time in milliseconds."))
+			),
+			FOutputVertexInterface(
+				TOutputDataVertexModel<FAudioBuffer>(TEXT("Envelope"), LOCTEXT("EnvelopeTooltip", "The output envelope"))
+			)
+		);
+
+		return Interface;
+	}
+
+	TUniquePtr<IOperator> FADSROperator::CreateOperator(const FCreateOperatorParams& InParams, FBuildErrorArray& OutErrors)
+	{
+		const FADSRNode& ADSRNode = static_cast<const FADSRNode&>(InParams.Node);
+
+
+		const FDataReferenceCollection& InputCollection = InParams.InputDataReferences;
+
+		// TODO: Could no-op this if the bop is not connected.
+		FBopReadRef Bop = InputCollection.GetDataReadReferenceOrConstruct<FBop>(TEXT("Bop"));
+
+		// TODO: If none of these are connected, could pre-generate ADSR envelope and return a different operator. 
+		FADSRReferences ADSRReferences = 
+			{
+				InputCollection.GetDataReadReferenceOrConstruct<FFloatTime>(TEXT("Attack"), ADSRNode.GetDefaultAttackMs(), ETimeResolution::Milliseconds),
+				InputCollection.GetDataReadReferenceOrConstruct<FFloatTime>(TEXT("Decay"), ADSRNode.GetDefaultDecayMs(), ETimeResolution::Milliseconds),
+				InputCollection.GetDataReadReferenceOrConstruct<FFloatTime>(TEXT("Sustain"), ADSRNode.GetDefaultSustainMs(), ETimeResolution::Milliseconds),
+				InputCollection.GetDataReadReferenceOrConstruct<FFloatTime>(TEXT("Release"), ADSRNode.GetDefaultReleaseMs(), ETimeResolution::Milliseconds)
+			};
+	
+		return MakeUnique<FADSROperator>(InParams.OperatorSettings, Bop, ADSRReferences);
 	}
 
 	FADSRNode::FADSRNode(const FString& InName, float InDefaultAttackMs, float InDefaultDecayMs, float InDefaultSustainMs, float InDefaultReleaseMs)
-	:	FNode(InName)
+	:	FNodeFacade(InName, TFacadeOperatorClass<FADSROperator>())
 	,	DefaultAttackMs(InDefaultAttackMs)
 	,	DefaultDecayMs(InDefaultDecayMs)
 	,	DefaultSustainMs(InDefaultSustainMs)
 	,	DefaultReleaseMs(InDefaultReleaseMs)
 	{
-		AddInputDataVertex<FBop>(TEXT("Bop"), LOCTEXT("BopTooltip", "Trigger for envelope."));
-		AddInputDataVertex<FFloatTime>(TEXT("Attack"), LOCTEXT("AttackTooltip", "Attack time in milliseconds."));
-		AddInputDataVertex<FFloatTime>(TEXT("Decay"), LOCTEXT("DecayTooltip", "Decay time in milliseconds."));
-		AddInputDataVertex<FFloatTime>(TEXT("Sustain"), LOCTEXT("SustainTooltip", "Sustain time in milliseconds."));
-		AddInputDataVertex<FFloatTime>(TEXT("Release"), LOCTEXT("ReleaseTooltip", "Release time in milliseconds."));
-
-		AddOutputDataVertex<FAudioBuffer>(TEXT("Envelope"), LOCTEXT("EnvelopeTooltip", "The output envelope"));
 	}
 
 	FADSRNode::FADSRNode(const FNodeInitData& InitData)
@@ -200,16 +226,6 @@ namespace Metasound
 	{
 		return DefaultReleaseMs;
 	}
-
-	const FName& FADSRNode::GetClassName() const
-	{
-		return ClassName;
-	}
-
-	IOperatorFactory& FADSRNode::GetDefaultOperatorFactory() 
-	{
-		return Factory;
-	}
 }
 
 #undef LOCTEXT_NAMESPACE
diff --git a/Engine/Plugins/Runtime/Metasound/Source/MetasoundStandardNodes/Private/MetasoundAudioMultiplyNode.cpp b/Engine/Plugins/Runtime/Metasound/Source/MetasoundStandardNodes/Private/MetasoundAudioMultiplyNode.cpp
--- a/Engine/Plugins/Runtime/Metasound/Source/MetasoundStandardNodes/Private/MetasoundAudioMultiplyNode.cpp
+++ b/Engine/Plugins/Runtime/Metasound/Source/MetasoundStandardNodes/Private/MetasoundAudioMultiplyNode.cpp
@@ -3,6 +3,7 @@
 #include "MetasoundAudioMultiplyNode.h"
 #include "MetasoundExecutableOperator.h"
 #include "MetasoundDataReferenceTypes.h"
+#include "MetasoundFacade.h"
 #include "MetasoundNodeRegistrationMacro.h"
 #include "DSP/BufferVectorOperations.h"
 
@@ -10,11 +11,15 @@
 
 namespace Metasound
 {
-	METASOUND_REGISTER_NODE(FAudioMultiplyNode);
-
 	class FAudioMultiplyOperator : public TExecutableOperator<FAudioMultiplyOperator>
 	{
 		public:
+			static const FNodeInfo& GetNodeInfo();
+
+			static FVertexInterface DeclareVertexInterface();
+
+			static TUniquePtr<IOperator> CreateOperator(const FCreateOperatorParams& InParams, FBuildErrorArray& OutErrors);
+
 			FAudioMultiplyOperator(const FOperatorSettings& InSettings, const FAudioBufferReadRef& InBuffer1, const FAudioBufferReadRef& InBuffer2)
 			:	OperatorSettings(InSettings)
 			,	InputBuffer1(InBuffer1)
@@ -58,63 +63,68 @@ namespace Metasound
 			FDataReferenceCollection InputDataReferences;
 	};
 
-	const FName FAudioMultiplyNode::ClassName = FName(TEXT("AudioMultiply"));
+	FAudioMultiplyNode::FAudioMultiplyNode(const FString& InInstanceName)
+	:	FNodeFacade(InInstanceName, TFacadeOperatorClass<FAudioMultiplyOperator>())
+	{
+	}
+
+	FAudioMultiplyNode::FAudioMultiplyNode(const FNodeInitData& InInitData)
+	:	FAudioMultiplyNode(InInitData.InstanceName)
+	{
+	}
+
+	FVertexInterface FAudioMultiplyOperator::DeclareVertexInterface()
+	{
+		static const FVertexInterface Interface(
+			FInputVertexInterface(
+				TInputDataVertexModel<FAudioBuffer>(TEXT("InputBuffer1"), LOCTEXT("InputBuffer1Tooltip", "Input buffer to multiply.")),
+				TInputDataVertexModel<FAudioBuffer>(TEXT("InputBuffer2"), LOCTEXT("InputBuffer2Tooltip", "Input buffer to multiply."))
+			),
+			FOutputVertexInterface(
+				TOutputDataVertexModel<FAudioBuffer>(TEXT("Audio"), LOCTEXT("OutpuBufferTooltip", "The output audio."))
+			)
+		);
+
+		return Interface;
+	}
 
-	TUniquePtr<IOperator> FAudioMultiplyNode::FOperatorFactory::CreateOperator(const INode& InNode, const FOperatorSettings& InOperatorSettings, const FDataReferenceCollection& InInputDataReferences, TArray<TUniquePtr<IOperatorBuildError>>& OutErrors) 
+	TUniquePtr<IOperator> FAudioMultiplyOperator::CreateOperator(const FCreateOperatorParams& InParams, FBuildErrorArray& OutErrors)
 	{
-		const FAudioMultiplyNode& AudioMultiplyNode = static_cast<const FAudioMultiplyNode&>(InNode);
+		const FAudioMultiplyNode& AudioMultiplyNode = static_cast<const FAudioMultiplyNode&>(InParams.Node);
 
-		FAudioBufferReadRef InputBuffer1 = FAudioBufferReadRef::CreateNew();
-		FAudioBufferReadRef InputBuffer2 = FAudioBufferReadRef::CreateNew();
 
-		FAudioBuffer Ones(InOperatorSettings.GetNumFramesPerBlock());
+		FAudioBuffer Ones(InParams.OperatorSettings.GetNumFramesPerBlock());
 
 		// initialize default array to all ones. 
 		float* Data = Ones.GetData();
-		for (int32 i = 0; i < InOperatorSettings.GetNumFramesPerBlock(); i++)
+		for (int32 i = 0; i < InParams.OperatorSettings.GetNumFramesPerBlock(); i++)
 		{
 			Data[i] = 1.f;
 		}
 
-		if (!SetReadableRefIfInCollection(TEXT("InputBuffer1"), InInputDataReferences, InputBuffer1))
-		{
-			InputBuffer1 = FAudioBufferReadRef::CreateNew(Ones);
-		}
+		// TODO: return a helper that uses implicit conversion so that return types can be deduced without templates.
 
-		if (!SetReadableRefIfInCollection(TEXT("InputBuffer2"), InInputDataReferences, InputBuffer2))
-		{
-			InputBuffer2 = FAudioBufferReadRef::CreateNew(Ones);
-		}
-
-		return MakeUnique<FAudioMultiplyOperator>(InOperatorSettings, InputBuffer1, InputBuffer2);
-	}
+		const FDataReferenceCollection& InputCol = InParams.InputDataReferences;
 
-	FAudioMultiplyNode::FAudioMultiplyNode(const FString& InName)
-	:	FNode(InName)
-	{
-		AddInputDataVertex<FAudioBuffer>(TEXT("InputBuffer1"), LOCTEXT("InputBuffer1Tooltip", "Input buffer to multiply."));
-		AddInputDataVertex<FAudioBuffer>(TEXT("InputBuffer2"), LOCTEXT("InputBuffer2Tooltip", "Input buffer to multiply."));
+		FAudioBufferReadRef InputBuffer1 = InputCol.GetDataReadReferenceOrConstruct<FAudioBuffer>(TEXT("InputBuffer1"), Ones);
+		FAudioBufferReadRef InputBuffer2 = InputCol.GetDataReadReferenceOrConstruct<FAudioBuffer>(TEXT("InputBuffer2"), Ones);
 
-		AddOutputDataVertex<FAudioBuffer>(TEXT("Audio"), LOCTEXT("OutpuBufferTooltip", "The output audio."));
+		return MakeUnique<FAudioMultiplyOperator>(InParams.OperatorSettings, InputBuffer1, InputBuffer2);
 	}
 
-	FAudioMultiplyNode::FAudioMultiplyNode(const FNodeInitData& InInitData)
-		: FAudioMultiplyNode(InInitData.InstanceName)
-	{}
-
-	FAudioMultiplyNode::~FAudioMultiplyNode()
+	const FNodeInfo& FAudioMultiplyOperator::GetNodeInfo()
 	{
-	}
-
-	const FName& FAudioMultiplyNode::GetClassName() const
-	{
-		return ClassName;
-	}
+		static const FNodeInfo Info = {
+			FName(TEXT("AudioMultiply")),
+			LOCTEXT("NodeDescription", "Multiply two audio streams, sample by sample."),
+			PluginAuthor,
+			PluginNodeMissingPrompt
+		};
+
+		return Info;
+	};
 
-	IOperatorFactory& FAudioMultiplyNode::GetDefaultOperatorFactory() 
-	{
-		return Factory;
-	}
+	METASOUND_REGISTER_NODE(FAudioMultiplyNode);
 }
 
 #undef LOCTEXT_NAMESPACE //MetasoundAudioMultiplyNode
diff --git a/Engine/Plugins/Runtime/Metasound/Source/MetasoundStandardNodes/Private/MetasoundOscNode.cpp b/Engine/Plugins/Runtime/Metasound/Source/MetasoundStandardNodes/Private/MetasoundOscNode.cpp
--- a/Engine/Plugins/Runtime/Metasound/Source/MetasoundStandardNodes/Private/MetasoundOscNode.cpp
+++ b/Engine/Plugins/Runtime/Metasound/Source/MetasoundStandardNodes/Private/MetasoundOscNode.cpp
@@ -1,54 +1,34 @@
 // Copyright Epic Games, Inc. All Rights Reserved.
 
 #include "MetasoundOscNode.h"
+
+#include "MetasoundAudioBuffer.h"
 #include "MetasoundExecutableOperator.h"
-#include "MetasoundPrimitives.h"
+#include "MetasoundFacade.h"
+#include "MetasoundFrequency.h"
 #include "MetasoundNodeRegistrationMacro.h"
+#include "MetasoundPrimitives.h"
+#include "MetasoundVertex.h"
 
-#define LOCTEXT_NAMESPACE "MetasoundOscNode"
+#define LOCTEXT_NAMESPACE "MetasoundStandardNodes"
 
 namespace Metasound
 {
-	METASOUND_REGISTER_NODE(FOscNode)
-
 	class FOscOperator : public TExecutableOperator<FOscOperator>
 	{
 		public:
-			FOscOperator(const FOperatorSettings& InSettings, const FFrequencyReadRef& InFrequency)
-			:	OperatorSettings(InSettings)
-			,	TwoPi(2.f * PI)
-			,	Phase(0.f)
-			,	Frequency(InFrequency)
-			,	AudioBuffer(FAudioBufferWriteRef::CreateNew(InSettings))
-			{
-				check(AudioBuffer->Num() == InSettings.GetNumFramesPerBlock());
-
-				OutputDataReferences.AddDataReadReference(TEXT("Audio"), FAudioBufferReadRef(AudioBuffer));
-			}
-
-			virtual const FDataReferenceCollection& GetInputs() const override
-			{
-				return InputDataReferences;
-			}
-
-			virtual const FDataReferenceCollection& GetOutputs() const override
-			{
-				return OutputDataReferences;
-			}
-
-			void Execute()
-			{
-				const float PhaseDelta = Frequency->GetRadiansPerSample(OperatorSettings.GetSampleRate());
-				float* Data = AudioBuffer->GetData();
-
-				for (int32 i = 0; i < OperatorSettings.GetNumFramesPerBlock(); i++)
-				{
-					Data[i] = FMath::Sin(Phase);
-					Phase += PhaseDelta;
-				}
-
-				Phase -= FMath::FloorToFloat(Phase / TwoPi) * TwoPi;
-			}
+
+			static TUniquePtr<IOperator> CreateOperator(const FCreateOperatorParams& InParams, FBuildErrorArray& OutErrors);
+			static const FNodeInfo& GetNodeInfo();
+			static FVertexInterface DeclareVertexInterface();
+
+			FOscOperator(const FOperatorSettings& InSettings, const FFrequencyReadRef& InFrequency);
+
+			virtual const FDataReferenceCollection& GetInputs() const override;
+
+			virtual const FDataReferenceCollection& GetOutputs() const override;
+
+			void Execute();
 
 		private:
 			const FOperatorSettings OperatorSettings;
@@ -62,51 +42,96 @@ namespace Metasound
 			FDataReferenceCollection OutputDataReferences;
 	};
 
-	const FName FOscNode::ClassName = FName(TEXT("Osc"));
+	FOscOperator::FOscOperator(const FOperatorSettings& InSettings, const FFrequencyReadRef& InFrequency)
+	:	OperatorSettings(InSettings)
+	,	TwoPi(2.f * PI)
+	,	Phase(0.f)
+	,	Frequency(InFrequency)
+	,	AudioBuffer(FAudioBufferWriteRef::CreateNew(InSettings))
+	{
+		check(AudioBuffer->Num() == InSettings.GetNumFramesPerBlock());
+
+		OutputDataReferences.AddDataReadReference(TEXT("Audio"), FAudioBufferReadRef(AudioBuffer));
+	}
+
+	const FDataReferenceCollection& FOscOperator::GetInputs() const
+	{
+		return InputDataReferences;
+	}
 
-	TUniquePtr<IOperator> FOscNode::FOperatorFactory::CreateOperator(const INode& InNode, const FOperatorSettings& InOperatorSettings, const FDataReferenceCollection& InInputDataReferences, TArray<TUniquePtr<IOperatorBuildError>>& OutErrors) 
+	const FDataReferenceCollection& FOscOperator::GetOutputs() const
 	{
-		const FOscNode& OscNode = static_cast<const FOscNode&>(InNode);
-		FFrequencyReadRef Frequency = FFrequencyReadRef::CreateNew(OscNode.GetDefaultFrequency(), EFrequencyResolution::Hertz);
+		return OutputDataReferences;
+	}
 
-		if (InInputDataReferences.ContainsDataReadReference<FFrequency>(TEXT("Frequency")))
+	void FOscOperator::Execute()
+	{
+		const float PhaseDelta = Frequency->GetRadiansPerSample(OperatorSettings.GetSampleRate());
+		float* Data = AudioBuffer->GetData();
+
+		for (int32 i = 0; i < OperatorSettings.GetNumFramesPerBlock(); i++)
 		{
-			Frequency = InInputDataReferences.GetDataReadReference<FFrequency>(TEXT("Frequency"));
+			Data[i] = FMath::Sin(Phase);
+			Phase += PhaseDelta;
 		}
 
-		return MakeUnique<FOscOperator>(InOperatorSettings, Frequency);
+		Phase -= FMath::FloorToFloat(Phase / TwoPi) * TwoPi;
 	}
 
-	FOscNode::FOscNode(const FString& InName, float InDefaultFrequency)
-		:	FNode(InName)
-		,	DefaultFrequency(InDefaultFrequency)
+	const FNodeInfo& FOscOperator::GetNodeInfo()
 	{
-		AddInputDataVertex<FFrequency>(TEXT("Frequency"), LOCTEXT("FrequencyTooltip", "The frequency of oscillator."));
-		AddOutputDataVertex<FAudioBuffer>(TEXT("Audio"), LOCTEXT("AudioTooltip", "The output audio"));
+		static const FNodeInfo Info = {
+			FName(TEXT("Osc")),
+			LOCTEXT("Metasound_OscNodeDescription", "Emits an audio signal of a sinusoid."),
+			PluginAuthor,
+			PluginNodeMissingPrompt
+		};
+
+		return Info;
 	}
 
-	FOscNode::FOscNode(const FNodeInitData& InInitData)
-		: FOscNode(InInitData.InstanceName, 440.0f)
+	FVertexInterface FOscOperator::DeclareVertexInterface()
 	{
+		static const FVertexInterface Interface(
+			FInputVertexInterface(
+				TInputDataVertexModel<FFrequency>(TEXT("Frequency"), LOCTEXT("OscFrequencyDescription", "The frequency of oscillator."))
+			),
+			FOutputVertexInterface(
+				TOutputDataVertexModel<FAudioBuffer>(TEXT("Audio"), LOCTEXT("AudioTooltip", "The output audio"))
+			)
+		);
+
+		return Interface;
 	}
 
-	FOscNode::~FOscNode()
+
+	FOscNode::FOscNode(const FString& InName, float InDefaultFrequency)
+	:	FNodeFacade(InName, TFacadeOperatorClass<FOscOperator>())
+	,	DefaultFrequency(InDefaultFrequency)
 	{
 	}
 
-	float FOscNode::GetDefaultFrequency() const
+	FOscNode::FOscNode(const FNodeInitData& InInitData)
+		: FOscNode(InInitData.InstanceName, 440.0f)
 	{
-		return DefaultFrequency;
 	}
 
-	const FName& FOscNode::GetClassName() const
+	float FOscNode::GetDefaultFrequency() const
 	{
-		return ::Metasound::FOscNode::ClassName;
+		return DefaultFrequency;
 	}
 
-	IOperatorFactory& FOscNode::GetDefaultOperatorFactory() 
+
+	TUniquePtr<IOperator> FOscOperator::CreateOperator(const FCreateOperatorParams& InParams, FBuildErrorArray& OutErrors) 
 	{
-		return Factory;
+		const FOscNode& OscNode = static_cast<const FOscNode&>(InParams.Node);
+		const FDataReferenceCollection& InputCol = InParams.InputDataReferences;
+
+		FFrequencyReadRef Frequency = InputCol.GetDataReadReferenceOrConstruct<FFrequency>(TEXT("Frequency"), OscNode.GetDefaultFrequency(), EFrequencyResolution::Hertz);
+
+		return MakeUnique<FOscOperator>(InParams.OperatorSettings, Frequency);
 	}
+
+	METASOUND_REGISTER_NODE(FOscNode);
 }
 #undef LOCTEXT_NAMESPACE //MetasoundOscNode
diff --git a/Engine/Plugins/Runtime/Metasound/Source/MetasoundStandardNodes/Private/MetasoundPeriodicBopNode.cpp b/Engine/Plugins/Runtime/Metasound/Source/MetasoundStandardNodes/Private/MetasoundPeriodicBopNode.cpp
--- a/Engine/Plugins/Runtime/Metasound/Source/MetasoundStandardNodes/Private/MetasoundPeriodicBopNode.cpp
+++ b/Engine/Plugins/Runtime/Metasound/Source/MetasoundStandardNodes/Private/MetasoundPeriodicBopNode.cpp
@@ -3,16 +3,18 @@
 #include "MetasoundPeriodicBopNode.h"
 
 #include "CoreMinimal.h"
-#include "MetasoundNode.h"
-#include "MetasoundOperatorInterface.h"
+
 #include "MetasoundBuilderInterface.h"
-#include "MetasoundNodeInterface.h"
 #include "MetasoundDataReferenceCollection.h"
 #include "MetasoundExecutableOperator.h"
-#include "MetasoundPrimitives.h"
+#include "MetasoundFacade.h"
+#include "MetasoundNode.h"
+#include "MetasoundNodeInterface.h"
 #include "MetasoundNodeRegistrationMacro.h"
+#include "MetasoundOperatorInterface.h"
+#include "MetasoundPrimitives.h"
 
-#define LOCTEXT_NAMESPACE "MetasoundPeriodicBopNode"
+#define LOCTEXT_NAMESPACE "MetasoundStandardNodes"
 
 namespace Metasound
 {
@@ -21,47 +23,20 @@ namespace Metasound
 	class FPeriodicBopOperator : public TExecutableOperator<FPeriodicBopOperator>
 	{
 		public:
+			static const FNodeInfo& GetNodeInfo();
+			static FVertexInterface DeclareVertexInterface();
+			static TUniquePtr<IOperator> CreateOperator(const FCreateOperatorParams& InParams, FBuildErrorArray& OutErrors);
+
 			static constexpr float MinimumPeriodSeconds = 0.001f;
 			static constexpr float MinimumPeriodSamples = 20.f;
 
-			FPeriodicBopOperator(const FOperatorSettings& InSettings, const FFloatTimeReadRef& InPeriod)
-			:	OperatorSettings(InSettings)
-			,	Bop(FBopWriteRef::CreateNew())
-			,	Period(InPeriod)
-			,	ExecuteDurationInSamples(InSettings.GetNumFramesPerBlock())
-			,	SampleCountdown(0.f)
-			{
-				OutputDataReferences.AddDataReadReference(TEXT("Bop"), FBopReadRef(Bop));
-			}
-
-			virtual const FDataReferenceCollection& GetInputs() const override
-			{
-				return InputDataReferences;
-			}
-
-			virtual const FDataReferenceCollection& GetOutputs() const override
-			{
-				return OutputDataReferences;
-			}
-
-			void Execute()
-			{
-				// Advance internal counter to get rid of old bops.
-				Bop->Advance(OperatorSettings.GetNumFramesPerBlock());
-
-				float PeriodInSamples = FMath::Max(Period->GetSeconds(), MinimumPeriodSeconds) * OperatorSettings.GetSampleRate();
-
-				PeriodInSamples = FMath::Max(PeriodInSamples, MinimumPeriodSamples);
-
-				while ((SampleCountdown - ExecuteDurationInSamples) <= 0.f)
-				{
-					uint32 Frame = FMath::RoundToInt(SampleCountdown);
-					Bop->BopFrame(Frame);
-					SampleCountdown += PeriodInSamples;
-				}
-
-				SampleCountdown -= ExecuteDurationInSamples;
-			}
+			FPeriodicBopOperator(const FOperatorSettings& InSettings, const FFloatTimeReadRef& InPeriod);
+
+			virtual const FDataReferenceCollection& GetInputs() const override;
+
+			virtual const FDataReferenceCollection& GetOutputs() const override;
+
+			void Execute();
 
 		private:
 
@@ -76,54 +51,101 @@ namespace Metasound
 			FDataReferenceCollection OutputDataReferences;
 	};
 
+	FPeriodicBopOperator::FPeriodicBopOperator(const FOperatorSettings& InSettings, const FFloatTimeReadRef& InPeriod)
+	:	OperatorSettings(InSettings)
+	,	Bop(FBopWriteRef::CreateNew())
+	,	Period(InPeriod)
+	,	ExecuteDurationInSamples(InSettings.GetNumFramesPerBlock())
+	,	SampleCountdown(0.f)
+	{
+		OutputDataReferences.AddDataReadReference(TEXT("Bop"), FBopReadRef(Bop));
+	}
+
+	const FDataReferenceCollection& FPeriodicBopOperator::GetInputs() const
+	{
+		return InputDataReferences;
+	}
+
+	const FDataReferenceCollection& FPeriodicBopOperator::GetOutputs() const
+	{
+		return OutputDataReferences;
+	}
+
+	void FPeriodicBopOperator::Execute()
+	{
+		// Advance internal counter to get rid of old bops.
+		Bop->Advance(OperatorSettings.GetNumFramesPerBlock());
+
+		float PeriodInSamples = FMath::Max(Period->GetSeconds(), MinimumPeriodSeconds) * OperatorSettings.GetSampleRate();
+
+		PeriodInSamples = FMath::Max(PeriodInSamples, MinimumPeriodSamples);
+
+		while ((SampleCountdown - ExecuteDurationInSamples) <= 0.f)
+		{
+			uint32 Frame = FMath::RoundToInt(SampleCountdown);
+			Bop->BopFrame(Frame);
+			SampleCountdown += PeriodInSamples;
+		}
 
-	TUniquePtr<IOperator> FPeriodicBopNode::FOperatorFactory::CreateOperator(const INode& InNode, const FOperatorSettings& InOperatorSettings, const FDataReferenceCollection& InInputDataReferences, TArray<TUniquePtr<IOperatorBuildError>>& OutErrors)
+		SampleCountdown -= ExecuteDurationInSamples;
+	}
+
+
+	TUniquePtr<IOperator> FPeriodicBopOperator::CreateOperator(const FCreateOperatorParams& InParams, FBuildErrorArray& OutErrors)
 	{
-		const FPeriodicBopNode& PeriodicBopNode = static_cast<const FPeriodicBopNode&>(InNode);
+		const FPeriodicBopNode& PeriodicBopNode = static_cast<const FPeriodicBopNode&>(InParams.Node);
 
-		// Note- won't this cause the Period input to not update after we build the operator?
 		FFloatTimeReadRef Period = FFloatTimeReadRef::CreateNew(PeriodicBopNode.GetDefaultPeriodInSeconds(), ETimeResolution::Seconds);
 
-		if (InInputDataReferences.ContainsDataReadReference<FFloatTime>(TEXT("Period")))
+		if (InParams.InputDataReferences.ContainsDataReadReference<FFloatTime>(TEXT("Period")))
 		{
-			Period = InInputDataReferences.GetDataReadReference<FFloatTime>(TEXT("Period"));
+			Period = InParams.InputDataReferences.GetDataReadReference<FFloatTime>(TEXT("Period"));
 		}
 
-		return MakeUnique<FPeriodicBopOperator>(InOperatorSettings, Period);
+		return MakeUnique<FPeriodicBopOperator>(InParams.OperatorSettings, Period);
+	}
+
+	const FNodeInfo& FPeriodicBopOperator::GetNodeInfo()
+	{
+		static const FNodeInfo Info = {
+			FName(TEXT("PeriodicBop")),
+			LOCTEXT("Metasound_PeriodicBopNodeDescription", "Emits a bop periodically based on the period duration given."),
+			PluginAuthor,
+			PluginNodeMissingPrompt
+		};
+
+		return Info;
 	}
 
-	const FName FPeriodicBopNode::ClassName = FName(TEXT("PeriodicBop"));
+	FVertexInterface FPeriodicBopOperator::DeclareVertexInterface()
+	{
+		static const FVertexInterface Interface(
+			FInputVertexInterface(
+				TInputDataVertexModel<FFloatTime>(TEXT("Period"), LOCTEXT("PeriodTooltip", "The period of the bops in seconds."))
+			),
+			FOutputVertexInterface(
+				TOutputDataVertexModel<FBop>(TEXT("Bop"), LOCTEXT("BopTooltip", "The output bop"))
+			)
+		);
+
+		return Interface;
+	}
 
 	FPeriodicBopNode::FPeriodicBopNode(const FString& InName, float InDefaultPeriodInSeconds)
-	:	FNode(InName)
+	:	FNodeFacade(InName, TFacadeOperatorClass<FPeriodicBopOperator>())
 	,	DefaultPeriod(InDefaultPeriodInSeconds)
 	{
-		AddInputDataVertex<FFloatTime>(TEXT("Period"), LOCTEXT("PeriodTooltip", "The period of the bops in seconds."));
-		AddOutputDataVertex<FBop>(TEXT("Bop"), LOCTEXT("BopTooltip", "The output bop"));
 	}
 
 	FPeriodicBopNode::FPeriodicBopNode(const FNodeInitData& InInitData)
 		: FPeriodicBopNode(InInitData.InstanceName, 1.0f)
 	{}
 
-	FPeriodicBopNode::~FPeriodicBopNode()
-	{
-	}
-
 	float FPeriodicBopNode::GetDefaultPeriodInSeconds() const
 	{
 		return DefaultPeriod;
 	}
 
-	const FName& FPeriodicBopNode::GetClassName() const
-	{
-		return ClassName;
-	}
-
-	IOperatorFactory& FPeriodicBopNode::GetDefaultOperatorFactory()
-	{
-		return Factory;
-	}
 }
 
 #undef LOCTEXT_NAMESPACE //MetasoundPeriodicBopNode
diff --git a/Engine/Plugins/Runtime/Metasound/Source/MetasoundStandardNodes/Public/MetasoundADSRNode.h b/Engine/Plugins/Runtime/Metasound/Source/MetasoundStandardNodes/Public/MetasoundADSRNode.h
--- a/Engine/Plugins/Runtime/Metasound/Source/MetasoundStandardNodes/Public/MetasoundADSRNode.h
+++ b/Engine/Plugins/Runtime/Metasound/Source/MetasoundStandardNodes/Public/MetasoundADSRNode.h
@@ -3,11 +3,8 @@
 
 #include "CoreMinimal.h"
 #include "Internationalization/Text.h"
-#include "MetasoundNode.h"
-#include "MetasoundOperatorInterface.h"
-#include "MetasoundBuilderInterface.h"
-#include "MetasoundDataReferenceCollection.h"
 
+#include "MetasoundFacade.h"
 
 namespace Metasound
 {
@@ -15,28 +12,9 @@ namespace Metasound
 	 *
 	 *  Creates an Attack, Decay Sustain, Release audio processer node. 
 	 */
-	class METASOUNDSTANDARDNODES_API FADSRNode : public FNode
+	class METASOUNDSTANDARDNODES_API FADSRNode : public FNodeFacade
 	{
-		// The operator factory for this node.
-		class FOperatorFactory : public IOperatorFactory
-		{
-			public:
-				virtual TUniquePtr<IOperator> CreateOperator(const INode& InNode, const FOperatorSettings& InOperatorSettings, const FDataReferenceCollection& InInputDataReferences, TArray<TUniquePtr<IOperatorBuildError>>& OutErrors) override;
-
-				template<typename DataType>
-				void SetReadableRefIfInCollection(const FString& InName, const FDataReferenceCollection& InCollection, TDataReadReference<DataType>& DataRef)
-				{
-					// TODO: add a helper function to FDataReferenceCollection to do a SetDataIfContains.
-					if (InCollection.ContainsDataReadReference<DataType>(InName))
-					{
-						DataRef = InCollection.GetDataReadReference<DataType>(InName);
-					}
-				}
-		};
-
 		public:
-			/** Class name for FADSRNode */
-			static const FName ClassName;
 
 			/** FADSR node constructor.
 			 *
@@ -64,34 +42,10 @@ namespace Metasound
 			/** Return default attack time in milliseconds */
 			float GetDefaultReleaseMs() const;
 
-			/** Returns the type name of the FADSRNode */
-			virtual const FName& GetClassName() const override;
-
-			virtual const FText& GetDescription() const override
-			{
-				static const FText StaticDescription = NSLOCTEXT("MetasoundGraphCore", "Metasound_ADSRNodeDescription", "Emits an ADSR (Attack, decay, sustain, & release) envelope when bopped.");
-				return StaticDescription;
-			}
-
-			virtual const FText& GetAuthorName() const override
-			{
-				return PluginAuthor;
-			}
-
-			virtual const FText& GetPromptIfMissing() const override
-			{
-				return PluginNodeMissingPrompt;
-			}
-
-			/** Return a factory for building a metasound operator. */
-			virtual IOperatorFactory& GetDefaultOperatorFactory() override;
-
 		private:
 			float DefaultAttackMs;
 			float DefaultDecayMs;
 			float DefaultSustainMs;
 			float DefaultReleaseMs;
-
-			FOperatorFactory Factory;
 	};
 } // namespace Metasound
diff --git a/Engine/Plugins/Runtime/Metasound/Source/MetasoundStandardNodes/Public/MetasoundAudioMultiplyNode.h b/Engine/Plugins/Runtime/Metasound/Source/MetasoundStandardNodes/Public/MetasoundAudioMultiplyNode.h
--- a/Engine/Plugins/Runtime/Metasound/Source/MetasoundStandardNodes/Public/MetasoundAudioMultiplyNode.h
+++ b/Engine/Plugins/Runtime/Metasound/Source/MetasoundStandardNodes/Public/MetasoundAudioMultiplyNode.h
@@ -3,67 +3,17 @@
 
 #include "CoreMinimal.h"
 #include "Internationalization/Text.h"
-#include "MetasoundNode.h"
-#include "MetasoundOperatorInterface.h"
-#include "MetasoundBuilderInterface.h"
-#include "MetasoundDataReferenceCollection.h"
+
+#include "MetasoundFacade.h"
 
 
 namespace Metasound
 {
-	class METASOUNDSTANDARDNODES_API FAudioMultiplyNode : public FNode
+	class METASOUNDSTANDARDNODES_API FAudioMultiplyNode : public FNodeFacade
 	{
-		class FOperatorFactory : public IOperatorFactory
-		{
-			public:
-				virtual TUniquePtr<IOperator> CreateOperator(const INode& InNode, const FOperatorSettings& InOperatorSettings, const FDataReferenceCollection& InInputDataReferences, TArray<TUniquePtr<IOperatorBuildError>>& OutErrors) override;
-
-				template<typename ParamType>
-				bool SetReadableRefIfInCollection(const FString& InParamName, const FDataReferenceCollection& InCollection, TDataReadReference<ParamType>& ParamRef)
-				{
-					// TODO: add a helper function to FDataReferenceCollection to do a SetParamIfContains.
-					if (InCollection.ContainsDataReadReference<ParamType>(InParamName))
-					{
-						ParamRef = InCollection.GetDataReadReference<ParamType>(InParamName);
-						return true;
-					}
-
-					return false;
-				}
-
-		};
-
 		public:
-			static const FName ClassName;
-
-			FAudioMultiplyNode(const FString& InName);
-
-			// Constructor used by the Metasound Frontend.
+			FAudioMultiplyNode(const FString& InInstanceName);
 			FAudioMultiplyNode(const FNodeInitData& InInitData);
-
-			virtual ~FAudioMultiplyNode();
-
-			virtual const FName& GetClassName() const override;
-
-			virtual IOperatorFactory& GetDefaultOperatorFactory() override;
-
-			virtual const FText& GetDescription() const override
-			{
-				static const FText StaticDescription = NSLOCTEXT("MetasoundGraphCore", "Metasound_AudioMultipleNodeDescription", "Multiplies two audio signals together (useful for amplitude modulation).");
-				return StaticDescription;
-			}
-
-			virtual const FText& GetAuthorName() const override
-			{
-				return PluginAuthor;
-			}
-
-			virtual const FText& GetPromptIfMissing() const override
-			{
-				return PluginNodeMissingPrompt;
-			}
-
-		private:
-			FOperatorFactory Factory;
 	};
+
 } // namespace Metasound
diff --git a/Engine/Plugins/Runtime/Metasound/Source/MetasoundStandardNodes/Public/MetasoundOscNode.h b/Engine/Plugins/Runtime/Metasound/Source/MetasoundStandardNodes/Public/MetasoundOscNode.h
--- a/Engine/Plugins/Runtime/Metasound/Source/MetasoundStandardNodes/Public/MetasoundOscNode.h
+++ b/Engine/Plugins/Runtime/Metasound/Source/MetasoundStandardNodes/Public/MetasoundOscNode.h
@@ -3,54 +3,30 @@
 #pragma once
 
 #include "CoreMinimal.h"
-#include "MetasoundNode.h"
-#include "MetasoundOperatorInterface.h"
+
+#include "MetasoundAudioBuffer.h"
 #include "MetasoundBuilderInterface.h"
 #include "MetasoundDataReferenceCollection.h"
+#include "MetasoundExecutableOperator.h"
+#include "MetasoundFacade.h"
+#include "MetasoundFrequency.h"
+#include "MetasoundNode.h"
+#include "MetasoundOperatorInterface.h"
 
 namespace Metasound
 {
-	class METASOUNDSTANDARDNODES_API FOscNode : public FNode
+	class METASOUNDSTANDARDNODES_API FOscNode : public FNodeFacade
 	{
-			class FOperatorFactory : public IOperatorFactory
-			{
-				virtual TUniquePtr<IOperator> CreateOperator(const INode& InNode, const FOperatorSettings& InOperatorSettings, const FDataReferenceCollection& InInputDataReferences, TArray<TUniquePtr<IOperatorBuildError>>& OutErrors) override;
-			};
-
 		public:
-			static const FName ClassName;
-
-			FOscNode(const FString& InName, float InDefaultFrequency);
 
-			// constructor used by the Metasound Frontend.
-			FOscNode(const FNodeInitData& InInitData);
+		FOscNode(const FString& InName, float InDefaultFrequency);
 
-			virtual ~FOscNode();
+		FOscNode(const FNodeInitData& InInitData);
 
-			float GetDefaultFrequency() const;
-
-			virtual const FName& GetClassName() const override;
-
-			virtual IOperatorFactory& GetDefaultOperatorFactory() override;
-
-			virtual const FText& GetDescription() const override
-			{
-				static const FText StaticDescription = NSLOCTEXT("MetasoundGraphCore", "Metasound_OscNodeDescription", "Emits an audio signal of a sinusoid.");
-				return StaticDescription;
-			}
-
-			virtual const FText& GetAuthorName() const override
-			{
-				return PluginAuthor;
-			}
-
-			virtual const FText& GetPromptIfMissing() const override
-			{
-				return PluginNodeMissingPrompt;
-			}
+		float GetDefaultFrequency() const;
 
 		private:
-			float DefaultFrequency;
-			FOperatorFactory Factory;
+
+		float DefaultFrequency = 440.f;
 	};
 }
diff --git a/Engine/Plugins/Runtime/Metasound/Source/MetasoundStandardNodes/Public/MetasoundPeriodicBopNode.h b/Engine/Plugins/Runtime/Metasound/Source/MetasoundStandardNodes/Public/MetasoundPeriodicBopNode.h
--- a/Engine/Plugins/Runtime/Metasound/Source/MetasoundStandardNodes/Public/MetasoundPeriodicBopNode.h
+++ b/Engine/Plugins/Runtime/Metasound/Source/MetasoundStandardNodes/Public/MetasoundPeriodicBopNode.h
@@ -5,52 +5,27 @@
 #include "CoreMinimal.h"
 #include "MetasoundNode.h"
 #include "MetasoundBuilderInterface.h"
+#include "MetasoundExecutableOperator.h"
+#include "MetasoundFacade.h"
 #include "MetasoundOperatorInterface.h"
 #include "MetasoundNodeInterface.h"
 #include "MetasoundDataReferenceCollection.h"
 #include "MetasoundDataReferenceTypes.h"
+#include "MetasoundTime.h"
 
 namespace Metasound
 {
-	class METASOUNDSTANDARDNODES_API FPeriodicBopNode : public FNode
+	class METASOUNDSTANDARDNODES_API FPeriodicBopNode : public FNodeFacade
 	{
-		class FOperatorFactory : public IOperatorFactory
-		{
-			virtual TUniquePtr<IOperator> CreateOperator(const INode& InNode, const FOperatorSettings& InOperatorSettings, const FDataReferenceCollection& InInputDataReferences, TArray<TUniquePtr<IOperatorBuildError>>& OutErrors) override;
-		};
-
 		public:
-			static const FName ClassName;
-
 			FPeriodicBopNode(const FString& InName, float InDefaultPeriodInSeconds);
-			FPeriodicBopNode(const FNodeInitData& InInitData);
 
-			virtual ~FPeriodicBopNode();
+			FPeriodicBopNode(const FNodeInitData& InInitData);
 
 			float GetDefaultPeriodInSeconds() const;
 
-			virtual const FName& GetClassName() const override;
-
-			virtual IOperatorFactory& GetDefaultOperatorFactory() override;
-
-			virtual const FText& GetDescription() const override
-			{
-				static const FText StaticDescription = NSLOCTEXT("MetasoundGraphCore", "Metasound_PeriodicBopNodeDescription", "Emits a bop periodically based on the period duration given.");
-				return StaticDescription;
-			}
-
-			virtual const FText& GetAuthorName() const override
-			{
-				return PluginAuthor;
-			}
-
-			virtual const FText& GetPromptIfMissing() const override
-			{
-				return PluginNodeMissingPrompt;
-			}
-
 		private:
+
 			float DefaultPeriod;
-			FOperatorFactory Factory;
 	};
 }
