Fix bug in FGraphLinter which looks for duplicate input connections.
Created a unit tests for FGraphLinter.

diff --git a/Engine/Plugins/Runtime/Metasound/Source/MetasoundGraphCore/Private/MetasoundBuildError.cpp b/Engine/Plugins/Runtime/Metasound/Source/MetasoundGraphCore/Private/MetasoundBuildError.cpp
--- a/Engine/Plugins/Runtime/Metasound/Source/MetasoundGraphCore/Private/MetasoundBuildError.cpp
+++ b/Engine/Plugins/Runtime/Metasound/Source/MetasoundGraphCore/Private/MetasoundBuildError.cpp
@@ -31,13 +31,13 @@ namespace Metasound
 	}
 
 	/** Returns an array of destinations associated with the error. */
-	const TArray<FInputDataDestination>& FBuildErrorBase::GetDestinations() const
+	const TArray<FInputDataDestination>& FBuildErrorBase::GetInputDataDestinations() const
 	{
 		return Destinations;
 	}
 	
 	/** Returns an array of sources associated with the error. */
-	const TArray<FOutputDataSource>& FBuildErrorBase::GetSources() const
+	const TArray<FOutputDataSource>& FBuildErrorBase::GetOutputDataSources() const
 	{
 		return Sources;
 	}
@@ -49,7 +49,7 @@ namespace Metasound
 	}
 
 	/** Returns an array of edges associated with the error. */
-	const TArray<FDataEdge>& FBuildErrorBase::GetEdges() const
+	const TArray<FDataEdge>& FBuildErrorBase::GetDataEdges() const
 	{
 		return Edges;
 	}
@@ -74,12 +74,12 @@ namespace Metasound
 		Sources.Append(InSources.GetData(), InSources.Num());
 	}
 
-	void FBuildErrorBase::AddEdge(const FDataEdge& InEdge)
+	void FBuildErrorBase::AddDataEdge(const FDataEdge& InEdge)
 	{
 		Edges.Add(InEdge);
 	}
 
-	void FBuildErrorBase::AddEdges(TArrayView<const FDataEdge> InEdges)
+	void FBuildErrorBase::AddDataEdges(TArrayView<const FDataEdge> InEdges)
 	{
 		Edges.Append(InEdges.GetData(), InEdges.Num());
 	}
@@ -115,7 +115,7 @@ namespace Metasound
 	FDanglingVertexError::FDanglingVertexError(const FDataEdge& InEdge)
 	:	FDanglingVertexError()
 	{
-		AddEdge(InEdge);
+		AddDataEdge(InEdge);
 	}
 
 	const FName FMissingVertexError::ErrorType = FName(TEXT("MetasoundMissingVertex"));
@@ -137,7 +137,7 @@ namespace Metasound
 	FDuplicateInputError::FDuplicateInputError(const TArrayView<FDataEdge> InEdges)
 	:	FBuildErrorBase(ErrorType, LOCTEXT("DuplicateInputError", "Node has duplicate connections to input"))
 	{
-		AddEdges(InEdges);
+		AddDataEdges(InEdges);
 	}
 
 	const FName FGraphCycleError::ErrorType = FName(TEXT("MetasoundGraphCycleError"));
@@ -146,7 +146,7 @@ namespace Metasound
 	:	FBuildErrorBase(ErrorType, LOCTEXT("GraphCycleError", "Graph contains cycles."))
 	{
 		AddNodes(InNodes);
-		AddEdges(InEdges);
+		AddDataEdges(InEdges);
 	}
 
 	const FName FNodePrunedError::ErrorType = FName(TEXT("MetasoundNodePrunedError"));
@@ -201,7 +201,7 @@ namespace Metasound
 	FInvalidConnectionDataTypeError::FInvalidConnectionDataTypeError(const FDataEdge& InEdge)
 	:	FBuildErrorBase(ErrorType, LOCTEXT("InvalidConnectionDataTypeError", "Source and destination data types do not match."))
 	{
-		AddEdge(InEdge);
+		AddDataEdge(InEdge);
 	}
 }
 
diff --git a/Engine/Plugins/Runtime/Metasound/Source/MetasoundGraphCore/Private/MetasoundGraphLinter.cpp b/Engine/Plugins/Runtime/Metasound/Source/MetasoundGraphCore/Private/MetasoundGraphLinter.cpp
--- a/Engine/Plugins/Runtime/Metasound/Source/MetasoundGraphCore/Private/MetasoundGraphLinter.cpp
+++ b/Engine/Plugins/Runtime/Metasound/Source/MetasoundGraphCore/Private/MetasoundGraphLinter.cpp
@@ -72,10 +72,10 @@ namespace Metasound
 			return &InEdge; 
 		};
 
-		// Raw pointer comparison for sorting.
-		bool CompareEdgeDestinationNode(const FDataEdge& InLHS, const FDataEdge& InRHS)
+		// For sorting edges by destination.
+		bool CompareEdgeDestination(const FDataEdge& InLHS, const FDataEdge& InRHS)
 		{
-			return InLHS.To.Node < InRHS.To.Node;
+			return InLHS.To < InRHS.To;
 		}
 	}
 
@@ -164,43 +164,57 @@ namespace Metasound
 
 		Algo::Transform(InGraph.GetDataEdges(), Edges, GetEdgePointer);
 
-		Edges.Sort(CompareEdgeDestinationNode);
+		Edges.Sort(CompareEdgeDestination);
 
 		if (Edges.Num() < 1)
 		{
 			return bIsValid;
 		}
 
+		// Loop through edges sorted by destinations. Look for duplicate destinations.
 		int32 GroupStartIndex = 0;
 		int32 GroupEndIndex = 0;
-		const INode* CurrentNode = Edges[GroupStartIndex]->To.Node;
+		const FInputDataDestination* CurrentDestination = &Edges[GroupStartIndex]->To;
+
+		auto AddErrorsIfGroupHasDuplicates = [&]()
+		{
+			if (GroupEndIndex > GroupStartIndex)
+			{
+				// This condition is hit when more than one edge refers to the
+				// same destination.
+				bIsValid = false;
+
+				// Collected edges with same destination.
+				TArray<FDataEdge> Duplicates;
+				for (int32 j = GroupStartIndex; j <= GroupEndIndex; j++)
+				{
+					Duplicates.Add(*Edges[j]);
+				}
+
+				// Add error.
+				AddBuildError<FDuplicateInputError>(OutErrors, Duplicates);
+			}
+		};
 
 		for (int32 i = 1; i < Edges.Num(); i++)
 		{
-			const INode* Node = Edges[i]->To.Node;
-			if (Node == CurrentNode)
+			const FInputDataDestination* Destination = &Edges[i]->To;
+
+			if (*Destination == *CurrentDestination)
 			{
 				GroupEndIndex = i;
 			}
 			else
 			{
-				if (GroupEndIndex > GroupStartIndex)
-				{
-					bIsValid = false;
-
-					TArray<FDataEdge> Duplicates;
-					for (int32 j = GroupStartIndex; j <= GroupEndIndex; j++)
-					{
-						Duplicates.Add(*Edges[j]);
-						AddBuildError<FDuplicateInputError>(OutErrors, Duplicates);
-					}
-				}
+				AddErrorsIfGroupHasDuplicates();
 
 				GroupStartIndex = GroupEndIndex = i;
-				CurrentNode = Node;
+				CurrentDestination = Destination;
 			}
 		}
 
+		AddErrorsIfGroupHasDuplicates();
+
 		return bIsValid;
 	}
 
diff --git a/Engine/Plugins/Runtime/Metasound/Source/MetasoundGraphCore/Private/MetasoundNodeInterface.cpp b/Engine/Plugins/Runtime/Metasound/Source/MetasoundGraphCore/Private/MetasoundNodeInterface.cpp
--- a/Engine/Plugins/Runtime/Metasound/Source/MetasoundGraphCore/Private/MetasoundNodeInterface.cpp
+++ b/Engine/Plugins/Runtime/Metasound/Source/MetasoundGraphCore/Private/MetasoundNodeInterface.cpp
@@ -9,15 +9,66 @@ namespace Metasound
 		return (InLeft.Node == InRight.Node) && (InLeft.Vertex == InRight.Vertex);
 	}
 
+	bool operator!=(const FOutputDataSource& InLeft, const FOutputDataSource& InRight)
+	{
+		return !(InLeft == InRight);
+	}
+
+	bool operator<(const FOutputDataSource& InLeft, const FOutputDataSource& InRight)
+	{
+		if (InLeft.Node == InRight.Node)
+		{
+			return InLeft.Vertex < InRight.Vertex;
+		}
+		else
+		{
+			return InLeft.Node < InRight.Node;
+		}
+	}
+
 	/** Check if two FInputDataDestinations are equal. */
 	bool operator==(const FInputDataDestination& InLeft, const FInputDataDestination& InRight)
 	{
 		return (InLeft.Node == InRight.Node) && (InLeft.Vertex == InRight.Vertex);
 	}
 
+	bool operator!=(const FInputDataDestination& InLeft, const FInputDataDestination& InRight)
+	{
+		return !(InLeft == InRight);
+	}
+
+	bool operator<(const FInputDataDestination& InLeft, const FInputDataDestination& InRight)
+	{
+		if (InLeft.Node == InRight.Node)
+		{
+			return InLeft.Vertex < InRight.Vertex;
+		}
+		else
+		{
+			return InLeft.Node < InRight.Node;
+		}
+	}
+
 	/** Check if two FDataEdges are equal. */
 	bool operator==(const FDataEdge& InLeft, const FDataEdge& InRight)
 	{
 		return (InLeft.From == InRight.From) && (InLeft.To == InRight.To);
 	}
+
+	bool operator!=(const FDataEdge& InLeft, const FDataEdge& InRight)
+	{
+		return !(InLeft == InRight);
+	}
+
+	bool operator<(const FDataEdge& InLeft, const FDataEdge& InRight)
+	{
+		if (InLeft.From == InRight.From)
+		{
+			return InLeft.To < InRight.To;	
+		}
+		else
+		{
+			return InLeft.From < InRight.From;
+		}
+	}
 }
diff --git a/Engine/Plugins/Runtime/Metasound/Source/MetasoundGraphCore/Private/MetasoundVertex.cpp b/Engine/Plugins/Runtime/Metasound/Source/MetasoundGraphCore/Private/MetasoundVertex.cpp
--- a/Engine/Plugins/Runtime/Metasound/Source/MetasoundGraphCore/Private/MetasoundVertex.cpp
+++ b/Engine/Plugins/Runtime/Metasound/Source/MetasoundGraphCore/Private/MetasoundVertex.cpp
@@ -78,6 +78,23 @@ namespace Metasound
 		return !(InLHS == InRHS);
 	}
 
+	bool operator<(const FInputDataVertex& InLHS, const FInputDataVertex& InRHS)
+	{
+		if (InLHS == InRHS)
+		{
+			return false;
+		}
+		
+		if (InLHS.GetVertexName() == InRHS.GetVertexName())
+		{
+			return InLHS.GetDataTypeName().FastLess(InRHS.GetDataTypeName());
+		}
+		else
+		{
+			return InLHS.GetVertexName() < InRHS.GetVertexName();
+		}
+	}
+
 	FOutputDataVertex::FOutputDataVertex()
 	:	VertexModel(MakeUnique<FEmptyVertexModel>(TEXT(""), FText::GetEmpty()))
 	{}
@@ -151,6 +168,23 @@ namespace Metasound
 		return !(InLHS == InRHS);
 	}
 
+	bool operator<(const FOutputDataVertex& InLHS, const FOutputDataVertex& InRHS)
+	{
+		if (InLHS == InRHS)
+		{
+			return false;
+		}
+		
+		if (InLHS.GetVertexName() == InRHS.GetVertexName())
+		{
+			return InLHS.GetDataTypeName().FastLess(InRHS.GetDataTypeName());
+		}
+		else
+		{
+			return InLHS.GetVertexName() < InRHS.GetVertexName();
+		}
+	}
+
 			/** Construct with an input and output interface. */
 	FVertexInterface::FVertexInterface(const FInputVertexInterface& InInputs, const FOutputVertexInterface& InOutputs)
 	:	InputInterface(InInputs)
diff --git a/Engine/Plugins/Runtime/Metasound/Source/MetasoundGraphCore/Public/MetasoundBuildError.h b/Engine/Plugins/Runtime/Metasound/Source/MetasoundGraphCore/Public/MetasoundBuildError.h
--- a/Engine/Plugins/Runtime/Metasound/Source/MetasoundGraphCore/Public/MetasoundBuildError.h
+++ b/Engine/Plugins/Runtime/Metasound/Source/MetasoundGraphCore/Public/MetasoundBuildError.h
@@ -30,13 +30,13 @@ namespace Metasound
 			virtual const FText& GetErrorDescription() const override;
 
 			/** Returns an array of destinations associated with the error. */
-			virtual const TArray<FInputDataDestination>& GetDestinations() const override;
+			virtual const TArray<FInputDataDestination>& GetInputDataDestinations() const override;
 			
 			/** Returns an array of sources associated with the error. */
-			virtual const TArray<FOutputDataSource>& GetSources() const override;
+			virtual const TArray<FOutputDataSource>& GetOutputDataSources() const override;
 
 			/** Returns an array of edges associated with the error. */
-			virtual const TArray<FDataEdge>& GetEdges() const override;
+			virtual const TArray<FDataEdge>& GetDataEdges() const override;
 
 			/** Returns an array of Nodes associated with the error. */
 			virtual const TArray<const INode*>& GetNodes() const override;
@@ -52,8 +52,8 @@ namespace Metasound
 			void AddOutputDataSources(TArrayView<const FOutputDataSource> InOutputDataSources);
 
 			// Add edges to be associated with error.
-			void AddEdge(const FDataEdge& InEdge);
-			void AddEdges(TArrayView<const FDataEdge> InEdges);
+			void AddDataEdge(const FDataEdge& InEdge);
+			void AddDataEdges(TArrayView<const FDataEdge> InEdges);
 
 			// Add nodes to be associated with error.
 			void AddNode(const INode& InNode);
diff --git a/Engine/Plugins/Runtime/Metasound/Source/MetasoundGraphCore/Public/MetasoundBuilderInterface.h b/Engine/Plugins/Runtime/Metasound/Source/MetasoundGraphCore/Public/MetasoundBuilderInterface.h
--- a/Engine/Plugins/Runtime/Metasound/Source/MetasoundGraphCore/Public/MetasoundBuilderInterface.h
+++ b/Engine/Plugins/Runtime/Metasound/Source/MetasoundGraphCore/Public/MetasoundBuilderInterface.h
@@ -24,16 +24,16 @@ namespace Metasound
 			virtual const FText& GetErrorDescription() const = 0;
 
 			/** Returns an array of destinations associated with the error. */
-			virtual const TArray<FInputDataDestination>& GetDestinations() const = 0;
+			virtual const TArray<FInputDataDestination>& GetInputDataDestinations() const = 0;
 			
 			/** Returns an array of sources associated with the error. */
-			virtual const TArray<FOutputDataSource>& GetSources() const = 0;
+			virtual const TArray<FOutputDataSource>& GetOutputDataSources() const = 0;
 
 			/** Returns an array of Nodes associated with the error. */
 			virtual const TArray<const INode*>& GetNodes() const = 0;
 
 			/** Returns an array of edges associated with the error. */
-			virtual const TArray<FDataEdge>& GetEdges() const = 0;
+			virtual const TArray<FDataEdge>& GetDataEdges() const = 0;
 	};
 
 	/** FCreateOperatorParams holds the parameters provided to operator factories
diff --git a/Engine/Plugins/Runtime/Metasound/Source/MetasoundGraphCore/Public/MetasoundNodeInterface.h b/Engine/Plugins/Runtime/Metasound/Source/MetasoundGraphCore/Public/MetasoundNodeInterface.h
--- a/Engine/Plugins/Runtime/Metasound/Source/MetasoundGraphCore/Public/MetasoundNodeInterface.h
+++ b/Engine/Plugins/Runtime/Metasound/Source/MetasoundGraphCore/Public/MetasoundNodeInterface.h
@@ -141,6 +141,8 @@ namespace Metasound
 
 		/** Check if two FOutputDataSources are equal. */
 		friend bool METASOUNDGRAPHCORE_API operator==(const FOutputDataSource& InLeft, const FOutputDataSource& InRight);
+		friend bool METASOUNDGRAPHCORE_API operator!=(const FOutputDataSource& InLeft, const FOutputDataSource& InRight);
+		friend bool METASOUNDGRAPHCORE_API operator<(const FOutputDataSource& InLeft, const FOutputDataSource& InRight);
 	};
 
 	/** FInputDataSource describes the destination of data which produced 
@@ -166,6 +168,8 @@ namespace Metasound
 
 		/** Check if two FInputDataDestinations are equal. */
 		friend bool METASOUNDGRAPHCORE_API operator==(const FInputDataDestination& InLeft, const FInputDataDestination& InRight);
+		friend bool METASOUNDGRAPHCORE_API operator!=(const FInputDataDestination& InLeft, const FInputDataDestination& InRight);
+		friend bool METASOUNDGRAPHCORE_API operator<(const FInputDataDestination& InLeft, const FInputDataDestination& InRight);
 	};
 
 	/** Key type for an FOutputDataSource or FInputDataDestination. */
@@ -214,6 +218,8 @@ namespace Metasound
 
 		/** Check if two FDataEdges are equal. */
 		friend bool METASOUNDGRAPHCORE_API operator==(const FDataEdge& InLeft, const FDataEdge& InRight);
+		friend bool METASOUNDGRAPHCORE_API operator!=(const FDataEdge& InLeft, const FDataEdge& InRight);
+		friend bool METASOUNDGRAPHCORE_API operator<(const FDataEdge& InLeft, const FDataEdge& InRight);
 	};
 
 
diff --git a/Engine/Plugins/Runtime/Metasound/Source/MetasoundGraphCore/Public/MetasoundVertex.h b/Engine/Plugins/Runtime/Metasound/Source/MetasoundGraphCore/Public/MetasoundVertex.h
--- a/Engine/Plugins/Runtime/Metasound/Source/MetasoundGraphCore/Public/MetasoundVertex.h
+++ b/Engine/Plugins/Runtime/Metasound/Source/MetasoundGraphCore/Public/MetasoundVertex.h
@@ -252,6 +252,7 @@ namespace Metasound
 
 			friend bool METASOUNDGRAPHCORE_API operator==(const FInputDataVertex& InLHS, const FInputDataVertex& InRHS);
 			friend bool METASOUNDGRAPHCORE_API operator!=(const FInputDataVertex& InLHS, const FInputDataVertex& InRHS);
+			friend bool METASOUNDGRAPHCORE_API operator<(const FInputDataVertex& InLHS, const FInputDataVertex& InRHS);
 
 		private:
 
@@ -342,6 +343,7 @@ namespace Metasound
 
 			friend bool METASOUNDGRAPHCORE_API operator==(const FOutputDataVertex& InLHS, const FOutputDataVertex& InRHS);
 			friend bool METASOUNDGRAPHCORE_API operator!=(const FOutputDataVertex& InLHS, const FOutputDataVertex& InRHS);
+			friend bool METASOUNDGRAPHCORE_API operator<(const FOutputDataVertex& InLHS, const FOutputDataVertex& InRHS);
 
 		private:
 
