Converter Nodes for Metasounds. Allows for use of the METASOUND_REGISTER_CONVERTER macros, as well as autogenerated converters.

diff --git a/Engine/Plugins/Runtime/Metasound/Source/MetasoundFrontend/Private/MetasoundFrontend.cpp b/Engine/Plugins/Runtime/Metasound/Source/MetasoundFrontend/Private/MetasoundFrontend.cpp
--- a/Engine/Plugins/Runtime/Metasound/Source/MetasoundFrontend/Private/MetasoundFrontend.cpp
+++ b/Engine/Plugins/Runtime/Metasound/Source/MetasoundFrontend/Private/MetasoundFrontend.cpp
@@ -621,7 +621,13 @@ namespace Metasound
 				return OutConnectability;
 			}
 
-			//@todo: scan for possible converter nodes here. (UEAU-473)
+			OutConnectability.PossibleConverterNodeClasses = FMetasoundFrontendRegistryContainer::Get()->GetPossibleConverterNodes(InHandle.GetOutputType(), GetInputType());
+			if (OutConnectability.PossibleConverterNodeClasses.Num() > 0)
+			{
+				OutConnectability.Connectable = FConnectability::EConnectable::YesWithConverterNode;
+				return OutConnectability;
+			}
+
 			return OutConnectability;
 		}
 
@@ -708,11 +714,33 @@ namespace Metasound
 			return false;
 		}
 
-		bool FInputHandle::ConnectWithConverterNode(FOutputHandle& InHandle, FString& InNodeClassName)
+		bool FInputHandle::ConnectWithConverterNode(FOutputHandle& InHandle, const FConverterNodeInfo& InNodeClassName)
 		{
-			// (UEAU-473)
-			ensureAlwaysMsgf(false, TEXT("Implement me!"));
-			return false;
+			FDescPath OwningGraphPath = Path::GetOuterGraphPath(NodePtr.GetPath());
+			FHandleInitParams GraphHandleInitParams =
+			{
+				NodePtr.GetAccessPoint(),
+				OwningGraphPath,
+				FString(), // Caching the graph name is not necessary here.
+				OwningAsset
+			};
+
+			FGraphHandle OwningGraphHandle = FGraphHandle(FHandleInitParams::PrivateToken, GraphHandleInitParams);
+			check(OwningGraphHandle.IsValid());
+
+			// Generate the converter node.
+			FNodeHandle Converter = OwningGraphHandle.AddNewNode(InNodeClassName.NodeKey);
+
+			check(Converter.IsValid());
+
+			FInputHandle ConverterInput = Converter.GetInputWithName(InNodeClassName.PreferredConverterInputPin);
+			FOutputHandle ConverterOutput = Converter.GetOutputWithName(InNodeClassName.PreferredConverterOutputPin);
+
+			check(ConverterInput.IsValid() && ConverterOutput.IsValid());
+
+			// Connect the output InHandle to the converter, than connect the converter to this input.
+			ensureAlways(ConverterInput.Connect(InHandle));
+			return ensureAlways(Connect(ConverterOutput));
 		}
 
 		const FMetasoundNodeConnectionDescription* FInputHandle::GetConnectionDescription() const
@@ -888,7 +916,7 @@ namespace Metasound
 			return InHandle.Connect(*this);
 		}
 
-		bool FOutputHandle::ConnectWithConverterNode(FInputHandle& InHandle, FString& InNodeClassName)
+		bool FOutputHandle::ConnectWithConverterNode(FInputHandle& InHandle, const FConverterNodeInfo& InNodeClassName)
 		{
 			return InHandle.ConnectWithConverterNode(*this, InNodeClassName);
 		}
@@ -2007,6 +2035,17 @@ namespace Metasound
 			return FNodeHandle(FHandleInitParams::PrivateToken, InitParams, NewNodeDescription.ObjectTypeOfNode);
 		}
 
+		FNodeHandle FGraphHandle::AddNewNode(const FNodeRegistryKey& InNodeClass)
+		{
+			// Construct a FNodeClassInfo from this lookup key.
+			FNodeClassInfo ClassInfo;
+			ClassInfo.LookupKey = InNodeClass;
+			ClassInfo.NodeName = InNodeClass.NodeName.ToString();
+			ClassInfo.NodeType = EMetasoundClassType::External;
+
+			return AddNewNode(ClassInfo);
+		}
+
 		bool FGraphHandle::RemoveNode(const FNodeHandle& InNode)
 		{
 			if (!ensureAlwaysMsgf(InNode.GetNodeType() != EMetasoundClassType::Input && InNode.GetNodeType() != EMetasoundClassType::Output,
diff --git a/Engine/Plugins/Runtime/Metasound/Source/MetasoundFrontend/Private/MetasoundFrontendBaseClasses.cpp b/Engine/Plugins/Runtime/Metasound/Source/MetasoundFrontend/Private/MetasoundFrontendBaseClasses.cpp
--- a/Engine/Plugins/Runtime/Metasound/Source/MetasoundFrontend/Private/MetasoundFrontendBaseClasses.cpp
+++ b/Engine/Plugins/Runtime/Metasound/Source/MetasoundFrontend/Private/MetasoundFrontendBaseClasses.cpp
@@ -992,6 +992,17 @@ namespace Metasound
 			return (InPathForOutputNode << 3)[Path::EFromClass::ToOutputs][*OutputName];
 		}
 
+		FDescPath Path::GetOuterGraphPath(FDescPath InPath)
+		{
+			// Unwind element by element until we hit a graph.
+			while (InPath.Path.Num() > 1 && InPath.Path.Last().CurrentDescType != Path::EDescType::Graph)
+			{
+				InPath.Path.Pop();
+			}
+
+			return InPath;
+		}
+
 		FString Path::GetPrintableString(FDescPath InPath)
 		{
 			FString OutString = FString(TEXT("//"));
diff --git a/Engine/Plugins/Runtime/Metasound/Source/MetasoundFrontend/Private/MetasoundFrontendRegistries.cpp b/Engine/Plugins/Runtime/Metasound/Source/MetasoundFrontend/Private/MetasoundFrontendRegistries.cpp
--- a/Engine/Plugins/Runtime/Metasound/Source/MetasoundFrontend/Private/MetasoundFrontendRegistries.cpp
+++ b/Engine/Plugins/Runtime/Metasound/Source/MetasoundFrontend/Private/MetasoundFrontendRegistries.cpp
@@ -14,6 +14,31 @@ FMetasoundFrontendRegistryContainer::FMetasoundFrontendRegistryContainer()
 {
 }
 
+FMetasoundFrontendRegistryContainer::FNodeRegistryKey FMetasoundFrontendRegistryContainer::GetRegistryKeyForNodeInternal(const Metasound::INode& InNode)
+{
+	using FInputVertexInterface = ::Metasound::FInputVertexInterface;
+	using FOutputVertexInterface = ::Metasound::FOutputVertexInterface;
+
+	const FName NodeName = InNode.GetClassName();
+	const FInputVertexInterface& Inputs = InNode.GetDefaultVertexInterface().GetInputInterface();
+	const FOutputVertexInterface& Outputs = InNode.GetDefaultVertexInterface().GetOutputInterface();
+
+	// Construct a hash using a combination of the class name, input names and output names.
+	uint32 NodeHash = FCrc::StrCrc32(*NodeName.ToString());
+
+	for (auto& InputTuple : Inputs)
+	{
+		NodeHash = HashCombine(NodeHash, FCrc::StrCrc32(*InputTuple.Value.GetVertexTypeName().ToString()));
+	}
+
+	for (auto& OutputTuple : Outputs)
+	{
+		NodeHash = HashCombine(NodeHash, FCrc::StrCrc32(*OutputTuple.Value.GetVertexTypeName().ToString()));
+	}
+
+	return { NodeName, NodeHash };
+}
+
 FMetasoundFrontendRegistryContainer* FMetasoundFrontendRegistryContainer::LazySingleton = nullptr;
 
 FMetasoundFrontendRegistryContainer* FMetasoundFrontendRegistryContainer::Get()
@@ -168,6 +193,19 @@ TUniquePtr<Metasound::INode> FMetasoundFrontendRegistryContainer::ConstructExter
 	}
 }
 
+TArray<::Metasound::Frontend::FConverterNodeInfo> FMetasoundFrontendRegistryContainer::GetPossibleConverterNodes(const FName& FromDataType, const FName& ToDataType)
+{
+	FConverterNodeRegistryKey InKey = { FromDataType, ToDataType };
+	if (!ConverterNodeRegistry.Contains(InKey))
+	{
+		return TArray<FConverterNodeInfo>();
+	}
+	else
+	{
+		return ConverterNodeRegistry[InKey].PotentialConverterNodes;
+	}
+}
+
 Metasound::ELiteralArgType FMetasoundFrontendRegistryContainer::GetDesiredLiteralTypeForDataType(FName InDataType) const
 {
 	if (!DataTypeRegistry.Contains(InDataType))
@@ -305,52 +343,50 @@ bool FMetasoundFrontendRegistryContainer::RegisterExternalNode(FNodeGetterCallba
 
 	Metasound::INode& DummyNode = *DummyNodePtr;
 
-	// First, build the key.
-	const FName NodeName = DummyNode.GetClassName();
-	const FInputVertexInterface& Inputs = DummyNode.GetDefaultVertexInterface().GetInputInterface();
-	const FOutputVertexInterface& Outputs = DummyNode.GetDefaultVertexInterface().GetOutputInterface();
-
-	// Construct a hash using a combination of the class name, input names and output names.
-	uint32 NodeHash = FCrc::StrCrc32(*NodeName.ToString());
+	FNodeRegistryKey InKey = GetRegistryKeyForNodeInternal(DummyNode);
 
-	TArray<FName> InputTypes;
-	TArray<FName> OutputTypes;
-
-	for (auto& InputTuple : Inputs)
-	{
-		NodeHash = HashCombine(NodeHash, FCrc::StrCrc32(*InputTuple.Value.GetVertexName()));
-		InputTypes.Add(InputTuple.Value.GetDataTypeName());
-	}
-
-	for (auto& OutputTuple : Outputs)
-	{
-		NodeHash = HashCombine(NodeHash, FCrc::StrCrc32(*OutputTuple.Value.GetVertexName()));
-		OutputTypes.Add(OutputTuple.Value.GetDataTypeName());
-	}
-
-	FNodeRegistryKey InKey = { NodeName, NodeHash };
+	const FName NodeName = DummyNode.GetClassName();
 
 	// check to see if an identical node was already registered, and log
 	ensureAlwaysMsgf(!ExternalNodeRegistry.Contains(InKey), TEXT("Node with identical name, inputs and outputs to node %s was already registered. The previously registered node will be overwritten. This could also happen because METASOUND_REGISTER_NODE is in a public header."), *NodeName.ToString());
 
 	UE_LOG(LogTemp, Display, TEXT("Registered Metasound Node %s"), *NodeName.ToString());
 
+	using FInputVertexInterface = ::Metasound::FInputVertexInterface;
+	using FOutputVertexInterface = ::Metasound::FOutputVertexInterface;
+
+	const FInputVertexInterface& Inputs = DummyNode.GetDefaultVertexInterface().GetInputInterface();
+	const FOutputVertexInterface& Outputs = DummyNode.GetDefaultVertexInterface().GetOutputInterface();
+
 	const bool bShouldLogInputsAndOutputs = true;
 	if (bShouldLogInputsAndOutputs)
 	{
 		UE_LOG(LogTemp, Display, TEXT("    %d inputs:"), Inputs.Num());
 		for (auto& InputTuple : Inputs)
 		{
-			UE_LOG(LogTemp, Display, TEXT("      %s (of type %s)"), *InputTuple.Value.GetVertexName(), *InputTuple.Value.GetDataTypeName().ToString());
+			UE_LOG(LogTemp, Display, TEXT("      %s (of type %s)"), *InputTuple.Value.GetVertexName(), *InputTuple.Value.GetVertexTypeName().ToString());
 		}
 
 		UE_LOG(LogTemp, Display, TEXT("    %d outputs:"), Outputs.Num());
 		for (auto& OutputTuple : Outputs)
 		{
-			UE_LOG(LogTemp, Display, TEXT("      %s (of type %s)"), *OutputTuple.Value.GetVertexName(), *OutputTuple.Value.GetDataTypeName().ToString());
+			UE_LOG(LogTemp, Display, TEXT("      %s (of type %s)"), *OutputTuple.Value.GetVertexName(), *OutputTuple.Value.GetVertexTypeName().ToString());
 		}
 	}
 
+	TArray<FName> InputTypes;
+	TArray<FName> OutputTypes;
+
+	for (auto& InputTuple : Inputs)
+	{
+		InputTypes.Add(InputTuple.Value.GetVertexTypeName());
+	}
+
+	for (auto& OutputTuple : Outputs)
+	{
+		OutputTypes.Add(OutputTuple.Value.GetVertexTypeName());
+	}
+
 	FNodeRegistryElement RegistryElement = FNodeRegistryElement(MoveTemp(InCallback));
 	RegistryElement.InputTypes = MoveTemp(InputTypes);
 	RegistryElement.OutputTypes = MoveTemp(OutputTypes);
@@ -360,6 +396,27 @@ bool FMetasoundFrontendRegistryContainer::RegisterExternalNode(FNodeGetterCallba
 	return true;
 }
 
+bool FMetasoundFrontendRegistryContainer::RegisterConversionNode(const FConverterNodeRegistryKey& InNodeKey, const FConverterNodeInfo& InNodeInfo)
+{
+	if (!ConverterNodeRegistry.Contains(InNodeKey))
+	{
+		ConverterNodeRegistry.Add(InNodeKey);
+	}
+
+	FConverterNodeRegistryValue& ConverterNodeList = ConverterNodeRegistry[InNodeKey];
+
+	if (ensureAlways(!ConverterNodeList.PotentialConverterNodes.Contains(InNodeInfo)))
+	{
+		ConverterNodeList.PotentialConverterNodes.Add(InNodeInfo);
+		return true;
+	}
+	else
+	{
+		// If we hit this, someone attempted to add the same converter node to our list multiple times.
+		return false;
+	}
+}
+
 TArray<FName> FMetasoundFrontendRegistryContainer::GetAllValidDataTypes()
 {
 	TArray<FName> OutDataTypes;
diff --git a/Engine/Plugins/Runtime/Metasound/Source/MetasoundFrontend/Public/MetasoundAutoConverterNode.h b/Engine/Plugins/Runtime/Metasound/Source/MetasoundFrontend/Public/MetasoundAutoConverterNode.h
new file mode 100644
--- /dev/null
+++ b/Engine/Plugins/Runtime/Metasound/Source/MetasoundFrontend/Public/MetasoundAutoConverterNode.h
@@ -0,0 +1,174 @@
+// Copyright Epic Games, Inc. All Rights Reserved.
+#pragma once
+
+#include "CoreMinimal.h"
+#include "MetasoundBuilderInterface.h"
+#include "MetasoundNodeInterface.h"
+#include "MetasoundOperatorInterface.h"
+#include "MetasoundDataReference.h"
+#include "MetasoundExecutableOperator.h"
+
+#include <type_traits>
+
+#define LOCTEXT_NAMESPACE "MetasoundGraphCore"
+
+
+namespace Metasound
+{
+	// This convenience node can be registered and will invoke static_cast<ToDataType>(FromDataType) every time it is executed.
+	template<typename FromDataType, typename ToDataType>
+	class TAutoConverterNode : public INode
+	{
+		static_assert(TDataReferenceTypeInfo<FromDataType>::bIsValidSpecialization, "Please use DECLARE_METASOUND_DATA_REFERENCE_TYPES with this class before trying to create an converter node with it.");
+		static_assert(TDataReferenceTypeInfo<ToDataType>::bIsValidSpecialization, "Please use DECLARE_METASOUND_DATA_REFERENCE_TYPES with this class before trying to create an converter node with it.");
+		static_assert(std::is_convertible<FromDataType, ToDataType>::value, "Tried to create an auto converter node between two types we can't static_cast between.");
+
+	public:
+		static FString& GetInputName()
+		{
+			static FString InputName = FString(TDataReferenceTypeInfo<FromDataType>::TypeName);
+			return InputName;
+		}
+
+		static FString& GetOutputName()
+		{
+			static FString OutputName = FString(TDataReferenceTypeInfo<ToDataType>::TypeName);
+			return OutputName;
+		}
+
+	private:
+		class FConverterOperator : public TExecutableOperator<FConverterOperator>
+		{
+			public:
+
+				FConverterOperator(TDataReadReference<FromDataType> InFromDataReference)
+					: FromData(InFromDataReference)
+					, ToData(TDataWriteReference<ToDataType>::CreateNew()) // TODO: this currently only works with default constructible datatypes. Can we do something similar to what we did with literals?
+				{
+					Outputs.AddDataReadReference<ToDataType>(GetOutputName(), TDataReadReference<ToDataType>(ToData));
+				}
+
+				virtual ~FConverterOperator() {}
+
+				virtual const FDataReferenceCollection& GetInputs() const override
+				{
+					return Inputs;
+				}
+
+				virtual const FDataReferenceCollection& GetOutputs() const override
+				{
+					return Outputs;
+				}
+
+				void Execute()
+				{
+					*ToData = static_cast<ToDataType>(*FromData);
+				}
+
+			private:
+				TDataReadReference<FromDataType> FromData;
+				TDataWriteReference<ToDataType> ToData;
+
+				FDataReferenceCollection Inputs;
+				FDataReferenceCollection Outputs;
+		};
+
+		class FCoverterOperatorFactory : public IOperatorFactory
+		{
+			public:
+				FCoverterOperatorFactory() = default;
+
+				virtual TUniquePtr<IOperator> CreateOperator(const FCreateOperatorParams& InParams, FBuildErrorArray& OutErrors) override
+				{
+					if (!InParams.InputDataReferences.ContainsDataReadReference<FromDataType>(GetInputName()))
+					{
+						// TODO: Add build error.
+						return TUniquePtr<IOperator>(nullptr);
+					}
+
+					return MakeUnique<FConverterOperator>(InParams.InputDataReferences.GetDataReadReference<FromDataType>(GetInputName()));
+				}
+		};
+
+		public:
+			static FVertexInterface DeclareVertexInterface()
+			{
+				return FVertexInterface(
+					FInputVertexInterface(
+						TInputDataVertexModel<FromDataType>(GetInputName(), FText::GetEmpty())
+					),
+					FOutputVertexInterface(
+						TOutputDataVertexModel<ToDataType>(GetOutputName(), FText::GetEmpty())
+					)
+				);
+			}
+
+			TAutoConverterNode(const FNodeInitData& InInitData)
+				: NodeDescription(InInitData.InstanceName)
+				, Interface(DeclareVertexInterface())
+				, Factory(MakeOperatorFactoryRef<FCoverterOperatorFactory>())
+			{
+			}
+
+			virtual ~TAutoConverterNode() = default;
+
+			virtual const FName& GetClassName() const override
+			{
+				static const FName ClassName(*FString::Printf(TEXT("%s To %s"), *GetInputName(), *GetOutputName()));
+				return ClassName;
+			}
+
+			virtual const FString& GetInstanceName() const override
+			{
+				return NodeDescription;
+			}
+
+			virtual const FText& GetDescription() const override
+			{
+				static const FText Description = LOCTEXT("Metasound_ConverterNodeDescription", "Converts between two different data types.");
+				return Description;
+			}
+
+			virtual const FText& GetAuthorName() const override
+			{
+				return PluginAuthor;
+			}
+
+			virtual const FText& GetPromptIfMissing() const override
+			{
+				return PluginNodeMissingPrompt;
+			}
+
+			virtual const FVertexInterface& GetVertexInterface() const override
+			{
+				return Interface;
+			}
+
+			virtual const FVertexInterface& GetDefaultVertexInterface() const override
+			{
+				return Interface;
+			}
+
+			virtual bool SetVertexInterface(const FVertexInterface& InInterface) override
+			{
+				return Interface == InInterface;
+			}
+
+			virtual bool IsVertexInterfaceSupported(const FVertexInterface& InInterface) const override
+			{
+				return Interface == InInterface;
+			}
+
+			virtual FOperatorFactorySharedRef GetDefaultOperatorFactory() const override
+			{
+				return Factory;
+			}
+
+		private:
+			FString NodeDescription;
+			FVertexInterface Interface;
+			FOperatorFactorySharedRef Factory;
+			
+	};
+}
+#undef LOCTEXT_NAMESPACE
diff --git a/Engine/Plugins/Runtime/Metasound/Source/MetasoundFrontend/Public/MetasoundConverterNodeRegistrationMacro.h b/Engine/Plugins/Runtime/Metasound/Source/MetasoundFrontend/Public/MetasoundConverterNodeRegistrationMacro.h
new file mode 100644
--- /dev/null
+++ b/Engine/Plugins/Runtime/Metasound/Source/MetasoundFrontend/Public/MetasoundConverterNodeRegistrationMacro.h
@@ -0,0 +1,57 @@
+// Copyright Epic Games, Inc. All Rights Reserved.
+
+#pragma once
+
+#include "CoreMinimal.h"
+#include "Templates/Casts.h"
+
+#include "MetasoundDataReference.h"
+#include "MetasoundOutputNode.h"
+#include "MetasoundInputNode.h"
+#include "MetasoundOperatorInterface.h"
+#include "MetasoundFrontendRegistries.h"
+
+template<typename TNodeType, typename TFromType, typename TToType>
+bool RegisterConversionNode(const FString& FromPin, const FString& ToPin)
+{
+	using FConverterNodeRegistryKey = ::Metasound::Frontend::FConverterNodeRegistryKey;
+	using FConverterNodeInfo = ::Metasound::Frontend::FConverterNodeInfo;
+
+	// if we reenter this code (because DECLARE_METASOUND_DATA_REFERENCE_TYPES was called twice with the same type),
+	// we catch it here.
+	static bool bAlreadyRegisteredThisConverter = false;
+	if (bAlreadyRegisteredThisConverter)
+	{
+		UE_LOG(LogTemp, Display, TEXT("Tried to call REGISTER_METASOUND_CONVERTER twice with the same class. ignoring the second call. Likely because REGISTER_METASOUND_CONVERTER is in a header that's used in multiple modules. Consider moving it to a private header or cpp file."))
+		return false;
+	}
+
+	bAlreadyRegisteredThisConverter = true;
+
+	// Get the FNames from our datatypes
+	FName FromType = ::Metasound::TDataReferenceTypeInfo<TFromType>::TypeName;
+	FName ToType = ::Metasound::TDataReferenceTypeInfo<TToType>::TypeName;
+
+	FConverterNodeRegistryKey RegistryKey = { FromType, ToType };
+
+	FConverterNodeInfo ConverterNodeInfo =
+	{
+		FromPin,
+		ToPin,
+		FMetasoundFrontendRegistryContainer::GetRegistryKeyForNode<TNodeType>()
+	};
+
+	return FMetasoundFrontendRegistryContainer::Get()->RegisterConversionNode(RegistryKey, ConverterNodeInfo);
+}
+
+// Convenience macro for token pasting the three parameters for the converter node into a unique static bool name:
+#define MS_CONVERTER_BOOL(Node, FromDataType, ToDataType) MS_CONVERTER_BOOL_HIDDEN(Node, FromDataType, ToDataType) 
+#define MS_CONVERTER_BOOL_HIDDEN(Node, FromDataType, ToDataType) Node ## FromDataType ## ToDataType
+
+#define MS_INVOKE_CONVERTER_REGISTRATION_HIDDEN(NodeType, FromDataType, ToDataType, FromOutputPin, ToInputPin) FMetasoundFrontendRegistryContainer::Get()->EnqueueInitCommand([](){ ::RegisterConversionNode<NodeType, FromDataType, ToDataType >(FString(TEXT( #FromOutputPin )), FString(TEXT( #ToInputPin ))); })
+
+#define MS_INVOKE_CONSTRUCTOR_REGISTRATION(NodeType, FromDataType, ToDataType, FromOutputPin, ToInputPin) MS_CONVERTER_BOOL(NodeType, FromDataType, ToDataType) = MS_INVOKE_CONVERTER_REGISTRATION_HIDDEN(NodeType, FromDataType, ToDataType, FromOutputPin, ToInputPin)
+
+// This macro exposes a previously registered node as a valid way to convert from one datatype to another.
+// When a user tries to connect a FromDataType output to a ToDataType input pin, this node will pop up as a potential node to insert between the two.
+#define REGISTER_METASOUND_CONVERTER(NodeType, FromDataType, ToDataType, FromOutputPin, ToInputPin) static const bool MS_INVOKE_CONSTRUCTOR_REGISTRATION(NodeType, FromDataType, ToDataType, FromOutputPin, ToInputPin);
diff --git a/Engine/Plugins/Runtime/Metasound/Source/MetasoundFrontend/Public/MetasoundDataTypeRegistrationMacro.h b/Engine/Plugins/Runtime/Metasound/Source/MetasoundFrontend/Public/MetasoundDataTypeRegistrationMacro.h
--- a/Engine/Plugins/Runtime/Metasound/Source/MetasoundFrontend/Public/MetasoundDataTypeRegistrationMacro.h
+++ b/Engine/Plugins/Runtime/Metasound/Source/MetasoundFrontend/Public/MetasoundDataTypeRegistrationMacro.h
@@ -10,7 +10,11 @@
 #include "MetasoundInputNode.h"
 #include "MetasoundOperatorInterface.h"
 #include "MetasoundFrontendRegistries.h"
+#include "MetasoundAutoConverterNode.h"
+#include "MetasoundNodeRegistrationMacro.h"
+#include "MetasoundConverterNodeRegistrationMacro.h"
 
+#include <type_traits>
 
 // SFINAE used to optionally invoke subclasses of IAudioProxyDataFactory when we can.
 template<typename UClassToUse, typename TEnableIf<TIsDerivedFrom<UClassToUse, IAudioProxyDataFactory>::Value, bool>::Type = true>
@@ -33,6 +37,44 @@ IAudioProxyDataFactory* CastToAudioProxyDataFactory(UObject* InObject)
 	return nullptr;
 }
 
+// This utility function can be used to check to see if we can static cast between two types, and autogenerate a node for that static cast.
+template<typename TFromDataType, typename TToDataType, typename TEnableIf<std::is_convertible<TFromDataType, TToDataType>::value, bool>::Type = true>
+void AttemptToRegisterConverter()
+{
+	using FConverterNode = Metasound::TAutoConverterNode<TFromDataType, TToDataType>;
+
+	if (!TIsSame<TFromDataType, TToDataType>::Value && !FMetasoundFrontendRegistryContainer::Get()->IsNodeRegistered<FConverterNode>())
+	{
+		ensureAlways(RegisterNodeWithFrontend<FConverterNode>());
+		
+		bool bSucessfullyRegisteredConversionNode = RegisterConversionNode<FConverterNode, TFromDataType, TToDataType>(FConverterNode::GetInputName(), FConverterNode::GetOutputName());
+		ensureAlways(bSucessfullyRegisteredConversionNode);
+	}
+}
+
+template<typename TFromDataType, typename TToDataType, typename TEnableIf<!std::is_convertible<TFromDataType, TToDataType>::value, bool>::Type = true>
+void AttemptToRegisterConverter()
+{
+	// This implementation intentionally noops, because static_cast<TFromDataType>(TToDataType&) is invalid.
+}
+
+// Here we attempt to infer and autogenerate conversions for basic datatypes.
+template<typename TDataType>
+void RegisterConverterNodes()
+{
+	// Conversions to this data type:
+	AttemptToRegisterConverter<bool, TDataType>();
+	AttemptToRegisterConverter<int32, TDataType>();
+	AttemptToRegisterConverter<float, TDataType>();
+	AttemptToRegisterConverter<FString, TDataType>();
+
+	// Conversions from this data type:
+	AttemptToRegisterConverter<TDataType, bool>();
+	AttemptToRegisterConverter<TDataType, int32>();
+	AttemptToRegisterConverter<TDataType, float>();
+	AttemptToRegisterConverter<TDataType, FString>();
+}
+
 template<typename TDataType, ::Metasound::ELiteralArgType PreferredArgType = ::Metasound::ELiteralArgType::None, typename UClassToUse = UObject>
 bool RegisterDataTypeWithFrontend()
 {
@@ -111,6 +153,9 @@ bool RegisterDataTypeWithFrontend()
 
 	bool bSucceeded = FMetasoundFrontendRegistryContainer::Get()->RegisterDataType(RegistryInfo, MoveTemp(Callbacks));
 	ensureAlwaysMsgf(bSucceeded, TEXT("Failed to register data type %s in the node registry!"), ::Metasound::TDataReferenceTypeInfo<TDataType>::TypeName);
+	
+	RegisterConverterNodes<TDataType>();
+	
 	return bSucceeded;
 }
 
diff --git a/Engine/Plugins/Runtime/Metasound/Source/MetasoundFrontend/Public/MetasoundFrontend.h b/Engine/Plugins/Runtime/Metasound/Source/MetasoundFrontend/Public/MetasoundFrontend.h
--- a/Engine/Plugins/Runtime/Metasound/Source/MetasoundFrontend/Public/MetasoundFrontend.h
+++ b/Engine/Plugins/Runtime/Metasound/Source/MetasoundFrontend/Public/MetasoundFrontend.h
@@ -155,7 +155,7 @@ namespace Metasound
 			// If Connectable is EConnectable::YesWithConverterNode,
 			// this will be a populated list of nodes we can use 
 			// to convert between the input and output.
-			TArray<FNodeClassInfo> PossibleConverterNodeClasses;
+			TArray<FConverterNodeInfo> PossibleConverterNodeClasses;
 		};
 
 		struct METASOUNDFRONTEND_API FHandleInitParams
@@ -216,7 +216,7 @@ namespace Metasound
 
 			FConnectability CanConnectTo(const FInputHandle& InHandle) const;
 			bool Connect(FInputHandle& InHandle);
-			bool ConnectWithConverterNode(FInputHandle& InHandle, FString& InNodeClassName);
+			bool ConnectWithConverterNode(FInputHandle& InHandle, const FConverterNodeInfo& InNodeClassName);
 			bool Disconnect(FInputHandle& InHandle);
 
 		private:
@@ -256,7 +256,7 @@ namespace Metasound
 
 			FConnectability CanConnectTo(const FOutputHandle& InHandle) const;
 			bool Connect(FOutputHandle& InHandle);
-			bool ConnectWithConverterNode(FOutputHandle& InHandle, FString& InNodeClassName);
+			bool ConnectWithConverterNode(FOutputHandle& InHandle, const FConverterNodeInfo& InNodeClassName);
 			bool Disconnect(FOutputHandle& InHandle);
 			bool Disconnect();
 
@@ -423,6 +423,7 @@ namespace Metasound
 			bool ClearLiteralForInput(const FString& InInputName);
 
 			FNodeHandle AddNewNode(const FNodeClassInfo& InNodeClass);
+			FNodeHandle AddNewNode(const FNodeRegistryKey& InNodeClass);
 
 			// Remove the node corresponding to this node handle.
 			// On success, invalidates the received node handle.
diff --git a/Engine/Plugins/Runtime/Metasound/Source/MetasoundFrontend/Public/MetasoundFrontendBaseClasses.h b/Engine/Plugins/Runtime/Metasound/Source/MetasoundFrontend/Public/MetasoundFrontendBaseClasses.h
--- a/Engine/Plugins/Runtime/Metasound/Source/MetasoundFrontend/Public/MetasoundFrontendBaseClasses.h
+++ b/Engine/Plugins/Runtime/Metasound/Source/MetasoundFrontend/Public/MetasoundFrontendBaseClasses.h
@@ -454,7 +454,7 @@ namespace Metasound
 						return FDescPath();
 						break;
 					}
-					/** Individiual inputs, outpus and nodes, as well as metadata, do not have any child elements. */
+					/** Individual inputs, outputs and nodes, as well as metadata, do not have any child elements. */
 					case Path::EDescType::Input:
 					case Path::EDescType::Output:
 					case Path::EDescType::Node:
@@ -508,6 +508,8 @@ namespace Metasound
 
 			FDescPath GetOutputDescriptionPath(FDescPath InPathForOutputNode, const FString& OutputName);
 
+			FDescPath GetOuterGraphPath(FDescPath InPath);
+
 			// Generates a human readable string for a path into a Metasound description. Used for debugging.
 			FString GetPrintableString(FDescPath InPath);
 		}
diff --git a/Engine/Plugins/Runtime/Metasound/Source/MetasoundFrontend/Public/MetasoundFrontendRegistries.h b/Engine/Plugins/Runtime/Metasound/Source/MetasoundFrontend/Public/MetasoundFrontendRegistries.h
--- a/Engine/Plugins/Runtime/Metasound/Source/MetasoundFrontend/Public/MetasoundFrontendRegistries.h
+++ b/Engine/Plugins/Runtime/Metasound/Source/MetasoundFrontend/Public/MetasoundFrontendRegistries.h
@@ -95,6 +95,48 @@ namespace Metasound
 			{
 			}
 		};
+
+		struct METASOUNDFRONTEND_API FConverterNodeRegistryKey
+		{
+			// The datatype one would like to convert from.
+			FName FromDataType;
+
+			// The datatype one would like to convert to.
+			FName ToDataType;
+
+			FORCEINLINE bool operator==(const FConverterNodeRegistryKey& Other) const
+			{
+				return FromDataType == Other.FromDataType && ToDataType == Other.ToDataType;
+			}
+
+			friend uint32 GetTypeHash(const ::Metasound::Frontend::FConverterNodeRegistryKey& InKey)
+			{
+				return HashCombine(GetTypeHash(InKey.FromDataType), GetTypeHash(InKey.ToDataType));
+			}
+		};
+
+		struct METASOUNDFRONTEND_API FConverterNodeInfo
+		{
+			// If this node has multiple input pins, we use this to designate which pin should be used.
+			FString PreferredConverterInputPin;
+
+			// If this node has multiple output pins, we use this to designate which pin should be used.
+			FString PreferredConverterOutputPin;
+
+			// The key for this node in the node registry.
+			FNodeRegistryKey NodeKey;
+
+			FORCEINLINE bool operator==(const FConverterNodeInfo& Other) const
+			{
+				return NodeKey == Other.NodeKey;
+			}
+		};
+
+		struct METASOUNDFRONTEND_API FConverterNodeRegistryValue
+		{
+			// A list of nodes that can perform a conversion between the two datatypes described in the FConverterNodeRegistryKey for this map element.
+			TArray<FConverterNodeInfo> PotentialConverterNodes;
+		};
 	}
 
 	struct FDataTypeConstructorCallbacks
@@ -115,17 +157,27 @@ namespace Metasound
  */
 class METASOUNDFRONTEND_API FMetasoundFrontendRegistryContainer
 {
+	using FConverterNodeRegistryKey = ::Metasound::Frontend::FConverterNodeRegistryKey;
+	using FConverterNodeRegistryValue = ::Metasound::Frontend::FConverterNodeRegistryValue;
+	using FConverterNodeInfo = ::Metasound::Frontend::FConverterNodeInfo;
+
+	using FNodeRegistryKey = Metasound::Frontend::FNodeRegistryKey;
+	using FNodeRegistryElement = Metasound::Frontend::FNodeRegistryElement;
+
+	using FDataTypeRegistryInfo = Metasound::FDataTypeRegistryInfo;
+	using FDataTypeConstructorCallbacks = ::Metasound::FDataTypeConstructorCallbacks;
+
 public:
 	static FMetasoundFrontendRegistryContainer* Get();
 	static void ShutdownMetasoundFrontend();
 
 	FMetasoundFrontendRegistryContainer(const FMetasoundFrontendRegistryContainer&) = delete;
 
+	// This is called on module startup. This invokes any registration commands enqueued by our registration macros.
 	void InitializeFrontend();
-	bool EnqueueInitCommand(TUniqueFunction<void()>&& InFunc);
 
-	using FNodeRegistryKey = Metasound::Frontend::FNodeRegistryKey;
-	using FNodeRegistryElement = Metasound::Frontend::FNodeRegistryElement;
+
+	bool EnqueueInitCommand(TUniqueFunction<void()>&& InFunc);
 
 	TMap<FNodeRegistryKey, FNodeRegistryElement>& GetExternalNodeRegistry();
 
@@ -135,8 +187,13 @@ public:
 	Metasound::FDataTypeLiteralParam GenerateLiteralForUObject(const FName& InDataType, UObject* InObject);
 	Metasound::FDataTypeLiteralParam GenerateLiteralForUObjectArray(const FName& InDataType, TArray<UObject*> InObjectArray);
 
+	// Create a new instance of a C++ implemented node from the registry.
 	TUniquePtr<Metasound::INode> ConstructExternalNode(const FName& InNodeType, uint32 InNodeHash, const Metasound::FNodeInitData& InInitData);
 
+	// Returns a list of possible nodes to use to convert from FromDataType to ToDataType.
+	// Returns an empty array if none are available.
+	TArray<FConverterNodeInfo> GetPossibleConverterNodes(const FName& FromDataType, const FName& ToDataType);
+
 	// Get the desired kind of literal for a given data type. Returns EConstructorArgType::Invalid if the data type couldn't be found.
 	Metasound::ELiteralArgType GetDesiredLiteralTypeForDataType(FName InDataType) const;
 
@@ -146,19 +203,39 @@ public:
 	// Get whether we can build a literal of this specific type for InDataType.
 	bool DoesDataTypeSupportLiteralType(FName InDataType, Metasound::ELiteralArgType InLiteralType) const;
 
-	bool RegisterDataType(const ::Metasound::FDataTypeRegistryInfo& InDataInfo, ::Metasound::FDataTypeConstructorCallbacks&& InCallbacks);
+	bool RegisterDataType(const FDataTypeRegistryInfo& InDataInfo, FDataTypeConstructorCallbacks&& InCallbacks);
 	bool RegisterExternalNode(FNodeGetterCallback&& InCallback);
+	bool RegisterConversionNode(const FConverterNodeRegistryKey& InNodeKey, const FConverterNodeInfo& InNodeInfo);
+
+	template <typename TNodeType>
+	static FNodeRegistryKey GetRegistryKeyForNode()
+	{
+		Metasound::FNodeInitData DummyInitData;
+
+		TNodeType DummyNode = TNodeType(DummyInitData);
+
+		return GetRegistryKeyForNodeInternal(DummyNode);
+	}
+
+	template <typename TNodeType>
+	bool IsNodeRegistered()
+	{
+		return ExternalNodeRegistry.Contains(GetRegistryKeyForNode<TNodeType>());
+	}
 
 	// Return any data types that can be used as a metasound input type or output type.
 	TArray<FName> GetAllValidDataTypes();
 
 	// Get info about a specific data type (what kind of literals we can use, etc.)
 	// @returns false if InDataType wasn't found in the registry. 
-	bool GetInfoForDataType(FName InDataType, Metasound::FDataTypeRegistryInfo& OutInfo);
+	bool GetInfoForDataType(FName InDataType, FDataTypeRegistryInfo& OutInfo);
 
 private:
 	FMetasoundFrontendRegistryContainer();
 
+	// Internal convenience function for generating a hash for an INode.
+	static FNodeRegistryKey GetRegistryKeyForNodeInternal(const Metasound::INode& InNode);
+
 	static FMetasoundFrontendRegistryContainer* LazySingleton;
 
 	// This buffer is used to enqueue nodes and datatypes to register when the module has been initialized,
@@ -166,7 +243,7 @@ private:
 	// The bad news is that TInlineAllocator is the safest allocator to use on static init.
 	// The good news is that none of these lambdas typically have captures, so this should have low memory overhead.
 	static constexpr int32 MaxNumNodesAndDatatypesToInitialize = 8192;
-	TArray<TUniqueFunction<void()>/*, TInlineAllocator<MaxNumNodesAndDatatypesToInitialize>*/> LazyInitCommands;
+	TArray<TUniqueFunction<void()>, TInlineAllocator<MaxNumNodesAndDatatypesToInitialize>> LazyInitCommands;
 	
 	FCriticalSection LazyInitCommandCritSection;
 	bool bHasModuleBeenInitialized;
@@ -174,6 +251,9 @@ private:
 	// Registry in which we keep all information about nodes implemented in C++.
 	TMap<FNodeRegistryKey, FNodeRegistryElement> ExternalNodeRegistry;
 
+	// Registry in which we keep lists of possible nodes to use to convert between two datatypes
+	TMap<FConverterNodeRegistryKey, FConverterNodeRegistryValue> ConverterNodeRegistry;
+
 	struct FDataTypeRegistryElement
 	{
 		Metasound::FDataTypeConstructorCallbacks Callbacks;
diff --git a/Engine/Plugins/Runtime/Metasound/Source/MetasoundFrontend/Public/MetasoundNodeRegistrationMacro.h b/Engine/Plugins/Runtime/Metasound/Source/MetasoundFrontend/Public/MetasoundNodeRegistrationMacro.h
--- a/Engine/Plugins/Runtime/Metasound/Source/MetasoundFrontend/Public/MetasoundNodeRegistrationMacro.h
+++ b/Engine/Plugins/Runtime/Metasound/Source/MetasoundFrontend/Public/MetasoundNodeRegistrationMacro.h
@@ -29,7 +29,7 @@ bool RegisterNodeWithFrontend()
 	static bool bAlreadyRegisteredThisDataType = false;
 	if (bAlreadyRegisteredThisDataType)
 	{
-		UE_LOG(LogTemp, Display, TEXT("Tried to call REGISTER_METASOUND_DATATYPE twice with the same class. ignoring the second call. Likely because METASOUND_REGISTER_NODE is in a header that's used in multiple modules. Consider moving it to a private header or cpp file."));
+		UE_LOG(LogTemp, Display, TEXT("Tried to call METASOUND_REGISTER_NODE twice with the same class. ignoring the second call. Likely because METASOUND_REGISTER_NODE is in a header that's used in multiple modules. Consider moving it to a private header or cpp file."));
 		return false;
 	}
 
