Remove need for magic number.

diff --git a/Engine/Plugins/Runtime/Metasound/Source/MetasoundEngine/Public/MetasoundWave.h b/Engine/Plugins/Runtime/Metasound/Source/MetasoundEngine/Public/MetasoundWave.h
--- a/Engine/Plugins/Runtime/Metasound/Source/MetasoundEngine/Public/MetasoundWave.h
+++ b/Engine/Plugins/Runtime/Metasound/Source/MetasoundEngine/Public/MetasoundWave.h
@@ -66,6 +66,6 @@ namespace Metasound
 		}
 	};
 
-	DECLARE_METASOUND_DATA_REFERENCE_TYPES(FWave, METASOUNDENGINE_API, 0x0ddba11, FWaveTypeInfo, FWaveReadRef, FWaveWriteRef)
-	DECLARE_METASOUND_DATA_REFERENCE_TYPES(FWaveAsset, METASOUNDENGINE_API, 0x0ddba42, FWaveAssetTypeInfo, FWaveAssetReadRef, FWaveAssetWriteRef)
+	DECLARE_METASOUND_DATA_REFERENCE_TYPES(FWave, METASOUNDENGINE_API, FWaveTypeInfo, FWaveReadRef, FWaveWriteRef)
+	DECLARE_METASOUND_DATA_REFERENCE_TYPES(FWaveAsset, METASOUNDENGINE_API, FWaveAssetTypeInfo, FWaveAssetReadRef, FWaveAssetWriteRef)
 }
diff --git a/Engine/Plugins/Runtime/Metasound/Source/MetasoundGraphCore/Public/MetasoundDataReference.h b/Engine/Plugins/Runtime/Metasound/Source/MetasoundGraphCore/Public/MetasoundDataReference.h
--- a/Engine/Plugins/Runtime/Metasound/Source/MetasoundGraphCore/Public/MetasoundDataReference.h
+++ b/Engine/Plugins/Runtime/Metasound/Source/MetasoundGraphCore/Public/MetasoundDataReference.h
@@ -3,26 +3,27 @@
 #pragma once
 
 #include "CoreMinimal.h"
-#include "Misc/TVariant.h"
+#include "MetasoundOperatorSettings.h"
+#include <type_traits>
 
 /** Macro to make declaring a metasound parameter simple.  */
 // Declares a metasound parameter type by
 // - Adding typedefs for commonly used template types.
 // - Defining parameter type traits.
-#define DECLARE_METASOUND_DATA_REFERENCE_TYPES(DataType, ModuleApi, DataTypeMagicNumber, DataTypeInfoTypeName, DataReadReferenceTypeName, DataWriteReferenceTypeName) \
+#define DECLARE_METASOUND_DATA_REFERENCE_TYPES(DataType, ModuleApi, DataTypeInfoTypeName, DataReadReferenceTypeName, DataWriteReferenceTypeName) \
 	template<> \
 	struct ::Metasound::TDataReferenceTypeInfo<DataType> \
 	{ \
 		static ModuleApi const TCHAR* TypeName; \
-		static constexpr ::Metasound::FDataTypeMagicNumber MagicNumber = (DataTypeMagicNumber); \
+		static constexpr ::Metasound::FMetasoundDataTypeId TypeId = static_cast<::Metasound::FMetasoundDataTypeId>(::Metasound::GetMetasoundDataTypeIdPtr<DataType>()); \
 		static constexpr bool bIsStringParsable = TTestIfDataTypeCtorIsImplemented<DataType, const FString&>::Value; \
-		static constexpr bool bIsBoolParsable = TTestIfDataTypeCtorIsImplemented<DataType, bool>::Value; \
-		static constexpr bool bIsIntParsable = TTestIfDataTypeCtorIsImplemented<DataType, int32>::Value; \
-		static constexpr bool bIsFloatParsable = TTestIfDataTypeCtorIsImplemented<DataType, float>::Value; \
-		static constexpr bool bIsProxyParsable =  TTestIfDataTypeCtorIsImplemented<DataType, const Audio::IProxyData&>::Value; \
-		static constexpr bool bIsProxyArrayParsable =  TTestIfDataTypeCtorIsImplemented<DataType, const Audio::IProxyData&>::Value; \
-		static constexpr bool bIsConstructableWithSettings = TTestIfDataTypeDefaultCtorIsImplemented<DataType>::Value; \
-		static constexpr bool bCanUseDefaultConstructor = TIsConstructible<DataType>::Value; \
+		static constexpr bool bIsBoolParsable = ::Metasound::TTestIfDataTypeCtorIsImplemented<DataType, bool>::Value; \
+		static constexpr bool bIsIntParsable = ::Metasound::TTestIfDataTypeCtorIsImplemented<DataType, int32>::Value; \
+		static constexpr bool bIsFloatParsable = ::Metasound::TTestIfDataTypeCtorIsImplemented<DataType, float>::Value; \
+		static constexpr bool bIsProxyParsable =  ::Metasound::TTestIfDataTypeCtorIsImplemented<DataType, const Audio::IProxyData&>::Value; \
+		static constexpr bool bIsProxyArrayParsable =  ::Metasound::TTestIfDataTypeCtorIsImplemented<DataType, const Audio::IProxyData&>::Value; \
+		static constexpr bool bIsConstructableWithSettings = ::Metasound::TTestIfDataTypeSettingsCtorIsImplemented<DataType>::Value; \
+		static constexpr bool bCanUseDefaultConstructor = std::is_constructible<DataType>::value; \
 		static constexpr bool bIsValidSpecialization = true; \
 	}; \
 	\
@@ -35,33 +36,64 @@
 #define IMPL_METASOUND_DATA_TYPE(DataType, DataTypeName) \
 	const TCHAR* ::Metasound::TDataReferenceTypeInfo<DataType>::TypeName = TEXT(DataTypeName);
 
+namespace Audio
+{
+	// Forward declare
+	class IProxyData;
+}
+
 namespace Metasound
 {
-	/** ID for parameter type. Should be unique for different C++ types. */
-	typedef int32 FDataTypeMagicNumber;
+	/** Used to generate a unique ID for each registered data type. */
+	template<typename DataType>
+	struct TMetasoundDataTypeIdPtr
+	{
+		/** static const address of specific type as an ID. */
+		static const DataType* const Ptr;
+	};
 
-	/** A Data Reference Interface.
-	 *
-	 * A parameter references provides information and access to a shared object in the graph. 
+	template<typename DataType>
+	const DataType* const TMetasoundDataTypeIdPtr<DataType>::Ptr = nullptr;
+
+	/** Returns an address of a pointer which is defined once per a Metasound
+	 * data type. 
 	 */
-	class IDataReference
+	template <typename DataType>
+	constexpr const DataType* const* GetMetasoundDataTypeIdPtr() noexcept 
 	{
-		public:
+		return &TMetasoundDataTypeIdPtr<DataType>::Ptr;
+	}
 
-		virtual ~IDataReference() 
-		{
-		}
+	/** ID for parameter type. Should be unique for different C++ types. */
+	using FMetasoundDataTypeId = const void*;
 
-		/** Returns the name of the data type. */
-		virtual const FName& GetDataTypeName() const = 0;
+	/** Test if a data type has a constructor which accepts FOperatorSettings */
+	template <typename TDataType>
+	struct TTestIfDataTypeSettingsCtorIsImplemented
+	{
+	public:
+		static constexpr bool Value = std::is_constructible<TDataType, const ::Metasound::FOperatorSettings&>::value;
+	};
 
-		/** Returns the ID of the parameter type. */
-		virtual FDataTypeMagicNumber GetDataTypeMagicNumber() const = 0;
+	/** Test if a stat type has a constructor which accepts the given data type 
+	 * and a FOperatorSettings.
+	 */
+	template <typename TDataType, typename TTypeToParse>
+	struct TTestIfDataTypeCtorIsImplemented
+	{
+	private:
+		static constexpr bool bSupportsConstructionWithSettings = 
+			std::is_constructible<TDataType, TTypeToParse, const ::Metasound::FOperatorSettings&>::value
+			|| std::is_constructible<TDataType, TTypeToParse, ::Metasound::FOperatorSettings>::value;
 
-		/** Creates a copy of the parameter type. */
-		virtual TUniquePtr<IDataReference> Clone() const = 0;
+		static constexpr bool bSupportsConstructionWithoutSettings = std::is_constructible<TDataType, TTypeToParse>::value;
+
+	public:
+
+		static constexpr bool Value = bSupportsConstructionWithSettings || bSupportsConstructionWithoutSettings;
 	};
 
+
 	/** Helper class to enforce specialization of TDataReferenceTypeInfo */
 	template<typename DataType>
 	struct TSpecializationHelper 
@@ -82,7 +114,7 @@ namespace Metasound
 		static constexpr const TCHAR TypeName[] = TEXT("");
 
 		/** Magic number used to check the data type when casting. */
-		static constexpr const FDataTypeMagicNumber MagicNumber = -1;
+		static constexpr const FMetasoundDataTypeId TypeId = nullptr;
 
 		// This static assert is triggered if TDataReferenceTypeInfo is used 
 		// without specialization.
@@ -99,11 +131,32 @@ namespace Metasound
 		static constexpr bool bIsValidSpecialization = false;
 	};
 
+
+	/** Return the data type name for a registered data type. */
+	template<typename DataType>
+	const FName GetMetasoundDataTypeName() 
+	{
+		static const FName TypeName = FName(TDataReferenceTypeInfo<std::decay_t<DataType>>::TypeName);
+
+		return TypeName;
+	}
+
+	/** Return the data type ID for a registered data type. 
+	 *
+	 * This ID is runtime constant but may change between executions and builds.
+	 */
+	template<typename DataType>
+	const FMetasoundDataTypeId GetMetasoundDataTypeId() 
+	{
+		return TDataReferenceTypeInfo<std::decay_t<DataType>>::TypeId;
+	}
+
+	/** Specialize void data type for internal use. */
 	template<>
 	struct TDataReferenceTypeInfo<void>
 	{
 		static METASOUNDGRAPHCORE_API const TCHAR* TypeName;
-		static constexpr const FDataTypeMagicNumber MagicNumber = -1;
+		static constexpr FMetasoundDataTypeId TypeId = static_cast<FMetasoundDataTypeId>(GetMetasoundDataTypeIdPtr<void>());
 		static constexpr bool bIsStringParsable = false;
 		static constexpr bool bIsBoolParsable = false;
 		static constexpr bool bIsIntParsable = false;
@@ -115,23 +168,40 @@ namespace Metasound
 		static constexpr bool bIsValidSpecialization = false;
 	};
 
-	// TODO: comments
-	template<typename DataType>
-	const FName GetMetasoundDataTypeName() 
+	/** A Data Reference Interface.
+	 *
+	 * A parameter references provides information and access to a shared object in the graph. 
+	 */
+	class IDataReference
 	{
-		static const FName TypeName = FName(TDataReferenceTypeInfo<typename TDecay<DataType>::Type >::TypeName);
+		public:
 
-		return TypeName;
-	}
+		virtual ~IDataReference() 
+		{
+		}
+
+		/** Returns the name of the data type. */
+		virtual const FName& GetDataTypeName() const = 0;
+
+		/** Returns the ID of the parameter type. */
+		virtual FMetasoundDataTypeId GetDataTypeId() const = 0;
+
+		/** Creates a copy of the parameter type. */
+		virtual TUniquePtr<IDataReference> Clone() const = 0;
+	};
 	
-	// TODO: comments
+	/** Test if an IDataReference contains the same data type as the template
+	 * parameter.
+	 *
+	 * @return True if the IDataReference contains the DataType. False otherwise. 
+	 */
 	template<typename DataType>
-	bool IsReferenceOfType(const IDataReference& InReference)
+	bool IsDataReferenceOfType(const IDataReference& InReference)
 	{
 		static const FName TypeName = GetMetasoundDataTypeName<DataType>();
-		static const FDataTypeMagicNumber MagicNumber = TDataReferenceTypeInfo<typename TDecay<DataType>::Type >::MagicNumber;
+		static const FMetasoundDataTypeId TypeId = GetMetasoundDataTypeId<DataType>();
 
-		return (InReference.GetDataTypeName() == TypeName) && (InReference.GetDataTypeMagicNumber() == MagicNumber);
+		return (InReference.GetDataTypeName() == TypeName) && (InReference.GetDataTypeId() == TypeId);
 	}
 
 	// This enum is used as a token to explicitly delineate when we should create a new object for the reference,
@@ -149,6 +219,7 @@ namespace Metasound
 	template <typename DataType>
 	class TDataReference : public IDataReference
 	{
+		static_assert(std::is_same<DataType, std::decay_t<DataType>>::value, "Data types used as data references must not decay");
 	protected:
 		/**
 		 * This constructor forwards arguments to an underlying constructor.
@@ -165,44 +236,12 @@ namespace Metasound
 			typedef TDataReferenceTypeInfo<DataType> FInfoType;
 
 
-			/*
-			 * This constructor forwards arguments to the constructor of the underlying DataType.
-			 *
-			 * SFINAE is used to disable copy and move constructors as those are defined separately. 
-			 
-			template <
-				typename... ArgTypes,
-				typename = typename TEnableIf<
-					TAndValue<
-						sizeof...(ArgTypes) != 0,
-						TOrValue<
-							sizeof...(ArgTypes) != 1,
-							TNot< TIsDerivedFrom< typename TDecay< typename TNthTypeFromParameterPack< 0, ArgTypes... >::Type >::Type,  TDataReference<DataType> > >
-						>
-					>::Value
-				>::Type
-			>
-			TDataReference(ArgTypes&&... Args)
-			:	ObjectReference(MakeShared<DataType, ESPMode::NotThreadSafe>(Forward<ArgTypes>(Args)...))
-			{
-			}
-
-			
-			
-			// TODO: having issues compiling this on linux when no default constructor is available. 
-			// Construct operator with no arguments if the DataType has a default constructor.
-			template< typename = typename TEnableIf< TIsConstructible<DataType>::Value >::Type >
-			TDataReference()
-			:	ObjectReference(MakeShared<DataType, ESPMode::NotThreadSafe>())
-			{
-			}
-			*/
-
-			// This should be used to construct a new DataType object and return this TDataReference as a wrapper around it.
+			/** This should be used to construct a new DataType object and return this TDataReference as a wrapper around it.
+			 */
 			template <typename... ArgTypes>
 			static TDataReference<DataType> CreateNew(ArgTypes&&... Args)
 			{
-				static_assert(TIsConstructible<DataType, ArgTypes...>::Value, "Tried to call TDataReference::CreateNew with args that don't match any constructor for an underlying type!");
+				static_assert(std::is_constructible<DataType, ArgTypes...>::value, "Tried to call TDataReference::CreateNew with args that don't match any constructor for an underlying type!");
 				return TDataReference<DataType>(EDataRefShouldConstruct::NewObject, Forward<ArgTypes>(Args)...);
 			}
 
@@ -221,24 +260,16 @@ namespace Metasound
 			/** Return the name of the underlying type. */
 			virtual const FName& GetDataTypeName() const override
 			{
-				static const FName TypeName = GetMetasoundDataTypeName<DataType>();
-				return TypeName;
-			}
+				static const FName Name = GetMetasoundDataTypeName<DataType>();
 
-			/** Return the ID of the underlying type. */
-			virtual FDataTypeMagicNumber GetDataTypeMagicNumber() const override
-			{
-				return FInfoType::MagicNumber;
+				return Name;
 			}
 
-			/** Create a clone of this object. 
-			virtual TUniquePtr<IDataReference> Clone() const override
+			/** Return the ID of the underlying type. */
+			virtual FMetasoundDataTypeId GetDataTypeId() const override
 			{
-				typedef TDataReference<DataType> FDataReference;
-
-				return MakeUnique< FDataReference >(*this);
+				return GetMetasoundDataTypeId<DataType>();
 			}
-			*/
 
 		protected:
 
@@ -266,40 +297,11 @@ namespace Metasound
 		public:
 			typedef TDataReference<DataType> FDataReference;
 
-			/*
-			// Construct operator type with arguments to "DataType" constructor. 
-			template<
-				typename... ArgTypes,
-				typename = typename TEnableIf<
-					TAndValue<
-						sizeof...(ArgTypes) != 0,
-						TOrValue<
-							sizeof...(ArgTypes) != 1,
-							TNot< TIsDerivedFrom< typename TDecay< typename TNthTypeFromParameterPack< 0, ArgTypes... >::Type >::Type, TDataReference<DataType> > >
-						>
-					>::Value
-				>::Type
-			>
-			explicit TDataWriteReference(ArgTypes&&... Args)
-			:	FDataReference(Forward<ArgTypes>(Args)...)
-			{
-			}
-
-			
-			// Construct operator with no arguments if the DataType has a default constructor.
-			template< typename = typename TEnableIf< TIsConstructible<DataType>::Value >::Type >
-			TDataWriteReference()
-			:	FDataReference()
-			{
-			}
-
-			*/
-
-			// This should be used to construct a new DataType object and return this TDataWriteReference as a wrapper around it.
+			/** This should be used to construct a new DataType object and return this TDataWriteReference as a wrapper around it. */
 			template <typename... ArgTypes>
 			static TDataWriteReference<DataType> CreateNew(ArgTypes&&... Args)
 			{
-				static_assert(TIsConstructible<DataType, ArgTypes...>::Value, "Tried to call TDataWriteReference::CreateNew with args that don't match any constructor for an underlying type!");
+				static_assert(std::is_constructible<DataType, ArgTypes...>::value, "Tried to call TDataWriteReference::CreateNew with args that don't match any constructor for an underlying type!");
 				return TDataWriteReference<DataType>(EDataRefShouldConstruct::NewObject, Forward<ArgTypes>(Args)...);
 			}
 
@@ -325,14 +327,12 @@ namespace Metasound
 			FORCEINLINE DataType& operator*() const
 			{
 				return *TDataReference<DataType>::ObjectReference;
-				//return *ObjectReference;
 			}
 
 			/** Non-const access to the underlying parameter object. */
 			FORCEINLINE DataType* operator->() const
 			{
 				return TDataReference<DataType>::ObjectReference.operator->();
-				//return ObjectReference.operator->();
 			}
 
 			/** Create a clone of this parameter reference. */
@@ -393,39 +393,11 @@ namespace Metasound
 		public:
 			typedef TDataReference<DataType> FDataReference;
 
-			/*
-			// Construct operator type with arguments to "DataType" constructor. 
-			template<
-				typename... ArgTypes,
-				typename = typename TEnableIf<
-					TAndValue<
-						sizeof...(ArgTypes) != 0,
-						TOrValue<
-							sizeof...(ArgTypes) != 1,
-							TNot< TIsDerivedFrom< typename TDecay< typename TNthTypeFromParameterPack< 0, ArgTypes... >::Type >::Type, TDataReference<DataType> > >
-						>
-					>::Value
-				>::Type
-			>
-			explicit TDataReadReference(ArgTypes&&... Args)
-			:	FDataReference(Forward<ArgTypes>(Args)...)
-			{
-			}
-
-			
-			// Construct operator with no arguments if the DataType has a default constructor.
-			template< typename = typename TEnableIf< TIsConstructible<DataType>::Value >::Type >
-			TDataReadReference()
-			:	FDataReference()
-			{
-			}
-			*/
-
 			// This should be used to construct a new DataType object and return this TDataReadReference as a wrapper around it.
 			template <typename... ArgTypes>
 			static TDataReadReference<DataType> CreateNew(ArgTypes&&... Args)
 			{
-				static_assert(TIsConstructible<DataType, ArgTypes...>::Value, "Tried to call TDataReadReference::CreateNew with args that don't match any constructor for an underlying type!");
+				static_assert(std::is_constructible<DataType, ArgTypes...>::value, "Tried to call TDataReadReference::CreateNew with args that don't match any constructor for an underlying type!");
 				return TDataReadReference<DataType>(EDataRefShouldConstruct::NewObject, Forward<ArgTypes>(Args)...);
 			}
 
diff --git a/Engine/Plugins/Runtime/Metasound/Source/MetasoundGraphCore/Public/MetasoundDataReferenceCollection.h b/Engine/Plugins/Runtime/Metasound/Source/MetasoundGraphCore/Public/MetasoundDataReferenceCollection.h
--- a/Engine/Plugins/Runtime/Metasound/Source/MetasoundGraphCore/Public/MetasoundDataReferenceCollection.h
+++ b/Engine/Plugins/Runtime/Metasound/Source/MetasoundGraphCore/Public/MetasoundDataReferenceCollection.h
@@ -187,7 +187,7 @@ namespace Metasound
 
 				check(nullptr != DataRefPtr);
 
-				check(HasMatchingTypeInfo<DataType>(*DataRefPtr));
+				check(IsDataReferenceOfType<DataType>(*DataRefPtr));
 
 				return FDataRefType(*static_cast<const FDataRefType*>(DataRefPtr));
 			}
@@ -230,7 +230,7 @@ namespace Metasound
 
 				check(nullptr != DataRefPtr);
 
-				check(HasMatchingTypeInfo<DataType>(*DataRefPtr));
+				check(IsDataReferenceOfType<DataType>(*DataRefPtr));
 
 				return FDataRefType(*static_cast<const FDataRefType*>(DataRefPtr));
 			}
@@ -269,24 +269,13 @@ namespace Metasound
 
 					if (nullptr != DataRefPtr)
 					{
-						return HasMatchingTypeInfo<DataType>(*DataRefPtr);
+						return IsDataReferenceOfType<DataType>(*DataRefPtr);
 					}
 				}
 
 				return false;
 			}
 
-			template<typename DataType>
-			bool HasMatchingTypeInfo(const IDataReference& InRef) const
-			{
-				static const FName DataTypeName = FName(TDataReferenceTypeInfo<DataType>::TypeName);
-
-				bool bTypeNameMatch = InRef.GetDataTypeName() == DataTypeName;
-				bool bTypeIdMatch = InRef.GetDataTypeMagicNumber() == TDataReferenceTypeInfo<DataType>::MagicNumber;
-
-				return bTypeNameMatch && bTypeIdMatch;
-			}
-
 			const IDataReference* GetDataReference(const FDataReferenceMap& InMap, const FString& InName) const;
 
 			void AddDataReference(FDataReferenceMap& InMap, const FString& InName, FDataRefWrapper&& InDataReference);
diff --git a/Engine/Plugins/Runtime/Metasound/Source/MetasoundGraphCore/Public/MetasoundOperatorInterface.h b/Engine/Plugins/Runtime/Metasound/Source/MetasoundGraphCore/Public/MetasoundOperatorInterface.h
--- a/Engine/Plugins/Runtime/Metasound/Source/MetasoundGraphCore/Public/MetasoundOperatorInterface.h
+++ b/Engine/Plugins/Runtime/Metasound/Source/MetasoundGraphCore/Public/MetasoundOperatorInterface.h
@@ -6,8 +6,10 @@
 
 #include "MetasoundDataReference.h"
 #include "MetasoundDataReferenceCollection.h"
-#include "MetasoundOperatorSettings.h"
-#include "IAudioProxyInitializer.h"
+#include "MetasoundOperatorSettings.h" // < TODO: For literal param.  Should remove when literal param is moved. 
+#include "Misc/TVariant.h" // < TODO: For literal param.  Should remove when literal param is moved. 
+#include "IAudioProxyInitializer.h" // < TODO: For literal param.  Should remove when literal param is moved. 
+
 
 namespace Metasound
 {
@@ -56,32 +58,7 @@ namespace Metasound
 	};
 }
 
-template <typename TDataType, typename TTypeToParse>
-struct TTestIfDataTypeCtorIsImplemented
-{
-private:
-	static constexpr bool bSupportsConstructionWithSettings = 
-		TIsConstructible<TDataType, TTypeToParse, const ::Metasound::FOperatorSettings&>::Value
-		|| TIsConstructible<TDataType, TTypeToParse, ::Metasound::FOperatorSettings&>::Value
-		|| TIsConstructible<TDataType, TTypeToParse, ::Metasound::FOperatorSettings>::Value;
-
-	static constexpr bool bSupportsConstructionWithoutSettings = TIsConstructible<TDataType, TTypeToParse>::Value;
 
-public:
-	/*
-	static constexpr bool Value = TIsConstructible<TDataType, TTypeToParse, const ::Metasound::FOperatorSettings&>::Value
-		|| TIsConstructible<TDataType, TTypeToParse>::Value;
-		*/
-
-	static constexpr bool Value = bSupportsConstructionWithSettings || bSupportsConstructionWithoutSettings;
-};
-
-template <typename TDataType>
-struct TTestIfDataTypeDefaultCtorIsImplemented
-{
-public:
-	static constexpr bool Value = TIsConstructible<TDataType, const ::Metasound::FOperatorSettings&>::Value;
-};
 
 namespace Metasound
 {
@@ -118,6 +95,7 @@ namespace Metasound
 	}
 	*/
 
+	// TODO: move data type literal to own file.
 	/**
 	 * Convenience wrapper for safely invoking the correct constructor.
 	 * This can be used for int32, float, bool, or FString.
@@ -235,7 +213,7 @@ namespace Metasound
 				}
 				case ELiteralArgType::None:
 				{
-					return TTestIfDataTypeDefaultCtorIsImplemented<DataType>::Value || TIsConstructible<DataType>::Value;
+					return TTestIfDataTypeSettingsCtorIsImplemented<DataType>::Value || std::is_constructible<DataType>::value;
 				}
 				default:
 				{
diff --git a/Engine/Plugins/Runtime/Metasound/Source/MetasoundGraphCore/Public/MetasoundVertex.h b/Engine/Plugins/Runtime/Metasound/Source/MetasoundGraphCore/Public/MetasoundVertex.h
--- a/Engine/Plugins/Runtime/Metasound/Source/MetasoundGraphCore/Public/MetasoundVertex.h
+++ b/Engine/Plugins/Runtime/Metasound/Source/MetasoundGraphCore/Public/MetasoundVertex.h
@@ -146,7 +146,7 @@ namespace Metasound
 		 */
 		virtual bool IsReferenceOfSameType(const IDataReference& InReference) const override
 		{
-			return IsReferenceOfType<DataType>(InReference);
+			return IsDataReferenceOfType<DataType>(InReference);
 		}
 
 
diff --git a/Engine/Plugins/Runtime/Metasound/Source/MetasoundStandardNodes/Public/MetasoundAudioBuffer.h b/Engine/Plugins/Runtime/Metasound/Source/MetasoundStandardNodes/Public/MetasoundAudioBuffer.h
--- a/Engine/Plugins/Runtime/Metasound/Source/MetasoundStandardNodes/Public/MetasoundAudioBuffer.h
+++ b/Engine/Plugins/Runtime/Metasound/Source/MetasoundStandardNodes/Public/MetasoundAudioBuffer.h
@@ -154,5 +154,5 @@ namespace Metasound
 #endif				
 	};
 
-	DECLARE_METASOUND_DATA_REFERENCE_TYPES(FAudioBuffer, METASOUNDSTANDARDNODES_API, 0x1b02dac2 , FAudioBufferTypeInfo, FAudioBufferReadRef, FAudioBufferWriteRef);
+	DECLARE_METASOUND_DATA_REFERENCE_TYPES(FAudioBuffer, METASOUNDSTANDARDNODES_API, FAudioBufferTypeInfo, FAudioBufferReadRef, FAudioBufferWriteRef);
 };
diff --git a/Engine/Plugins/Runtime/Metasound/Source/MetasoundStandardNodes/Public/MetasoundAudioFormats.h b/Engine/Plugins/Runtime/Metasound/Source/MetasoundStandardNodes/Public/MetasoundAudioFormats.h
--- a/Engine/Plugins/Runtime/Metasound/Source/MetasoundStandardNodes/Public/MetasoundAudioFormats.h
+++ b/Engine/Plugins/Runtime/Metasound/Source/MetasoundStandardNodes/Public/MetasoundAudioFormats.h
@@ -456,15 +456,15 @@ namespace Metasound
 	};
 
 
-	DECLARE_METASOUND_DATA_REFERENCE_TYPES(FUnformattedAudio, METASOUNDSTANDARDNODES_API, 0xd78a3ed1 , FUnformattedAudioTypeInfo, FUnformattedAudioReadRef, FUnformattedAudioWriteRef);
+	DECLARE_METASOUND_DATA_REFERENCE_TYPES(FUnformattedAudio, METASOUNDSTANDARDNODES_API, FUnformattedAudioTypeInfo, FUnformattedAudioReadRef, FUnformattedAudioWriteRef);
 	
 
-	DECLARE_METASOUND_DATA_REFERENCE_TYPES(FMultichannelAudioFormat, METASOUNDSTANDARDNODES_API, 0x56bdcbe0 , FMultichannelAudioFormatTypeInfo, FMultichannelAudioFormatReadRef, FMultichannelAudioFormatWriteRef);
+	DECLARE_METASOUND_DATA_REFERENCE_TYPES(FMultichannelAudioFormat, METASOUNDSTANDARDNODES_API, FMultichannelAudioFormatTypeInfo, FMultichannelAudioFormatReadRef, FMultichannelAudioFormatWriteRef);
 	
 
-	DECLARE_METASOUND_DATA_REFERENCE_TYPES(FMonoAudioFormat, METASOUNDSTANDARDNODES_API, 0x6f468c8c, FMonoAudioFormatTypeInfo, FMonoAudioFormatReadRef, FMonoAudioFormatWriteRef);
+	DECLARE_METASOUND_DATA_REFERENCE_TYPES(FMonoAudioFormat, METASOUNDSTANDARDNODES_API, FMonoAudioFormatTypeInfo, FMonoAudioFormatReadRef, FMonoAudioFormatWriteRef);
 
-	DECLARE_METASOUND_DATA_REFERENCE_TYPES(FStereoAudioFormat, METASOUNDSTANDARDNODES_API, 0xb55304e2 , FStereoAudioFormatTypeInfo, FStereoAudioFormatReadRef, FStereoAudioFormatWriteRef);
+	DECLARE_METASOUND_DATA_REFERENCE_TYPES(FStereoAudioFormat, METASOUNDSTANDARDNODES_API, FStereoAudioFormatTypeInfo, FStereoAudioFormatReadRef, FStereoAudioFormatWriteRef);
 	
 	
 }
diff --git a/Engine/Plugins/Runtime/Metasound/Source/MetasoundStandardNodes/Public/MetasoundDataReferenceTypes.h b/Engine/Plugins/Runtime/Metasound/Source/MetasoundStandardNodes/Public/MetasoundDataReferenceTypes.h
--- a/Engine/Plugins/Runtime/Metasound/Source/MetasoundStandardNodes/Public/MetasoundDataReferenceTypes.h
+++ b/Engine/Plugins/Runtime/Metasound/Source/MetasoundStandardNodes/Public/MetasoundDataReferenceTypes.h
@@ -85,5 +85,5 @@ namespace Metasound
 			TArray<uint32> Frames;
 	};
 
-	DECLARE_METASOUND_DATA_REFERENCE_TYPES(FBop, METASOUNDSTANDARDNODES_API, 2, FBopTypeInfo, FBopReadRef, FBopWriteRef);
+	DECLARE_METASOUND_DATA_REFERENCE_TYPES(FBop, METASOUNDSTANDARDNODES_API, FBopTypeInfo, FBopReadRef, FBopWriteRef);
 }
diff --git a/Engine/Plugins/Runtime/Metasound/Source/MetasoundStandardNodes/Public/MetasoundFrequency.h b/Engine/Plugins/Runtime/Metasound/Source/MetasoundStandardNodes/Public/MetasoundFrequency.h
--- a/Engine/Plugins/Runtime/Metasound/Source/MetasoundStandardNodes/Public/MetasoundFrequency.h
+++ b/Engine/Plugins/Runtime/Metasound/Source/MetasoundStandardNodes/Public/MetasoundFrequency.h
@@ -150,5 +150,5 @@ namespace Metasound
 			float FrequencyInHz;
 	};
 
-	DECLARE_METASOUND_DATA_REFERENCE_TYPES(FFrequency, METASOUNDSTANDARDNODES_API, 0xf1c430a4, FFrequencyTypeInfo, FFrequencyReadRef, FFrequencyWriteRef)
+	DECLARE_METASOUND_DATA_REFERENCE_TYPES(FFrequency, METASOUNDSTANDARDNODES_API, FFrequencyTypeInfo, FFrequencyReadRef, FFrequencyWriteRef)
 }
diff --git a/Engine/Plugins/Runtime/Metasound/Source/MetasoundStandardNodes/Public/MetasoundPrimitives.h b/Engine/Plugins/Runtime/Metasound/Source/MetasoundStandardNodes/Public/MetasoundPrimitives.h
--- a/Engine/Plugins/Runtime/Metasound/Source/MetasoundStandardNodes/Public/MetasoundPrimitives.h
+++ b/Engine/Plugins/Runtime/Metasound/Source/MetasoundStandardNodes/Public/MetasoundPrimitives.h
@@ -12,14 +12,14 @@
 
 namespace Metasound
 {
-	DECLARE_METASOUND_DATA_REFERENCE_TYPES(bool, METASOUNDSTANDARDNODES_API, 0xcc9d715e, FBoolTypeInfo, FBoolReadRef, FBoolWriteRef);
+	DECLARE_METASOUND_DATA_REFERENCE_TYPES(bool, METASOUNDSTANDARDNODES_API, FBoolTypeInfo, FBoolReadRef, FBoolWriteRef);
 	
-	DECLARE_METASOUND_DATA_REFERENCE_TYPES(int32, METASOUNDSTANDARDNODES_API, 0xf9920034, FInt32TypeInfo, FInt32ReadRef, FInt32WriteRef);
+	DECLARE_METASOUND_DATA_REFERENCE_TYPES(int32, METASOUNDSTANDARDNODES_API, FInt32TypeInfo, FInt32ReadRef, FInt32WriteRef);
 
-	DECLARE_METASOUND_DATA_REFERENCE_TYPES(int64, METASOUNDSTANDARDNODES_API, 0xfc8f8ada, FInt64TypeInfo, FInt64ReadRef, FInt64WriteRef);
+	DECLARE_METASOUND_DATA_REFERENCE_TYPES(int64, METASOUNDSTANDARDNODES_API, FInt64TypeInfo, FInt64ReadRef, FInt64WriteRef);
 
-	DECLARE_METASOUND_DATA_REFERENCE_TYPES(float, METASOUNDSTANDARDNODES_API, 0xe7adc4db, FFloatTypeInfo, FFloatReadRef, FFloatWriteRef);
+	DECLARE_METASOUND_DATA_REFERENCE_TYPES(float, METASOUNDSTANDARDNODES_API, FFloatTypeInfo, FFloatReadRef, FFloatWriteRef);
 
-	DECLARE_METASOUND_DATA_REFERENCE_TYPES(double, METASOUNDSTANDARDNODES_API, 0x8f48a053, FDoubleTypeInfo, FDoubleReadRef, FDoubleWriteRef);
+	DECLARE_METASOUND_DATA_REFERENCE_TYPES(double, METASOUNDSTANDARDNODES_API, FDoubleTypeInfo, FDoubleReadRef, FDoubleWriteRef);
 }
 
diff --git a/Engine/Plugins/Runtime/Metasound/Source/MetasoundStandardNodes/Public/MetasoundTime.h b/Engine/Plugins/Runtime/Metasound/Source/MetasoundStandardNodes/Public/MetasoundTime.h
--- a/Engine/Plugins/Runtime/Metasound/Source/MetasoundStandardNodes/Public/MetasoundTime.h
+++ b/Engine/Plugins/Runtime/Metasound/Source/MetasoundStandardNodes/Public/MetasoundTime.h
@@ -624,19 +624,19 @@ namespace Metasound
 	 */
 	typedef TTime<float, float> FFloatTime;
 
-	DECLARE_METASOUND_DATA_REFERENCE_TYPES(FFloatTime, METASOUNDSTANDARDNODES_API, 0x2490848e, FFloatTimeTypeInfo, FFloatTimeReadRef, FFloatTimeWriteRef);
+	DECLARE_METASOUND_DATA_REFERENCE_TYPES(FFloatTime, METASOUNDSTANDARDNODES_API, FFloatTimeTypeInfo, FFloatTimeReadRef, FFloatTimeWriteRef);
 
 	/** FDoubleTime stores and represents all time values as 64bit floating
 	 * point values.
 	 */
 	typedef TTime<double, double> FDoubleTime;
 
-	DECLARE_METASOUND_DATA_REFERENCE_TYPES(FDoubleTime, METASOUNDSTANDARDNODES_API, 0x407fcff1, FDoubleTimeTypeInfo, FDoubleTimeReadRef, FDoubleTimeWriteRef);
+	DECLARE_METASOUND_DATA_REFERENCE_TYPES(FDoubleTime, METASOUNDSTANDARDNODES_API, FDoubleTimeTypeInfo, FDoubleTimeReadRef, FDoubleTimeWriteRef);
 
 	/** FSampleTime stores time quantized to sample boundaries. It returns time
 	 * as 32bit floating point values.
 	 */
 	typedef TTime<FTimeSampleCounter, float> FSampleTime;
 
-	DECLARE_METASOUND_DATA_REFERENCE_TYPES(FSampleTime, METASOUNDSTANDARDNODES_API, 0xf1de8f2d, FSampleTimeTypeInfo, FSampleTimeReadRef, FSampleTimeWriteRef);
+	DECLARE_METASOUND_DATA_REFERENCE_TYPES(FSampleTime, METASOUNDSTANDARDNODES_API, FSampleTimeTypeInfo, FSampleTimeReadRef, FSampleTimeWriteRef);
 }
