Ignoring metasound data reference type id

diff --git a/Engine/Plugins/Runtime/Metasound/Source/MetasoundGraphCore/Public/MetasoundDataReference.h b/Engine/Plugins/Runtime/Metasound/Source/MetasoundGraphCore/Public/MetasoundDataReference.h
--- a/Engine/Plugins/Runtime/Metasound/Source/MetasoundGraphCore/Public/MetasoundDataReference.h
+++ b/Engine/Plugins/Runtime/Metasound/Source/MetasoundGraphCore/Public/MetasoundDataReference.h
@@ -201,7 +201,12 @@ namespace Metasound
 		static const FName TypeName = GetMetasoundDataTypeName<DataType>();
 		static const FMetasoundDataTypeId TypeId = GetMetasoundDataTypeId<DataType>();
 
-		return (InReference.GetDataTypeName() == TypeName) && (InReference.GetDataTypeId() == TypeId);
+		bool bEqualTypeName = InReference.GetDataTypeName() == TypeName;
+		// TODO: need to move data type ID definition to traits implementation.
+		//bool bEqualTypeId = InReference.GetDataTypeId() == TypeId;
+		bool bEqualTypeId = true;
+
+		return (bEqualTypeName && bEqualTypeId);
 	}
 
 	// This enum is used as a token to explicitly delineate when we should create a new object for the reference,
