Initial implementation of the metasound transmission systems, as well as send and receive nodes.

diff --git a/Engine/Plugins/Runtime/Metasound/Source/MetasoundFrontend/MetasoundFrontend.Build.cs b/Engine/Plugins/Runtime/Metasound/Source/MetasoundFrontend/MetasoundFrontend.Build.cs
--- a/Engine/Plugins/Runtime/Metasound/Source/MetasoundFrontend/MetasoundFrontend.Build.cs
+++ b/Engine/Plugins/Runtime/Metasound/Source/MetasoundFrontend/MetasoundFrontend.Build.cs
@@ -26,7 +26,7 @@ namespace UnrealBuildTool.Rules
 			(
 				new string[]
 				{
-					"MetasoundGraphCore",
+					"MetasoundGraphCore"
 				}
 			);
 
diff --git a/Engine/Plugins/Runtime/Metasound/Source/MetasoundFrontend/Private/MetasoundFrontend.cpp b/Engine/Plugins/Runtime/Metasound/Source/MetasoundFrontend/Private/MetasoundFrontend.cpp
--- a/Engine/Plugins/Runtime/Metasound/Source/MetasoundFrontend/Private/MetasoundFrontend.cpp
+++ b/Engine/Plugins/Runtime/Metasound/Source/MetasoundFrontend/Private/MetasoundFrontend.cpp
@@ -12,6 +12,12 @@
 #include "Serialization/MemoryReader.h"
 #include "MetasoundFrontendRegistries.h"
 
+#include "MetasoundDataTypeRegistrationMacro.h"
+
+#include "MetasoundPrimitives.h"
+
+#include "MetasoundAudioBuffer.h"
+#include "MetasoundRouter.h"
 
 static int32 MetasoundUndoRollLimitCvar = 128;
 FAutoConsoleVariableRef CVarMetasoundUndoRollLimit(
@@ -2154,7 +2160,7 @@ namespace Metasound
 					}
 				}
 
-				if (ensureAlwaysMsgf(IndexOfDependencyToRemove > 0, TEXT(R"(Couldn't find node class %s in the list of dependencies for this graph, but found it in the nodes list.
+				if (ensureAlwaysMsgf(IndexOfDependencyToRemove >= 0, TEXT(R"(Couldn't find node class %s in the list of dependencies for this graph, but found it in the nodes list.
 				This likely means that the underlying class description is corrupted.)"), *NodeClassName))
 				{
 					DependencyIDs.RemoveAt(IndexOfDependencyToRemove);
@@ -2727,4 +2733,22 @@ class FMetasoundFrontendModule : public IModuleInterface
 	}
 };
 
+
+REGISTER_METASOUND_DATATYPE(bool, "Primitive:Bool", ::Metasound::ELiteralArgType::Boolean)
+REGISTER_METASOUND_DATATYPE(int32, "Primitive:Int32", ::Metasound::ELiteralArgType::Integer)
+REGISTER_METASOUND_DATATYPE(int64, "Primitive:Int64", ::Metasound::ELiteralArgType::Integer)
+REGISTER_METASOUND_DATATYPE(float, "Primitive:Float", ::Metasound::ELiteralArgType::Float)
+REGISTER_METASOUND_DATATYPE(double, "Primitive:Double", ::Metasound::ELiteralArgType::Float)
+REGISTER_METASOUND_DATATYPE(FString, "Primitive:String", ::Metasound::ELiteralArgType::String)
+
+namespace Metasound
+{
+	REGISTER_METASOUND_DATATYPE(FAudioBuffer, "Audio:Buffer")
+
+
+	static_assert(TDataReferenceTypeInfo<FSendAddress>::bIsValidSpecialization, "ahhhhh");
+
+	REGISTER_METASOUND_DATATYPE(FSendAddress, "Send:Address")
+}
+
 IMPLEMENT_MODULE(FMetasoundFrontendModule, MetasoundFrontend);
diff --git a/Engine/Plugins/Runtime/Metasound/Source/MetasoundFrontend/Private/MetasoundRouter.cpp b/Engine/Plugins/Runtime/Metasound/Source/MetasoundFrontend/Private/MetasoundRouter.cpp
new file mode 100644
--- /dev/null
+++ b/Engine/Plugins/Runtime/Metasound/Source/MetasoundFrontend/Private/MetasoundRouter.cpp
@@ -0,0 +1,197 @@
+// Copyright Epic Games, Inc. All Rights Reserved.
+
+#include "MetasoundRouter.h"
+#include "MetasoundOperatorInterface.h"
+#include "CoreMinimal.h"
+#include "HAL/IConsoleManager.h"
+
+// Convenience exec commands to push values to global params.
+static FAutoConsoleCommand GPushFloatCommand(
+	TEXT("au.metasound.SetFloat"),
+	TEXT("Use this with au.metasound.SetFloat [address] [value]. Pushes a parameter value directly to a global address, which can then be received by Metasounds using a Receive node."),
+	FConsoleCommandWithArgsDelegate::CreateStatic(
+		[](const TArray< FString >& Args)
+		{
+			if (Args.Num() < 2)
+			{
+				UE_LOG(LogTemp, Warning, TEXT("au.metasound.SetFloat should be called with two args- the address to send to and the value to send."));
+				return;
+			}
+
+			FName ChannelName = FName(*Args[0]);
+			float ValueToPush = TCString<TCHAR>::Atof(*Args[1]);
+
+			Metasound::FDataTypeLiteralParam LiteralParam(ValueToPush);
+
+			if (!Metasound::FDataTransmissionCenter::Get().PushLiteral(ChannelName, LiteralParam))
+			{
+				UE_LOG(LogTemp, Warning, TEXT("au.metasound.SetFloat failed, likely because the channel did not support floats or the channel did not exist."));
+			}
+		})
+);
+
+static FAutoConsoleCommand GPushBoolCommand(
+	TEXT("au.metasound.SetBool"),
+	TEXT("Use this with au.metasound.SetBool [address] [value]. Pushes a parameter value directly to a global address, which can then be received by Metasounds using a Receive node."),
+	FConsoleCommandWithArgsDelegate::CreateStatic(
+		[](const TArray< FString >& Args)
+		{
+			if (Args.Num() < 2)
+			{
+				UE_LOG(LogTemp, Warning, TEXT("au.metasound.SetBool should be called with two args- the address to send to and the value to send."));
+				return;
+			}
+
+			FName ChannelName = FName(*Args[0]);
+			int32 ValueAsInt = TCString<TCHAR>::Atoi(*Args[1]);
+
+			bool ValueToPush = ValueAsInt != 0;
+
+			Metasound::FDataTypeLiteralParam LiteralParam(ValueToPush);
+
+			if (!Metasound::FDataTransmissionCenter::Get().PushLiteral(ChannelName, LiteralParam))
+			{
+				UE_LOG(LogTemp, Warning, TEXT("au.metasound.SetBool failed, likely because the channel did not support floats or the channel did not exist."));
+			}
+		})
+);
+
+static FAutoConsoleCommand GPushIntCommand(
+	TEXT("au.metasound.SetInt"),
+	TEXT("Use this with au.metasound.SetInt [address] [value]. Pushes a parameter value directly to a global address, which can then be received by Metasounds using a Receive node."),
+	FConsoleCommandWithArgsDelegate::CreateStatic(
+		[](const TArray< FString >& Args)
+		{
+			if (Args.Num() < 2)
+			{
+				UE_LOG(LogTemp, Warning, TEXT("au.metasound.SetBool should be called with two args- the address to send to and the value to send."));
+				return;
+			}
+
+			FName ChannelName = FName(Args[0]);
+			int32 ValueToPush = TCString<TCHAR>::Atoi(*Args[1]);
+
+			Metasound::FDataTypeLiteralParam LiteralParam(ValueToPush);
+
+			if (!Metasound::FDataTransmissionCenter::Get().PushLiteral(ChannelName, LiteralParam))
+			{
+				UE_LOG(LogTemp, Warning, TEXT("au.metasound.SetInt failed, likely because the channel did not support floats or the channel did not exist."));
+			}
+		})
+);
+
+static FAutoConsoleCommand GPushStringCommand(
+	TEXT("au.metasound.SetString"),
+	TEXT("Use this with au.metasound.SetString [address] [value]. Pushes a parameter value directly to a global address, which can then be received by Metasounds using a Receive node."),
+	FConsoleCommandWithArgsDelegate::CreateStatic(
+		[](const TArray< FString >& Args)
+		{
+			if (Args.Num() < 2)
+			{
+				UE_LOG(LogTemp, Warning, TEXT("au.metasound.SetBool should be called with two args- the address to send to and the value to send."));
+				return;
+			}
+
+			FName ChannelName = FName(*Args[0]);
+
+			Metasound::FDataTypeLiteralParam LiteralParam(Args[1]);
+
+			if (!Metasound::FDataTransmissionCenter::Get().PushLiteral(ChannelName, LiteralParam))
+			{
+				UE_LOG(LogTemp, Warning, TEXT("au.metasound.SetString failed, likely because the channel did not support floats or the channel did not exist."));
+			}
+		})
+);
+
+
+namespace Metasound
+{
+
+	IReceiver::~IReceiver()
+	{
+		DataChannel->OnReceiverDestroyed();
+	}
+
+	ISender::~ISender()
+	{
+		DataChannel->OnSenderDestroyed();
+	}
+	
+	ITransmissionSubsystem::ITransmissionSubsystem(FName InSubsystemName)
+		: SubsystemName(InSubsystemName)
+	{
+		FDataTransmissionCenter::Get().RegisterSubsystem(this, SubsystemName);
+	}
+
+	ITransmissionSubsystem::~ITransmissionSubsystem()
+	{
+		FDataTransmissionCenter::Get().UnregisterSubsystem(SubsystemName);
+	}
+
+	Metasound::FDataTransmissionCenter& FDataTransmissionCenter::Get()
+	{
+		static FDataTransmissionCenter Singleton;
+		return Singleton;
+	}
+
+	FDataTransmissionCenter::FDataTransmissionCenter()
+	{
+	}
+
+	TArray<FName> FAddressRouter::GetAvailableChannels()
+	{
+		FScopeLock ScopeLock(&DataChannelMapMutationLock);
+
+		TArray<FName> AvailableChannels;
+
+		for (auto& DataChannelPair : DataChannelMap)
+		{
+			AvailableChannels.Add(DataChannelPair.Key);
+		}
+
+		return AvailableChannels;
+	}
+
+	FName FAddressRouter::GetDatatypeForChannel(FName InChannelName)
+	{
+		FScopeLock ScopeLock(&DataChannelMapMutationLock);
+
+		if (TSharedRef<IDataChannel>* FoundChannel = DataChannelMap.Find(InChannelName))
+		{
+			return (*FoundChannel)->GetDataType();
+		}
+		else
+		{
+			return FName();
+		}
+	}
+
+	TArray<FName> FInstanceLocalRouter::GetAvailableChannels(uint64 InInstanceID)
+	{
+		FScopeLock ScopeLock(&InstanceRouterMapMutationLock);
+
+		if (FAddressRouter* AddressRouter = InstanceRouterMap.Find(InInstanceID))
+		{
+			return AddressRouter->GetAvailableChannels();
+		}
+		else
+		{
+			return TArray<FName>();
+		}
+	}
+
+	FName FInstanceLocalRouter::GetDatatypeForChannel(uint64 InInstanceID, FName InChannelName)
+	{
+		FScopeLock ScopeLock(&InstanceRouterMapMutationLock);
+
+		if (FAddressRouter* AddressRouter = InstanceRouterMap.Find(InInstanceID))
+		{
+			return AddressRouter->GetDatatypeForChannel(InChannelName);
+		}
+		else
+		{
+			return FName();
+		}
+	}
+
+}
diff --git a/Engine/Plugins/Runtime/Metasound/Source/MetasoundFrontend/Public/MetasoundAudioBuffer.h b/Engine/Plugins/Runtime/Metasound/Source/MetasoundFrontend/Public/MetasoundAudioBuffer.h
new file mode 100644
--- /dev/null
+++ b/Engine/Plugins/Runtime/Metasound/Source/MetasoundFrontend/Public/MetasoundAudioBuffer.h
@@ -0,0 +1,126 @@
+// Copyright Epic Games, Inc. All Rights Reserved.
+
+#pragma once
+
+#include "CoreMinimal.h"
+#include "DSP/BufferVectorOperations.h"
+#include "MetasoundGraphCoreModule.h"
+#include "MetasoundDataReference.h"
+
+/** Define which determines whether to check that the size of the audio buffer has not changed since initialization */
+#define METASOUNDGRAPHCORE_CHECKAUDIONUM !(UE_BUILD_SHIPPING || UE_BUILD_TEST)
+
+namespace Metasound 
+{
+	/**  FAudioBuffer
+	 *
+	 * FAudioBuffer is the default buffer for passing audio data between nodes.  It should not be resized. 
+	 * It is compatible with functions accepting const references to AlignedFloatBuffer (const AlignedFloatBuffer&) 
+	 * arguments via an implicit conversion operator which exposes the underlying AlignedFloatBuffer container.
+	 *
+	 */
+	class METASOUNDFRONTEND_API FAudioBuffer
+	{
+		public:
+			/** Create an FAudioBuffer with a specific number of samples. 
+			 *
+			 * @param InNumSamples - Number of samples in buffer.
+			 */
+			explicit FAudioBuffer(int32 InNumSamples)
+			{
+				Buffer.AddUninitialized(InNumSamples);
+
+#if METASOUNDGRAPHCORE_CHECKAUDIONUM
+				// InitialNum should not change during the life of an FAudioBuffer
+				InitialNum = Buffer.Num();
+#endif				
+			}
+
+			/**
+			 * This is the constructor used by the frontend.
+			 */
+			FAudioBuffer(const FOperatorSettings& InSettings)
+			{
+				Buffer.AddUninitialized(InSettings.GetNumFramesPerBlock());
+
+#if METASOUNDGRAPHCORE_CHECKAUDIONUM
+				// InitialNum should not change during the life of an FAudioBuffer
+				InitialNum = Buffer.Num();
+#endif	
+			}
+
+			FAudioBuffer()
+			{
+#if METASOUNDGRAPHCORE_CHECKAUDIONUM
+				// InitialNum should not change during the life of an FAudioBuffer
+				InitialNum = Buffer.Num();
+#endif				
+			}
+
+			/** Return a pointer to the audio float data. */
+			FORCEINLINE const float* GetData() const
+			{
+#if METASOUNDGRAPHCORE_CHECKAUDIONUM
+				UE_CLOG(InitialNum != Buffer.Num(), LogMetasoundGraphCore, Error, TEXT("Metasound audio buffer size change detected.  Audio Buffers should not be resized."));
+#endif				
+				return Buffer.GetData();
+			}
+
+			/** Return a pointer to the audio float data. */
+			FORCEINLINE float* GetData()
+			{
+#if METASOUNDGRAPHCORE_CHECKAUDIONUM
+				UE_CLOG(InitialNum != Buffer.Num(), LogMetasoundGraphCore, Error, TEXT("Metasound audio buffer size change detected.  Audio Buffers should not be resized."));
+#endif				
+				return Buffer.GetData();
+			}
+
+			/** Return the number of samples in the audio buffer. */
+			FORCEINLINE int32 Num() const
+			{
+#if METASOUNDGRAPHCORE_CHECKAUDIONUM
+				UE_CLOG(InitialNum != Buffer.Num(), LogMetasoundGraphCore, Error, TEXT("Metasound audio buffer size change detected.  Audio Buffers should not be resized."));
+#endif				
+				return Buffer.Num();
+			}
+
+			/** Implicit conversion to Audio::AlignedFloatBuffer */
+			FORCEINLINE operator const Audio::AlignedFloatBuffer& () const
+			{
+#if METASOUNDGRAPHCORE_CHECKAUDIONUM
+				UE_CLOG(InitialNum != Buffer.Num(), LogMetasoundGraphCore, Error, TEXT("Metasound audio buffer size change detected.  Audio Buffers should not be resized."));
+#endif				
+				return Buffer;
+			}
+
+			/** Implicit conversion to Audio::AlignedFloatBuffer 
+			 *
+			 * WARNING: if the buffer is resized, it will cause errors. 
+			 */
+			FORCEINLINE operator Audio::AlignedFloatBuffer& ()
+			{
+#if METASOUNDGRAPHCORE_CHECKAUDIONUM
+				UE_CLOG(InitialNum != Buffer.Num(), LogMetasoundGraphCore, Error, TEXT("Metasound audio buffer size change detected.  Audio Buffers should not be resized."));
+#endif				
+				return Buffer;
+			}
+
+			FORCEINLINE void Zero()
+			{
+				if (Buffer.Num() > 0)
+				{
+					FMemory::Memset(Buffer.GetData(), 0, sizeof(float) * Buffer.Num());
+				}
+			}
+
+		private:
+
+			Audio::AlignedFloatBuffer Buffer;
+
+#if METASOUNDGRAPHCORE_CHECKAUDIONUM
+			int32 InitialNum;
+#endif				
+	};
+
+	DECLARE_METASOUND_DATA_REFERENCE_TYPES(FAudioBuffer, METASOUNDFRONTEND_API, FAudioBufferTypeInfo, FAudioBufferReadRef, FAudioBufferWriteRef);
+};
diff --git a/Engine/Plugins/Runtime/Metasound/Source/MetasoundFrontend/Public/MetasoundDataTypeRegistrationMacro.h b/Engine/Plugins/Runtime/Metasound/Source/MetasoundFrontend/Public/MetasoundDataTypeRegistrationMacro.h
--- a/Engine/Plugins/Runtime/Metasound/Source/MetasoundFrontend/Public/MetasoundDataTypeRegistrationMacro.h
+++ b/Engine/Plugins/Runtime/Metasound/Source/MetasoundFrontend/Public/MetasoundDataTypeRegistrationMacro.h
@@ -13,6 +13,9 @@
 #include "MetasoundAutoConverterNode.h"
 #include "MetasoundNodeRegistrationMacro.h"
 #include "MetasoundConverterNodeRegistrationMacro.h"
+#include "MetasoundSendNode.h"
+#include "MetasoundReceiveNode.h"
+#include "MetasoundRouter.h"
 
 #include <type_traits>
 
@@ -37,6 +40,35 @@ IAudioProxyDataFactory* CastToAudioProxyDataFactory(UObject* InObject)
 	return nullptr;
 }
 
+// Helper utility to test if we can transmit a datatype between a send and a receive node.
+template <typename TDataType>
+struct TIsTransmittable
+{
+private:
+	static constexpr bool bCanBeTransmitted =
+		std::is_copy_constructible<TDataType>::value|| TIsDerivedFrom<TDataType, ::Metasound::IAudioDatatype>::Value;
+
+public:
+
+	static constexpr bool Value = bCanBeTransmitted;
+};
+
+// This utility function can be used to optionally check to see if we can transmit a data type, and autogenerate send and receive nodes for that datatype.
+template<typename TDataType, typename TEnableIf<TIsTransmittable<TDataType>::Value, bool>::Type = true>
+void AttemptToRegisterSendAndReceiveNodes()
+{
+	ensureAlways(RegisterNodeWithFrontend<Metasound::TSendNode<TDataType>>());
+	ensureAlways(RegisterNodeWithFrontend<Metasound::TReceiveNode<TDataType>>());
+}
+
+template<typename TDataType, typename TEnableIf<!TIsTransmittable<TDataType>::Value, bool>::Type = true>
+void AttemptToRegisterSendAndReceiveNodes()
+{
+	// This implementation intentionally noops, because Metasound::TIsTransmittable is false for this datatype.
+	// This is either because the datatype is not trivially copyable, and thus can't be buffered between threads,
+	// or it's not an audio buffer type, which we use Audio::FPatchMixerSplitter instances for.
+}
+
 // This utility function can be used to check to see if we can static cast between two types, and autogenerate a node for that static cast.
 template<typename TFromDataType, typename TToDataType, typename TEnableIf<std::is_convertible<TFromDataType, TToDataType>::value, bool>::Type = true>
 void AttemptToRegisterConverter()
@@ -155,6 +187,7 @@ bool RegisterDataTypeWithFrontend()
 	ensureAlwaysMsgf(bSucceeded, TEXT("Failed to register data type %s in the node registry!"), ::Metasound::TDataReferenceTypeInfo<TDataType>::TypeName);
 	
 	RegisterConverterNodes<TDataType>();
+	AttemptToRegisterSendAndReceiveNodes<TDataType>();
 	
 	return bSucceeded;
 }
diff --git a/Engine/Plugins/Runtime/Metasound/Source/MetasoundFrontend/Public/MetasoundPrimitives.h b/Engine/Plugins/Runtime/Metasound/Source/MetasoundFrontend/Public/MetasoundPrimitives.h
new file mode 100644
--- /dev/null
+++ b/Engine/Plugins/Runtime/Metasound/Source/MetasoundFrontend/Public/MetasoundPrimitives.h
@@ -0,0 +1,22 @@
+// Copyright Epic Games, Inc. All Rights Reserved.
+
+#pragma once
+
+#include "CoreMinimal.h"
+#include "MetasoundDataReference.h"
+
+namespace Metasound
+{
+	DECLARE_METASOUND_DATA_REFERENCE_TYPES(bool, METASOUNDFRONTEND_API, FBoolTypeInfo, FBoolReadRef, FBoolWriteRef);
+	
+	DECLARE_METASOUND_DATA_REFERENCE_TYPES(int32, METASOUNDFRONTEND_API, FInt32TypeInfo, FInt32ReadRef, FInt32WriteRef);
+
+	DECLARE_METASOUND_DATA_REFERENCE_TYPES(int64, METASOUNDFRONTEND_API, FInt64TypeInfo, FInt64ReadRef, FInt64WriteRef);
+
+	DECLARE_METASOUND_DATA_REFERENCE_TYPES(float, METASOUNDFRONTEND_API, FFloatTypeInfo, FFloatReadRef, FFloatWriteRef);
+
+	DECLARE_METASOUND_DATA_REFERENCE_TYPES(double, METASOUNDFRONTEND_API, FDoubleTypeInfo, FDoubleReadRef, FDoubleWriteRef);
+
+	DECLARE_METASOUND_DATA_REFERENCE_TYPES(FString, METASOUNDFRONTEND_API, FStringTypeInfo, FStringReadRef, FStringWriteRef);
+}
+
diff --git a/Engine/Plugins/Runtime/Metasound/Source/MetasoundFrontend/Public/MetasoundReceiveNode.h b/Engine/Plugins/Runtime/Metasound/Source/MetasoundFrontend/Public/MetasoundReceiveNode.h
new file mode 100644
--- /dev/null
+++ b/Engine/Plugins/Runtime/Metasound/Source/MetasoundFrontend/Public/MetasoundReceiveNode.h
@@ -0,0 +1,187 @@
+// Copyright Epic Games, Inc. All Rights Reserved.
+#pragma once
+
+#include "CoreMinimal.h"
+#include "MetasoundBuilderInterface.h"
+#include "MetasoundNodeInterface.h"
+#include "MetasoundOperatorInterface.h"
+#include "MetasoundDataReference.h"
+#include "MetasoundExecutableOperator.h"
+#include "MetasoundRouter.h"
+
+#include <type_traits>
+
+#define LOCTEXT_NAMESPACE "MetasoundGraphCore"
+
+
+namespace Metasound
+{
+	// This convenience node can be registered and will invoke static_cast<ToDataType>(FromDataType) every time it is executed.
+	template<typename TDataType>
+	class TReceiveNode : public INode
+	{
+		static_assert(TDataReferenceTypeInfo<TDataType>::bIsValidSpecialization, "Please use DECLARE_METASOUND_DATA_REFERENCE_TYPES with this class before trying to create an converter node with it.");
+
+	public:
+		static FString& GetAddressInputName()
+		{
+			static FString InputName = FString(TEXT("Address"));
+			return InputName;
+		}
+
+		static FString& GetSendOutputName()
+		{
+			static FString SendInput = FString(TDataReferenceTypeInfo<TDataType>::TypeName);
+			return SendInput;
+		}
+
+	private:
+		class TReceiverOperator : public TExecutableOperator<TReceiverOperator>
+		{
+			public:
+
+				TReceiverOperator(TDataReadReference<FSendAddress> InSendAddress, const FOperatorSettings& InOperatorSettings)
+					: OutputData(TDataWriteReference<TDataType>::CreateNew())
+					, SendAddress(InSendAddress)
+					, CachedSendAddress(*InSendAddress)
+					, CachedReceiverParams({InOperatorSettings})
+					, Receiver(FDataTransmissionCenter::Get().RegisterNewReceiver<TDataType>(CachedSendAddress, CachedReceiverParams))
+				{
+					Outputs.AddDataReadReference<TDataType>(GetSendOutputName(), TDataReadReference<TDataType>(OutputData));
+				}
+
+				virtual ~TReceiverOperator() {}
+
+				virtual const FDataReferenceCollection& GetInputs() const override
+				{
+					return Inputs;
+				}
+
+				virtual const FDataReferenceCollection& GetOutputs() const override
+				{
+					return Outputs;
+				}
+
+				void Execute()
+				{
+					if (SendAddress->ChannelName != CachedSendAddress.ChannelName)
+					{
+						CachedSendAddress = *SendAddress;
+						Receiver = FDataTransmissionCenter::Get().RegisterNewReceiver<TDataType>(CachedSendAddress, CachedReceiverParams);
+						check(Receiver.IsValid());
+					}
+
+					if (Receiver->CanPop())
+					{
+						Receiver->Pop(*OutputData);
+					}
+				}
+
+			private:
+				TDataWriteReference<TDataType> OutputData;
+				TDataReadReference<FSendAddress> SendAddress;
+				FSendAddress CachedSendAddress;
+				FReceiverInitParams CachedReceiverParams;
+
+				TReceiverPtr<TDataType> Receiver;
+
+				FDataReferenceCollection Inputs;
+				FDataReferenceCollection Outputs;
+		};
+
+		class FSendOperatorFactory : public IOperatorFactory
+		{
+			public:
+				FSendOperatorFactory() = default;
+
+				virtual TUniquePtr<IOperator> CreateOperator(const FCreateOperatorParams& InParams, FBuildErrorArray& OutErrors) override
+				{
+					return MakeUnique<TReceiverOperator>(
+						InParams.InputDataReferences.GetDataReadReference<FSendAddress>(GetAddressInputName()),
+						InParams.OperatorSettings
+						);
+				}
+		};
+
+		public:
+
+			static FVertexInterface DeclareVertexInterface()
+			{
+				return FVertexInterface(
+					FInputVertexInterface(
+						TInputDataVertexModel<FSendAddress>(GetAddressInputName(), FText::GetEmpty())
+					),
+					FOutputVertexInterface(
+						TOutputDataVertexModel<TDataType>(GetSendOutputName(), FText::GetEmpty())
+					)
+				);
+			}
+
+			TReceiveNode(const FNodeInitData& InInitData)
+				: NodeDescription(InInitData.InstanceName)
+				, Interface(DeclareVertexInterface())
+				, Factory(MakeOperatorFactoryRef<FSendOperatorFactory>())
+			{
+			}
+
+			virtual ~TReceiveNode() = default;
+
+			virtual const FName& GetClassName() const override
+			{
+				static const FName ClassName(*FString::Printf(TEXT("Receive %s"), *GetSendOutputName()));
+				return ClassName;
+			}
+
+			virtual const FString& GetInstanceName() const override
+			{
+				return NodeDescription;
+			}
+
+			virtual const FText& GetDescription() const override
+			{
+				static const FText Description = LOCTEXT("Metasound_ReceiveNodeDescription", "receives data from a send node with the same name.");
+				return Description;
+			}
+
+			virtual const FText& GetAuthorName() const override
+			{
+				return PluginAuthor;
+			}
+
+			virtual const FText& GetPromptIfMissing() const override
+			{
+				return PluginNodeMissingPrompt;
+			}
+
+			virtual const FVertexInterface& GetVertexInterface() const override
+			{
+				return Interface;
+			}
+
+			virtual const FVertexInterface& GetDefaultVertexInterface() const override
+			{
+				return Interface;
+			}
+
+			virtual bool SetVertexInterface(const FVertexInterface& InInterface) override
+			{
+				return Interface == InInterface;
+			}
+
+			virtual bool IsVertexInterfaceSupported(const FVertexInterface& InInterface) const override
+			{
+				return Interface == InInterface;
+			}
+
+			virtual FOperatorFactorySharedRef GetDefaultOperatorFactory() const override
+			{
+				return Factory;
+			}
+
+		private:
+			FString NodeDescription;
+			FVertexInterface Interface;
+			FOperatorFactorySharedRef Factory;
+	};
+}
+#undef LOCTEXT_NAMESPACE
diff --git a/Engine/Plugins/Runtime/Metasound/Source/MetasoundFrontend/Public/MetasoundRouter.h b/Engine/Plugins/Runtime/Metasound/Source/MetasoundFrontend/Public/MetasoundRouter.h
new file mode 100644
--- /dev/null
+++ b/Engine/Plugins/Runtime/Metasound/Source/MetasoundFrontend/Public/MetasoundRouter.h
@@ -0,0 +1,1155 @@
+// Copyright Epic Games, Inc. All Rights Reserved.
+#pragma once
+
+#include "CoreMinimal.h"
+#include "UObject/NameTypes.h"
+
+#include "MetasoundOperatorSettings.h"
+#include "MetasoundDataReference.h"
+#include "MetasoundOperatorInterface.h"
+#include "MetasoundAudioBuffer.h"
+#include "Misc/Guid.h"
+#include "DSP/Dsp.h"
+#include "DSP/MultithreadedPatching.h"
+
+#include <type_traits>
+
+/**
+ * METASOUND TRANSMISSION SYSTEM
+ * This allows for transmission of arbitrary datatypes using addresses set by FNames.
+ * This system is used by Send and Receive nodes in the Metasound graph, but can also be used in C++ for arbitrary message passing.
+ *
+ * Typical use case:
+ * TSenderPtr<float> FloatSender FDataTransmissionCenter::Get().RegisterNewSend<float>(ExampleAddress, InitSettings);
+ * TReceiverPtr<float> FloatReceiver FDataTransmissionCenter::Get().RegisterNewReceiver<float>(ExampleAddress, InitSettings);
+ *
+ * //...
+ * FloatSender->Push(4.5f);
+ *
+ * //... elsewhere on a different thread:
+ * float ReceivedFloat = FloatReceiver->Pop();
+ *
+ * It's that easy! (hopefully!)
+ *
+ * In general, some performance considerations for this system:
+ * 1) a circular buffer is allocated for each address a sender is set up for, and is sized to the max Delay of any connected sender.
+ * 2) Senders and receivers are meant to have long lifecycles, and requesting them from the FDataTransmissionCenter may be expensive if done every tick.
+ * 3) Most non-audio data types cannot be mixed, so multiple senders to the same address will cause only one of the senders to be effective.
+ * 4) while the overall system is thread safe, individual TSender and TReceiver objects should only be used by a single thread.
+ */
+
+namespace Metasound
+{
+	// Enum for which scope to use 
+	enum class ETransmissionScope : uint8
+	{
+		ThisInstanceOnly,
+		Global,
+		Extension // This means that this metasound router talks to an implentation of IMetasoundTransmissionListener.
+	};
+
+	// This converts an ETransmissionScope into an FName that can be used in FSendAddress::Subsystem.
+	// For ETransmissionScope::Extension, the FName of that extension should be used instead.
+	FORCEINLINE FName GetSubsystemNameForSendScope(ETransmissionScope InScope)
+	{
+		switch (InScope)
+		{
+			case ETransmissionScope::ThisInstanceOnly:
+			{
+				static FName SubsystemName = TEXT("This Instance Only");
+				return SubsystemName;
+			}
+			case ETransmissionScope::Global:
+			{
+				static FName SubsystemName = TEXT("All Metasounds");
+				return SubsystemName;
+			}
+			default:
+			{
+				checkNoEntry();
+				return FName();
+			}
+		}
+	}
+
+	struct METASOUNDFRONTEND_API FSendAddress
+	{
+		// The specific subsystem this send should be broadcast to.
+		// This can be either local to this instance, or global to all metasounds, or a specific subsystem.
+		FName Subsystem;
+
+		// This is the name for the channel itself.
+		FName ChannelName;
+
+		// For instance-specific addresses, this FName is the instance of this metasound.
+		uint64 MetasoundInstanceID = INDEX_NONE;
+
+		friend uint32 GetTypeHash(const FSendAddress& Address)
+		{
+			uint32 HashedChannel = HashCombine(GetTypeHash(Address.Subsystem), GetTypeHash(Address.ChannelName));
+			uint32 HashedMetasoundInstance = HashCombine(HashedChannel, static_cast<uint32>(Address.MetasoundInstanceID % TNumericLimits<uint32>::Max()));
+			return HashedMetasoundInstance;
+		}
+
+		FORCEINLINE bool operator==(const FSendAddress& Other) const
+		{
+			return ChannelName == Other.ChannelName;
+		}
+
+		FSendAddress(const FString& InAddress)
+			: Subsystem(GetSubsystemNameForSendScope(ETransmissionScope::Global))
+			, ChannelName(*InAddress)
+		{
+		}
+
+		FSendAddress()
+			: Subsystem(GetSubsystemNameForSendScope(ETransmissionScope::Global))
+			, ChannelName(FName(TEXT("Default Channel")))
+		{}
+	};
+
+	// Base class that defines shared state and utilities for senders, receivers, and the shared channel state.
+	class METASOUNDFRONTEND_API IDataTransmissionBase
+	{
+	public:
+		IDataTransmissionBase() = delete;
+
+		IDataTransmissionBase(FName InDataType)
+			: DataType(InDataType)
+		{
+		}
+
+		virtual ~IDataTransmissionBase() = default;
+		
+		template<typename TSenderImplementation>
+		bool CheckType() const
+		{
+			FName DestinationTypeName = TSenderImplementation::GetDataTypeName();
+			return ensureAlwaysMsgf(DataType == DestinationTypeName, TEXT("Tried to downcast type %s to %s!"), *DataType.ToString(), *DestinationTypeName.ToString());
+		}
+
+		template<typename TSenderImplementation>
+		const TSenderImplementation& GetAs() const
+		{
+			check(CheckType<TSenderImplementation>());
+			return static_cast<TSenderImplementation&>(*this);
+		}
+
+		template<typename TSenderImplementation>
+		TSenderImplementation& GetAs()
+		{
+			check(CheckType<TSenderImplementation>());
+			return static_cast<TSenderImplementation&>(*this);
+		}
+
+		FName GetDataType() const
+		{
+			return DataType;
+		}
+
+	protected:
+		// The data type for this sender.
+		FName DataType;
+	};
+
+	class IDataChannel;
+
+	// This is a handle to something that will retrieve a datatype.
+	// Intentionally opaque- to use, one will need to call IReceiver::GetAs<TReceiver<YourDataType>>().
+	class METASOUNDFRONTEND_API IReceiver : public IDataTransmissionBase
+	{
+	public:
+		IReceiver(TSharedPtr<IDataChannel> InDataChannel, FName DataType)
+			: IDataTransmissionBase(DataType)
+			, DataChannel(InDataChannel)
+		{}
+
+
+		virtual ~IReceiver();
+
+	protected:
+		TSharedPtr<IDataChannel> DataChannel;
+	};
+
+	// This is a handle to something that will retrieve a datatype.
+	// Intentionally opaque- to use, one will need to call ISender::GetAs<TSender<YourDataType>>().
+	class METASOUNDFRONTEND_API ISender : public IDataTransmissionBase
+	{
+	public:
+		ISender(TSharedPtr<IDataChannel> InDataChannel, FName DataType)
+			: IDataTransmissionBase(DataType)
+			, DataChannel(InDataChannel)
+		{}
+
+		virtual ~ISender();
+
+	protected:
+		TSharedPtr<IDataChannel> DataChannel;
+	};
+
+	struct FSenderInitParams
+	{
+		FOperatorSettings OperatorSettings;
+		float DelayTimeInSeconds;
+	};
+
+	struct FReceiverInitParams
+	{
+		FOperatorSettings OperatorSettings;
+	};
+
+	// This contains an intermediary buffer for use between the send and receive nodes.
+	// Typically only used by ISender and IReceiver implementations.
+	class METASOUNDFRONTEND_API IDataChannel : public IDataTransmissionBase, public TSharedFromThis<IDataChannel>
+	{
+	public:
+		virtual ~IDataChannel() = default;
+
+		IDataChannel(FName DataType)
+			: IDataTransmissionBase(DataType)
+		{}
+
+		TUniquePtr<IReceiver> NewReceiver(const FReceiverInitParams& InitParams)
+		{
+			NumAliveReceivers.Increment();
+			return ConstructNewReceiverImplementation(InitParams);
+		}
+
+		TUniquePtr<ISender> NewSender(const FSenderInitParams& InitParams)
+		{
+			NumAliveSenders.Increment();
+			return ConstructNewSenderImplementation(InitParams);
+		}
+
+		uint32 GetNumActiveReceivers() const
+		{
+			return NumAliveReceivers.GetValue();
+		}
+
+		uint32 GetNumActiveSenders() const
+		{
+			return NumAliveSenders.GetValue();
+		}
+
+		virtual bool PushOpaque(ISender& InSender) { return true; };
+		virtual void PopOpaque(IReceiver& InReceiver) {};
+		virtual bool IsEmpty() const { return true; };
+
+		virtual FName GetDataType() { return FName(); };
+
+		virtual bool PushLiteral(const FDataTypeLiteralParam& InParam) { return false; }
+
+	protected:
+		virtual TUniquePtr<IReceiver> ConstructNewReceiverImplementation(const FReceiverInitParams& InitParams) { return nullptr; };
+		virtual TUniquePtr<ISender> ConstructNewSenderImplementation(const FSenderInitParams& InitParams) { return nullptr; };
+
+	private:
+		void OnSenderDestroyed()
+		{
+			check(NumAliveSenders.GetValue() > 0);
+			NumAliveSenders.Decrement();
+		}
+
+		void OnReceiverDestroyed()
+		{
+			check(NumAliveReceivers.GetValue() > 0);
+			NumAliveReceivers.Decrement();
+		}
+
+		friend class ISender;
+		friend class IReceiver;
+
+		FThreadSafeCounter NumAliveReceivers;
+		FThreadSafeCounter NumAliveSenders;
+	};
+
+	// Generic templates:
+	template<typename TDataType>
+	class TNonOperationalSender : public ISender
+	{
+	public:
+		bool Push(const TDataType& InValue) { return true; }
+	};
+
+	template<typename TDataType>
+	class TNonOperationalReceiver : public IReceiver
+	{
+	public:
+		bool Pop(TDataType& OutValue) { return true; }
+	};
+
+	template<typename TDataType>
+	class TNonOperationalDataChannel : public IDataChannel
+	{
+
+	public:
+
+		TNonOperationalDataChannel(const FOperatorSettings& InOperatorSettings)
+			: IDataChannel(FName())
+		{}
+
+		virtual bool PushOpaque(ISender& InSender) override 
+		{
+			return true;
+		}
+
+		virtual void PopOpaque(IReceiver& InReceiver) override
+		{
+
+		}
+
+		virtual bool IsEmpty() const override
+		{
+			return true;
+		}
+
+		virtual FName GetDataType() override
+		{
+			return GetDataTypeName();
+		}
+
+		static FName GetDataTypeName()
+		{
+			static FName TypeName = TDataReferenceTypeInfo<TDataType>::TypeName;
+			return TypeName;
+		}
+
+	protected:
+		virtual TUniquePtr<IReceiver> ConstructNewReceiverImplementation(const FReceiverInitParams& InitParams) override
+		{
+			return TUniquePtr<IReceiver>();
+		}
+
+		virtual TUniquePtr<ISender> ConstructNewSenderImplementation(const FSenderInitParams& InitParams) override
+		{
+			return TUniquePtr<ISender>();
+		}
+	};
+
+	// Specializations for trivially copyable types:
+	template<typename TDataType>
+	class TSender : public ISender
+	{
+		static_assert(std::is_copy_constructible<TDataType>::value, "The generic TDataChannel requires the DataType it is specialized with to be copy assignable.");
+
+		// This can be changed to change when we reallocate the internal buffer for this sender.
+		static constexpr float BufferResizeThreshold = 2.0f;
+
+	public:
+
+		bool Push(const TDataType& InElement)
+		{
+			SenderBuffer.Push(InElement);
+
+			if (SenderBuffer.Num() >= NumElementsToDelayBy)
+			{
+				return PushToDataChannel();
+			}
+			else
+			{
+				return true;
+			}
+		}
+
+		// Resets the delay for this specific sender. if this goes beyond the threshold set by BufferResizeThreshold, we reallocate.
+		bool SetDelay(float InSeconds)
+		{
+			Params.DelayTimeInSeconds = InSeconds;
+			NumElementsToDelayBy = FMath::Min(InSeconds * Params.OperatorSettings.GetActualBlockRate(), 1);
+
+			if (NumElementsToDelayBy >= SenderBuffer.GetCapacity())
+			{
+				SenderBuffer.SetCapacity(NumElementsToDelayBy * BufferResizeThreshold);
+			}
+		}
+
+		static FName GetDataTypeName()
+		{
+			static FName TypeName = TDataReferenceTypeInfo<TDataType>::TypeName;
+			return TypeName;
+		}
+
+		TSender(const FSenderInitParams& InitParams, TSharedPtr<IDataChannel> InDataChannel)
+			: ISender(InDataChannel, GetDataTypeName())
+			, Params(InitParams)
+		{
+
+		}
+
+		TDataType RetrievePayload()
+		{
+			check(Payload.IsValid());
+			return MoveTempIfPossible(*Payload);
+		}
+
+		// Push from the SenderBuffer to the data channel.
+		bool PushToDataChannel()
+		{
+			if (!Payload.IsValid())
+			{
+				Payload.Reset(new TDataType(SenderBuffer.Pop()));
+			}
+			else
+			{
+				*Payload = SenderBuffer.Pop();
+			}
+
+			return DataChannel->PushOpaque(*this);
+		}
+
+	private:
+		// This buffer acts as a delay line for this sender specifically.
+		Audio::TCircularAudioBuffer<TDataType> SenderBuffer;
+
+		// lazily initialized holder for payload, allocated on first call to Push.
+		// Placed on the heap so that we don't have to assume we know how to construct TDataType.
+		TUniquePtr<TDataType> Payload;
+
+		FSenderInitParams Params;
+
+		// The amount of pop operations before elements before this sender pushed appear.
+		uint32 NumElementsToDelayBy;
+	};
+
+	template<typename TDataType>
+	class TReceiver : public IReceiver
+	{
+		static_assert(std::is_copy_constructible<TDataType>::value, "The generic TDataChannel requires the DataType it is specialized with to be copy constructible.");
+	public:
+
+		bool CanPop() const
+		{
+			return !DataChannel->IsEmpty();
+		}
+
+		// Pop the latest value from the data channel.
+		bool Pop(TDataType& OutValue)
+		{
+			DataChannel->PopOpaque(*this);
+
+			check(Payload.IsValid());
+			OutValue = *Payload;
+			return true;
+		}
+
+		static FName GetDataTypeName()
+		{
+			static FName TypeName = TDataReferenceTypeInfo<TDataType>::TypeName;
+			return TypeName;
+		}
+
+		TReceiver(const FReceiverInitParams& InitParams, TSharedPtr<IDataChannel> InDataChannel)
+			: IReceiver(InDataChannel, GetDataTypeName())
+			, OperatorSettings(InitParams.OperatorSettings)
+		{
+		}
+
+		void PushPayload(const TDataType& InDataPayload)
+		{
+			if (!Payload.IsValid())
+			{
+				Payload.Reset(new TDataType(InDataPayload));
+			}
+			else
+			{
+				*Payload = InDataPayload;
+			}
+		}
+
+	private:
+		TUniquePtr<TDataType> Payload;
+		FOperatorSettings OperatorSettings;
+	};
+
+	// Generic templated implementation of IDataChannel that can use for any copyable type.
+	template<typename TDataType>
+	class TCopyableDataChannel : public IDataChannel
+	{
+		static_assert(std::is_copy_constructible<TDataType>::value, "The generic TDataChannel requires the DataType it is specialized with to be copy constructible.");
+
+	public:
+
+		TCopyableDataChannel(const FOperatorSettings& InOperatorSettings)
+			: IDataChannel(GetDataTypeName())
+			, OperatorSettings(InOperatorSettings)
+		{}
+
+		virtual bool PushOpaque(ISender& InSender) override
+		{
+			TSender<TDataType>& CastSender = InSender.GetAs<TSender<TDataType>>();
+
+			FScopeLock ScopeLock(&AtomicDataLock);
+			if (!AtomicData.IsValid())
+			{
+				AtomicData.Reset(new TDataType(CastSender.RetrievePayload()));
+			}
+			else
+			{
+				*AtomicData = CastSender.RetrievePayload();
+			}
+
+			return true;
+		}
+
+		virtual void PopOpaque(IReceiver& InReceiver) override
+		{
+			TReceiver<TDataType>& CastReceiver = InReceiver.GetAs<TReceiver<TDataType>>();
+
+			FScopeLock ScopeLock(&AtomicDataLock);
+			if (AtomicData.IsValid())
+			{
+				CastReceiver.PushPayload(*AtomicData);
+			}
+		}
+
+		virtual bool IsEmpty() const override
+		{
+			return !AtomicData.IsValid();
+		}
+
+		static FName GetDataTypeName()
+		{
+			static FName TypeName = TDataReferenceTypeInfo<TDataType>::TypeName;
+			return TypeName;
+		}
+
+		virtual FName GetDataType() override
+		{
+			return GetDataTypeName();
+		}
+
+		virtual bool PushLiteral(const FDataTypeLiteralParam& InParam)
+		{
+			if (InParam.IsCompatibleWithType<TDataType>())
+			{
+				TDataType DataToPush = InParam.ParseTo<TDataType>(OperatorSettings);
+
+				FScopeLock ScopeLock(&AtomicDataLock);
+				if (!AtomicData.IsValid())
+				{
+					AtomicData.Reset(new TDataType(DataToPush));
+				}
+				else
+				{
+					*AtomicData = DataToPush;
+				}
+
+				return true;
+			}
+			else
+			{
+				return false;
+			}
+		}
+
+	protected:
+		virtual TUniquePtr<IReceiver> ConstructNewReceiverImplementation(const FReceiverInitParams& InitParams) override
+		{
+			return TUniquePtr<IReceiver>(new TReceiver<TDataType>(InitParams, this->AsShared()));
+		}
+
+		virtual TUniquePtr<ISender> ConstructNewSenderImplementation(const FSenderInitParams& InitParams) override
+		{
+			return TUniquePtr<ISender>(new TSender<TDataType>(InitParams, this->AsShared()));
+		}
+
+	private:
+		FThreadSafeCounter ReceiverCounter;
+
+		// NOTE- In the future, this could be changed to use a TCircualrAudioBuffer or a TQueue for lockless operation.
+		// The primary challenge with this is handling multiple senders and receivers.
+		// In order to support multiple senders and receivers at scale, we'd need a good implementation of a bounded MPMC queue for arbitrary datatypes.
+		TUniquePtr<TDataType> AtomicData;
+		FCriticalSection AtomicDataLock;
+
+		FOperatorSettings OperatorSettings;
+	};
+
+	// AUDIO SPECIALIZATIONS:
+	template<typename TDataType>
+	class TAudioSender : public ISender
+	{
+		static constexpr float BufferResizeThreshold = 2.0f;
+		static const int32 MaxChannels = 8;
+	public:
+
+		bool Push(const TDataType& InElement)
+		{
+			const int32 NumChannelsToPush = InElement.GetNumChannels();
+			check(NumChannelsToPush <= MaxChannels);
+
+			const TArrayView<const FAudioBufferReadRef> InputBuffers = InElement.GetBuffers();
+
+			for (int32 ChannelIndex = 0; ChannelIndex < NumChannelsToPush; ChannelIndex++)
+			{
+				const FAudioBuffer& ChannelBuffer = *InputBuffers[ChannelIndex];
+				DelayLines[ChannelIndex].Push(ChannelBuffer.GetData(), ChannelBuffer.Num());
+
+				if (DelayLines[ChannelIndex].Num() > DelayTimeInFrames)
+				{
+					Audio::AlignedFloatBuffer& IntermediateBuffer = IntermediaryBuffers[ChannelIndex];
+					DelayLines[ChannelIndex].Pop(IntermediateBuffer.GetData(), IntermediateBuffer.Num());
+					DataChannelInputs[ChannelIndex].PushAudio(IntermediateBuffer.GetData(), IntermediateBuffer.Num());
+				}
+			}
+
+			return true;
+		}
+
+		// Resets the delay for this sender.
+		bool SetDelay(float InSeconds)
+		{
+			Params.DelayTimeInSeconds = InSeconds;
+			DelayTimeInFrames = FMath::Min<int32>(Params.DelayTimeInSeconds * Params.OperatorSettings.GetSampleRate(), Params.OperatorSettings.GetNumFramesPerBlock());
+			
+			if (DelayTimeInFrames >= DelayLines[0].Num())
+			{
+				int32 NumSamplesToBuffer = DelayTimeInFrames * BufferResizeThreshold;
+				DelayLines.Reset();
+				IntermediaryBuffers.Reset();
+
+				// For now, all senders/receivers/data channels will just allocate 8 delay lines, but this can be changed in the future.
+				for (int32 Index = 0; Index < MaxChannels; Index++)
+				{
+					DelayLines.Emplace(NumSamplesToBuffer);
+					Audio::AlignedFloatBuffer BufferForChannel;
+					BufferForChannel.AddZeroed(Params.OperatorSettings.GetNumFramesPerBlock());
+					IntermediaryBuffers.Add(MoveTemp(BufferForChannel));
+				}
+			}
+
+			return true;
+		}
+
+		static FName GetDataTypeName()
+		{
+			static FName TypeName = TDataReferenceTypeInfo<TDataType>::TypeName;
+			return TypeName;
+		}
+
+		TAudioSender(const FSenderInitParams& InitParams, TSharedPtr<IDataChannel> InDataChannel, TArray<Audio::FPatchInput>&& Inputs)
+			: ISender(InDataChannel, GetDataTypeName())
+			, DataChannelInputs(MoveTemp(Inputs))
+			, Params(InitParams)
+		{
+			SetDelay(InitParams.DelayTimeInSeconds);
+		}
+
+	private:
+		// This buffer acts as a delay line for this sender specifically.
+		TArray<Audio::TCircularAudioBuffer<float>> DelayLines;
+
+		uint32 DelayTimeInFrames;
+
+		TArray<Audio::FPatchInput> DataChannelInputs;
+
+		TArray<Audio::AlignedFloatBuffer> IntermediaryBuffers;
+
+		FSenderInitParams Params;
+	};
+
+	template<typename TDataType>
+	class TAudioReceiver : public IReceiver
+	{
+		static constexpr float BufferResizeThreshold = 2.0f;
+		static const int32 MaxChannels = 8;
+	public:
+
+		bool CanPop() const
+		{
+			// Audio receivers can and should always pop audio,
+			// because we zero out audio when a datachannel underruns.
+			return true;
+		}
+
+		bool Pop(TDataType& OutElement)
+		{
+			const int32 NumChannelsToPop = OutElement.GetNumChannels();
+			check(NumChannelsToPop <= MaxChannels);
+
+			DataChannel->PopOpaque(*this);
+
+			TArrayView<const FAudioBufferWriteRef> OutputBuffers = OutElement.GetBuffers();
+
+			for (int32 ChannelIndex = 0; ChannelIndex < NumChannelsToPop; ChannelIndex++)
+			{
+				FAudioBuffer& OutBuffer = *OutputBuffers[ChannelIndex];
+				DataChannelOutputs[ChannelIndex]->PopAudio(OutBuffer.GetData(), OutBuffer.Num(), false);
+			}
+
+			return true;
+		}
+
+		static FName GetDataTypeName()
+		{
+			static FName TypeName = TDataReferenceTypeInfo<TDataType>::TypeName;
+			return TypeName;
+		}
+
+		TAudioReceiver(const FReceiverInitParams& InitParams, TSharedPtr<IDataChannel> InDataChannel, TArray<Audio::FPatchOutputStrongPtr>&& Outputs)
+			: IReceiver(InDataChannel, GetDataTypeName())
+			, DataChannelOutputs(MoveTemp(Outputs))
+			, Params(InitParams)
+		{
+		}
+
+	private:
+
+		TArray<Audio::FPatchOutputStrongPtr> DataChannelOutputs;
+
+		FReceiverInitParams Params;
+	};
+
+	template<typename TDataType>
+	class TAudioDataChannel : public IDataChannel
+	{
+	private:
+		static const int32 MaxChannels = 8;
+	public:
+		TAudioDataChannel(const FOperatorSettings& InOperatorSettings)
+			: IDataChannel(GetDataTypeName())
+		{
+			AudioBuses.AddDefaulted(MaxChannels);
+		}
+
+		virtual bool PushOpaque(ISender& InSender) override
+		{
+			return true;
+		}
+
+		virtual void PopOpaque(IReceiver& InReceiver) override
+		{
+			for (Audio::FPatchMixerSplitter& Bus : AudioBuses)
+			{
+				Bus.ProcessAudio();
+			}
+		}
+
+		virtual bool IsEmpty() const override
+		{
+			return false;
+		}
+
+		static FName GetDataTypeName()
+		{
+			static FName TypeName = TDataReferenceTypeInfo<TDataType>::TypeName;
+			return TypeName;
+		}
+
+		virtual FName GetDataType() override
+		{
+			return GetDataTypeName();
+		}
+
+	protected:
+		virtual TUniquePtr<IReceiver> ConstructNewReceiverImplementation(const FReceiverInitParams& InitParams) override
+		{
+			TArray<Audio::FPatchOutputStrongPtr> Outputs;
+			for (Audio::FPatchMixerSplitter& Bus : AudioBuses)
+			{
+				Outputs.Add(Bus.AddNewOutput(8096, 1.0f));
+			}
+
+			return TUniquePtr<IReceiver>(new TAudioReceiver<TDataType>(InitParams, this->AsShared(), MoveTemp(Outputs)));
+		}
+
+		virtual TUniquePtr<ISender> ConstructNewSenderImplementation(const FSenderInitParams& InitParams) override
+		{
+			TArray<Audio::FPatchInput> Inputs;
+			for (Audio::FPatchMixerSplitter& Bus : AudioBuses)
+			{
+				Inputs.Add(Bus.AddNewInput(8096, 1.0f));
+			}
+
+			return TUniquePtr<ISender>(new TAudioSender<TDataType>(InitParams, this->AsShared(), MoveTemp(Inputs)));
+		}
+
+	private:
+		TArray<Audio::FPatchMixerSplitter> AudioBuses;
+	};
+
+
+	// Coalesce both our standard copying Sender/Receivers and our audio-specific sender/receiver implementations into TSenderPtr and TReceiverPtr:
+
+	template<typename TDataType, typename USenderType = typename std::conditional<TIsDerivedFrom<TDataType, IAudioDatatype>::Value, TAudioSender<TDataType>, TSender<TDataType>>::type>
+	using TSenderPtr = TUniquePtr<USenderType>;
+
+	template<typename TDataType, typename UReceiverType = typename std::conditional<TIsDerivedFrom<TDataType, IAudioDatatype>::Value, TAudioReceiver<TDataType>, TReceiver<TDataType>>::type>
+	using TReceiverPtr = TUniquePtr<UReceiverType>;
+
+	// SFINAE for creating the correct data channel type for the given datatype:
+
+	template<typename TDataType, typename TEnableIf<std::is_copy_constructible<TDataType>::value && !TIsDerivedFrom<TDataType, IAudioDatatype>::Value, bool>::Type = true>
+	TSharedRef<IDataChannel> MakeDataChannel(const FOperatorSettings& InSettings)
+	{
+		return MakeShareable(new TCopyableDataChannel<TDataType>(InSettings));
+	}
+
+	template<typename TDataType, typename TEnableIf<TIsDerivedFrom<TDataType, IAudioDatatype>::Value, bool>::Type = true>
+	TSharedRef<IDataChannel> MakeDataChannel(const FOperatorSettings& InSettings)
+	{
+		return MakeShareable(new TAudioDataChannel<TDataType>(InSettings));
+	}
+
+	template<typename TDataType, typename TEnableIf<!std::is_copy_constructible<TDataType>::value && !TIsDerivedFrom<TDataType, IAudioDatatype>::Value, bool>::Type = true>
+	TSharedRef<IDataChannel> MakeDataChannel(const FOperatorSettings& InSettings)
+	{
+		return MakeShareable(new TNonOperationalDataChannel<TDataType>(InSettings));
+	}
+
+	// Utility function for properly downcasting to the correct type based on whether TDataType is an audio stream:
+
+	// For Audio Senders:
+	template <typename TDataType, typename TEnableIf<TIsDerivedFrom<TDataType, IAudioDatatype>::Value, bool>::Type = true>
+	TSenderPtr<TDataType> Downcast(TUniquePtr<ISender>&& InPtr)
+	{
+		checkf(InPtr->CheckType<TAudioSender<TDataType>>(), TEXT("Tried to downcast an ISender of type %s to a TSender of type %s!"), *(InPtr->GetDataType().ToString()), *(TAudioSender<TDataType>::GetDataTypeName().ToString()));
+		return TSenderPtr<TDataType>(static_cast<TAudioSender<TDataType>*>(InPtr.Release()));
+	}
+
+	// For generic senders for copyable types:
+	template <typename TDataType, typename TEnableIf<std::is_copy_constructible<TDataType>::value && !TIsDerivedFrom<TDataType, IAudioDatatype>::Value, bool>::Type = true>
+	TSenderPtr<TDataType> Downcast(TUniquePtr<ISender>&& InPtr)
+	{
+		checkf(InPtr->CheckType<TSender<TDataType>>(), TEXT("Tried to downcast an ISender of type %s to a TSender of type %s!"), *(InPtr->GetDataType().ToString()), *(TSender<TDataType>::GetDataTypeName().ToString()));
+		return TSenderPtr<TDataType>(static_cast<TSender<TDataType>*>(InPtr.Release()));
+	}
+
+	// For invalid types (non-copyable AND not an audio type):
+	template <typename TDataType, typename TEnableIf<!std::is_copy_constructible<TDataType>::value && !TIsDerivedFrom<TDataType, IAudioDatatype>::Value, bool>::Type = true>
+	TSenderPtr<TDataType> Downcast(TUniquePtr<ISender>&& InPtr)
+	{
+		return nullptr;
+	}
+
+	// For audio receivers:
+	template <typename TDataType, typename TEnableIf<TIsDerivedFrom<TDataType, IAudioDatatype>::Value, bool>::Type = true>
+	TReceiverPtr<TDataType> Downcast(TUniquePtr<IReceiver>&& InPtr)
+	{
+		checkf(InPtr->CheckType<TAudioReceiver<TDataType>>(), TEXT("Tried to downcast an IReceiver of type %s to a TSender of type %s!"), *(InPtr->GetDataType().ToString()), *(TAudioReceiver<TDataType>::GetDataTypeName().ToString()));
+		return TReceiverPtr<TDataType>(static_cast<TAudioReceiver<TDataType>*>(InPtr.Release()));
+	}
+
+	// For generic receivers for copyable types:
+	template <typename TDataType, typename TEnableIf<std::is_copy_constructible<TDataType>::value && !TIsDerivedFrom<TDataType, IAudioDatatype>::Value, bool>::Type = true>
+	TReceiverPtr<TDataType> Downcast(TUniquePtr<IReceiver>&& InPtr)
+	{
+		checkf(InPtr->CheckType<TReceiver<TDataType>>(), TEXT("Tried to downcast an IReceiver of type %s to a TSender of type %s!"), *(InPtr->GetDataType().ToString()), *(TReceiver<TDataType>::GetDataTypeName().ToString()));
+		return TReceiverPtr<TDataType>(static_cast<TReceiver<TDataType>*>(InPtr.Release()));
+	}
+
+	// For invalid types (non-copyable AND not an audio type):
+	template <typename TDataType, typename TEnableIf<!std::is_copy_constructible<TDataType>::value && !TIsDerivedFrom<TDataType, IAudioDatatype>::Value, bool>::Type = true>
+	TReceiverPtr<TDataType> Downcast(TUniquePtr<IReceiver>&& InPtr)
+	{
+		return nullptr;
+	}
+
+	// Basic router that takes an FName address, 
+	class METASOUNDFRONTEND_API FAddressRouter
+	{
+	public:
+		TArray<FName> GetAvailableChannels();
+
+		FName GetDatatypeForChannel(FName InChannelName);
+
+		template<typename TDataType>
+		TSenderPtr<TDataType> RegisterNewSend(FName InChannelName, const FSenderInitParams& InitParams)
+		{
+			FScopeLock ScopeLock(&DataChannelMapMutationLock);
+
+			if (TSharedRef<IDataChannel>* ExistingChannelPtr = DataChannelMap.Find(InChannelName))
+			{
+				TSharedRef<IDataChannel>& ExistingChannel = *ExistingChannelPtr;
+
+				// ensure that the existing data channel is of the same datatype as the send we're trying to register.
+				return Downcast<TDataType>(ExistingChannel->NewSender(InitParams));
+			}
+			else
+			{
+				// This is the first time we're seeing this, add it to the map.
+				TSharedRef<IDataChannel>& ExistingChannel = DataChannelMap.Add(InChannelName, MakeDataChannel<TDataType>(InitParams.OperatorSettings));
+				return Downcast<TDataType>(ExistingChannel->NewSender(InitParams));
+			}
+		}
+
+		template<typename TDataType>
+		TReceiverPtr<TDataType> RegisterNewReceiver(FName InChannelName, const FReceiverInitParams& InitParams)
+		{
+			FScopeLock ScopeLock(&DataChannelMapMutationLock);
+
+			if (TSharedRef<IDataChannel>* ExistingChannelPtr = DataChannelMap.Find(InChannelName))
+			{
+				TSharedRef<IDataChannel>& ExistingChannel = *ExistingChannelPtr;
+
+				TUniquePtr<IReceiver> BaseReceiver = ExistingChannel->NewReceiver(InitParams);
+				return Downcast<TDataType>(MoveTemp(BaseReceiver));
+			}
+			else
+			{
+				// This is the first time we're seeing this, add it to the map.
+				TSharedRef<IDataChannel>& ExistingChannel = DataChannelMap.Add(InChannelName, MakeDataChannel<TDataType>(InitParams.OperatorSettings));
+
+				return Downcast<TDataType>(ExistingChannel->NewReceiver(InitParams));
+			}
+		}
+
+		bool PushLiteral(FName InChannelName, const FDataTypeLiteralParam& InParam)
+		{
+			FScopeLock ScopeLock(&DataChannelMapMutationLock);
+			if (TSharedRef<IDataChannel>* ExistingChannelPtr = DataChannelMap.Find(InChannelName))
+			{
+				TSharedRef<IDataChannel>& ExistingChannel = *ExistingChannelPtr;
+				return ExistingChannel->PushLiteral(InParam);
+			}
+			else
+			{
+				return false;
+			}
+		}
+
+		FAddressRouter(const FAddressRouter& Other)
+			: DataChannelMap(Other.DataChannelMap)
+		{
+		}
+
+		FAddressRouter()
+		{}
+
+	private:
+		TMap<FName, TSharedRef<IDataChannel>> DataChannelMap;
+		FCriticalSection DataChannelMapMutationLock;
+	};
+
+	// this class is simply a container of a map of instances to FAddressRouters.
+	class FInstanceLocalRouter
+	{
+	public:
+
+		TArray<FName> GetAvailableChannels(uint64 InInstanceID);
+
+		FName GetDatatypeForChannel(uint64 InInstanceID, FName InChannelName);
+
+		template<typename TDataType>
+		TSenderPtr<TDataType> RegisterNewSend(uint64 InInstanceID, FName InChannelName, const FSenderInitParams& InitParams)
+		{
+			FScopeLock ScopeLock(&InstanceRouterMapMutationLock);
+
+			if (FAddressRouter* AddressRouter = InstanceRouterMap.Find(InInstanceID))
+			{
+				return AddressRouter->RegisterNewSend<TDataType>(InChannelName, InitParams);
+			}
+			else
+			{
+				// This is the first time we're seeing this, add it to the map.
+				FAddressRouter& NewRouter = InstanceRouterMap.Add(InInstanceID, FAddressRouter());
+				return NewRouter.RegisterNewSend<TDataType>(InChannelName, InitParams);
+			}
+		}
+
+		template<typename TDataType>
+		TReceiverPtr<TDataType> RegisterNewReceiver(uint64 InInstanceID, FName InChannelName, const FReceiverInitParams& InitParams)
+		{
+			FScopeLock ScopeLock(&InstanceRouterMapMutationLock);
+
+			if (FAddressRouter* AddressRouter = InstanceRouterMap.Find(InInstanceID))
+			{
+				return AddressRouter->RegisterNewReceiver<TDataType>(InChannelName, InitParams);
+			}
+			else
+			{
+				// This is the first time we're seeing this, add it to the map.
+				FAddressRouter& NewRouter = InstanceRouterMap.Add(InInstanceID, FAddressRouter());
+				return NewRouter.RegisterNewReceiver<TDataType>(InChannelName, InitParams);
+			}
+		}
+
+	private:
+		TMap<uint64, FAddressRouter> InstanceRouterMap;
+		FCriticalSection InstanceRouterMapMutationLock;
+	};
+
+	// Interface for any subsystem that wants to register as an independently addressable subsystem.
+	// The name provided to the ITransmissionSubsystem constructor will appear in the drop down menu for metasound send and receive nodes,
+	// and your implementation of ITransmissionSubsystem will be notified whenever a new send or receive is created for it's subsystem.
+	class METASOUNDFRONTEND_API ITransmissionSubsystem
+	{
+	public:
+		virtual ~ITransmissionSubsystem();
+
+	protected:
+
+		ITransmissionSubsystem() = delete;
+		
+		// Invoked by the implementing class. Automatically registers this subsystem with the FDataTransmissionCenter singleton.
+		ITransmissionSubsystem(FName InSubsystemName);
+
+		virtual void OnNewSendRegistered(const FSendAddress& SendAddress, FName DataType) {};
+		virtual void OnNewReceiverRegistered(const FSendAddress& SendAddress, FName DataType) {};
+		virtual bool CanSupportDataType(FName DataType) { return true; }
+
+	private:
+		FName SubsystemName;
+
+		friend class FDataTransmissionCenter;
+	};
+
+	// Main entry point for all sender/receiver registration.
+	class METASOUNDFRONTEND_API FDataTransmissionCenter
+	{
+		struct FSubsystemData
+		{
+			ITransmissionSubsystem* SubsystemPtr;
+			FAddressRouter AddressRouter;
+
+			FSubsystemData(ITransmissionSubsystem* InSystemPtr)
+				: SubsystemPtr(InSystemPtr)
+			{}
+		};
+
+	public:
+
+		// Returns the universal router.
+		static FDataTransmissionCenter& Get();
+
+		template<typename TDataType>
+		bool DoesSubsytemSupportDatatype(FName InSubsystem)
+		{
+			if (InSubsystem == GetSubsystemNameForSendScope(ETransmissionScope::ThisInstanceOnly) || InSubsystem == GetSubsystemNameForSendScope(ETransmissionScope::Global))
+			{
+				return true;
+			}
+			else if (FSubsystemData* FoundSubsystem = SubsystemRouters.Find(InSubsystem))
+			{
+				ITransmissionSubsystem* SubsystemInterface = FoundSubsystem->SubsystemPtr;
+				check(SubsystemInterface);
+
+				// Ensure that this subsystem can support this send type.
+				return SubsystemInterface->CanSupportDataType(FName(TDataReferenceTypeInfo<TDataType>::TypeName));
+			}
+			else
+			{
+				// Otherwise, the subsystem FName was invalid.
+				ensureAlways(false);
+				return false;
+			}
+		}
+
+		// Creates a new object to push data to an address.
+		// Returns a new sender, or nullptr if registration failed.
+		template<typename TDataType>
+		TSenderPtr<TDataType> RegisterNewSend(const FSendAddress& InAddress, const FSenderInitParams& InitParams)
+		{
+			if (InAddress.Subsystem == GetSubsystemNameForSendScope(ETransmissionScope::ThisInstanceOnly))
+			{
+				return InstanceRouter.RegisterNewSend<TDataType>(InAddress.MetasoundInstanceID, InAddress.ChannelName, InitParams);
+			}
+			else if (InAddress.Subsystem == GetSubsystemNameForSendScope(ETransmissionScope::Global))
+			{
+				return GlobalRouter.RegisterNewSend<TDataType>(InAddress.ChannelName, InitParams);
+			}
+			else if (FSubsystemData* FoundSubsystem = SubsystemRouters.Find(InAddress.Subsystem))
+			{
+				ITransmissionSubsystem* SubsystemInterface = FoundSubsystem->SubsystemPtr;
+
+				// Ensure that this subsystem can support this send type.
+				if (!ensureAlways(SubsystemInterface && SubsystemInterface->CanSupportDataType(FName(TDataReferenceTypeInfo<TDataType>::TypeName))))
+				{
+					return TSenderPtr<TDataType>();
+				}
+
+				FScopeLock ScopeLock(&SubsystemRoutersMutationLock);
+				TSenderPtr<TDataType> Sender = FoundSubsystem->AddressRouter.RegisterNewSend<TDataType>(InAddress.ChannelName, InitParams);
+				FoundSubsystem->SubsystemPtr->OnNewSendRegistered(InAddress, FName(TDataReferenceTypeInfo<TDataType>::TypeName));
+				return MoveTemp(Sender);
+			}
+			else
+			{
+				// Otherwise, the subsystem FName was invalid.
+				ensureAlways(false);
+				return TSenderPtr<TDataType>();
+			}
+		}
+
+		// Registers a new object to poll data from an address.
+		// Returns a new receiver, or nullptr if registration failed.
+		template<typename TDataType>
+		TReceiverPtr<TDataType> RegisterNewReceiver(const FSendAddress& InAddress, const FReceiverInitParams& InitParams)
+		{
+			if (InAddress.Subsystem == GetSubsystemNameForSendScope(ETransmissionScope::ThisInstanceOnly))
+			{
+				return InstanceRouter.RegisterNewReceiver<TDataType>(InAddress.MetasoundInstanceID, InAddress.ChannelName, InitParams);
+			}
+			else if (InAddress.Subsystem == GetSubsystemNameForSendScope(ETransmissionScope::Global))
+			{
+				return GlobalRouter.RegisterNewReceiver<TDataType>(InAddress.ChannelName, InitParams);
+			}
+			else if (FSubsystemData* FoundSubsystem = SubsystemRouters.Find(InAddress.Subsystem))
+			{
+				ITransmissionSubsystem* SubsystemInterface = FoundSubsystem->SubsystemPtr;
+
+				// Ensure that this subsystem can support this send type.
+				if (!ensureAlways(SubsystemInterface && SubsystemInterface->CanSupportDataType(FName(TDataReferenceTypeInfo<TDataType>::TypeName))))
+				{
+					return TReceiverPtr<TDataType>();
+				}
+
+				FScopeLock ScopeLock(&SubsystemRoutersMutationLock);
+				TReceiverPtr<TDataType> Receiver = FoundSubsystem->AddressRouter.RegisterNewReceiver<TDataType>(InAddress.ChannelName, InitParams);
+				FoundSubsystem->SubsystemPtr->OnNewReceiverRegistered(InAddress, FName(TDataReferenceTypeInfo<TDataType>::TypeName));
+				return MoveTemp(Receiver);
+			}
+			else
+			{
+				// Otherwise, the subsystem FName was invalid.
+				ensureAlways(false);
+				return TReceiverPtr<TDataType>();
+			}
+		}
+
+		// Pushes a literal parameter to a specific data channel in the global router.
+		// returns false if the literal type isn't supported.
+		bool PushLiteral(FName GlobalChannelName, const FDataTypeLiteralParam& InParam)
+		{
+			return GlobalRouter.PushLiteral(GlobalChannelName, InParam);
+		}
+
+	private:
+		// Single map of FNames to IDataChannels
+		FAddressRouter GlobalRouter;
+
+		// Map of instance IDs to FAddressRouter instances:
+		FInstanceLocalRouter InstanceRouter;
+
+		// Map of Subsystem Names to FAddressRouter instances.
+		TMap<FName, FSubsystemData> SubsystemRouters;
+		FCriticalSection SubsystemRoutersMutationLock;
+
+		// These are used by the constructor and destructor of implementations of ITransmissionSubsystem.
+		void RegisterSubsystem(ITransmissionSubsystem* InSystem, FName InSubsytemName)
+		{
+			FScopeLock ScopeLock(&SubsystemRoutersMutationLock);
+
+			//check to make sure we're not adding a subsystem twice.
+			check(!SubsystemRouters.Contains(InSubsytemName));
+
+			SubsystemRouters.Emplace(InSubsytemName, InSystem);
+		}
+
+		void UnregisterSubsystem(FName InSubsystemName)
+		{
+			FScopeLock ScopeLock(&SubsystemRoutersMutationLock);
+
+			//check to make sure we're not adding a subsystem twice.
+			check(SubsystemRouters.Contains(InSubsystemName));
+
+			SubsystemRouters.Remove(InSubsystemName);
+		}
+
+		FDataTransmissionCenter();
+
+		friend class ITransmissionSubsystem;
+	};
+
+	DECLARE_METASOUND_DATA_REFERENCE_TYPES(FSendAddress, METASOUNDFRONTEND_API, FSendAddressTypeInfo, FSendAddressReadRef, FSendAddressWriteRef)
+}
diff --git a/Engine/Plugins/Runtime/Metasound/Source/MetasoundFrontend/Public/MetasoundSendNode.h b/Engine/Plugins/Runtime/Metasound/Source/MetasoundFrontend/Public/MetasoundSendNode.h
new file mode 100644
--- /dev/null
+++ b/Engine/Plugins/Runtime/Metasound/Source/MetasoundFrontend/Public/MetasoundSendNode.h
@@ -0,0 +1,184 @@
+// Copyright Epic Games, Inc. All Rights Reserved.
+#pragma once
+
+#include "CoreMinimal.h"
+#include "MetasoundBuilderInterface.h"
+#include "MetasoundNodeInterface.h"
+#include "MetasoundOperatorInterface.h"
+#include "MetasoundDataReference.h"
+#include "MetasoundExecutableOperator.h"
+#include "MetasoundRouter.h"
+
+#include <type_traits>
+
+#define LOCTEXT_NAMESPACE "MetasoundGraphCore"
+
+
+namespace Metasound
+{
+	// This convenience node can be registered and will invoke static_cast<ToDataType>(FromDataType) every time it is executed.
+	template<typename TDataType>
+	class TSendNode : public INode
+	{
+		static_assert(TDataReferenceTypeInfo<TDataType>::bIsValidSpecialization, "Please use DECLARE_METASOUND_DATA_REFERENCE_TYPES with this class before trying to create an converter node with it.");
+
+	public:
+		static FString& GetAddressInputName()
+		{
+			static FString InputName = FString(TEXT("Address"));
+			return InputName;
+		}
+
+		static FString& GetSendInputName()
+		{
+			static FString SendInput = FString(TDataReferenceTypeInfo<TDataType>::TypeName);
+			return SendInput;
+		}
+
+	private:
+		class TSendOperator : public TExecutableOperator<TSendOperator>
+		{
+			public:
+
+				TSendOperator(TDataReadReference<TDataType> InInputData, TDataReadReference<FSendAddress> InSendAddress, const FOperatorSettings& InOperatorSettings)
+					: InputData(InInputData)
+					, SendAddress(InSendAddress)
+					, CachedSendAddress(*InSendAddress)
+					, CachedSenderParams({InOperatorSettings, 0.0f})
+					, Sender(FDataTransmissionCenter::Get().RegisterNewSend<TDataType>(CachedSendAddress, CachedSenderParams))
+				{
+					Inputs.AddDataReadReference<FSendAddress>(GetAddressInputName(), SendAddress);
+					Inputs.AddDataReadReference<TDataType>(GetSendInputName(), TDataReadReference<TDataType>(InputData));
+				}
+
+				virtual ~TSendOperator() {}
+
+				virtual const FDataReferenceCollection& GetInputs() const override
+				{
+					return Inputs;
+				}
+
+				virtual const FDataReferenceCollection& GetOutputs() const override
+				{
+					return Outputs;
+				}
+
+				void Execute()
+				{
+					if (SendAddress->ChannelName != CachedSendAddress.ChannelName)
+					{
+						CachedSendAddress = *SendAddress;
+						Sender = FDataTransmissionCenter::Get().RegisterNewSend<TDataType>(CachedSendAddress, CachedSenderParams);
+						check(Sender.IsValid());
+					}
+
+					Sender->Push(*InputData);
+				}
+
+			private:
+				TDataReadReference<TDataType> InputData;
+				TDataReadReference<FSendAddress> SendAddress;
+				FSendAddress CachedSendAddress;
+				FSenderInitParams CachedSenderParams;
+
+				TSenderPtr<TDataType> Sender;
+
+				FDataReferenceCollection Inputs;
+				FDataReferenceCollection Outputs;
+		};
+
+		class FSendOperatorFactory : public IOperatorFactory
+		{
+			public:
+				FSendOperatorFactory() = default;
+
+				virtual TUniquePtr<IOperator> CreateOperator(const FCreateOperatorParams& InParams, FBuildErrorArray& OutErrors) override
+				{
+					return MakeUnique<TSendOperator>(InParams.InputDataReferences.GetDataReadReference<TDataType>(GetSendInputName()),
+						InParams.InputDataReferences.GetDataReadReference<FSendAddress>(GetAddressInputName()),
+						InParams.OperatorSettings
+						);
+				}
+		};
+
+		public:
+			static FVertexInterface DeclareVertexInterface()
+			{
+				return FVertexInterface(
+					FInputVertexInterface(
+						TInputDataVertexModel<FSendAddress>(GetAddressInputName(), FText::GetEmpty()),
+						TInputDataVertexModel<TDataType>(GetSendInputName(), FText::GetEmpty())
+					),
+					FOutputVertexInterface(
+					)
+				);
+			}
+
+			TSendNode(const FNodeInitData& InInitData)
+				: NodeDescription(InInitData.InstanceName)
+				, Interface(DeclareVertexInterface())
+				, Factory(MakeOperatorFactoryRef<FSendOperatorFactory>())
+			{
+			}
+
+			virtual ~TSendNode() = default;
+
+			virtual const FName& GetClassName() const override
+			{
+				static const FName ClassName(*FString::Printf(TEXT("Send %s"), *GetSendInputName()));
+				return ClassName;
+			}
+
+			virtual const FString& GetInstanceName() const override
+			{
+				return NodeDescription;
+			}
+
+			virtual const FText& GetDescription() const override
+			{
+				static const FText Description = LOCTEXT("Metasound_SendNodeDescription", "Sends data to any other metasound anywhere else.");
+				return Description;
+			}
+
+			virtual const FText& GetAuthorName() const override
+			{
+				return PluginAuthor;
+			}
+
+			virtual const FText& GetPromptIfMissing() const override
+			{
+				return PluginNodeMissingPrompt;
+			}
+
+			virtual const FVertexInterface& GetVertexInterface() const override
+			{
+				return Interface;
+			}
+
+			virtual const FVertexInterface& GetDefaultVertexInterface() const override
+			{
+				return Interface;
+			}
+
+			virtual bool SetVertexInterface(const FVertexInterface& InInterface) override
+			{
+				return Interface == InInterface;
+			}
+
+			virtual bool IsVertexInterfaceSupported(const FVertexInterface& InInterface) const override
+			{
+				return Interface == InInterface;
+			}
+
+			virtual FOperatorFactorySharedRef GetDefaultOperatorFactory() const override
+			{
+				return Factory;
+			}
+
+		private:
+			FString NodeDescription;
+			FVertexInterface Interface;
+			FOperatorFactorySharedRef Factory;
+	};
+}
+#undef LOCTEXT_NAMESPACE
diff --git a/Engine/Plugins/Runtime/Metasound/Source/MetasoundGenerator/MetasoundGenerator.Build.cs b/Engine/Plugins/Runtime/Metasound/Source/MetasoundGenerator/MetasoundGenerator.Build.cs
--- a/Engine/Plugins/Runtime/Metasound/Source/MetasoundGenerator/MetasoundGenerator.Build.cs
+++ b/Engine/Plugins/Runtime/Metasound/Source/MetasoundGenerator/MetasoundGenerator.Build.cs
@@ -25,7 +25,8 @@ namespace UnrealBuildTool.Rules
                     "Core",
 					"AudioMixer",
 					"MetasoundGraphCore",
-					"MetasoundStandardNodes"
+					"MetasoundStandardNodes",
+					"MetasoundFrontend"
                 }
             );
 		}
diff --git a/Engine/Plugins/Runtime/Metasound/Source/MetasoundGraphCore/Public/MetasoundDataReference.h b/Engine/Plugins/Runtime/Metasound/Source/MetasoundGraphCore/Public/MetasoundDataReference.h
--- a/Engine/Plugins/Runtime/Metasound/Source/MetasoundGraphCore/Public/MetasoundDataReference.h
+++ b/Engine/Plugins/Runtime/Metasound/Source/MetasoundGraphCore/Public/MetasoundDataReference.h
@@ -240,7 +240,6 @@ namespace Metasound
 
 			typedef TDataReferenceTypeInfo<DataType> FInfoType;
 
-
 			/** This should be used to construct a new DataType object and return this TDataReference as a wrapper around it.
 			 */
 			template <typename... ArgTypes>
@@ -462,5 +461,20 @@ namespace Metasound
 		None, // If this is set, we will invoke TType(const FOperatorSettings&) If that constructor exists, or the default constructor if not.
 		Invalid,
 	};
+
+	// This empty base class is used so that we can specialize various nodes (Send, Receive, etc.) for subclasses of IAudioDatatype.
+	class METASOUNDGRAPHCORE_API IAudioDatatype
+	{
+		/**
+		 * Audio datatypes require the following functions:
+		 * int32 GetNumChannels() const { return NumChannels; }
+		 *
+		 * int32 GetMaxNumChannels() const { return MaxNumChannels; }
+
+		 * const TArrayView<const FAudioBufferReadRef> GetBuffers() const { return ReadableBuffers; }
+
+		 * const TArrayView<const FAudioBufferWriteRef> GetBuffers() { return WritableBuffers; }
+		 */
+	};
 }
 
diff --git a/Engine/Plugins/Runtime/Metasound/Source/MetasoundGraphCore/Public/MetasoundOperatorInterface.h b/Engine/Plugins/Runtime/Metasound/Source/MetasoundGraphCore/Public/MetasoundOperatorInterface.h
--- a/Engine/Plugins/Runtime/Metasound/Source/MetasoundGraphCore/Public/MetasoundOperatorInterface.h
+++ b/Engine/Plugins/Runtime/Metasound/Source/MetasoundGraphCore/Public/MetasoundOperatorInterface.h
@@ -183,7 +183,7 @@ namespace Metasound
 		}
 
 		template<typename DataType>
-		bool IsCompatibleWithType()
+		bool IsCompatibleWithType() const
 		{
 			switch (ConstructorArgType)
 			{
@@ -227,7 +227,7 @@ namespace Metasound
 		// This is a bucket constructor for when ConstructorArgType is unsupported by the data type we're trying to generate.
 		// This will only be invoked if something has gone terribly wrong in the frontend.
 		template<typename DataType>
-		DataType InvalidConstructor()
+		DataType InvalidConstructor() const
 		{
 			checkNoEntry();
 
@@ -246,7 +246,7 @@ namespace Metasound
 		template<typename DataType, typename TEnableIf<
 			TIsConstructible<DataType, const TArray<Audio::IProxyDataPtr>&, const FOperatorSettings&>::Value
 			, bool >::Type = true >
-			DataType ParseAsAudioProxyArray(const FOperatorSettings& WithOperatorSettings)
+			DataType ParseAsAudioProxyArray(const FOperatorSettings& WithOperatorSettings) const
 		{
 			check(ConstructorArgType == ELiteralArgType::UObjectProxyArray);
 			const TArray<Audio::IProxyDataPtr>& ProxyArrayToParse = ConstructorArg.Get<TArray<Audio::IProxyDataPtr>>();
@@ -256,7 +256,7 @@ namespace Metasound
 		template<typename DataType, typename TEnableIf<
 			TIsConstructible<DataType, const TArray<Audio::IProxyDataPtr>&>::Value && !TIsConstructible<DataType, const TArray<Audio::IProxyDataPtr>&, const FOperatorSettings&>::Value
 			, bool >::Type = true >
-			DataType ParseAsAudioProxyArray(const FOperatorSettings& WithOperatorSettings)
+			DataType ParseAsAudioProxyArray(const FOperatorSettings& WithOperatorSettings) const
 		{
 			check(ConstructorArgType == ELiteralArgType::UObjectProxyArray);
 			const TArray<Audio::IProxyDataPtr>& ProxyArrayToParse = ConstructorArg.Get<TArray<Audio::IProxyDataPtr>>();
@@ -264,7 +264,7 @@ namespace Metasound
 		}
 
 		template<typename DataType>
-		DataType ParseAsAudioProxyArray(...)
+		DataType ParseAsAudioProxyArray(...) const
 		{
 			return InvalidConstructor<DataType>();
 		}
@@ -272,7 +272,7 @@ namespace Metasound
 		template<typename DataType, typename TEnableIf<
 			TIsConstructible<DataType, const Audio::IProxyData&, const FOperatorSettings&>::Value
 			, bool >::Type = true >
-		DataType ParseAsAudioProxy(const FOperatorSettings& WithOperatorSettings)
+		DataType ParseAsAudioProxy(const FOperatorSettings& WithOperatorSettings) const
 		{
 			check(ConstructorArgType == ELiteralArgType::UObjectProxy);
 			const Audio::IProxyDataPtr& ProxyToParse = ConstructorArg.Get<Audio::IProxyDataPtr>();
@@ -282,7 +282,7 @@ namespace Metasound
 		template<typename DataType, typename TEnableIf<
 			TIsConstructible<DataType, const Audio::IProxyData&>::Value && !TIsConstructible<DataType, const Audio::IProxyData&, const FOperatorSettings&>::Value
 			, bool >::Type = true >
-		DataType ParseAsAudioProxy(const FOperatorSettings& WithOperatorSettings)
+		DataType ParseAsAudioProxy(const FOperatorSettings& WithOperatorSettings) const
 		{
 			check(ConstructorArgType == ELiteralArgType::UObjectProxy);
 			const Audio::IProxyDataPtr& ProxyToParse = ConstructorArg.Get<Audio::IProxyDataPtr>();
@@ -291,7 +291,7 @@ namespace Metasound
 		}
 
 		template<typename DataType>
-		DataType ParseAsAudioProxy(...)
+		DataType ParseAsAudioProxy(...) const
 		{
 			return InvalidConstructor<DataType>();
 		}
@@ -299,7 +299,7 @@ namespace Metasound
 		template<typename DataType, typename TEnableIf< 
 			TIsConstructible<DataType, const FString&, const FOperatorSettings&>::Value
 			, bool >::Type = true >
-		DataType ParseAsString(const FOperatorSettings& WithOperatorSettings)
+		DataType ParseAsString(const FOperatorSettings& WithOperatorSettings) const
 		{
 			check(ConstructorArgType == ELiteralArgType::String);
 			FString StringToParse = ConstructorArg.Get<FString>();
@@ -309,7 +309,7 @@ namespace Metasound
 		template<typename DataType, typename TEnableIf< 
 			TIsConstructible<DataType, const FString&>::Value && !TIsConstructible<DataType, const FString&, const FOperatorSettings&>::Value
 			, bool >::Type = true >
-		DataType ParseAsString(const FOperatorSettings& WithOperatorSettings)
+		DataType ParseAsString(const FOperatorSettings& WithOperatorSettings) const
 		{
 			check(ConstructorArgType == ELiteralArgType::String);
 			FString StringToParse = ConstructorArg.Get<FString>();
@@ -317,7 +317,7 @@ namespace Metasound
 		}
 
 		template<typename DataType>
-		DataType ParseAsString(...)
+		DataType ParseAsString(...) const
 		{
 			return InvalidConstructor<DataType>();
 		}
@@ -325,7 +325,7 @@ namespace Metasound
 		template<typename DataType, typename TEnableIf< 
 			TIsConstructible<DataType, float, const FOperatorSettings&>::Value
 			, bool >::Type = true >
-		DataType ParseAsFloat(const FOperatorSettings& WithOperatorSettings)
+		DataType ParseAsFloat(const FOperatorSettings& WithOperatorSettings) const
 		{
 			check(ConstructorArgType == ELiteralArgType::Float);
 			float ValueToParse = ConstructorArg.Get<float>();
@@ -335,7 +335,7 @@ namespace Metasound
 		template<typename DataType, typename TEnableIf< 
 			TIsConstructible<DataType, float >::Value && !TIsConstructible<DataType, float, const FOperatorSettings&>::Value
 			, bool >::Type = true >
-		DataType ParseAsFloat(const FOperatorSettings& WithOperatorSettings)
+		DataType ParseAsFloat(const FOperatorSettings& WithOperatorSettings) const
 		{
 			check(ConstructorArgType == ELiteralArgType::Float);
 			float ValueToParse = ConstructorArg.Get<float>();
@@ -343,7 +343,7 @@ namespace Metasound
 		}
 
 		template<typename DataType>
-		DataType ParseAsFloat(...)
+		DataType ParseAsFloat(...) const
 		{
 			return InvalidConstructor<DataType>();
 		}
@@ -351,7 +351,7 @@ namespace Metasound
 		template<typename DataType, typename TEnableIf< 
 			TIsConstructible<DataType, int32, const FOperatorSettings&>::Value
 			, bool >::Type = true >
-		DataType ParseAsInt(const FOperatorSettings& WithOperatorSettings)
+		DataType ParseAsInt(const FOperatorSettings& WithOperatorSettings) const
 		{
 			check(ConstructorArgType == ELiteralArgType::Integer);
 			int32 ValueToParse = ConstructorArg.Get<int32>();
@@ -361,7 +361,7 @@ namespace Metasound
 		template<typename DataType, typename TEnableIf< 
 			TIsConstructible<DataType, int32>::Value && !TIsConstructible<DataType, int32, const FOperatorSettings&>::Value
 			, bool >::Type = true >
-		DataType ParseAsInt(const FOperatorSettings& WithOperatorSettings)
+		DataType ParseAsInt(const FOperatorSettings& WithOperatorSettings) const
 		{
 			check(ConstructorArgType == ELiteralArgType::Integer);
 			int32 ValueToParse = ConstructorArg.Get<int32>();
@@ -369,7 +369,7 @@ namespace Metasound
 		}
 
 		template<typename DataType>
-		DataType ParseAsInt(...)
+		DataType ParseAsInt(...) const
 		{
 			return InvalidConstructor<DataType>();
 		}
@@ -377,7 +377,7 @@ namespace Metasound
 		template<typename DataType, typename TEnableIf< 
 			TIsConstructible<DataType, bool, const FOperatorSettings&>::Value
 			, bool >::Type = true >
-		DataType ParseAsBool(const FOperatorSettings& WithOperatorSettings)
+		DataType ParseAsBool(const FOperatorSettings& WithOperatorSettings) const
 		{
 			check(ConstructorArgType == ELiteralArgType::Boolean);
 			bool ValueToParse = ConstructorArg.Get<bool>();
@@ -387,7 +387,7 @@ namespace Metasound
 		template<typename DataType, typename TEnableIf< 
 			TIsConstructible<DataType, bool>::Value && !TIsConstructible<DataType, bool, const FOperatorSettings&>::Value
 			, bool >::Type = true >
-		DataType ParseAsBool(const FOperatorSettings& WithOperatorSettings)
+		DataType ParseAsBool(const FOperatorSettings& WithOperatorSettings) const
 		{
 			check(ConstructorArgType == ELiteralArgType::Boolean);
 			bool ValueToParse = ConstructorArg.Get<bool>();
@@ -395,7 +395,7 @@ namespace Metasound
 		}
 
 		template<typename DataType>
-		DataType ParseAsBool(...)
+		DataType ParseAsBool(...) const
 		{
 			return InvalidConstructor<DataType>();
 		}
@@ -404,7 +404,7 @@ namespace Metasound
 		template<typename DataType, typename TEnableIf< 
 			TIsConstructible<DataType, const FOperatorSettings&>::Value
 			, bool >::Type = true >
-		DataType ParseWithDefaultConstructor(const FOperatorSettings& WithOperatorSettings)
+		DataType ParseWithDefaultConstructor(const FOperatorSettings& WithOperatorSettings) const
 		{
 			check(ConstructorArgType == ELiteralArgType::None);
 			return DataType(WithOperatorSettings);
@@ -413,14 +413,14 @@ namespace Metasound
 		template<typename DataType, typename TEnableIf< 
 			TIsConstructible<DataType>::Value && !TIsConstructible<DataType, const FOperatorSettings&>::Value
 			, bool>::Type = true >
-		DataType ParseWithDefaultConstructor(const FOperatorSettings& WithOperatorSettings)
+		DataType ParseWithDefaultConstructor(const FOperatorSettings& WithOperatorSettings) const
 		{
 			check(ConstructorArgType == ELiteralArgType::None);
 			return DataType();
 		}
 
 		template<typename DataType>
-		DataType ParseWithDefaultConstructor(...)
+		DataType ParseWithDefaultConstructor(...) const
 		{
 			return InvalidConstructor<DataType>();
 		}
@@ -429,7 +429,7 @@ namespace Metasound
 
 		// This can be used to invoke the correct constructor for this init param for a data type.
 		template<typename DataType>
-		DataType ParseTo(const FOperatorSettings& WithOperatorSettings)
+		DataType ParseTo(const FOperatorSettings& WithOperatorSettings) const
 		{
 			switch (ConstructorArgType)
 			{
diff --git a/Engine/Plugins/Runtime/Metasound/Source/MetasoundStandardNodes/Private/MetasoundAudioFormats.cpp b/Engine/Plugins/Runtime/Metasound/Source/MetasoundStandardNodes/Private/MetasoundAudioFormats.cpp
--- a/Engine/Plugins/Runtime/Metasound/Source/MetasoundStandardNodes/Private/MetasoundAudioFormats.cpp
+++ b/Engine/Plugins/Runtime/Metasound/Source/MetasoundStandardNodes/Private/MetasoundAudioFormats.cpp
@@ -18,7 +18,7 @@ private:
 
 namespace Metasound
 {
-	REGISTER_METASOUND_DATATYPE(FAudioBuffer, "Audio:Buffer")
+	// REGISTER_METASOUND_DATATYPE(FAudioBuffer, "Audio:Buffer")
 	REGISTER_METASOUND_DATATYPE(FUnformattedAudio, "Audio:Unformatted")
 	REGISTER_METASOUND_DATATYPE(FMultichannelAudioFormat, "Audio:Multichannel")
 
diff --git a/Engine/Plugins/Runtime/Metasound/Source/MetasoundStandardNodes/Private/MetasoundStandardNodesModule.cpp b/Engine/Plugins/Runtime/Metasound/Source/MetasoundStandardNodes/Private/MetasoundStandardNodesModule.cpp
--- a/Engine/Plugins/Runtime/Metasound/Source/MetasoundStandardNodes/Private/MetasoundStandardNodesModule.cpp
+++ b/Engine/Plugins/Runtime/Metasound/Source/MetasoundStandardNodes/Private/MetasoundStandardNodesModule.cpp
@@ -10,12 +10,6 @@
 #include "MetasoundTime.h"
 #include "MetasoundFrequency.h"
 
-REGISTER_METASOUND_DATATYPE(bool, "Primitive:Bool", ::Metasound::ELiteralArgType::Boolean)
-REGISTER_METASOUND_DATATYPE(int32, "Primitive:Int32", ::Metasound::ELiteralArgType::Integer)
-REGISTER_METASOUND_DATATYPE(int64, "Primitive:Int64", ::Metasound::ELiteralArgType::Integer)
-REGISTER_METASOUND_DATATYPE(float, "Primitive:Float", ::Metasound::ELiteralArgType::Float)
-REGISTER_METASOUND_DATATYPE(double, "Primitive:Double", ::Metasound::ELiteralArgType::Float)
-
 namespace Metasound
 {
 	REGISTER_METASOUND_DATATYPE(FBop, "Primitive:Bop", ::Metasound::ELiteralArgType::Boolean)
diff --git a/Engine/Plugins/Runtime/Metasound/Source/MetasoundStandardNodes/Public/MetasoundAudioBuffer.h b/Engine/Plugins/Runtime/Metasound/Source/MetasoundStandardNodes/Public/MetasoundAudioBuffer.h
deleted file mode 100644
diff --git a/Engine/Plugins/Runtime/Metasound/Source/MetasoundStandardNodes/Public/MetasoundAudioFormats.h b/Engine/Plugins/Runtime/Metasound/Source/MetasoundStandardNodes/Public/MetasoundAudioFormats.h
--- a/Engine/Plugins/Runtime/Metasound/Source/MetasoundStandardNodes/Public/MetasoundAudioFormats.h
+++ b/Engine/Plugins/Runtime/Metasound/Source/MetasoundStandardNodes/Public/MetasoundAudioFormats.h
@@ -19,7 +19,7 @@ namespace Metasound
 	 * of the FUnformattedAudio. All audio buffers within a FUnformattedAudio object must contain the 
 	 * same number of audio frames.
 	 */
-	class METASOUNDSTANDARDNODES_API FUnformattedAudio
+	class METASOUNDSTANDARDNODES_API FUnformattedAudio : public IAudioDatatype
 	{
 		public:
 			/** FUnformattedAudio Constructor.
@@ -30,6 +30,10 @@ namespace Metasound
 			 */
 			FUnformattedAudio(int32 InNumFrames, int32 InNumChannels, int32 InMaxNumChannels);
 
+			FUnformattedAudio()
+				: FUnformattedAudio(0, 0, 0)
+			{}
+
 			/**
 			 * FUnformattedAudio Constructor used by the metasound frontend.
 			 *
@@ -61,13 +65,13 @@ namespace Metasound
 			 *
 			 * This array will have GetNumChannels() elements.
 			 */
-			const TArrayView<const FAudioBufferReadRef>& GetBuffers() const { return ReadableBuffers; }
+			const TArrayView<const FAudioBufferReadRef> GetBuffers() const { return ReadableBuffers; }
 
 			/** Return an array view of the writable buffer references.
 			 *
 			 * This array will have GetNumChannels() elements.
 			 */
-			const TArrayView<const FAudioBufferWriteRef>& GetBuffers() { return WritableBuffers; }
+			const TArrayView<const FAudioBufferWriteRef> GetBuffers() { return WritableBuffers; }
 
 			/** Return an array of the readable buffer reference storage.
 			 *
@@ -103,7 +107,7 @@ namespace Metasound
 	 * of the FMultichannelAudioFormat. All audio buffers within a FMultichannelAudioFormat object must contain the 
 	 * same number of audio frames.
 	 */
-	class METASOUNDSTANDARDNODES_API FMultichannelAudioFormat
+	class METASOUNDSTANDARDNODES_API FMultichannelAudioFormat : public IAudioDatatype
 	{
 		public:
 			FMultichannelAudioFormat();
@@ -147,17 +151,20 @@ namespace Metasound
 			/** Return the number of audio channels. */
 			int32 GetNumChannels() const { return NumChannels; }
 
+			/** Return the maximum number of channels. Multichannel audio buffers cannot be resized after construction. */
+			int32 GetMaxNumChannels() const { return NumChannels; }
+
 			/** Return an array view of the readable buffer references.
 			 *
 			 * This array will have GetNumChannels() elements.
 			 */
-			const TArrayView<const FAudioBufferReadRef>& GetBuffers() const { return ReadableBuffers; }
+			const TArrayView<const FAudioBufferReadRef> GetBuffers() const { return ReadableBuffers; }
 
 			/** Return an array view of the writable buffer references.
 			 *
 			 * This array will have GetNumChannels() elements.
 			 */
-			const TArrayView<const FAudioBufferWriteRef>& GetBuffers() { return WritableBuffers; }
+			const TArrayView<const FAudioBufferWriteRef> GetBuffers() { return WritableBuffers; }
 
 			/** Return an array of the readable buffer reference storage.
 			 *
@@ -205,7 +212,7 @@ namespace Metasound
 	 * number of audio frames.
 	 */
 	template<int32 NumChannels>
-	class TStaticChannelAudioFormat
+	class TStaticChannelAudioFormat : public IAudioDatatype
 	{
 		public:
 			/** TStaticChannelAudioFormat Constructor
@@ -237,6 +244,9 @@ namespace Metasound
 				return NumChannels;
 			}
 
+			/** Return the maximum number of channels. static channel audio buffers cannot be resized after construction. */
+			int32 GetMaxNumChannels() const { return NumChannels; }
+
 			/** Return an readable buffer reference for a specific channel. */
 			template<int32 ChannelIndex>
 			FAudioBufferReadRef GetBuffer() const
diff --git a/Engine/Plugins/Runtime/Metasound/Source/MetasoundStandardNodes/Public/MetasoundPrimitives.h b/Engine/Plugins/Runtime/Metasound/Source/MetasoundStandardNodes/Public/MetasoundPrimitives.h
deleted file mode 100644
diff --git a/Engine/Plugins/Runtime/Metasound/Source/MetasoundStandardNodes/Public/MetasoundTime.h b/Engine/Plugins/Runtime/Metasound/Source/MetasoundStandardNodes/Public/MetasoundTime.h
--- a/Engine/Plugins/Runtime/Metasound/Source/MetasoundStandardNodes/Public/MetasoundTime.h
+++ b/Engine/Plugins/Runtime/Metasound/Source/MetasoundStandardNodes/Public/MetasoundTime.h
@@ -306,6 +306,10 @@ namespace Metasound
 				: TTime(InSeconds, ETimeResolution::Seconds, InSettings.GetSampleRate())
 			{}
 
+			TTime()
+				: TTime(0.0f, ETimeResolution::Seconds, SMALL_NUMBER)
+			{}
+
 			/** Construct a sample count based TTime using a number of samples
 			 * and a sample rate.
 			 *
diff --git a/Engine/Source/Runtime/SignalProcessing/Public/DSP/Dsp.h b/Engine/Source/Runtime/SignalProcessing/Public/DSP/Dsp.h
--- a/Engine/Source/Runtime/SignalProcessing/Public/DSP/Dsp.h
+++ b/Engine/Source/Runtime/SignalProcessing/Public/DSP/Dsp.h
@@ -610,6 +610,7 @@ namespace Audio
 
 		// Pushes some amount of samples into this circular buffer.
 		// Returns the amount of samples written.
+		// This can only be used for trivially copyable types.
 		int32 Push(const SampleType* InBuffer, uint32 NumSamples)
 		{
 			SampleType* DestBuffer = InternalBuffer.GetData();
@@ -627,6 +628,46 @@ namespace Audio
 			return NumToCopy;
 		}
 
+		// Push a single sample onto this buffer.
+		// Returns false if the buffer is full.
+		bool Push(const SampleType& InElement)
+		{
+			if (Remainder() == 0)
+			{
+				return false;
+			}
+			else
+			{
+				SampleType* DestBuffer = InternalBuffer.GetData();
+				const uint32 ReadIndex = ReadCounter.GetValue();
+				const uint32 WriteIndex = WriteCounter.GetValue();
+
+				DestBuffer[WriteIndex] = InElement;
+
+				WriteCounter.Set((WriteIndex + 1) % Capacity);
+				return true;
+			}
+		}
+
+		bool Push(SampleType&& InElement)
+		{
+			if (Remainder() == 0)
+			{
+				return false;
+			}
+			else
+			{
+				SampleType* DestBuffer = InternalBuffer.GetData();
+				const uint32 ReadIndex = ReadCounter.GetValue();
+				const uint32 WriteIndex = WriteCounter.GetValue();
+
+				DestBuffer[WriteIndex] = MoveTemp(InElement);
+
+				WriteCounter.Set((WriteIndex + 1) % Capacity);
+				return true;
+			}
+		}
+
 		// Same as Pop(), but does not increment the read counter.
 		int32 Peek(SampleType* OutBuffer, uint32 NumSamples)
 		{
@@ -646,6 +687,25 @@ namespace Audio
 			return NumToCopy;
 		}
 
+		// Peeks a single element.
+		// returns false if the element is empty.
+		bool Peek(SampleType& OutElement)
+		{
+			if (Num() == 0)
+			{
+				return false;
+			}
+			else
+			{
+				SampleType* SrcBuffer = InternalBuffer.GetData();
+				const uint32 ReadIndex = ReadCounter.GetValue();
+				
+				OutElement = SrcBuffer[ReadIndex];
+
+				return true;
+			}
+		}
+
 		// Pops some amount of samples into this circular buffer.
 		// Returns the amount of samples read.
 		int32 Pop(SampleType* OutBuffer, uint32 NumSamples)
@@ -658,6 +718,21 @@ namespace Audio
 			return NumSamplesRead;
 		}
 
+		// Pops a single element.
+		// Will assert if the buffer is empty. Please check Num() > 0 before calling.
+		SampleType Pop()
+		{
+			// Calling this when the buffer is empty is considered a fatal error.
+			check(Num() > 0);
+
+			SampleType* SrcBuffer = InternalBuffer.GetData();
+			const uint32 ReadIndex = ReadCounter.GetValue();
+
+			SampleType PoppedValue = MoveTempIfPossible(InternalBuffer[ReadIndex]);
+			ReadCounter.Set((ReadCounter.GetValue() + 1) % Capacity);
+			return PoppedValue;
+		}
+
 		// When called, seeks the read or write cursor to only retain either the NumSamples latest data
 		// (if bRetainOldestSamples is false) or the NumSamples oldest data (if bRetainOldestSamples is true)
 		// in the buffer. Cannot be used to increase the capacity of this buffer.
