WavePlayer Backcompat with old ICompressedAudioInfo

diff --git a/Engine/Plugins/Runtime/Metasound/Source/MetasoundEngine/MetasoundEngine.Build.cs b/Engine/Plugins/Runtime/Metasound/Source/MetasoundEngine/MetasoundEngine.Build.cs
--- a/Engine/Plugins/Runtime/Metasound/Source/MetasoundEngine/MetasoundEngine.Build.cs
+++ b/Engine/Plugins/Runtime/Metasound/Source/MetasoundEngine/MetasoundEngine.Build.cs
@@ -16,7 +16,8 @@ namespace UnrealBuildTool.Rules
 					"Core",
 					"MetasoundGraphCore",
 					"MetasoundGenerator",
-					"AudioExtensions"
+					"AudioExtensions",
+					"AudioCodecEngine"
 				}
 			);
 
diff --git a/Engine/Plugins/Runtime/Metasound/Source/MetasoundEngine/Private/MetasoundEngineModule.cpp b/Engine/Plugins/Runtime/Metasound/Source/MetasoundEngine/Private/MetasoundEngineModule.cpp
--- a/Engine/Plugins/Runtime/Metasound/Source/MetasoundEngine/Private/MetasoundEngineModule.cpp
+++ b/Engine/Plugins/Runtime/Metasound/Source/MetasoundEngine/Private/MetasoundEngineModule.cpp
@@ -12,7 +12,6 @@ DEFINE_LOG_CATEGORY(LogMetasoundEngine);
 
 namespace Metasound
 {
-	REGISTER_METASOUND_DATATYPE(FWave, "Primitive:Wave");
 	REGISTER_METASOUND_DATATYPE(FWaveAsset, "Primitive:WaveAsset", ::Metasound::ELiteralArgType::UObjectProxy, USoundWave)
 }
 
diff --git a/Engine/Plugins/Runtime/Metasound/Source/MetasoundEngine/Private/MetasoundWave.cpp b/Engine/Plugins/Runtime/Metasound/Source/MetasoundEngine/Private/MetasoundWave.cpp
--- a/Engine/Plugins/Runtime/Metasound/Source/MetasoundEngine/Private/MetasoundWave.cpp
+++ b/Engine/Plugins/Runtime/Metasound/Source/MetasoundEngine/Private/MetasoundWave.cpp
@@ -5,20 +5,12 @@
 #include "IAudioCodec.h"
 #include "IAudioCodecRegistry.h"
 #include "Sound/SoundWave.h"
+#include "DecoderInputFactory.h"
+#include "AudioDeviceManager.h"
+#include "AudioDevice.h"
 
 namespace Metasound
 {	
-	// Simple for now, just wrap around the TArray<uint8>
-	class FWaveDecoderInput : public Audio::FDecoderInputArrayView
-	{
-		FWaveReadRef WaveRef;	// These are not thread safe, and will need to be
-	public:
-		FWaveDecoderInput(const FWaveReadRef& InWaveRef)
-			: Audio::FDecoderInputArrayView(MakeArrayView(InWaveRef->CompressedBytes.GetData(), InWaveRef->CompressedBytes.Num()), 0)
-			, WaveRef(InWaveRef)
-		{}
-	};
-
 	// Small RAII helper.
 	struct FScopedBulkDataLock
 	{
@@ -88,22 +80,19 @@ namespace Metasound
 #endif //WITH_EDITORONLY_DATA
 	}
 	
-	FWave::FWave(USoundWave* InWave) 
-	{
-		DoInlineEncode(InWave,CompressedBytes);
-	}
 
-	FWave::FWave(const TArray<uint8>& InBytes) 
-		: CompressedBytes(InBytes)
-	{
-	}
+	FWaveAsset::FDecoderInputPtr FWaveAsset::CreateDecoderInput(const FWaveAssetReadRef& InWaveRef)
+	{				
+		FName OldFormat = FAudioDeviceManager::Get()->GetActiveAudioDevice()->GetRuntimeFormat(
+			const_cast<USoundWave*>(InWaveRef->GetSoundWave())
+		);
+
+		TUniquePtr<Audio::IDecoderInput> Input = Audio::CreateBackCompatDecoderInput(
+			OldFormat,
+			InWaveRef->GetSoundWave()
+		);		
+		return FWaveAsset::FDecoderInputPtr(Input.Release());
 
-	FWave::FDecoderInputPtr FWave::CreateDecoderInput(const FWaveReadRef& InWaveRef)
-	{
-		if( InWaveRef->CompressedBytes.Num() > 0)
-		{
-			return MakeShared<FWaveDecoderInput, ESPMode::ThreadSafe>(InWaveRef);
-		}
 		return nullptr;
 	}
 }
diff --git a/Engine/Plugins/Runtime/Metasound/Source/MetasoundEngine/Private/MetasoundWavePlayerNode.cpp b/Engine/Plugins/Runtime/Metasound/Source/MetasoundEngine/Private/MetasoundWavePlayerNode.cpp
--- a/Engine/Plugins/Runtime/Metasound/Source/MetasoundEngine/Private/MetasoundWavePlayerNode.cpp
+++ b/Engine/Plugins/Runtime/Metasound/Source/MetasoundEngine/Private/MetasoundWavePlayerNode.cpp
@@ -35,7 +35,7 @@ namespace Metasound
 		// Silent setup
 		FWavePlayerOperator(
 			const FOperatorSettings& InSettings, 
-			const FWaveReadRef& InWave )
+			const FWaveAssetReadRef& InWave )
 			: OperatorSettings(InSettings)
 			, Wave(InWave)
 			, AudioBuffer(FAudioBufferWriteRef::CreateNew(InSettings))
@@ -46,8 +46,8 @@ namespace Metasound
 		
 		FWavePlayerOperator(
 			const FOperatorSettings& InSettings,
-			const FWaveReadRef& InWave, 
-			FWave::FDecoderInputPtr&& InDecoderInput, 
+			const FWaveAssetReadRef& InWave, 
+			FWaveAsset::FDecoderInputPtr&& InDecoderInput, 
 			TUniquePtr<Audio::IDecoderOutput>&& InDecoderOutput,
 			TUniquePtr<Audio::IDecoder>&& InDecoder )
 			: OperatorSettings(InSettings)
@@ -107,25 +107,24 @@ namespace Metasound
 	private:
 		const FOperatorSettings OperatorSettings;
 
-		FWaveReadRef Wave;
+		FWaveAssetReadRef Wave;
 		FAudioBufferWriteRef AudioBuffer;
 		FDataReferenceCollection InputDataReferences;
 		FDataReferenceCollection OutputDataReferences;
 
 		// Decoder/IO. 
 		Audio::ICodec::FDecoderPtr Decoder;
-		FWave::FDecoderInputPtr DecoderInput;
+		FWaveAsset::FDecoderInputPtr DecoderInput;
 		TUniquePtr<Audio::IDecoderOutput> DecoderOutput;
 	};
 
 	const FNodeInfo FWavePlayerNode::Info = FNodeInfo(
-		{
-			FName(TEXT("Wave")),
-			LOCTEXT("Metasound_WavePlayerNodeDescription", "Plays a supplied Wave"),
-			PluginAuthor,
-			PluginNodeMissingPrompt
-		}
-	);
+	{
+		FName(TEXT("Wave Player")),
+		LOCTEXT("Metasound_WavePlayerNodeDescription", "Plays a supplied Wave"),
+		PluginAuthor,
+		PluginNodeMissingPrompt
+	});
 
 	TUniquePtr<IOperator> FWavePlayerNode::FOperatorFactory::CreateOperator(
 		const FCreateOperatorParams& InParams, 
@@ -137,16 +136,20 @@ namespace Metasound
 
 		const FDataReferenceCollection& InputCol = InParams.InputDataReferences;
 
-		FWaveReadRef Wave = InputCol.GetDataReadReferenceOrConstruct<FWave>(TEXT("Wave"));
+		FWaveAssetReadRef Wave = InputCol.GetDataReadReferenceOrConstruct<FWaveAsset>(TEXT("Wave"));
 
-		FWave::FDecoderInputPtr Input = FWave::CreateDecoderInput(Wave);
+		FWaveAsset::FDecoderInputPtr Input = FWaveAsset::CreateDecoderInput(Wave);
 		if (Input)
 		{
 			ICodecRegistry::FCodecPtr Codec = ICodecRegistry::Get().FindCodecByParsingInput(Input.Get());
 			if (Codec)
 			{
-				// V1, Ask for an output buffer the size of a frame.
-				IDecoderOutput::FRequirements Reqs { Float32_Interleaved, InParams.OperatorSettings.GetNumFramesPerBlock() };
+				IDecoderOutput::FRequirements Reqs 
+				{ 
+					Float32_Interleaved, 
+					InParams.OperatorSettings.GetNumFramesPerBlock(), 
+					InParams.OperatorSettings.GetSampleRate() 
+				};
 				TUniquePtr<IDecoderOutput> Output = IDecoderOutput::Create(Reqs);
 				TUniquePtr<IDecoder> Decoder = Codec->CreateDecoder(Input.Get(), Output.Get());
 
@@ -178,7 +181,7 @@ namespace Metasound
 		:	FNode(InName, FWavePlayerNode::Info)
 		,	Factory(MakeOperatorFactoryRef<FWavePlayerNode::FOperatorFactory>())
 	{
-		Interface.GetInputInterface().Add(TInputDataVertexModel<FWave>(TEXT("Wave"), LOCTEXT("WaveTooltip", "The Wave to be decoded")));
+		Interface.GetInputInterface().Add(TInputDataVertexModel<FWaveAsset>(TEXT("Wave"), LOCTEXT("WaveTooltip", "The Wave to be decoded")));
 		Interface.GetOutputInterface().Add(TOutputDataVertexModel<FAudioBuffer>(TEXT("Audio"), LOCTEXT("AudioTooltip", "The output audio")));
 	}
 
diff --git a/Engine/Plugins/Runtime/Metasound/Source/MetasoundEngine/Public/MetasoundWave.h b/Engine/Plugins/Runtime/Metasound/Source/MetasoundEngine/Public/MetasoundWave.h
--- a/Engine/Plugins/Runtime/Metasound/Source/MetasoundEngine/Public/MetasoundWave.h
+++ b/Engine/Plugins/Runtime/Metasound/Source/MetasoundEngine/Public/MetasoundWave.h
@@ -18,30 +18,11 @@ class USoundWave;
 namespace Metasound
 {
 	// Forward declare ReadRef
-	class FWave;
-	typedef TDataReadReference<FWave> FWaveReadRef;
-		
-	class METASOUNDENGINE_API FWave
-	{
-		TArray<uint8> CompressedBytes;
-		friend class FWaveDecoderInput;
-
-	public:
-		FWave() = default;
-
-		// For testing only.
-		FWave(const TArray<uint8>& InBytes);
-		
-		FWave(USoundWave* InWave);
-
-		using FDecoderInputPtr = TSharedPtr<Audio::IDecoderInput,ESPMode::ThreadSafe>;
+	class FWaveAsset;
+	typedef TDataReadReference<FWaveAsset> FWaveAssetReadRef;
 
-		// Factory function to create a Decoder input
-		static FDecoderInputPtr CreateDecoderInput(
-			const FWaveReadRef& InWaveRef);
-	};
-	// Metasound datatype that holds onto a weak ptr. Mostly used as a placeholder until we have a proper proxy type.
-	class FWaveAsset
+	// Metasound data type that holds onto a weak ptr. Mostly used as a placeholder until we have a proper proxy type.
+	class METASOUNDENGINE_API FWaveAsset
 	{
 	private:
 		TWeakObjectPtr<USoundWave> SoundWave;
@@ -64,8 +45,11 @@ namespace Metasound
 		{
 			return SoundWave.Get();
 		}
+
+		using FDecoderInputPtr = TSharedPtr<Audio::IDecoderInput,ESPMode::ThreadSafe>;
+		static FDecoderInputPtr CreateDecoderInput(
+			const FWaveAssetReadRef& InWaveRef);
 	};
 
-	DECLARE_METASOUND_DATA_REFERENCE_TYPES(FWave, METASOUNDENGINE_API, FWaveTypeInfo, FWaveReadRef, FWaveWriteRef)
 	DECLARE_METASOUND_DATA_REFERENCE_TYPES(FWaveAsset, METASOUNDENGINE_API, FWaveAssetTypeInfo, FWaveAssetReadRef, FWaveAssetWriteRef)
 }
diff --git a/Engine/Source/Runtime/AudioCodecEngine/AudioCodecEngine.Build.cs b/Engine/Source/Runtime/AudioCodecEngine/AudioCodecEngine.Build.cs
new file mode 100644
--- /dev/null
+++ b/Engine/Source/Runtime/AudioCodecEngine/AudioCodecEngine.Build.cs
@@ -0,0 +1,21 @@
+// Copyright Epic Games, Inc. All Rights Reserved.
+
+namespace UnrealBuildTool.Rules
+{
+	public class AudioCodecEngine : ModuleRules
+	{
+		public AudioCodecEngine(ReadOnlyTargetRules Target) : base(Target)
+		{
+			PrivateDependencyModuleNames.AddRange(
+				new string[]
+				{
+					"Core",
+					"CoreUObject",
+                    "Engine",
+					"AudioExtensions",
+				//	"TargetPlatform"
+				}
+			);
+		}
+	}
+}
diff --git a/Engine/Source/Runtime/AudioCodecEngine/Private/AudioCodecEngineModule.cpp b/Engine/Source/Runtime/AudioCodecEngine/Private/AudioCodecEngineModule.cpp
new file mode 100644
--- /dev/null
+++ b/Engine/Source/Runtime/AudioCodecEngine/Private/AudioCodecEngineModule.cpp
@@ -0,0 +1,29 @@
+// Copyright Epic Games, Inc. All Rights Reserved.
+
+#include "Modules/ModuleManager.h"
+
+#include "IAudioCodecRegistry.h"
+#include "DecoderBackCompat.h"
+
+class FAudioCodecEngineModule : public IModuleInterface
+{
+public:
+	void StartupModule() override
+	{
+		Audio::ICodecRegistry::Get().RegisterCodec(
+			MakeUnique<Audio::FBackCompatCodec>()
+		);
+	}
+	void ShutdownModule() override
+	{
+		using namespace Audio;
+		if (ICodecRegistry::FCodecPtr Codec = ICodecRegistry::Get().FindCodecByName(
+			FBackCompatCodec::GetDetailsStatic().Name,
+			FBackCompatCodec::GetDetailsStatic().Version))
+		{
+			ICodecRegistry::Get().UnregisterCodec(Codec);
+		}
+	}
+};
+
+IMPLEMENT_MODULE(FAudioCodecEngineModule, AudioCodecEngine);
diff --git a/Engine/Source/Runtime/AudioCodecEngine/Private/DecoderBackCompat.cpp b/Engine/Source/Runtime/AudioCodecEngine/Private/DecoderBackCompat.cpp
new file mode 100644
--- /dev/null
+++ b/Engine/Source/Runtime/AudioCodecEngine/Private/DecoderBackCompat.cpp
@@ -0,0 +1,97 @@
+// Copyright Epic Games, Inc. All Rights Reserved.
+
+#include "DecoderBackCompat.h"
+#include "DecoderInputBackCompat.h"
+#include "AudioDecompress.h"
+
+namespace Audio
+{
+	static const FCodecDetails Details =
+	{
+		TEXT("BackCompat"),
+		TEXT("BackCompat"),
+		1,
+		{ FCodecFeatures::HasDecoder }
+	};
+
+	bool FBackCompatCodec::SupportsPlatform(FName InPlatformName) const 
+	{
+		// Everything right now.
+		return true;
+	}
+	const FCodecDetails& FBackCompatCodec::GetDetails() const 
+	{
+		return Details;
+	}
+
+	const FCodecDetails& FBackCompatCodec::GetDetailsStatic()
+	{
+		return Details;
+	}
+
+	ICodec::FDecoderPtr FBackCompatCodec::CreateDecoder(
+		IDecoder::FDecoderInputPtr InSrc,
+		IDecoder::FDecoderOutputPtr InDst)
+	{
+		return MakeUnique<FBackCompat>(
+			InSrc,
+			InDst
+		);
+	}
+
+	FBackCompat::FBackCompat(IDecoder::FDecoderInputPtr InSrc, IDecoder::FDecoderOutputPtr InDst) 
+		: Src(InSrc)
+		, Dst(InDst)
+	{
+		audio_ensure(InSrc->FindSection(Desc));
+		FDecodedFormatInfo Info { Desc.NumChannels, Desc.NumFramesPerSec, EBitRepresentation::Int16_Interleaved };
+		Reqs = Dst->GetRequirements(Info);
+		ResidualBuffer.SetNum(Reqs.NumSampleFramesWanted * Desc.NumChannels);
+	}
+
+	Audio::IDecoder::FDecodeReturn FBackCompat::Decode()
+	{
+		uint32 NumFramesRemaining = Reqs.NumSampleFramesWanted;
+		FBackCompatInput& BackCompatSrc = static_cast<FBackCompatInput&>(*Src);
+		ICompressedAudioInfo* Info = BackCompatSrc.GetInfo();
+		
+		bool bFinished = false;
+		IDecoderOutput::FPushedAudioDetails PushedDetails(
+			Desc.NumFramesPerSec,
+			Desc.NumChannels,
+			FrameOffset
+		);
+
+		uint32 BuffSizeInBytes = ResidualBuffer.Num() * sizeof(int16);
+		uint32 BuffSizeInFrames = Reqs.NumSampleFramesWanted;
+		uint8* Buff = (uint8*)ResidualBuffer.GetData();
+
+		while(!bFinished && NumFramesRemaining > 0)
+		{
+			if (BackCompatSrc.Wave->IsStreaming())
+			{
+				bFinished = Info->StreamCompressedData(Buff, false, BuffSizeInBytes);
+			}
+			else
+			{
+				bFinished = Info->ReadCompressedData(Buff, false, BuffSizeInBytes);
+			}						
+			PushedDetails.SampleFramesStartOffset = FrameOffset;
+			
+			Dst->PushAudio( 
+				PushedDetails,
+				MakeArrayView(ResidualBuffer) 
+			);
+			
+			FrameOffset += BuffSizeInFrames;
+			NumFramesRemaining -= BuffSizeInFrames;
+		}
+		
+		if( bFinished )
+		{
+			return EDecodeResult::Finished;
+		}
+		return EDecodeResult::MoreDataRemaining;
+	}
+
+}
\ No newline at end of file
diff --git a/Engine/Source/Runtime/AudioCodecEngine/Private/DecoderBackCompat.h b/Engine/Source/Runtime/AudioCodecEngine/Private/DecoderBackCompat.h
new file mode 100644
--- /dev/null
+++ b/Engine/Source/Runtime/AudioCodecEngine/Private/DecoderBackCompat.h
@@ -0,0 +1,45 @@
+// Copyright Epic Games, Inc. All Rights Reserved.
+
+#pragma once
+
+#include "IAudioCodec.h"
+
+namespace Audio
+{
+	struct FBackCompatCodec : public ICodec
+	{
+		// Static version.
+		static const FCodecDetails& GetDetailsStatic();
+
+		// Query.
+		bool SupportsPlatform(FName InPlatformName) const override;
+		const FCodecDetails& GetDetails() const override;
+
+		// Factory for decoders
+		FDecoderPtr CreateDecoder(
+			IDecoder::FDecoderInputPtr InSrc,
+			IDecoder::FDecoderOutputPtr InDst) override;
+	
+	};
+
+	struct FBackCompat : public IDecoder
+	{
+		IDecoder::FDecoderInputPtr Src = nullptr;
+		IDecoder::FDecoderOutputPtr Dst = nullptr;
+		FFormatDescriptorSection Desc;
+		IDecoderOutput::FRequirements Reqs;
+		TArray<int16> ResidualBuffer;
+		uint32 FrameOffset = 0;
+
+		FBackCompat(			
+			IDecoder::FDecoderInputPtr InSrc,
+			IDecoder::FDecoderOutputPtr InDst );
+
+		bool HasError() const override
+		{
+			return false;
+		}
+
+		FDecodeReturn Decode() override;
+	};
+}
diff --git a/Engine/Source/Runtime/AudioCodecEngine/Private/DecoderInputBackCompat.cpp b/Engine/Source/Runtime/AudioCodecEngine/Private/DecoderInputBackCompat.cpp
new file mode 100644
--- /dev/null
+++ b/Engine/Source/Runtime/AudioCodecEngine/Private/DecoderInputBackCompat.cpp
@@ -0,0 +1,103 @@
+// Copyright Epic Games, Inc. All Rights Reserved.
+
+#include "DecoderInputBackCompat.h"
+#include "AudioDevice.h"
+#include "AudioDeviceManager.h"
+#include "Interfaces/IAudioFormat.h"
+#include "AudioDecompress.h"
+#include "Sound/SoundWave.h"
+#include "DecoderBackCompat.h"
+
+namespace Audio
+{	
+	ICompressedAudioInfo* FBackCompatInput::GetInfo(
+		FFormatDescriptorSection* OutDescriptor /*= nullptr*/) const
+	{
+		USoundWave* NonConstWave = const_cast<USoundWave*>(Wave.Get());
+		FAudioDeviceHandle Handle = FAudioDeviceManager::Get()->GetActiveAudioDevice();
+
+		if (!OldInfoObject.IsValid())
+		{
+			OldInfoObject.Reset(Handle->CreateCompressedAudioInfo(NonConstWave));
+			audio_ensure(OldInfoObject.IsValid());
+
+			FSoundQualityInfo Info;
+			if (Wave->IsStreaming())
+			{
+				if (!OldInfoObject->StreamCompressedInfo(NonConstWave, &Info))
+				{
+					return nullptr;
+				}
+			}
+			else
+			{
+				if (NonConstWave->ResourceData == nullptr)
+				{
+					NonConstWave->InitAudioResource(Handle->GetRuntimeFormat(NonConstWave));
+				}
+
+				if (!OldInfoObject->ReadCompressedInfo(NonConstWave->ResourceData, NonConstWave->ResourceSize, &Info))
+				{
+					return nullptr;
+				}
+			}
+
+			Desc.NumChannels		= Info.NumChannels;
+			Desc.NumFramesPerSec	= Info.SampleRate;
+			Desc.NumFrames			= (uint32)((float)Info.Duration * Info.SampleRate);
+			Desc.NumBytesPerPacket	= ~0;
+
+			Desc.CodecName			= FBackCompatCodec::GetDetailsStatic().Name;
+			Desc.CodecFamilyName	= FBackCompatCodec::GetDetailsStatic().FamilyName;
+			Desc.CodecVersion		= FBackCompatCodec::GetDetailsStatic().Version;		
+		}
+		if( OutDescriptor )
+		{		
+			*OutDescriptor = Desc;
+		}		
+		return OldInfoObject.Get();
+	}
+
+	bool FBackCompatInput::FindSection(FEncodedSectionBase& OutSection)
+	{
+		if (FFormatDescriptorSection::kSectionName == OutSection.SectionName)
+		{
+			FFormatDescriptorSection& Descriptor = static_cast<FFormatDescriptorSection&>(OutSection);
+			return GetInfo(&Descriptor) != nullptr;
+		}
+		return false;
+	}
+
+	bool FBackCompatInput::HasError() const
+	{
+		return GetInfo() == nullptr;
+	}
+
+	int64 FBackCompatInput::Tell() const
+	{
+		// Not implemented.
+		audio_ensure(false);
+		return -1;
+	}
+
+	bool FBackCompatInput::IsEndOfStream() const
+	{
+		// Not implemented.
+		audio_ensure(false);
+		return false;
+	}
+
+	TArrayView<const uint8> FBackCompatInput::PeekNextPacket(int32 InMaxPacketLength) const
+	{
+		// Not implemented.
+		audio_ensure(false);
+		return MakeArrayView<const uint8>(nullptr,0);
+	}
+
+	TArrayView<const uint8> FBackCompatInput::PopNextPacket(int32 InPacketSize)
+	{
+		// Not implemented.
+		audio_ensure(false);
+		return MakeArrayView<const uint8>(nullptr,0);
+	}
+}
diff --git a/Engine/Source/Runtime/AudioCodecEngine/Private/DecoderInputBackCompat.h b/Engine/Source/Runtime/AudioCodecEngine/Private/DecoderInputBackCompat.h
new file mode 100644
--- /dev/null
+++ b/Engine/Source/Runtime/AudioCodecEngine/Private/DecoderInputBackCompat.h
@@ -0,0 +1,43 @@
+// Copyright Epic Games, Inc. All Rights Reserved.
+
+#pragma once
+
+#include "IAudioCodec.h"
+#include "UObject/NameTypes.h"
+#include "AudioDecompress.h"
+
+class ICompressedAudioInfo; 	// Forward declares.
+
+namespace Audio
+{
+	struct FBackCompatInput : public IDecoderInput
+	{
+		FName OldFormatName;
+		TWeakObjectPtr<const USoundWave> Wave;
+		mutable FFormatDescriptorSection Desc;
+		mutable TUniquePtr<ICompressedAudioInfo> OldInfoObject;
+
+		FBackCompatInput(
+			FName InOldFormatName,
+			const USoundWave* InWave)
+			: OldFormatName(InOldFormatName)
+			, Wave(MakeWeakObjectPtr(InWave))
+		{
+		}
+
+		bool HasError() const override;
+		bool IsEndOfStream() const override;
+
+		ICompressedAudioInfo* GetInfo(
+			FFormatDescriptorSection* OutDescriptor = nullptr) const;
+
+		bool FindSection(FEncodedSectionBase& OutSection) override;
+		int64 Tell() const override;
+		
+		TArrayView<const uint8> PeekNextPacket(
+			int32 InMaxPacketLength) const override;
+
+		TArrayView<const uint8> PopNextPacket(
+			int32 InPacketSize) override;
+	};
+}
diff --git a/Engine/Source/Runtime/AudioCodecEngine/Private/DecoderInputFactory.cpp b/Engine/Source/Runtime/AudioCodecEngine/Private/DecoderInputFactory.cpp
new file mode 100644
--- /dev/null
+++ b/Engine/Source/Runtime/AudioCodecEngine/Private/DecoderInputFactory.cpp
@@ -0,0 +1,16 @@
+// Copyright Epic Games, Inc. All Rights Reserved.
+
+#include "DecoderInputFactory.h"
+#include "DecoderInputBackCompat.h"
+#include "IAudioCodec.h"
+#include "Sound/SoundWave.h"
+
+namespace Audio
+{		
+	TUniquePtr<Audio::IDecoderInput> CreateBackCompatDecoderInput(
+		FName InOldFormatName,
+		const USoundWave* InSoundWave)
+	{
+		return MakeUnique<FBackCompatInput>(InOldFormatName, InSoundWave);
+	}
+}
diff --git a/Engine/Source/Runtime/AudioCodecEngine/Public/DecoderInputFactory.h b/Engine/Source/Runtime/AudioCodecEngine/Public/DecoderInputFactory.h
new file mode 100644
--- /dev/null
+++ b/Engine/Source/Runtime/AudioCodecEngine/Public/DecoderInputFactory.h
@@ -0,0 +1,19 @@
+// Copyright Epic Games, Inc. All Rights Reserved.
+
+#pragma once
+
+#include "Templates/UniquePtr.h"
+
+// Forward declares.
+class USoundWave;
+
+namespace Audio
+{
+	// Forward declares.
+	struct IDecoderInput;
+
+	// Just loose for now.
+	AUDIOCODECENGINE_API TUniquePtr<IDecoderInput> CreateBackCompatDecoderInput(
+		FName InOldFormatName,
+		const USoundWave* InSoundWave );
+}
diff --git a/Engine/Source/Runtime/AudioExtensions/Private/IAudioCodecRegistry.cpp b/Engine/Source/Runtime/AudioExtensions/Private/IAudioCodecRegistry.cpp
--- a/Engine/Source/Runtime/AudioExtensions/Private/IAudioCodecRegistry.cpp
+++ b/Engine/Source/Runtime/AudioExtensions/Private/IAudioCodecRegistry.cpp
@@ -202,8 +202,8 @@ namespace Audio
 				return nullptr;
 			}
 
-			FCodecPtr Ptr = FindCodecByFamilyName(
-				Format.CodecFamilyName,
+			FCodecPtr Ptr = FindCodecByName(
+				Format.CodecName,
 				Format.CodecVersion);
 
 			if( Ptr )
@@ -211,6 +211,16 @@ namespace Audio
 				return Ptr;
 			}
 
+			FCodecPtr FamilyPtr = FindCodecByFamilyName(
+				Format.CodecFamilyName,
+				Format.CodecVersion);
+			
+			if (FamilyPtr)
+			{
+				return FamilyPtr;
+			}
+
+
 			// Fail.
 			return nullptr;		
 		}
diff --git a/Engine/Source/Runtime/AudioExtensions/Public/IAudioCodec.h b/Engine/Source/Runtime/AudioExtensions/Public/IAudioCodec.h
--- a/Engine/Source/Runtime/AudioExtensions/Public/IAudioCodec.h
+++ b/Engine/Source/Runtime/AudioExtensions/Public/IAudioCodec.h
@@ -40,18 +40,6 @@ namespace Audio
 	public:
 		FFcc() = default;
 
-#if 0 // Clang doesn't like this.
-		template<uint32 N>
-		constexpr FFcc(const char(&InStr)[N])
-			: Value(0)
-		{
-			static_assert(N <= 5, "Four Or Less Characters Strings Only");
-			for (uint32 i = 0, Shift = 0; i < N; ++i, Shift += 8)
-			{
-				Value |= static_cast<uint32>(InStr[i] << Shift);
-			}
-		}
-#endif 
 		constexpr FFcc(const char One, const char Two, const char Three, const char Four)
 		{
 			Value = static_cast<uint32>(One) |
@@ -139,8 +127,9 @@ namespace Audio
 	{	
 		struct FRequirements
 		{
-			EBitRepresentation DownstreamFormat	= Float32_Interleaved;	// Our pipeline is all float downstream, so this might not be necessary.
+			EBitRepresentation DownstreamFormat	= Float32_Interleaved;	
 			int32 NumSampleFramesWanted			= 256;
+			int32 NumSampleFramesPerSecond		= 48000;
 		};
 
 		// Factory.
@@ -219,9 +208,6 @@ namespace Audio
 
 	struct AUDIOEXTENSIONS_API FCodecFeatures
 	{	
-		bool HasFeature(FName Feature) const { return true; }
-		void GetFeatures(TArray<FName>& OutFeatureNames) const {}
-
 		enum EFeatureBitField
 		{
 			HwDecode,
@@ -235,13 +221,12 @@ namespace Audio
 			HasDecoder,
 			HasEncoder,
 		};
-		uint32_t FeaturesBitField;		// Store as bitfield.
+		uint32_t FeaturesBitField = 0;		// Store as bitfield.
 
 		bool HasFeature(EFeatureBitField InBit) const { return ( FeaturesBitField & (uint32)InBit); };
 
-		FCodecFeatures(std::initializer_list<EFeatureBitField> InitList)
+		constexpr FCodecFeatures(std::initializer_list<EFeatureBitField> InitList)
 		{
-			FeaturesBitField = 0;
 			for (EFeatureBitField i : InitList)
 			{
 				FeaturesBitField |= 1 << i;
@@ -748,25 +733,13 @@ namespace Audio
 		{}
 	};
 
-	struct AUDIOEXTENSIONS_API ICodec : public IModularFeature
+	struct AUDIOEXTENSIONS_API ICodec 
 	{
 		virtual ~ICodec() {}
 
-		static FName GetModularFeatureName() 
-		{ 
-			static const FName sName = TEXT("AudioCodec"); 
-			return sName; 
-		}	
-
-		// TODO. Add lifetime subscription delegates.
-		// OnDecoderCreated 
-		// OnEncoderCreated
-		// OnStaticInit
-		// OnStaticDeint
-		
 		// Query.
-		virtual bool SupportsPlatform(FName InPlatformName) const = 0;		// This could move to the details.
-		virtual const FCodecDetails& GetDetails() const = 0;				// Name/Version/ImplName/Features.
+		virtual bool SupportsPlatform(FName InPlatformName) const = 0;	
+		virtual const FCodecDetails& GetDetails() const = 0;			
 		FName GetName() const { return GetDetails().Name; };
 
 		// Factory for encoders
@@ -785,28 +758,7 @@ namespace Audio
 		{ 
 			return nullptr; 
 		}			
-
-	protected:
-		// Register Modular features.
-		void RegisterModularFeature() { IModularFeatures::Get().RegisterModularFeature(ICodec::GetModularFeatureName(), this); }
-		void UnRegisterModularFeature() { IModularFeatures::Get().UnregisterModularFeature(ICodec::GetModularFeatureName(), this); }
 	};
-
-	// template this with the codec class.
-	// Convenience class for when you have just one codec in a module.
-	struct AUDIOEXTENSIONS_API IAudioCodecModule : public ICodec,
-												   public IModuleInterface
-	{
-		virtual void StartupModule() override
-		{
-			RegisterModularFeature();
-		}
-		virtual void ShutdownModule() override
-		{
-			UnRegisterModularFeature();
-		}
-	};
-
 } // namespace Audio
 
 #pragma endregion Codec
