Fixing -Wc++11-narrowing warning as errror CIS in WavePlayerNode.

diff --git a/Engine/Plugins/Runtime/Metasound/Source/MetasoundEngine/Private/MetasoundWavePlayerNode.cpp b/Engine/Plugins/Runtime/Metasound/Source/MetasoundEngine/Private/MetasoundWavePlayerNode.cpp
--- a/Engine/Plugins/Runtime/Metasound/Source/MetasoundEngine/Private/MetasoundWavePlayerNode.cpp
+++ b/Engine/Plugins/Runtime/Metasound/Source/MetasoundEngine/Private/MetasoundWavePlayerNode.cpp
@@ -148,7 +148,7 @@ namespace Metasound
 				{ 
 					Float32_Interleaved, 
 					InParams.OperatorSettings.GetNumFramesPerBlock(), 
-					InParams.OperatorSettings.GetSampleRate() 
+					static_cast<int32>(InParams.OperatorSettings.GetSampleRate()) 
 				};
 				TUniquePtr<IDecoderOutput> Output = IDecoderOutput::Create(Reqs);
 				TUniquePtr<IDecoder> Decoder = Codec->CreateDecoder(Input.Get(), Output.Get());
