Metasound Bop implementation

diff --git a/Engine/Plugins/Runtime/Metasound/Source/MetasoundEngine/Private/MetasoundSource.cpp b/Engine/Plugins/Runtime/Metasound/Source/MetasoundEngine/Private/MetasoundSource.cpp
--- a/Engine/Plugins/Runtime/Metasound/Source/MetasoundEngine/Private/MetasoundSource.cpp
+++ b/Engine/Plugins/Runtime/Metasound/Source/MetasoundEngine/Private/MetasoundSource.cpp
@@ -7,10 +7,10 @@
 #include "AssetRegistryModule.h"
 
 #include "MetasoundAssetBase.h"
+#include "MetasoundBop.h"
 #include "MetasoundGenerator.h"
 #include "MetasoundAudioFormats.h"
 #include "MetasoundPrimitives.h"
-#include "MetasoundDataReferenceTypes.h"
 
 #include "MetasoundArchetypeRegistration.h"
 
@@ -82,7 +82,7 @@ ISoundGeneratorPtr UMetasoundSource::CreateSoundGenerator(const FSoundGeneratorI
 		{
 			MoveTemp(Operator),
 			Outputs.GetDataReadReferenceOrConstruct<Metasound::FAudioBuffer>(GetAudioOutputName(), InSettings.GetNumFramesPerBlock()),
-			Outputs.GetDataReadReferenceOrConstruct<Metasound::FBop>(GetIsFinishedOutputName(), false)
+			Outputs.GetDataReadReferenceOrConstruct<Metasound::FBop>(GetIsFinishedOutputName(), false, InSettings)
 		};
 
 		return ISoundGeneratorPtr(new Metasound::FMetasoundGenerator(MoveTemp(InitParams)));
diff --git a/Engine/Plugins/Runtime/Metasound/Source/MetasoundEngine/Private/MetasoundWavePlayerNode.cpp b/Engine/Plugins/Runtime/Metasound/Source/MetasoundEngine/Private/MetasoundWavePlayerNode.cpp
--- a/Engine/Plugins/Runtime/Metasound/Source/MetasoundEngine/Private/MetasoundWavePlayerNode.cpp
+++ b/Engine/Plugins/Runtime/Metasound/Source/MetasoundEngine/Private/MetasoundWavePlayerNode.cpp
@@ -2,7 +2,6 @@
 
 #include "MetasoundWavePlayerNode.h"
 #include "MetasoundExecutableOperator.h"
-#include "MetasoundDataReferenceTypes.h"
 #include "MetasoundNodeRegistrationMacro.h"
 #include "MetasoundPrimitives.h"
 #include "MetasoundWave.h"
diff --git a/Engine/Plugins/Runtime/Metasound/Source/MetasoundFrontend/Private/MetasoundRouter.cpp b/Engine/Plugins/Runtime/Metasound/Source/MetasoundFrontend/Private/MetasoundRouter.cpp
--- a/Engine/Plugins/Runtime/Metasound/Source/MetasoundFrontend/Private/MetasoundRouter.cpp
+++ b/Engine/Plugins/Runtime/Metasound/Source/MetasoundFrontend/Private/MetasoundRouter.cpp
@@ -68,7 +68,7 @@ static FAutoConsoleCommand GPushIntCommand(
 				return;
 			}
 
-			FName ChannelName = FName(Args[0]);
+			FName ChannelName = FName(*Args[0]);
 			int32 ValueToPush = TCString<TCHAR>::Atoi(*Args[1]);
 
 			Metasound::FDataTypeLiteralParam LiteralParam(ValueToPush);
diff --git a/Engine/Plugins/Runtime/Metasound/Source/MetasoundFrontend/Public/MetasoundReceiveNode.h b/Engine/Plugins/Runtime/Metasound/Source/MetasoundFrontend/Public/MetasoundReceiveNode.h
--- a/Engine/Plugins/Runtime/Metasound/Source/MetasoundFrontend/Public/MetasoundReceiveNode.h
+++ b/Engine/Plugins/Runtime/Metasound/Source/MetasoundFrontend/Public/MetasoundReceiveNode.h
@@ -16,12 +16,31 @@
 
 namespace Metasound
 {
+	// TODO: should consider making this a subclass of TInputNode. Both share requirement
+	// of needing a default DataReference. The construction of that gets tricky so 
+	// would be nice to consolidate all that logic. 
+	//
+	//
+	// TODO: Maybe it could work like this...
+	// Can instantiate a factory with default args. Factory is valid as long as constructor of data type supports
+	// parameters (w/ or w/o operator settings.) Factory makes a lambda for creating data type references.
+	// Factory gets instantiated when node is declared.  Basically lets us forward constructor 
+	// args until the creation of the lambda, which requires a copy to capture "[=]". 
+	//
+	// But that still leaves the issue of setting those values on the node. Would be best to have that in NodeInitParams
+	// but its not clear how that is passed about. Maybe something that checks whether it can construct from literal. 
+	//
+	// Or you could go another route where you check whether the Operator constructor takes a specific set of references and 
+	// use that to auto-generate factories.  
+	
 	// This convenience node can be registered and will invoke static_cast<ToDataType>(FromDataType) every time it is executed.
 	template<typename TDataType>
 	class TReceiveNode : public INode
 	{
 		static_assert(TDataReferenceTypeInfo<TDataType>::bIsValidSpecialization, "Please use DECLARE_METASOUND_DATA_REFERENCE_TYPES with this class before trying to create an converter node with it.");
 
+
+
 	public:
 		static FString& GetAddressInputName()
 		{
@@ -38,10 +57,11 @@ namespace Metasound
 	private:
 		class TReceiverOperator : public TExecutableOperator<TReceiverOperator>
 		{
+				TReceiverOperator() = delete;
 			public:
 
-				TReceiverOperator(TDataReadReference<FSendAddress> InSendAddress, const FOperatorSettings& InOperatorSettings)
-					: OutputData(TDataWriteReference<TDataType>::CreateNew())
+				TReceiverOperator(const TDataWriteReference<TDataType>& Data, TDataReadReference<FSendAddress> InSendAddress, const FOperatorSettings& InOperatorSettings)
+					: OutputData(Data)
 					, SendAddress(InSendAddress)
 					, CachedSendAddress(*InSendAddress)
 					, CachedReceiverParams({InOperatorSettings})
@@ -89,14 +109,43 @@ namespace Metasound
 				FDataReferenceCollection Outputs;
 		};
 
-		class FSendOperatorFactory : public IOperatorFactory
+		template<typename T>
+		struct TReceiverOperatorFactoryHelper
+		{
+			using FDataTypeInfo = TDataReferenceTypeInfo<T>;
+
+			static constexpr bool bUseDefaultCtor = FDataTypeInfo::bCanUseDefaultConstructor && !FDataTypeInfo::bIsConstructableWithSettings;
+			static constexpr bool bUseSettingsCtor = FDataTypeInfo::bIsConstructableWithSettings;
+		};
+
+		class FReceiverOperatorFactory : public IOperatorFactory
 		{
 			public:
-				FSendOperatorFactory() = default;
+				FReceiverOperatorFactory() = default;
+
+				template<
+					typename U = TDataType,
+					typename std::enable_if< TReceiverOperatorFactoryHelper<U>::bUseDefaultCtor, int>::type = 0
+				>
+				TDataWriteReference<TDataType> CreateDefaultWriteReference(const FCreateOperatorParams& InParams) const
+				{
+					return TDataWriteReference<TDataType>::CreateNew();
+				}
+
+				template<
+					typename U = TDataType,
+					typename std::enable_if< TReceiverOperatorFactoryHelper<U>::bUseSettingsCtor, int>::type = 0
+				>
+				TDataWriteReference<TDataType> CreateDefaultWriteReference(const FCreateOperatorParams& InParams) const
+				{
+					return TDataWriteReference<TDataType>::CreateNew(InParams.OperatorSettings);
+				}
+
 
 				virtual TUniquePtr<IOperator> CreateOperator(const FCreateOperatorParams& InParams, FBuildErrorArray& OutErrors) override
 				{
 					return MakeUnique<TReceiverOperator>(
+						CreateDefaultWriteReference(InParams),
 						InParams.InputDataReferences.GetDataReadReference<FSendAddress>(GetAddressInputName()),
 						InParams.OperatorSettings
 						);
@@ -120,7 +169,7 @@ namespace Metasound
 			TReceiveNode(const FNodeInitData& InInitData)
 				: NodeDescription(InInitData.InstanceName)
 				, Interface(DeclareVertexInterface())
-				, Factory(MakeOperatorFactoryRef<FSendOperatorFactory>())
+				, Factory(MakeOperatorFactoryRef<FReceiverOperatorFactory>())
 			{
 			}
 
diff --git a/Engine/Plugins/Runtime/Metasound/Source/MetasoundGenerator/Private/MetasoundGenerator.cpp b/Engine/Plugins/Runtime/Metasound/Source/MetasoundGenerator/Private/MetasoundGenerator.cpp
--- a/Engine/Plugins/Runtime/Metasound/Source/MetasoundGenerator/Private/MetasoundGenerator.cpp
+++ b/Engine/Plugins/Runtime/Metasound/Source/MetasoundGenerator/Private/MetasoundGenerator.cpp
@@ -2,8 +2,8 @@
 #include "MetasoundGenerator.h"
 
 #include "DSP/Dsp.h"
+#include "MetasoundBop.h"
 #include "MetasoundOperatorInterface.h"
-#include "MetasoundDataReferenceTypes.h"
 #include "MetasoundOscNode.h"
 #include "MetasoundAudioMultiplyNode.h"
 #include "MetasoundADSRNode.h"
diff --git a/Engine/Plugins/Runtime/Metasound/Source/MetasoundGenerator/Private/SynthComponentMetasoundExample.cpp b/Engine/Plugins/Runtime/Metasound/Source/MetasoundGenerator/Private/SynthComponentMetasoundExample.cpp
--- a/Engine/Plugins/Runtime/Metasound/Source/MetasoundGenerator/Private/SynthComponentMetasoundExample.cpp
+++ b/Engine/Plugins/Runtime/Metasound/Source/MetasoundGenerator/Private/SynthComponentMetasoundExample.cpp
@@ -144,7 +144,7 @@ void USynthComponentMetasoundExample::PostEditChangeProperty(struct FPropertyCha
 }
 #endif
 
-void USynthComponentMetasoundExample::SetGraphOperator(FOperatorUniquePtr InGraphOperator, const FString& InOutputAudioName)
+void USynthComponentMetasoundExample::SetGraphOperator(FOperatorUniquePtr InGraphOperator, const FOperatorSettings& InSettings, const FString& InOutputAudioName)
 {
 	using namespace Metasound;
 
@@ -152,7 +152,7 @@ void USynthComponentMetasoundExample::SetGraphOperator(FOperatorUniquePtr InGrap
 	{
 		FDataReferenceCollection Outputs = InGraphOperator->GetOutputs();
 		FAudioBufferReadRef Audio = Outputs.GetDataReadReferenceOrConstruct<FAudioBuffer>(InOutputAudioName, 256);
-		FBopReadRef OnFinished = FBopReadRef::CreateNew(); // unused in this example
+		FBopReadRef OnFinished = FBopReadRef::CreateNew(InSettings); // unused in this example
 		
 		// Multichannel version:
 		//NumChannels = Audio->GetNumChannels();
@@ -176,7 +176,7 @@ void USynthComponentMetasoundExample::SetGraphOperator(FOperatorUniquePtr InGrap
 	}
 }
 
-bool USynthComponentMetasoundExample::PushGraphOperator(FOperatorUniquePtr InGraphOperator, const FString& InOutputAudioName)
+bool USynthComponentMetasoundExample::PushGraphOperator(FOperatorUniquePtr InGraphOperator, const FOperatorSettings& InOperatorSettings, const FString& InOutputAudioName)
 {
 	using namespace Metasound;
 
@@ -189,7 +189,7 @@ bool USynthComponentMetasoundExample::PushGraphOperator(FOperatorUniquePtr InGra
 	FAudioBufferReadRef AudioReadRef = Outputs.GetDataReadReferenceOrConstruct<FAudioBuffer>(InOutputAudioName, 256);
 	
 	// On Finished isn't used in this example.
-	FBopReadRef OnFinishedRef = FBopReadRef::CreateNew();
+	FBopReadRef OnFinishedRef = FBopReadRef::CreateNew(InOperatorSettings);
 
 	FMetasoundGeneratorInitParams GeneratorInitParams =
 	{
@@ -225,7 +225,7 @@ ISoundGeneratorPtr USynthComponentMetasoundExample::CreateSoundGenerator(const F
 	TArray<IOperatorBuilder::FBuildErrorPtr> Errors;
 	FOperatorUniquePtr Operator = Builder.BuildGraphOperator(*ExampleGraph->Graph, Errors);
 
-	SetGraphOperator(MoveTemp(Operator), TEXT("Audio"));
+	SetGraphOperator(MoveTemp(Operator), Settings, TEXT("Audio"));
 
 	return Generator;
 }
diff --git a/Engine/Plugins/Runtime/Metasound/Source/MetasoundGenerator/Public/MetasoundGenerator.h b/Engine/Plugins/Runtime/Metasound/Source/MetasoundGenerator/Public/MetasoundGenerator.h
--- a/Engine/Plugins/Runtime/Metasound/Source/MetasoundGenerator/Public/MetasoundGenerator.h
+++ b/Engine/Plugins/Runtime/Metasound/Source/MetasoundGenerator/Public/MetasoundGenerator.h
@@ -4,11 +4,11 @@
 
 #include "CoreMinimal.h"
 #include "DSP/Dsp.h"
+#include "MetasoundBop.h"
 #include "MetasoundOperatorInterface.h"
 #include "MetasoundExecutableOperator.h"
 #include "MetasoundAudioFormats.h"
 #include "MetasoundPrimitives.h"
-#include "MetasoundDataReferenceTypes.h"
 #include "Sound/SoundGenerator.h"
 
 namespace Metasound
diff --git a/Engine/Plugins/Runtime/Metasound/Source/MetasoundGenerator/Public/SynthComponentMetasoundExample.h b/Engine/Plugins/Runtime/Metasound/Source/MetasoundGenerator/Public/SynthComponentMetasoundExample.h
--- a/Engine/Plugins/Runtime/Metasound/Source/MetasoundGenerator/Public/SynthComponentMetasoundExample.h
+++ b/Engine/Plugins/Runtime/Metasound/Source/MetasoundGenerator/Public/SynthComponentMetasoundExample.h
@@ -7,6 +7,7 @@
 #include "Components/SynthComponent.h"
 
 #include "MetasoundOperatorInterface.h"
+#include "MetasoundOperatorSettings.h"
 #include "MetasoundNodeInterface.h"
 #include "MetasoundAudioFormats.h"
 #include "MetasoundGenerator.h"
@@ -25,6 +26,7 @@ class METASOUNDGENERATOR_API USynthComponentMetasoundExample : public USynthComp
 public:
 	using FOperatorUniquePtr = TUniquePtr<Metasound::IOperator>;
 	using FAudioBufferReadRef = Metasound::FAudioBufferReadRef;
+	using FOperatorSettings = Metasound::FOperatorSettings;
 
     USynthComponentMetasoundExample(const FObjectInitializer& ObjInitializer);
 
@@ -48,18 +50,20 @@ public:
 	 * the next time CreateSoundGenerator(...) is called.
 	 *
 	 * @param InGraphOperator - Operator which executes entire MetasoundGraph.
+	 * @param InOperatorSettings - Settings associated with the graph operator.
 	 * @param InOutputAudioName - Name of the FAudioBufferReadRef data ref in the graph output.
 	 */
-	void SetGraphOperator(FOperatorUniquePtr InGraphOperator, const FString& InOutputAudioName);
+	void SetGraphOperator(FOperatorUniquePtr InGraphOperator, const FOperatorSettings& InOperatorSettings, const FString& InOutputAudioName);
 
 	/** Attempt to push a new graph operator to the existing metasound generator.
 	 *
 	 * @param InGraphOperator - Operator which executes entire MetasoundGraph.
+	 * @param InOperatorSettings - Settings associated with the graph operator.
 	 * @param InOutputAudioName - Name of the FAudioBufferReadRef data ref in the graph output.
 	 *
 	 * @return True on success, false on failure.
 	 */
-	bool PushGraphOperator(FOperatorUniquePtr InGraphOperator, const FString& InOutputAudioName);
+	bool PushGraphOperator(FOperatorUniquePtr InGraphOperator, const FOperatorSettings& InOperatorSettings, const FString& InOutputAudioName);
 
     virtual ISoundGeneratorPtr CreateSoundGenerator(const FSoundGeneratorInitParams& InParams) override;
 
diff --git a/Engine/Plugins/Runtime/Metasound/Source/MetasoundStandardNodes/Private/MetasoundADSRNode.cpp b/Engine/Plugins/Runtime/Metasound/Source/MetasoundStandardNodes/Private/MetasoundADSRNode.cpp
--- a/Engine/Plugins/Runtime/Metasound/Source/MetasoundStandardNodes/Private/MetasoundADSRNode.cpp
+++ b/Engine/Plugins/Runtime/Metasound/Source/MetasoundStandardNodes/Private/MetasoundADSRNode.cpp
@@ -1,11 +1,13 @@
 // Copyright Epic Games, Inc. All Rights Reserved.
 
 #include "MetasoundADSRNode.h"
+
+#include "DSP/Envelope.h"
+#include "MetasoundAudioBuffer.h"
+#include "MetasoundBop.h"
 #include "MetasoundExecutableOperator.h"
 #include "MetasoundTime.h"
-#include "MetasoundDataReferenceTypes.h"
 #include "MetasoundNodeRegistrationMacro.h"
-#include "DSP/Envelope.h"
 
 #define LOCTEXT_NAMESPACE "MetasoundStandardNodes"
 
@@ -36,13 +38,18 @@ namespace Metasound
 			void Execute();
 
 		private:
+
+			void GenerateEnvelope(int32 InStartFrame, int32 InEndFrame);
+
 			const FOperatorSettings OperatorSettings;
 
 			// TODO: write envelope gen for metasound more suited to this processing structure. 
 			Audio::FEnvelope Envelope;
 
 			FBopReadRef Bop;
-			int32 StopEnvelopePos;
+			bool bNotReleased;
+			FSampleTime TimeTillRelease;
+			FSampleTime TimePerBlock;
 			FADSRReferences ADSRReferences;
 			FAudioBufferWriteRef EnvelopeBuffer;
 
@@ -53,7 +60,9 @@ namespace Metasound
 	FADSROperator::FADSROperator(const FOperatorSettings& InSettings, const FBopReadRef& InBop, const FADSRReferences& InADSRData)
 	:	OperatorSettings(InSettings)
 	,	Bop(InBop)
-	,	StopEnvelopePos(-1)
+	,	bNotReleased(false)
+	,	TimeTillRelease(0, InSettings.GetSampleRate())
+	,	TimePerBlock(InSettings.GetNumFramesPerBlock(), InSettings.GetSampleRate())
 	,	ADSRReferences(InADSRData)
 	,	EnvelopeBuffer(FAudioBufferWriteRef::CreateNew(InSettings.GetNumFramesPerBlock()))
 	{
@@ -74,66 +83,68 @@ namespace Metasound
 		return OutputDataReferences;
 	}
 
-	void FADSROperator::Execute()
+	void FADSROperator::GenerateEnvelope(int32 InStartFrame, int32 InEndFrame)
 	{
-		Envelope.SetAttackTime(ADSRReferences.Attack->GetMilliseconds());
-		Envelope.SetDecayTime(ADSRReferences.Decay->GetMilliseconds());
-		Envelope.SetReleaseTime(ADSRReferences.Release->GetMilliseconds());
-
-		float ADSRMilliseconds = ADSRReferences.Attack->GetMilliseconds();
-		ADSRMilliseconds += ADSRReferences.Decay->GetMilliseconds();
-		ADSRMilliseconds += ADSRReferences.Sustain->GetMilliseconds();
+		float* EnvelopeData = EnvelopeBuffer->GetData();
 
-		// TODO: Bops need to be sorted.
-		// TODO: Make convenient way to do bop loops like this since it might happen alot.
-		// TODO: Bops need to be done on 4-sample boundaries to allow for SIMD.
-	
+		const int32 ReleaseFrame = TimeTillRelease.GetNumSamples();
 
-		int32 BopIndex = 0;
-		int32 StartPos = 0;
-		int32 NextBop = BopIndex < Bop->Num() ? (*Bop)[BopIndex] : OperatorSettings.GetNumFramesPerBlock() + 1;
-		int32 EndPos = FMath::Min(OperatorSettings.GetNumFramesPerBlock(), NextBop);
-		
-		if (StopEnvelopePos > 0)
+		if (bNotReleased && (ReleaseFrame < InEndFrame))
 		{
-			StopEnvelopePos -= OperatorSettings.GetNumFramesPerBlock();
-			EndPos = FMath::Min(StopEnvelopePos, EndPos);
-		}
+			// If the envelope is expected to be released during these frames,
+			// need to split generation of envelope so envelope can be stopped.
+			for (int32 i = InStartFrame; i < ReleaseFrame; i++)
+			{
+				EnvelopeData[i] = Envelope.Generate();
+			}
 
-		float* EnvelopeData = EnvelopeBuffer->GetData();
+			Envelope.Stop();
 
-		// TODO: GetNumFramesPerBlock() need to be on 4 sample boundaries to allow for SIMD.
-		// TODO: just make GetNumFramesPerBlock() an int32 instead of a uint32
-		while (StartPos < (int32)OperatorSettings.GetNumFramesPerBlock())
-		{
-			int32 i = StartPos;
-			for (; i < EndPos; i++)
+			bNotReleased = false;
+
+			for (int32 i = FMath::Max(InStartFrame, ReleaseFrame); i < InEndFrame; i++)
 			{
 				EnvelopeData[i] = Envelope.Generate();
 			}
-
-			StartPos = EndPos;
-			
-			if (EndPos == StopEnvelopePos)
+		}
+		else
+		{
+			for (int32 i = InStartFrame; i < InEndFrame; i++)
 			{
-				// Stop sustain, start release.
-				Envelope.Stop();
-				StopEnvelopePos = -1;
-				EndPos = FMath::Min((int32)OperatorSettings.GetNumFramesPerBlock(), NextBop);
+				EnvelopeData[i] = Envelope.Generate();
 			}
+		}
+	}
+
 
-			if (EndPos == NextBop)
+	void FADSROperator::Execute()
+	{
+		Envelope.SetAttackTime(ADSRReferences.Attack->GetMilliseconds());
+		Envelope.SetDecayTime(ADSRReferences.Decay->GetMilliseconds());
+		Envelope.SetReleaseTime(ADSRReferences.Release->GetMilliseconds());
+
+		if (bNotReleased)
+		{
+			// Keep track of when next release should happen
+			TimeTillRelease -= TimePerBlock;
+		}
+		
+		Bop->ExecuteBlock(
+			// OnPreBop
+			[&](int32 StartFrame, int32 EndFrame) { GenerateEnvelope(StartFrame, EndFrame); },
+			// OnBop
+			[&](int32 StartFrame, int32 EndFrame)
 			{
-				// Trigger next envelope. 
 				Envelope.Start();
-				StopEnvelopePos = EndPos + FMath::RoundToInt(OperatorSettings.GetSampleRate() * 0.001f * ADSRMilliseconds);
 
-				BopIndex++;
+				// Capture attack, decay and sustain time (skip release).
+				TimeTillRelease = *ADSRReferences.Attack + *ADSRReferences.Decay + *ADSRReferences.Sustain;
+
+				bNotReleased = true;
 
-				NextBop = BopIndex < Bop->Num() ? (*Bop)[BopIndex] : OperatorSettings.GetNumFramesPerBlock() + 1;
-				EndPos = FMath::Min(StopEnvelopePos, FMath::Min((int32)OperatorSettings.GetNumFramesPerBlock(), NextBop));
+				GenerateEnvelope(StartFrame, EndFrame);
 			}
-		}
+		);
 	}
 
 	const FNodeInfo& FADSROperator::GetNodeInfo()
@@ -153,10 +164,10 @@ namespace Metasound
 		static const FVertexInterface Interface(
 			FInputVertexInterface(
 				TInputDataVertexModel<FBop>(TEXT("Bop"), LOCTEXT("BopTooltip", "Trigger for envelope.")),
-				TInputDataVertexModel<FFloatTime>(TEXT("Attack"), LOCTEXT("AttackTooltip", "Attack time in milliseconds.")),
-				TInputDataVertexModel<FFloatTime>(TEXT("Decay"), LOCTEXT("DecayTooltip", "Decay time in milliseconds.")),
-				TInputDataVertexModel<FFloatTime>(TEXT("Sustain"), LOCTEXT("SustainTooltip", "Sustain time in milliseconds.")),
-				TInputDataVertexModel<FFloatTime>(TEXT("Release"), LOCTEXT("ReleaseTooltip", "Release time in milliseconds."))
+				TInputDataVertexModel<FFloatTime>(TEXT("Attack"), LOCTEXT("AttackTooltip", "Attack time.")),
+				TInputDataVertexModel<FFloatTime>(TEXT("Decay"), LOCTEXT("DecayTooltip", "Decay time.")),
+				TInputDataVertexModel<FFloatTime>(TEXT("Sustain"), LOCTEXT("SustainTooltip", "Sustain time.")),
+				TInputDataVertexModel<FFloatTime>(TEXT("Release"), LOCTEXT("ReleaseTooltip", "Release time."))
 			),
 			FOutputVertexInterface(
 				TOutputDataVertexModel<FAudioBuffer>(TEXT("Envelope"), LOCTEXT("EnvelopeTooltip", "The output envelope"))
@@ -174,7 +185,7 @@ namespace Metasound
 		const FDataReferenceCollection& InputCollection = InParams.InputDataReferences;
 
 		// TODO: Could no-op this if the bop is not connected.
-		FBopReadRef Bop = InputCollection.GetDataReadReferenceOrConstruct<FBop>(TEXT("Bop"));
+		FBopReadRef Bop = InputCollection.GetDataReadReferenceOrConstruct<FBop>(TEXT("Bop"), InParams.OperatorSettings);
 
 		// TODO: If none of these are connected, could pre-generate ADSR envelope and return a different operator. 
 		FADSRReferences ADSRReferences = 
diff --git a/Engine/Plugins/Runtime/Metasound/Source/MetasoundStandardNodes/Private/MetasoundAudioMultiplyNode.cpp b/Engine/Plugins/Runtime/Metasound/Source/MetasoundStandardNodes/Private/MetasoundAudioMultiplyNode.cpp
--- a/Engine/Plugins/Runtime/Metasound/Source/MetasoundStandardNodes/Private/MetasoundAudioMultiplyNode.cpp
+++ b/Engine/Plugins/Runtime/Metasound/Source/MetasoundStandardNodes/Private/MetasoundAudioMultiplyNode.cpp
@@ -1,10 +1,11 @@
 // Copyright Epic Games, Inc. All Rights Reserved.
 
 #include "MetasoundAudioMultiplyNode.h"
+#include "MetasoundAudioBuffer.h"
 #include "MetasoundExecutableOperator.h"
-#include "MetasoundDataReferenceTypes.h"
 #include "MetasoundFacade.h"
 #include "MetasoundNodeRegistrationMacro.h"
+#include "MetasoundOperatorSettings.h"
 #include "DSP/BufferVectorOperations.h"
 
 #define LOCTEXT_NAMESPACE "MetasoundAudioMultiplyNode"
diff --git a/Engine/Plugins/Runtime/Metasound/Source/MetasoundStandardNodes/Private/MetasoundBop.cpp b/Engine/Plugins/Runtime/Metasound/Source/MetasoundStandardNodes/Private/MetasoundBop.cpp
new file mode 100644
--- /dev/null
+++ b/Engine/Plugins/Runtime/Metasound/Source/MetasoundStandardNodes/Private/MetasoundBop.cpp
@@ -0,0 +1,117 @@
+// Copyright Epic Games, Inc. All Rights Reserved.
+
+#include "MetasoundBop.h"
+
+#include "MetasoundDataReference.h"
+#include "MetasoundDataTypeRegistrationMacro.h"
+#include "MetasoundOperatorSettings.h"
+
+namespace Metasound
+{
+	FBop::FBop(bool bShouldBop, const FOperatorSettings& InSettings)
+	:	NumFramesPerBlock(InSettings.GetNumFramesPerBlock())
+	,	SampleRate(InSettings.GetSampleRate())
+	{
+		if (bShouldBop)
+		{
+			BopFrame(0);
+		}
+	}
+
+	FBop::FBop(int32 InFrameToBop, const FOperatorSettings& InSettings)
+	:	NumFramesPerBlock(InSettings.GetNumFramesPerBlock())
+	,	SampleRate(InSettings.GetSampleRate())
+	{
+		BopFrame(InFrameToBop);
+	}
+
+	FBop::FBop(const FOperatorSettings& InSettings)
+		: FBop(false, InSettings)
+	{}
+
+	void FBop::BopFrame(int32 InFrameToBop)
+	{
+		// Insert bop frame index into sorted bopped frames
+		BoppedFrames.Insert(InFrameToBop, Algo::LowerBound(BoppedFrames, InFrameToBop));
+
+		if (InFrameToBop < NumFramesPerBlock)
+		{
+			// Update last bop pos. 
+			UpdateLastBopIndexInBlock();
+		}
+
+		bHasBop = true;
+	}
+
+
+	void FBop::AdvanceBlock()
+	{
+		Advance(NumFramesPerBlock);
+	}
+
+	void FBop::Advance(int32 InNumFrames)
+	{
+		if (bHasBop)
+		{
+			const int32 NumToRemove = Algo::LowerBound(BoppedFrames, InNumFrames);
+
+			if (NumToRemove)
+			{
+				BoppedFrames.RemoveAt(0, NumToRemove, false /* bAllowShrinking */);
+			}
+
+			const int32 Num = BoppedFrames.Num();
+
+			if (Num > 0)
+			{
+				for (int32 i = 0; i < Num; i++)
+				{
+					BoppedFrames[i] -= InNumFrames;
+				}
+			}
+			else
+			{
+				bHasBop = false;
+			}
+
+			UpdateLastBopIndexInBlock();
+		}
+	}
+
+	int32 FBop::Num() const
+	{
+		return BoppedFrames.Num();
+	}
+
+	int32 FBop::operator[](int32 InBopIndex) const
+	{
+		return BoppedFrames[InBopIndex];
+	}
+
+	bool FBop::IsBopped() const
+	{
+		return bHasBop;
+	}
+
+	bool FBop::IsBoppedInBlock() const
+	{
+		return LastBopIndexInBlock > 0;
+	}
+
+	FBop::operator bool() const 
+	{
+		return IsBoppedInBlock();
+	}
+
+	void FBop::Reset()
+	{
+		BoppedFrames.Reset();
+		LastBopIndexInBlock = INDEX_NONE;
+		bHasBop = false;
+	}
+
+	void FBop::UpdateLastBopIndexInBlock()
+	{
+		LastBopIndexInBlock = Algo::LowerBound(BoppedFrames, NumFramesPerBlock);
+	}
+}
diff --git a/Engine/Plugins/Runtime/Metasound/Source/MetasoundStandardNodes/Private/MetasoundPeriodicBopNode.cpp b/Engine/Plugins/Runtime/Metasound/Source/MetasoundStandardNodes/Private/MetasoundPeriodicBopNode.cpp
--- a/Engine/Plugins/Runtime/Metasound/Source/MetasoundStandardNodes/Private/MetasoundPeriodicBopNode.cpp
+++ b/Engine/Plugins/Runtime/Metasound/Source/MetasoundStandardNodes/Private/MetasoundPeriodicBopNode.cpp
@@ -53,7 +53,7 @@ namespace Metasound
 
 	FPeriodicBopOperator::FPeriodicBopOperator(const FOperatorSettings& InSettings, const FFloatTimeReadRef& InPeriod)
 	:	OperatorSettings(InSettings)
-	,	Bop(FBopWriteRef::CreateNew())
+	,	Bop(FBopWriteRef::CreateNew(InSettings))
 	,	Period(InPeriod)
 	,	ExecuteDurationInSamples(InSettings.GetNumFramesPerBlock())
 	,	SampleCountdown(0.f)
@@ -74,7 +74,7 @@ namespace Metasound
 	void FPeriodicBopOperator::Execute()
 	{
 		// Advance internal counter to get rid of old bops.
-		Bop->Advance(OperatorSettings.GetNumFramesPerBlock());
+		Bop->AdvanceBlock();
 
 		float PeriodInSamples = FMath::Max(Period->GetSeconds(), MinimumPeriodSeconds) * OperatorSettings.GetSampleRate();
 
diff --git a/Engine/Plugins/Runtime/Metasound/Source/MetasoundStandardNodes/Private/MetasoundStandardNodesModule.cpp b/Engine/Plugins/Runtime/Metasound/Source/MetasoundStandardNodes/Private/MetasoundStandardNodesModule.cpp
--- a/Engine/Plugins/Runtime/Metasound/Source/MetasoundStandardNodes/Private/MetasoundStandardNodesModule.cpp
+++ b/Engine/Plugins/Runtime/Metasound/Source/MetasoundStandardNodes/Private/MetasoundStandardNodesModule.cpp
@@ -5,8 +5,8 @@
 #include "CoreMinimal.h"
 #include "Modules/ModuleInterface.h"
 #include "Modules/ModuleManager.h"
+#include "MetasoundBop.h"
 #include "MetasoundPrimitives.h"
-#include "MetasoundDataReferenceTypes.h"
 #include "MetasoundTime.h"
 #include "MetasoundFrequency.h"
 
diff --git a/Engine/Plugins/Runtime/Metasound/Source/MetasoundStandardNodes/Public/MetasoundBop.h b/Engine/Plugins/Runtime/Metasound/Source/MetasoundStandardNodes/Public/MetasoundBop.h
new file mode 100644
--- /dev/null
+++ b/Engine/Plugins/Runtime/Metasound/Source/MetasoundStandardNodes/Public/MetasoundBop.h
@@ -0,0 +1,170 @@
+// Copyright Epic Games, Inc. All Rights Reserved.
+
+#pragma once
+
+#include "MetasoundDataReference.h"
+#include "MetasoundDataTypeRegistrationMacro.h"
+#include "MetasoundOperatorSettings.h"
+
+namespace Metasound
+{
+	/** FBop supports sample accurate triggering, sample accurate internal tracking,
+	 * and a conveient interface for running bop-aligned audio signal processing
+	 * routines on buffers.. 
+	 *
+	 * FBops are triggered using FBop::BopTime or FBop::BopFrame.
+	 * FBops track time internally by calling FBop::Advance.
+	 * Executing audio signal processing on buffers can be performed by calling
+	 * FBop::ExecuteBlock or FBop::LookAhead.
+	 */
+	class METASOUNDSTANDARDNODES_API FBop
+	{
+		public:
+			/** FBop constructor. 
+			 *
+			 * @param bShouldBop - If true, bops first sample.
+			 * @param InSettings - Operator settings.
+			 */
+			explicit FBop(bool bShouldBop, const FOperatorSettings& InSettings);
+
+			/** FBop constructor. 
+			 *
+			 * @param InFrameToBop - Set specific frame to bop.
+			 * @param InSettings - Operator settings.
+			 */
+			explicit FBop(int32 InFrameToBop, const FOperatorSettings& InSettings);
+
+			/** FBop constructor. By default it is not bopped.
+			 *
+			 * @param InSettings - Operator settings.
+			 */
+			FBop(const FOperatorSettings& InSettings);
+
+			/** Bop a specific frame in the future.
+			 *
+			 * @param InFrameToBop - Index of frame to bop.
+			 */
+			void BopFrame(int32 InFrameToBop);
+
+			/** Bop a specific time in the future.
+			 *
+			 * @param InTime - Time to set bop.
+			 */
+			template<typename TimeType>
+			void BopTime(const TimeType& InTime)
+			{
+				BopFrame(InTime.GetNumSamples(SampleRate));
+			}
+
+			/** Advance internal frame counters by block size. */
+			void AdvanceBlock();
+
+			/** Advance internal frame counters by specific frame count. */
+			void Advance(int32 InNumFrames);
+
+			/** Number of bopped frames. */
+			int32 Num() const;
+
+			/** Returns frame index for a given bop index. 
+			 *
+			 * @param InIndex - Index of bop. Must be a value between 0 and Num().
+			 *
+			 * @return The frame of the bop.
+			 */
+			int32 operator[](int32 InBopIndex) const;
+
+			/** Returns true if there are any bopped frames. */
+			bool IsBopped() const;
+
+			/** Returns true there is a bop in the current block of audio. */
+			bool IsBoppedInBlock() const;
+
+			/** Implicit conversion of FBop into bool by calling IsBoppedInBlock()
+			 *
+			 * @return Returns true if frame is bopped in current block.
+			 */
+			operator bool() const;
+
+			/** Removes all bopped frames. */
+			void Reset();
+
+			/** Executes one block of frames and calls underlying InPreBop and InOnBop
+			 * functions with frame indices.
+			 *
+			 * @param InPreBop - A function which handles frames before the first 
+			 *                   bop in the current block. The function must accept 
+			 *                   the arguments (int32 StartFrame, int32 EndFrame).
+			 * @param InOnBop - A function which handles frames starting with the 
+			 *                  bops index and ending the next bop index or the
+			 *                  number of frames in a block.. The function must
+			 *                  accept the arguments (int32 StartFrame, int32 EndFrame).
+			 */
+			template<typename PreBopType, typename OnBopType>
+			void ExecuteBlock(PreBopType InPreBop, OnBopType InOnBop) const
+			{
+				ExecuteFrames(NumFramesPerBlock, LastBopIndexInBlock, InPreBop, InOnBop);
+			}
+
+			/** Executes a desired number of frames and calls underlying InPreBop 
+			 * and InOnBop functions with frame indices.
+			 *
+			 * @param InNumFrames - The number of frames to look ahead.
+			 * @param InPreBop - A function which handles frames before the first 
+			 *                   bop in the range of frames. The function must accept 
+			 *                   the arguments (int32 StartFrame, int32 EndFrame).
+			 * @param InOnBop - A function which handles frames starting with the 
+			 *                  bops index and ending the next bop index or InNumFrames
+			 *             		The function must accept the arguments 
+			 *             		(int32 StartFrame, int32 EndFrame).
+			 */
+			template<typename PreBopType, typename OnBopType>
+			void LookAhead(int32 InNumFrames, PreBopType InPreBop, OnBopType InOnBop) const
+			{
+				if (InNumFrames > 0)
+				{
+					int32 LastBopIndexInLookAhead= Algo::LowerBound(BoppedFrames, InNumFrames);
+					
+					ExecuteFrames(InNumFrames, LastBopIndexInLookAhead, InPreBop, InOnBop);
+				}
+			}
+
+		private:
+			void UpdateLastBopIndexInBlock();
+
+			template<typename PreBopType, typename OnBopType>
+			void ExecuteFrames(int32 InNumFrames, int32 InLastBopIndex, PreBopType InPreBop, OnBopType InOnBop) const
+			{
+				if (InLastBopIndex <= 0)
+				{
+					InPreBop(0, InNumFrames);
+					return;
+				}
+
+				const int32* BoppedFramesData = BoppedFrames.GetData();
+
+				if (BoppedFramesData[0] > 0)
+				{
+					InPreBop(0, BoppedFramesData[0]);
+				}
+
+				const int32 NextToLastBopIndex = InLastBopIndex - 1;
+
+				for (int32 i = 0; i < NextToLastBopIndex; i++)
+				{
+					InOnBop(BoppedFramesData[i], BoppedFramesData[i + 1]);
+				}
+
+				InOnBop(BoppedFramesData[NextToLastBopIndex], InNumFrames);
+			}
+
+
+			TArray<int32> BoppedFrames;
+
+			bool bHasBop = false;
+			int32 NumFramesPerBlock = 0;
+			float SampleRate;
+			int32 LastBopIndexInBlock = INDEX_NONE;
+	};
+
+	DECLARE_METASOUND_DATA_REFERENCE_TYPES(FBop, METASOUNDSTANDARDNODES_API, FBopTypeInfo, FBopReadRef, FBopWriteRef);
+}
diff --git a/Engine/Plugins/Runtime/Metasound/Source/MetasoundStandardNodes/Public/MetasoundDataReferenceTypes.h b/Engine/Plugins/Runtime/Metasound/Source/MetasoundStandardNodes/Public/MetasoundDataReferenceTypes.h
deleted file mode 100644
diff --git a/Engine/Plugins/Runtime/Metasound/Source/MetasoundStandardNodes/Public/MetasoundPeriodicBopNode.h b/Engine/Plugins/Runtime/Metasound/Source/MetasoundStandardNodes/Public/MetasoundPeriodicBopNode.h
--- a/Engine/Plugins/Runtime/Metasound/Source/MetasoundStandardNodes/Public/MetasoundPeriodicBopNode.h
+++ b/Engine/Plugins/Runtime/Metasound/Source/MetasoundStandardNodes/Public/MetasoundPeriodicBopNode.h
@@ -4,13 +4,13 @@
 
 #include "CoreMinimal.h"
 #include "MetasoundNode.h"
+#include "MetasoundBop.h"
 #include "MetasoundBuilderInterface.h"
 #include "MetasoundExecutableOperator.h"
 #include "MetasoundFacade.h"
 #include "MetasoundOperatorInterface.h"
 #include "MetasoundNodeInterface.h"
 #include "MetasoundDataReferenceCollection.h"
-#include "MetasoundDataReferenceTypes.h"
 #include "MetasoundTime.h"
 
 namespace Metasound
