Prevent against NULL inputs in WavePlayerNode, which stop EngineTest crashing in the case of bad arguments.

diff --git a/Engine/Plugins/Runtime/Metasound/Source/MetasoundEngine/Private/MetasoundWave.cpp b/Engine/Plugins/Runtime/Metasound/Source/MetasoundEngine/Private/MetasoundWave.cpp
--- a/Engine/Plugins/Runtime/Metasound/Source/MetasoundEngine/Private/MetasoundWave.cpp
+++ b/Engine/Plugins/Runtime/Metasound/Source/MetasoundEngine/Private/MetasoundWave.cpp
@@ -83,16 +83,18 @@ namespace Metasound
 
 	FWaveAsset::FDecoderInputPtr FWaveAsset::CreateDecoderInput(const FWaveAssetReadRef& InWaveRef)
 	{				
-		FName OldFormat = FAudioDeviceManager::Get()->GetActiveAudioDevice()->GetRuntimeFormat(
-			const_cast<USoundWave*>(InWaveRef->GetSoundWave())
-		);
-
-		TUniquePtr<Audio::IDecoderInput> Input = Audio::CreateBackCompatDecoderInput(
-			OldFormat,
-			InWaveRef->GetSoundWave()
-		);		
-		return FWaveAsset::FDecoderInputPtr(Input.Release());
+		if (InWaveRef->GetSoundWave())
+		{
+			FName OldFormat = FAudioDeviceManager::Get()->GetActiveAudioDevice()->GetRuntimeFormat(
+				const_cast<USoundWave*>(InWaveRef->GetSoundWave())
+			);
 
+			TUniquePtr<Audio::IDecoderInput> Input = Audio::CreateBackCompatDecoderInput(
+				OldFormat,
+				InWaveRef->GetSoundWave()
+			);
+			return FWaveAsset::FDecoderInputPtr(Input.Release());
+		}
 		return nullptr;
 	}
 }
diff --git a/Engine/Plugins/Runtime/Metasound/Source/MetasoundEngine/Private/MetasoundWavePlayerNode.cpp b/Engine/Plugins/Runtime/Metasound/Source/MetasoundEngine/Private/MetasoundWavePlayerNode.cpp
--- a/Engine/Plugins/Runtime/Metasound/Source/MetasoundEngine/Private/MetasoundWavePlayerNode.cpp
+++ b/Engine/Plugins/Runtime/Metasound/Source/MetasoundEngine/Private/MetasoundWavePlayerNode.cpp
@@ -137,40 +137,43 @@ namespace Metasound
 
 		FWaveAssetReadRef Wave = InputCol.GetDataReadReferenceOrConstruct<FWaveAsset>(TEXT("Wave"));
 
-		FWaveAsset::FDecoderInputPtr Input = FWaveAsset::CreateDecoderInput(Wave);
-		if (Input)
+		if(Wave->GetSoundWave())
 		{
-			ICodecRegistry::FCodecPtr Codec = ICodecRegistry::Get().FindCodecByParsingInput(Input.Get());
-			if (Codec)
+			FWaveAsset::FDecoderInputPtr Input = FWaveAsset::CreateDecoderInput(Wave);
+			if (Input)
 			{
-				IDecoderOutput::FRequirements Reqs 
-				{ 
-					Float32_Interleaved, 
-					InParams.OperatorSettings.GetNumFramesPerBlock(), 
-					static_cast<int32>(InParams.OperatorSettings.GetSampleRate()) 
-				};
-				TUniquePtr<IDecoderOutput> Output = IDecoderOutput::Create(Reqs);
-				TUniquePtr<IDecoder> Decoder = Codec->CreateDecoder(Input.Get(), Output.Get());
-
-				return MakeUnique<FWavePlayerOperator>(
-					InParams.OperatorSettings, 
-					Wave, 
-					MoveTemp(Input), 
-					MoveTemp(Output), 
-					MoveTemp(Decoder)
-				);
+				ICodecRegistry::FCodecPtr Codec = ICodecRegistry::Get().FindCodecByParsingInput(Input.Get());
+				if (Codec)
+				{
+					IDecoderOutput::FRequirements Reqs 
+					{ 
+						Float32_Interleaved, 
+						InParams.OperatorSettings.GetNumFramesPerBlock(), 
+						static_cast<int32>(InParams.OperatorSettings.GetSampleRate()) 
+					};
+					TUniquePtr<IDecoderOutput> Output = IDecoderOutput::Create(Reqs);
+					TUniquePtr<IDecoder> Decoder = Codec->CreateDecoder(Input.Get(), Output.Get());
+
+					return MakeUnique<FWavePlayerOperator>(
+						InParams.OperatorSettings, 
+						Wave, 
+						MoveTemp(Input), 
+						MoveTemp(Output), 
+						MoveTemp(Decoder)
+					);
+				}
+				else
+				{
+					AddBuildError<FWavePlayerError>(OutErrors, TEXT("FailedToFindCodec"),
+						LOCTEXT("FailedToFindCodec", "Failed to find codec for opening the supplied Wave"));
+				}
 			}
 			else
 			{
-				AddBuildError<FWavePlayerError>(OutErrors, TEXT("FailedToFindCodec"),
-					LOCTEXT("FailedToFindCodec", "Failed to find codec for opening the supplied Wave"));
+				AddBuildError<FWavePlayerError>(OutErrors, TEXT("FailedToParseInput"),
+					LOCTEXT("FailedToParseInput", "Failed to parse the compressed data"));
 			}
 		}
-		else
-		{
-			AddBuildError<FWavePlayerError>(OutErrors, TEXT("FailedToParseInput"),
-				LOCTEXT("FailedToParseInput", "Failed to parse the compressed data"));
-		}
 
 		// Create the player without any inputs, will just produce silence.
 		return MakeUnique<FWavePlayerOperator>(InParams.OperatorSettings, Wave);
