Fixing some build issues from auto converter node by adding TDataReferenceFactory

diff --git a/Engine/Plugins/Runtime/Metasound/Source/MetasoundFrontend/Public/MetasoundAutoConverterNode.h b/Engine/Plugins/Runtime/Metasound/Source/MetasoundFrontend/Public/MetasoundAutoConverterNode.h
--- a/Engine/Plugins/Runtime/Metasound/Source/MetasoundFrontend/Public/MetasoundAutoConverterNode.h
+++ b/Engine/Plugins/Runtime/Metasound/Source/MetasoundFrontend/Public/MetasoundAutoConverterNode.h
@@ -3,6 +3,7 @@
 
 #include "CoreMinimal.h"
 #include "MetasoundBuilderInterface.h"
+#include "MetasoundBuildError.h"
 #include "MetasoundNodeInterface.h"
 #include "MetasoundOperatorInterface.h"
 #include "MetasoundDataReference.h"
@@ -37,15 +38,18 @@ namespace Metasound
 		}
 
 	private:
+		/** FConverterOperator converts from "FromDataType" to "ToDataType" using
+		 * a implicit conversion operators. 
+		 */
 		class FConverterOperator : public TExecutableOperator<FConverterOperator>
 		{
 			public:
 
-				FConverterOperator(TDataReadReference<FromDataType> InFromDataReference)
+				FConverterOperator(TDataReadReference<FromDataType> InFromDataReference, TDataWriteReference<ToDataType> InToDataReference)
 					: FromData(InFromDataReference)
-					, ToData(TDataWriteReference<ToDataType>::CreateNew()) // TODO: this currently only works with default constructible datatypes. Can we do something similar to what we did with literals?
+					, ToData(InToDataReference) 
 				{
-					Outputs.AddDataReadReference<ToDataType>(GetOutputName(), TDataReadReference<ToDataType>(ToData));
+					Outputs.AddDataReadReference<ToDataType>(GetOutputName(), ToData);
 				}
 
 				virtual ~FConverterOperator() {}
@@ -73,6 +77,9 @@ namespace Metasound
 				FDataReferenceCollection Outputs;
 		};
 
+		/** FConverterOperatorFactory creates an operator which converts from 
+		 * "FromDataType" to "ToDataType". 
+		 */
 		class FCoverterOperatorFactory : public IOperatorFactory
 		{
 			public:
@@ -80,13 +87,25 @@ namespace Metasound
 
 				virtual TUniquePtr<IOperator> CreateOperator(const FCreateOperatorParams& InParams, FBuildErrorArray& OutErrors) override
 				{
+					TDataWriteReference<ToDataType> WriteReference = TDataReferenceFactory<ToDataType>::CreateNewWriteReference(InParams.OperatorSettings);
+
 					if (!InParams.InputDataReferences.ContainsDataReadReference<FromDataType>(GetInputName()))
 					{
-						// TODO: Add build error.
-						return TUniquePtr<IOperator>(nullptr);
+						if (ensure(InParams.Node.GetVertexInterface().ContainsInputVertex(GetInputName())))
+						{
+							// There should be something hooked up to the converter node. Report it as an error. 
+							FInputDataDestination Dest(InParams.Node, InParams.Node.GetVertexInterface().GetInputVertex(GetInputName()));
+							AddBuildError<FMissingInputDataReferenceError>(OutErrors, Dest);
+						}
+
+						// We can still build something even though there is an error. 
+						TDataReadReference<FromDataType> ReadReference = TDataReferenceFactory<FromDataType>::CreateNewReadReference(InParams.OperatorSettings);
+						return MakeUnique<FConverterOperator>(ReadReference, WriteReference);
 					}
 
-					return MakeUnique<FConverterOperator>(InParams.InputDataReferences.GetDataReadReference<FromDataType>(GetInputName()));
+					TDataReadReference<FromDataType> ReadReference = InParams.InputDataReferences.GetDataReadReference<FromDataType>(GetInputName());
+
+					return MakeUnique<FConverterOperator>(ReadReference, WriteReference);
 				}
 		};
 
diff --git a/Engine/Plugins/Runtime/Metasound/Source/MetasoundFrontend/Public/MetasoundDataTypeRegistrationMacro.h b/Engine/Plugins/Runtime/Metasound/Source/MetasoundFrontend/Public/MetasoundDataTypeRegistrationMacro.h
--- a/Engine/Plugins/Runtime/Metasound/Source/MetasoundFrontend/Public/MetasoundDataTypeRegistrationMacro.h
+++ b/Engine/Plugins/Runtime/Metasound/Source/MetasoundFrontend/Public/MetasoundDataTypeRegistrationMacro.h
@@ -8,6 +8,7 @@
 #include "MetasoundDataReference.h"
 #include "MetasoundOutputNode.h"
 #include "MetasoundInputNode.h"
+#include "MetasoundPrimitives.h"
 #include "MetasoundOperatorInterface.h"
 #include "MetasoundFrontendRegistries.h"
 #include "MetasoundAutoConverterNode.h"
@@ -69,8 +70,10 @@ void AttemptToRegisterSendAndReceiveNodes()
 	// or it's not an audio buffer type, which we use Audio::FPatchMixerSplitter instances for.
 }
 
+
+
 // This utility function can be used to check to see if we can static cast between two types, and autogenerate a node for that static cast.
-template<typename TFromDataType, typename TToDataType, typename TEnableIf<std::is_convertible<TFromDataType, TToDataType>::value, bool>::Type = true>
+template<typename TFromDataType, typename TToDataType, typename std::enable_if<std::is_convertible<TFromDataType, TToDataType>::value, int>::type = 0>
 void AttemptToRegisterConverter()
 {
 	using FConverterNode = Metasound::TAutoConverterNode<TFromDataType, TToDataType>;
@@ -84,7 +87,7 @@ void AttemptToRegisterConverter()
 	}
 }
 
-template<typename TFromDataType, typename TToDataType, typename TEnableIf<!std::is_convertible<TFromDataType, TToDataType>::value, bool>::Type = true>
+template<typename TFromDataType, typename TToDataType, typename std::enable_if<!std::is_convertible<TFromDataType, TToDataType>::value, int>::type = 0>
 void AttemptToRegisterConverter()
 {
 	// This implementation intentionally noops, because static_cast<TFromDataType>(TToDataType&) is invalid.
diff --git a/Engine/Plugins/Runtime/Metasound/Source/MetasoundFrontend/Public/MetasoundReceiveNode.h b/Engine/Plugins/Runtime/Metasound/Source/MetasoundFrontend/Public/MetasoundReceiveNode.h
--- a/Engine/Plugins/Runtime/Metasound/Source/MetasoundFrontend/Public/MetasoundReceiveNode.h
+++ b/Engine/Plugins/Runtime/Metasound/Source/MetasoundFrontend/Public/MetasoundReceiveNode.h
@@ -109,43 +109,15 @@ namespace Metasound
 				FDataReferenceCollection Outputs;
 		};
 
-		template<typename T>
-		struct TReceiverOperatorFactoryHelper
-		{
-			using FDataTypeInfo = TDataReferenceTypeInfo<T>;
-
-			static constexpr bool bUseDefaultCtor = FDataTypeInfo::bCanUseDefaultConstructor && !FDataTypeInfo::bIsConstructableWithSettings;
-			static constexpr bool bUseSettingsCtor = FDataTypeInfo::bIsConstructableWithSettings;
-		};
-
 		class FReceiverOperatorFactory : public IOperatorFactory
 		{
 			public:
 				FReceiverOperatorFactory() = default;
 
-				template<
-					typename U = TDataType,
-					typename std::enable_if< TReceiverOperatorFactoryHelper<U>::bUseDefaultCtor, int>::type = 0
-				>
-				TDataWriteReference<TDataType> CreateDefaultWriteReference(const FCreateOperatorParams& InParams) const
-				{
-					return TDataWriteReference<TDataType>::CreateNew();
-				}
-
-				template<
-					typename U = TDataType,
-					typename std::enable_if< TReceiverOperatorFactoryHelper<U>::bUseSettingsCtor, int>::type = 0
-				>
-				TDataWriteReference<TDataType> CreateDefaultWriteReference(const FCreateOperatorParams& InParams) const
-				{
-					return TDataWriteReference<TDataType>::CreateNew(InParams.OperatorSettings);
-				}
-
-
 				virtual TUniquePtr<IOperator> CreateOperator(const FCreateOperatorParams& InParams, FBuildErrorArray& OutErrors) override
 				{
 					return MakeUnique<TReceiverOperator>(
-						CreateDefaultWriteReference(InParams),
+						TDataReferenceFactory<TDataType>::CreateNewWriteReference(InParams.OperatorSettings),
 						InParams.InputDataReferences.GetDataReadReference<FSendAddress>(GetAddressInputName()),
 						InParams.OperatorSettings
 						);
diff --git a/Engine/Plugins/Runtime/Metasound/Source/MetasoundGraphCore/Public/MetasoundDataReference.h b/Engine/Plugins/Runtime/Metasound/Source/MetasoundGraphCore/Public/MetasoundDataReference.h
--- a/Engine/Plugins/Runtime/Metasound/Source/MetasoundGraphCore/Public/MetasoundDataReference.h
+++ b/Engine/Plugins/Runtime/Metasound/Source/MetasoundGraphCore/Public/MetasoundDataReference.h
@@ -476,5 +476,136 @@ namespace Metasound
 		 * const TArrayView<const FAudioBufferWriteRef> GetBuffers() { return WritableBuffers; }
 		 */
 	};
+
+
+	// TODO: add creation of data type from literal to TDataReferenceFactory by introducing TVariant interface. 
+	// TODO: swap FOperatorSettings and ArgType order in code which uses literals. 
+	
+	/** TDataReferenceFactory
+	 *
+	 * This factory provides a convenient interface for creating new data references. 
+	 * The factory prioritizes the available data reference constructors and chooses
+	 * the highest prioritized constructor which satisfies the provided constructor
+	 * arguments. This allows callers to construct with a consistent interface for 
+	 * varying data reference types. 
+	 *
+	 * The constructors are prioritized as follows
+	 * 1. A constructor which accepts an FOperatorSettings and provided arguments.
+	 * 2. A constructor which accepts the provided arguments.
+	 * 3. A constructor which accepts an FOperatorSettings.
+	 * 4. The default constructor.
+	 *
+	 * If none of these constructors exist for the provided arguments, then the
+	 * factory creation method will not compile. 
+	 */
+	template<typename DataType>
+	class TDataReferenceFactory
+	{
+		private:
+
+			/** TDataReferenceFactoryHelper determines which constructors are available 
+			 * for a given data type and arguments.
+			 */
+			template<typename T, typename... ArgTypes>
+			struct TDataReferenceFactoryHelper
+			{
+				static constexpr bool bIsDefaultConstructible = std::is_constructible<T>::value;
+				static constexpr bool bIsConstructibleWithSettings = std::is_constructible<T, const FOperatorSettings&>::value;
+				static constexpr bool bIsConstructibleWithArgs = std::is_constructible<T, ArgTypes...>::value;
+				static constexpr bool bIsConstructibleWithSettingsAndArgs = std::is_constructible<T, const FOperatorSettings&, ArgTypes...>::value;
+			};
+
+			/** TDataReferenceFactoryHelper specialization which handles an empty parameter pack. */
+			template<typename T>
+			struct TDataReferenceFactoryHelper<T>
+			{
+				static constexpr bool bIsDefaultConstructible = std::is_constructible<T>::value;
+				static constexpr bool bIsConstructibleWithSettings = std::is_constructible<T, const FOperatorSettings&>::value;
+				static constexpr bool bIsConstructibleWithArgs = false;
+				static constexpr bool bIsConstructibleWithSettingsAndArgs = false;
+			};
+
+			/** TDataReferenceFactoryChooser prioritizes the available constructors. */
+			template<typename T, typename... ArgTypes>
+			struct TDataReferenceFactoryChooser
+			{
+				using FHelper = TDataReferenceFactoryHelper<T, ArgTypes...>;
+
+				static constexpr bool bIsConstructibleWithoutArg = FHelper::bIsDefaultConstructible || FHelper::bIsConstructibleWithSettings;
+
+				// Only one of these booleans should be true. These enable the appropriate
+				// "CreateNewReference(...)" function. If multiple are true, than a
+				// duplicate symbol will exist. 
+				static constexpr bool bCreateWithSettingsAndArgs = FHelper::bIsConstructibleWithSettingsAndArgs;
+				static constexpr bool bCreateWithArgs = FHelper::bIsConstructibleWithArgs && !bCreateWithSettingsAndArgs;
+				static constexpr bool bCreateWithSettings = FHelper::bIsConstructibleWithSettings && !bCreateWithArgs;
+				static constexpr bool bCreateWithDefaultConstructor = FHelper::bIsDefaultConstructible && !bCreateWithSettings;
+
+				static constexpr bool bCreateWithAnyConstructor = bCreateWithSettingsAndArgs || bCreateWithArgs || bCreateWithSettings || bCreateWithDefaultConstructor;
+
+				// TODO: 
+				static_assert(bCreateWithAnyConstructor, "No constructor possible");
+			};
+
+			/** Create a reference using FOperatorSettings and arguments. */
+			template<
+				typename ReferenceType, 
+				typename... ArgTypes,
+				typename std::enable_if< TDataReferenceFactoryChooser<DataType, ArgTypes...>::bCreateWithSettingsAndArgs, int>::type = 0
+			>
+			static ReferenceType CreateNewReference(const FOperatorSettings& InSettings, ArgTypes&&... Args)
+			{
+				return ReferenceType::CreateNew(InSettings, Forward<ArgTypes>(Args)...);
+			}
+
+			/** Create a reference using arguments. */
+			template<
+				typename ReferenceType, 
+				typename... ArgTypes,
+				typename std::enable_if< TDataReferenceFactoryChooser<DataType, ArgTypes...>::bCreateWithArgs, int>::type = 0
+			>
+			static ReferenceType CreateNewReference(const FOperatorSettings& InSettings, ArgTypes&&... Args)
+			{
+				return ReferenceType::CreateNew(Forward<ArgTypes>(Args)...);
+			}
+
+			/** Create a reference using FOperatorSettings. */
+			template<
+				typename ReferenceType, 
+				typename... ArgTypes,
+				typename std::enable_if< TDataReferenceFactoryChooser<DataType, ArgTypes...>::bCreateWithSettings, int>::type = 0
+			>
+			static ReferenceType CreateNewReference(const FOperatorSettings& InSettings, ArgTypes&&... Args)
+			{
+				return ReferenceType::CreateNew(InSettings);
+			}
+
+			/** Create a reference using the default constructor */
+			template<
+				typename ReferenceType, 
+				typename... ArgTypes,
+				typename std::enable_if< TDataReferenceFactoryChooser<DataType, ArgTypes...>::bCreateWithDefaultConstructor, int>::type = 0
+			>
+			static ReferenceType CreateNewReference(const FOperatorSettings& InSettings, ArgTypes&&... Args)
+			{
+				return ReferenceType::CreateNew();
+			}
+
+		public:
+
+			/** Create a new TDataWriteReference<DataType> using the given FOperatorSettings and ArgTypes. */
+			template<typename... ArgTypes>
+			static TDataWriteReference<DataType> CreateNewWriteReference(const FOperatorSettings& InSettings, ArgTypes&&... Args)
+			{
+				return CreateNewReference<TDataWriteReference<DataType>>(InSettings, Forward<ArgTypes>(Args)...);
+			}
+
+			/** Create a new TDataReadReference<DataType> using the given FOperatorSettings and ArgTypes. */
+			template<typename... ArgTypes>
+			static TDataReadReference<DataType> CreateNewReadReference(const FOperatorSettings& InSettings, ArgTypes&&... Args)
+			{
+				return CreateNewReference<TDataReadReference<DataType>>(InSettings, Forward<ArgTypes>(Args)...);
+			}
+	};
 }
 
diff --git a/Engine/Plugins/Runtime/Metasound/Source/MetasoundStandardNodes/Private/MetasoundAudioFormats.cpp b/Engine/Plugins/Runtime/Metasound/Source/MetasoundStandardNodes/Private/MetasoundAudioFormats.cpp
--- a/Engine/Plugins/Runtime/Metasound/Source/MetasoundStandardNodes/Private/MetasoundAudioFormats.cpp
+++ b/Engine/Plugins/Runtime/Metasound/Source/MetasoundStandardNodes/Private/MetasoundAudioFormats.cpp
@@ -6,32 +6,17 @@
 #include "MetasoundDataReference.h"
 #include "MetasoundAudioBuffer.h"
 
-struct FMyType
-{
-public:
-	FMyType() = delete;
-
-private:
-	FMyType Get();
-	friend class FMyOtherType;
-};
-
 namespace Metasound
 {
-	// REGISTER_METASOUND_DATATYPE(FAudioBuffer, "Audio:Buffer")
 	REGISTER_METASOUND_DATATYPE(FUnformattedAudio, "Audio:Unformatted")
-	REGISTER_METASOUND_DATATYPE(FMultichannelAudioFormat, "Audio:Multichannel")
+	REGISTER_METASOUND_DATATYPE(FMonoAudioFormat, "Audio:Mono")
+	REGISTER_METASOUND_DATATYPE(FStereoAudioFormat, "Audio:Stereo")
 
-
-
-	// TODO: Neither of these are trivially constructible, becuase they are actually views into a preexisting buffer.
-	//REGISTER_METASOUND_DATATYPE(FMonoAudioFormat)
-	//REGISTER_METASOUND_DATATYPE(FStereoAudioFormat)
-	IMPL_METASOUND_DATA_TYPE(FMonoAudioFormat, "Audio:Mono")
-	IMPL_METASOUND_DATA_TYPE(FStereoAudioFormat, "Audio:Stereo")
+	// FMultichannelAudio cannot be used as an input type because it's channel count
+	// cannot be changed at runtime. Hence, it is not registered.
+	IMPL_METASOUND_DATA_TYPE(FMultichannelAudioFormat, "Audio:Multichannel")
 
 	/* FUnformattedAudio */
-
 	FUnformattedAudio::FUnformattedAudio(int32 InNumFrames, int32 InNumChannels, int32 InMaxNumChannels)
 	:	NumFrames(InNumFrames)
 	,	NumChannels(0)
diff --git a/Engine/Plugins/Runtime/Metasound/Source/MetasoundStandardNodes/Public/MetasoundAudioFormats.h b/Engine/Plugins/Runtime/Metasound/Source/MetasoundStandardNodes/Public/MetasoundAudioFormats.h
--- a/Engine/Plugins/Runtime/Metasound/Source/MetasoundStandardNodes/Public/MetasoundAudioFormats.h
+++ b/Engine/Plugins/Runtime/Metasound/Source/MetasoundStandardNodes/Public/MetasoundAudioFormats.h
@@ -238,6 +238,12 @@ namespace Metasound
 				ReadableBuffers = ReadableBufferStorage;
 			}
 
+			TStaticChannelAudioFormat(const FOperatorSettings& InOperatorSettings)
+			:	TStaticChannelAudioFormat(InOperatorSettings.GetNumFramesPerBlock())
+			{
+			}
+
+
 			/** Return the number of audio channels. */
 			int32 GetNumChannels() const
 			{
@@ -468,13 +474,9 @@ namespace Metasound
 
 	DECLARE_METASOUND_DATA_REFERENCE_TYPES(FUnformattedAudio, METASOUNDSTANDARDNODES_API, FUnformattedAudioTypeInfo, FUnformattedAudioReadRef, FUnformattedAudioWriteRef);
 	
-
 	DECLARE_METASOUND_DATA_REFERENCE_TYPES(FMultichannelAudioFormat, METASOUNDSTANDARDNODES_API, FMultichannelAudioFormatTypeInfo, FMultichannelAudioFormatReadRef, FMultichannelAudioFormatWriteRef);
-	
 
 	DECLARE_METASOUND_DATA_REFERENCE_TYPES(FMonoAudioFormat, METASOUNDSTANDARDNODES_API, FMonoAudioFormatTypeInfo, FMonoAudioFormatReadRef, FMonoAudioFormatWriteRef);
 
 	DECLARE_METASOUND_DATA_REFERENCE_TYPES(FStereoAudioFormat, METASOUNDSTANDARDNODES_API, FStereoAudioFormatTypeInfo, FStereoAudioFormatReadRef, FStereoAudioFormatWriteRef);
-	
-	
 }
