Update input node and extend metasound data factory to support literals

diff --git a/Engine/Plugins/Runtime/Metasound/Source/MetasoundEngine/Private/MetasoundEngineModule.cpp b/Engine/Plugins/Runtime/Metasound/Source/MetasoundEngine/Private/MetasoundEngineModule.cpp
--- a/Engine/Plugins/Runtime/Metasound/Source/MetasoundEngine/Private/MetasoundEngineModule.cpp
+++ b/Engine/Plugins/Runtime/Metasound/Source/MetasoundEngine/Private/MetasoundEngineModule.cpp
@@ -6,14 +6,12 @@
 #include "MetasoundSource.h"
 #include "Modules/ModuleManager.h"
 #include "MetasoundWave.h"
+#include "MetasoundDataReference.h"
 #include "MetasoundDataTypeRegistrationMacro.h"
 
 DEFINE_LOG_CATEGORY(LogMetasoundEngine);
 
-namespace Metasound
-{
-	REGISTER_METASOUND_DATATYPE(FWaveAsset, "Primitive:WaveAsset", ::Metasound::ELiteralArgType::UObjectProxy, USoundWave)
-}
+REGISTER_METASOUND_DATATYPE(Metasound::FWaveAsset, "Primitive:WaveAsset", Metasound::ELiteralArgType::UObjectProxy, USoundWave);
 
 class FMetasoundEngineModule : public IMetasoundEngineModule
 {
diff --git a/Engine/Plugins/Runtime/Metasound/Source/MetasoundEngine/Private/MetasoundSource.cpp b/Engine/Plugins/Runtime/Metasound/Source/MetasoundEngine/Private/MetasoundSource.cpp
--- a/Engine/Plugins/Runtime/Metasound/Source/MetasoundEngine/Private/MetasoundSource.cpp
+++ b/Engine/Plugins/Runtime/Metasound/Source/MetasoundEngine/Private/MetasoundSource.cpp
@@ -82,7 +82,7 @@ ISoundGeneratorPtr UMetasoundSource::CreateSoundGenerator(const FSoundGeneratorI
 		{
 			MoveTemp(Operator),
 			Outputs.GetDataReadReferenceOrConstruct<Metasound::FAudioBuffer>(GetAudioOutputName(), InSettings.GetNumFramesPerBlock()),
-			Outputs.GetDataReadReferenceOrConstruct<Metasound::FBop>(GetIsFinishedOutputName(), false, InSettings)
+			Outputs.GetDataReadReferenceOrConstruct<Metasound::FBop>(GetIsFinishedOutputName(), InSettings, false)
 		};
 
 		return ISoundGeneratorPtr(new Metasound::FMetasoundGenerator(MoveTemp(InitParams)));
diff --git a/Engine/Plugins/Runtime/Metasound/Source/MetasoundEngine/Private/MetasoundWavePlayerNode.cpp b/Engine/Plugins/Runtime/Metasound/Source/MetasoundEngine/Private/MetasoundWavePlayerNode.cpp
--- a/Engine/Plugins/Runtime/Metasound/Source/MetasoundEngine/Private/MetasoundWavePlayerNode.cpp
+++ b/Engine/Plugins/Runtime/Metasound/Source/MetasoundEngine/Private/MetasoundWavePlayerNode.cpp
@@ -21,11 +21,18 @@ namespace Metasound
 	class FWavePlayerError : public FBuildErrorBase
 	{
 	public:
-		FWavePlayerError(FName InName, FText InErrorDescription)
-			: FBuildErrorBase(InName, InErrorDescription)
-		{}
+		FWavePlayerError(const FWavePlayerNode& InNode, FText InErrorDescription)
+			: FBuildErrorBase(ErrorType, InErrorDescription)
+		{
+			AddNode(InNode);
+		}
+
 		virtual ~FWavePlayerError() = default;
+
+		static const FName ErrorType;
 	};
+
+	const FName FWavePlayerError::ErrorType = FName(TEXT("WavePlayerError"));
 	
 	class FWavePlayerOperator : public TExecutableOperator<FWavePlayerOperator>
 	{
@@ -164,16 +171,18 @@ namespace Metasound
 				}
 				else
 				{
-					AddBuildError<FWavePlayerError>(OutErrors, TEXT("FailedToFindCodec"),
-						LOCTEXT("FailedToFindCodec", "Failed to find codec for opening the supplied Wave"));
+					AddBuildError<FWavePlayerError>(OutErrors, WaveNode, LOCTEXT("FailedToFindCodec", "Failed to find codec for opening the supplied Wave"));
 				}
 			}
 			else
 			{
-				AddBuildError<FWavePlayerError>(OutErrors, TEXT("FailedToParseInput"),
-					LOCTEXT("FailedToParseInput", "Failed to parse the compressed data"));
+				AddBuildError<FWavePlayerError>(OutErrors, WaveNode, LOCTEXT("FailedToParseInput", "Failed to parse the compressed data"));
 			}
 		}
+		else
+		{
+			AddBuildError<FWavePlayerError>(OutErrors, WaveNode, LOCTEXT("NoSoundWave", "No Sound Wave"));
+		}
 
 		// Create the player without any inputs, will just produce silence.
 		return MakeUnique<FWavePlayerOperator>(InParams.OperatorSettings, Wave);
diff --git a/Engine/Plugins/Runtime/Metasound/Source/MetasoundFrontend/Private/MetasoundFrontend.cpp b/Engine/Plugins/Runtime/Metasound/Source/MetasoundFrontend/Private/MetasoundFrontend.cpp
--- a/Engine/Plugins/Runtime/Metasound/Source/MetasoundFrontend/Private/MetasoundFrontend.cpp
+++ b/Engine/Plugins/Runtime/Metasound/Source/MetasoundFrontend/Private/MetasoundFrontend.cpp
@@ -2741,14 +2741,7 @@ REGISTER_METASOUND_DATATYPE(float, "Primitive:Float", ::Metasound::ELiteralArgTy
 REGISTER_METASOUND_DATATYPE(double, "Primitive:Double", ::Metasound::ELiteralArgType::Float)
 REGISTER_METASOUND_DATATYPE(FString, "Primitive:String", ::Metasound::ELiteralArgType::String)
 
-namespace Metasound
-{
-	REGISTER_METASOUND_DATATYPE(FAudioBuffer, "Audio:Buffer")
-
-
-	static_assert(TDataReferenceTypeInfo<FSendAddress>::bIsValidSpecialization, "ahhhhh");
-
-	REGISTER_METASOUND_DATATYPE(FSendAddress, "Send:Address")
-}
+REGISTER_METASOUND_DATATYPE(Metasound::FAudioBuffer, "Audio:Buffer")
+REGISTER_METASOUND_DATATYPE(Metasound::FSendAddress, "Send:Address")
 
 IMPLEMENT_MODULE(FMetasoundFrontendModule, MetasoundFrontend);
diff --git a/Engine/Plugins/Runtime/Metasound/Source/MetasoundFrontend/Public/MetasoundAudioBuffer.h b/Engine/Plugins/Runtime/Metasound/Source/MetasoundFrontend/Public/MetasoundAudioBuffer.h
--- a/Engine/Plugins/Runtime/Metasound/Source/MetasoundFrontend/Public/MetasoundAudioBuffer.h
+++ b/Engine/Plugins/Runtime/Metasound/Source/MetasoundFrontend/Public/MetasoundAudioBuffer.h
@@ -123,4 +123,20 @@ namespace Metasound
 	};
 
 	DECLARE_METASOUND_DATA_REFERENCE_TYPES(FAudioBuffer, METASOUNDFRONTEND_API, FAudioBufferTypeInfo, FAudioBufferReadRef, FAudioBufferWriteRef);
+
+
+	// This empty base class is used so that we can specialize various nodes (Send, Receive, etc.) for subclasses of IAudioDatatype.
+	class METASOUNDFRONTEND_API IAudioDataType
+	{
+		/**
+		 * Audio datatypes require the following member functions:
+		 * int32 GetNumChannels() const { return NumChannels; }
+		 *
+		 * int32 GetMaxNumChannels() const { return MaxNumChannels; }
+
+		 * const TArrayView<const FAudioBufferReadRef> GetBuffers() const { return ReadableBuffers; }
+
+		 * const TArrayView<const FAudioBufferWriteRef> GetBuffers() { return WritableBuffers; }
+		 */
+	};
 };
diff --git a/Engine/Plugins/Runtime/Metasound/Source/MetasoundFrontend/Public/MetasoundAutoConverterNode.h b/Engine/Plugins/Runtime/Metasound/Source/MetasoundFrontend/Public/MetasoundAutoConverterNode.h
--- a/Engine/Plugins/Runtime/Metasound/Source/MetasoundFrontend/Public/MetasoundAutoConverterNode.h
+++ b/Engine/Plugins/Runtime/Metasound/Source/MetasoundFrontend/Public/MetasoundAutoConverterNode.h
@@ -6,6 +6,7 @@
 #include "MetasoundBuildError.h"
 #include "MetasoundNodeInterface.h"
 #include "MetasoundOperatorInterface.h"
+#include "MetasoundDataFactory.h"
 #include "MetasoundDataReference.h"
 #include "MetasoundExecutableOperator.h"
 
@@ -20,20 +21,18 @@ namespace Metasound
 	template<typename FromDataType, typename ToDataType>
 	class TAutoConverterNode : public INode
 	{
-		static_assert(TDataReferenceTypeInfo<FromDataType>::bIsValidSpecialization, "Please use DECLARE_METASOUND_DATA_REFERENCE_TYPES with this class before trying to create an converter node with it.");
-		static_assert(TDataReferenceTypeInfo<ToDataType>::bIsValidSpecialization, "Please use DECLARE_METASOUND_DATA_REFERENCE_TYPES with this class before trying to create an converter node with it.");
 		static_assert(std::is_convertible<FromDataType, ToDataType>::value, "Tried to create an auto converter node between two types we can't static_cast between.");
 
 	public:
-		static FString& GetInputName()
+		static const FString& GetInputName()
 		{
-			static FString InputName = FString(TDataReferenceTypeInfo<FromDataType>::TypeName);
+			static const FString InputName = GetMetasoundDataTypeString<FromDataType>();
 			return InputName;
 		}
 
-		static FString& GetOutputName()
+		static const FString& GetOutputName()
 		{
-			static FString OutputName = FString(TDataReferenceTypeInfo<ToDataType>::TypeName);
+			static const FString OutputName = GetMetasoundDataTypeString<ToDataType>();
 			return OutputName;
 		}
 
@@ -87,7 +86,7 @@ namespace Metasound
 
 				virtual TUniquePtr<IOperator> CreateOperator(const FCreateOperatorParams& InParams, FBuildErrorArray& OutErrors) override
 				{
-					TDataWriteReference<ToDataType> WriteReference = TDataReferenceFactory<ToDataType>::CreateNewWriteReference(InParams.OperatorSettings);
+					TDataWriteReference<ToDataType> WriteReference = TDataWriteReferenceFactory<ToDataType>::CreateAny(InParams.OperatorSettings);
 
 					if (!InParams.InputDataReferences.ContainsDataReadReference<FromDataType>(GetInputName()))
 					{
@@ -99,7 +98,7 @@ namespace Metasound
 						}
 
 						// We can still build something even though there is an error. 
-						TDataReadReference<FromDataType> ReadReference = TDataReferenceFactory<FromDataType>::CreateNewReadReference(InParams.OperatorSettings);
+						TDataReadReference<FromDataType> ReadReference = TDataReadReferenceFactory<FromDataType>::CreateAny(InParams.OperatorSettings);
 						return MakeUnique<FConverterOperator>(ReadReference, WriteReference);
 					}
 
diff --git a/Engine/Plugins/Runtime/Metasound/Source/MetasoundFrontend/Public/MetasoundDataTypeRegistrationMacro.h b/Engine/Plugins/Runtime/Metasound/Source/MetasoundFrontend/Public/MetasoundDataTypeRegistrationMacro.h
--- a/Engine/Plugins/Runtime/Metasound/Source/MetasoundFrontend/Public/MetasoundDataTypeRegistrationMacro.h
+++ b/Engine/Plugins/Runtime/Metasound/Source/MetasoundFrontend/Public/MetasoundDataTypeRegistrationMacro.h
@@ -5,18 +5,19 @@
 #include "CoreMinimal.h"
 #include "Templates/Casts.h"
 
+#include "MetasoundAutoConverterNode.h"
+#include "MetasoundConverterNodeRegistrationMacro.h"
+#include "MetasoundDataFactory.h"
 #include "MetasoundDataReference.h"
-#include "MetasoundOutputNode.h"
-#include "MetasoundInputNode.h"
-#include "MetasoundPrimitives.h"
-#include "MetasoundOperatorInterface.h"
 #include "MetasoundFrontendRegistries.h"
-#include "MetasoundAutoConverterNode.h"
+#include "MetasoundInputNode.h"
 #include "MetasoundNodeRegistrationMacro.h"
-#include "MetasoundConverterNodeRegistrationMacro.h"
-#include "MetasoundSendNode.h"
+#include "MetasoundOperatorInterface.h"
+#include "MetasoundOutputNode.h"
+#include "MetasoundPrimitives.h"
 #include "MetasoundReceiveNode.h"
 #include "MetasoundRouter.h"
+#include "MetasoundSendNode.h"
 
 #include <type_traits>
 
@@ -47,7 +48,7 @@ struct TIsTransmittable
 {
 private:
 	static constexpr bool bCanBeTransmitted =
-		std::is_copy_constructible<TDataType>::value|| TIsDerivedFrom<TDataType, ::Metasound::IAudioDatatype>::Value;
+		std::is_copy_constructible<TDataType>::value|| TIsDerivedFrom<TDataType, ::Metasound::IAudioDataType>::Value;
 
 public:
 
@@ -127,7 +128,7 @@ bool RegisterDataTypeWithFrontend()
 	// Lambdas that generate our template-instantiated input and output nodes:
 	FInputNodeConstructorCallback InputNodeConstructor = [](::Metasound::FInputNodeConstructorParams&& InParams) -> TUniquePtr<::Metasound::INode>
 	{
-		return TUniquePtr<::Metasound::INode>(new ::Metasound::TInputNode<TDataType>(InParams.InNodeName, InParams.InVertexName, MoveTemp(InParams.InitParam), InParams.InSettings));
+		return TUniquePtr<::Metasound::INode>(new ::Metasound::TInputNode<TDataType>(InParams.InNodeName, InParams.InVertexName, MoveTemp(InParams.InitParam)));
 	};
 
 	FOutputNodeConstructorCallback OutputNodeConstructor = [](const ::Metasound::FOutputNodeConstrutorParams& InParams) -> TUniquePtr<::Metasound::INode>
@@ -167,27 +168,30 @@ bool RegisterDataTypeWithFrontend()
 		};
 	}
 
-	static FName DataTypeName = FName(::Metasound::TDataReferenceTypeInfo<TDataType>::TypeName);
+	static const FName DataTypeName = ::Metasound::GetMetasoundDataTypeName<TDataType>();
 
 	// Pack all of our various constructor lambdas to a single struct.
 	::Metasound::FDataTypeConstructorCallbacks Callbacks = { MoveTemp(InputNodeConstructor), MoveTemp(OutputNodeConstructor), MoveTemp(ProxyGenerator) };
 
+
 	::Metasound::FDataTypeRegistryInfo RegistryInfo;
+
 	RegistryInfo.DataTypeName = DataTypeName;
 	RegistryInfo.PreferredLiteralType = PreferredArgType;
-	RegistryInfo.bIsBoolParsable = ::Metasound::TDataReferenceTypeInfo<TDataType>::bIsBoolParsable;
-	RegistryInfo.bIsIntParsable = ::Metasound::TDataReferenceTypeInfo<TDataType>::bIsIntParsable;
-	RegistryInfo.bIsFloatParsable = ::Metasound::TDataReferenceTypeInfo<TDataType>::bIsFloatParsable;
-	RegistryInfo.bIsStringParsable = ::Metasound::TDataReferenceTypeInfo<TDataType>::bIsStringParsable;
-	RegistryInfo.bIsProxyParsable = ::Metasound::TDataReferenceTypeInfo<TDataType>::bIsProxyParsable;
-	RegistryInfo.bIsProxyArrayParsable = ::Metasound::TDataReferenceTypeInfo<TDataType>::bIsProxyArrayParsable;
-	RegistryInfo.bIsConstructableWithSettings = ::Metasound::TDataReferenceTypeInfo<TDataType>::bIsConstructableWithSettings;
-	RegistryInfo.bIsDefaultConstructible = ::Metasound::TDataReferenceTypeInfo<TDataType>::bCanUseDefaultConstructor;
+
+	RegistryInfo.bIsBoolParsable = ::Metasound::TIsParsable<TDataType, bool>::Value;
+	RegistryInfo.bIsIntParsable = ::Metasound::TIsParsable<TDataType, int32>::Value;
+	RegistryInfo.bIsFloatParsable = ::Metasound::TIsParsable<TDataType, float>::Value;
+	RegistryInfo.bIsStringParsable = ::Metasound::TIsParsable<TDataType, FString>::Value;
+	RegistryInfo.bIsProxyParsable = ::Metasound::TIsParsable<TDataType, const Audio::IProxyDataPtr&>::Value;
+	RegistryInfo.bIsProxyArrayParsable = ::Metasound::TIsParsable<TDataType, const TArray<Audio::IProxyDataPtr>& >::Value;
+	RegistryInfo.bIsConstructableWithSettings = std::is_constructible<TDataType, const ::Metasound::FOperatorSettings&>::value;
+	RegistryInfo.bIsDefaultConstructible = std::is_constructible<TDataType>::value;
 	
 	RegistryInfo.ProxyGeneratorClass = UClassToUse::StaticClass();
 
 	bool bSucceeded = FMetasoundFrontendRegistryContainer::Get()->RegisterDataType(RegistryInfo, MoveTemp(Callbacks));
-	ensureAlwaysMsgf(bSucceeded, TEXT("Failed to register data type %s in the node registry!"), ::Metasound::TDataReferenceTypeInfo<TDataType>::TypeName);
+	ensureAlwaysMsgf(bSucceeded, TEXT("Failed to register data type %s in the node registry!"), *::Metasound::GetMetasoundDataTypeString<TDataType>());
 	
 	RegisterConverterNodes<TDataType>();
 	AttemptToRegisterSendAndReceiveNodes<TDataType>();
@@ -195,6 +199,16 @@ bool RegisterDataTypeWithFrontend()
 	return bSucceeded;
 }
 
+template<typename DataType>
+struct TMetasoundDataTypeRegistration
+{
+	static_assert(std::is_same<DataType, typename std::decay<DataType>::type>::value, "DataType and decayed DataType must be the same");
+	
+	static constexpr bool bCanRegister = ::Metasound::TInputNode<DataType>::bCanRegister;
+
+	static const bool bSuccessfullyRegistered;
+};
+
 // This should be used to expose a datatype as a potential input or output for a metasound graph.
 // The first argument to the macro is the class to expose.
 // the second argument is the display name of that type in the Metasound editor.
@@ -204,10 +218,19 @@ bool RegisterDataTypeWithFrontend()
 // If ::Metasound::
 // Metasound::ELiteralArgType::Invalid can be used to enforce that we don't provide space for a literal, in which case you should have a default constructor or a constructor that takes [const FOperatorSettings&] implemented.
 // If you pass in a preferred arg type, please make sure that the passed in datatype has a matching constructor, since we won't check this until runtime.
+
+#define CANNOT_REGISTER_METASOUND_DATA_TYPE_ASSERT_STRING(DataType) \
+"To register " #DataType " to be used as a Metasounds input or output type, it needs a default constructor or one of the following constructors must be implemented:  " \
+#DataType "(const ::Metasound::FOperatorSettings& InSettings), " \
+#DataType "(const ::Metasound::FOperatorSettings& InSettings, bool InValue), " \
+#DataType "(const ::Metasound::FOperatorSettings& InSettings, int32 InValue), " \
+#DataType "(const ::Metasound::FOperatorSettings& InSettings, float InValue), " \
+#DataType "(const ::Metasound::FOperatorSettings& InSettings, const FString& InString)" \
+#DataType "(const ::Metasound::FOperatorSettings& InSettings, const Audio::IProxyDataPtr& InData),  or " \
+#DataType "(const ::Metasound::FOperatorSettings& InSettings, const TArray<Audio::IProxyDataPtr>& InProxyArray)."
+
 #define REGISTER_METASOUND_DATATYPE(DataType, DataTypeName, ...) \
-	IMPL_METASOUND_DATA_TYPE(DataType, DataTypeName) \
-	static_assert(::Metasound::TDataReferenceTypeInfo<DataType>::bIsValidSpecialization, "Please call DECLARE_METASOUND_DATA_REFERENCE_TYPES(" #DataType "...) before calling REGISTER_METASOUND_DATATYPE(" #DataType")."); \
-	static constexpr bool bCanRegister##DataType = ::Metasound::TDataReferenceTypeInfo<DataType>::bIsStringParsable || ::Metasound::TDataReferenceTypeInfo<DataType>::bIsBoolParsable || ::Metasound::TDataReferenceTypeInfo<DataType>::bIsIntParsable || ::Metasound::TDataReferenceTypeInfo<DataType>::bIsFloatParsable || ::Metasound::TDataReferenceTypeInfo<DataType>::bIsProxyParsable || ::Metasound::TDataReferenceTypeInfo<DataType>::bIsProxyArrayParsable || ::Metasound::TDataReferenceTypeInfo<DataType>::bIsConstructableWithSettings || ::Metasound::TDataReferenceTypeInfo<DataType>::bCanUseDefaultConstructor; \
-	static_assert(bCanRegister##DataType , "To register " #DataType " to be used as a Metasounds input or output type, it needs a default constructor or one of the following constructors must be implemented:  " #DataType "(const ::Metasound::FOperatorSettings& InSettings), " #DataType "(int32 InValue, const ::Metasound::FOperatorSettings& InSettings), " #DataType "(float InValue, const ::Metasound::FOperatorSettings& InSettings), " #DataType "(bool InValue, const ::Metasound::FOperatorSettings& InSettings), " #DataType "(const Audio::IProxyData& InData, const ::Metasound::FOperatorSettings& InSettings), " #DataType "(const TArray<Audio::IProxyDataPtr>& InProxyArray, const ::Metasound::FOperatorSettings& InSettings) or " #DataType "(const FString& InString, const ::Metasound::FOperatorSettings& InSettings)."); \
-	static const bool bSuccessfullyRegistered##DataType = FMetasoundFrontendRegistryContainer::Get()->EnqueueInitCommand([](){ ::RegisterDataTypeWithFrontend<DataType, ##__VA_ARGS__>(); }); // This static bool is useful for debugging, but also is the only way the compiler will let us call this function outside of an expression.
+	DEFINE_METASOUND_DATA_TYPE(DataType, DataTypeName); \
+	static_assert(::TMetasoundDataTypeRegistration<DataType>::bCanRegister, CANNOT_REGISTER_METASOUND_DATA_TYPE_ASSERT_STRING(DataType)); \
+	template<> const bool ::TMetasoundDataTypeRegistration<DataType>::bSuccessfullyRegistered = ::FMetasoundFrontendRegistryContainer::Get()->EnqueueInitCommand([](){ ::RegisterDataTypeWithFrontend<DataType, ##__VA_ARGS__>(); }); // This static bool is useful for debugging, but also is the only way the compiler will let us call this function outside of an expression.
 
diff --git a/Engine/Plugins/Runtime/Metasound/Source/MetasoundFrontend/Public/MetasoundInputNode.h b/Engine/Plugins/Runtime/Metasound/Source/MetasoundFrontend/Public/MetasoundInputNode.h
--- a/Engine/Plugins/Runtime/Metasound/Source/MetasoundFrontend/Public/MetasoundInputNode.h
+++ b/Engine/Plugins/Runtime/Metasound/Source/MetasoundFrontend/Public/MetasoundInputNode.h
@@ -11,113 +11,99 @@
 
 namespace Metasound
 {
+	/** TInputOperator supplies a writeable input and a readable output. */
 	template<typename DataType>
-	class TInputNode : public FNode
+	class TInputOperator : public IOperator
 	{
-		static_assert(TDataReferenceTypeInfo<DataType>::bIsValidSpecialization, "Please use DECLARE_METASOUND_DATA_REFERENCE_TYPES with this class before trying to create an input node with it.");
+		public:
+			using FDataWriteReference = TDataWriteReference<DataType>;
 
-		class FInputOperator : public IOperator
-		{
-			public:
-				using FDataWriteReference = TDataWriteReference<DataType>;
+			TInputOperator(const FString& InDataReferenceName, FDataWriteReference InDataReference)
+			{
+				Inputs.AddDataWriteReference<DataType>(InDataReferenceName, InDataReference);
+				Outputs.AddDataReadReference<DataType>(InDataReferenceName, InDataReference);
+			}
 
-				FInputOperator(const FString& InDataReferenceName, FDataWriteReference InDataReference)
-				{
-					Inputs.AddDataWriteReference<DataType>(InDataReferenceName, InDataReference);
-					Outputs.AddDataReadReference<DataType>(InDataReferenceName, InDataReference);
-				}
+			virtual ~TInputOperator() {}
+
+			virtual const FDataReferenceCollection& GetInputs() const override
+			{
+				return Inputs;
+			}
+
+			virtual const FDataReferenceCollection& GetOutputs() const override
+			{
+				return Outputs;
+			}
+
+			virtual FExecuteFunction GetExecuteFunction() override
+			{
+				return nullptr;
+			}
 
-				virtual ~FInputOperator() {}
+		private:
+			FDataReferenceCollection Outputs;
+			FDataReferenceCollection Inputs;
+	};
 
-				virtual const FDataReferenceCollection& GetInputs() const override
-				{
-					return Inputs;
-				}
+	/** TInputOperatorLiteralFactory creates an input by passing it a literal. */
+	template<typename DataType>
+	class TInputOperatorLiteralFactory : public IOperatorFactory
+	{
+		public:
+			// If the data type is parsable from a literal type, then the data type 
+			// can be registered as an input type with the frontend.  To make a 
+			// DataType registerable, either create a constructor for the data type
+			// which accepts the one of the supported literal types with an optional 
+			// FOperatorSettings argument, or create a default constructor, or specialize
+			// this factory with an implementation for that specific data type.
+			static constexpr bool bCanRegister = TLiteralTraits<DataType>::bIsParsableFromAnyLiteralType;
+
+			using FDataWriteReference = TDataWriteReference<DataType>;
+
+			TInputOperatorLiteralFactory(FDataTypeLiteralParam&& InInitParam)
+			:	InitParam(MoveTemp(InInitParam))
+			{
+			}
 
-				virtual const FDataReferenceCollection& GetOutputs() const override
-				{
-					return Outputs;
-				}
+			virtual TUniquePtr<IOperator> CreateOperator(const FCreateOperatorParams& InParams, FBuildErrorArray& OutErrors) override;
 
-				virtual FExecuteFunction GetExecuteFunction() override
-				{
-					return nullptr;
-				}
+		private:
 
-			private:
-				FDataReferenceCollection Outputs;
-				FDataReferenceCollection Inputs;
-		};
+			FDataTypeLiteralParam InitParam;
+	};
 
+	/** TInputOperatorFactory initializes the DataType at construction. It uses
+	 * the DataType's copy operator to create additional version. 
+	 */
+	template<typename DataType>
+	class TInputOperatorFactory : public IOperatorFactory
+	{
 		public:
 
-		// TODO: Would be nice if we didn't utilize the copy constructor by default here. But it's a nice catch-all
-		// for our use case.  Could store constructor parameters of the paramtype, but run the risk of holding rvalues
-		// or pointers to objects which are no longer valid. 
-
-		// TODO: instead of this copy factory, register a DataType factory lambda that uses a FDataTypeLiteralParam
-		class FCopyOperatorFactory : public IOperatorFactory
-		{
-			static_assert(TIsConstructible<DataType, const DataType&>::Value, "Data Type must be copy constructible!");
-				
-			public:
-				using FDataWriteReference = TDataWriteReference<DataType>;
-				using FInputNodeType = TInputNode<DataType>;
-
-				template< typename... ArgTypes >
-				FCopyOperatorFactory(ArgTypes&&... Args)
-				:	Data(Forward<ArgTypes>(Args)...)
-				{
-				}
-
-				/*
-				template<typename = typename TEnableIf< TIsConstructible<DataType>::Value >::Type >
-				FCopyOperatorFactory()
-				:	Data()
-				{
-				}
-				*/
-
-				FCopyOperatorFactory(FDataTypeLiteralParam InitParam, const FOperatorSettings& InSettings)
-					: Data(InitParam.ParseTo<DataType>(InSettings))
-				{
-				}
-
-				virtual TUniquePtr<IOperator> CreateOperator(const FCreateOperatorParams& InParams, FBuildErrorArray& OutErrors) override
-				{
-					// Use copy constructor to create a new parameter reference.
-					FDataWriteReference DataRef = FDataWriteReference::CreateNew(Data);
-					// TODO: Write special version of this for audio types since they will need to be constructed based upon the operator settings. 
-					// TODO: Need to ponder how inputs are initialized, but template specialization might do the trick. 
-
-					const FInputNodeType& InputNode = static_cast<const FInputNodeType&>(InParams.Node);
-
-					return MakeUnique<FInputOperator>(InputNode.GetVertexName(), DataRef);
-				}
-
-			private:
-				DataType Data;
-		};
-			// Constructors that need to get working on linux:
-			/*
-			template<typename = typename TEnableIf< TIsConstructible<DataType>::Value >::Type>
-			TInputNode(const FString& InNodeDescription, const FString& InVertexName)
-			:	NodeDescription(InNodeDescription)
-			,	VertexName(InVertexName)
-			,   Factory()
+			using FDataWriteReference = TDataWriteReference<DataType>;
+
+			template<typename... ArgTypes>
+			TInputOperatorFactory(ArgTypes&&... Args)
+			:	Data(Forward<ArgTypes>(Args)...)
 			{
-				FOutputDataVertex OutputVertex = MakeOutputDataVertex<DataType>(VertexName, FText::GetEmpty());
-				FDataVertexKey OutputVertexKey = MakeDataVertexKey(OutputVertex);
+			}
 
-				Outputs.Add(OutputVertexKey, OutputVertex);
+			virtual TUniquePtr<IOperator> CreateOperator(const FCreateOperatorParams& InParams, FBuildErrorArray& OutErrors) override;
 
-				FInputDataVertex InputVertex = MakeInputDataVertex<DataType>(VertexName, FText::GetEmpty());
-				FDataVertexKey InputVertexKey = MakeDataVertexKey(InputVertex);
+		private:
+			DataType Data;
+	};
 
-				Inputs.Add(InputVertexKey, InputVertex);
+	/** TInputNode represents an input to a metasound graph. */
+	template<typename DataType>
+	class TInputNode : public FNode
+	{
 
-			}
-			*/
+		public:
+
+			// If true, this node can be instantiated by the FrontEnd.
+			static constexpr bool bCanRegister = TInputOperatorLiteralFactory<DataType>::bCanRegister;
 
 			static const FNodeInfo& GetNodeInfo()
 			{
@@ -146,25 +132,27 @@ namespace Metasound
 				);
 			}
 
+			/* Construct a TInputNode using the TInputOperatorFactory<> and forwarding 
+			 * Args to the TInputOperatorFactory constructor.*/
 			template<typename... ArgTypes>
 			TInputNode(const FString& InNodeDescription, const FString& InVertexName, ArgTypes&&... Args)
 			:	FNode(InNodeDescription, GetNodeInfo())
 			,	VertexName(InVertexName)
 			,	Interface(DeclareVertexInterface(InVertexName))
-			,	Factory(MakeOperatorFactoryRef<FCopyOperatorFactory>(Forward<ArgTypes>(Args)...))
+			,	Factory(MakeOperatorFactoryRef<TInputOperatorFactory<DataType>>(Forward<ArgTypes>(Args)...))
 			{
-				static_assert(TIsConstructible<DataType, ArgTypes...>::Value, "Tried to construct TInputNode<DataType> with arguments that don't match any constructor for DataType.");
 			}
 
-			explicit TInputNode(const FString& InNodeDescription, const FString& InVertexName, FDataTypeLiteralParam&& InParam, const FOperatorSettings& InSettings)
+			/* Construct a TInputNode using the TInputOperatorLiteralFactory<> and moving
+			 * InParam to the TInputOperatorLiteralFactory constructor.*/
+			explicit TInputNode(const FString& InNodeDescription, const FString& InVertexName, FDataTypeLiteralParam&& InParam)
 			:	FNode(InNodeDescription, GetNodeInfo())
 			,	VertexName(InVertexName)
 			,	Interface(DeclareVertexInterface(InVertexName))
-			, 	Factory(MakeOperatorFactoryRef<FCopyOperatorFactory>(MoveTemp(InParam), InSettings))
+			, 	Factory(MakeOperatorFactoryRef<TInputOperatorLiteralFactory<DataType>>(MoveTemp(InParam)))
 			{
 			}
 
-
 			const FString& GetVertexName() const
 			{
 				return VertexName;
@@ -200,7 +188,34 @@ namespace Metasound
 
 			FVertexInterface Interface;
 			FOperatorFactorySharedRef Factory;
-
 	};
+
+
+	template<typename DataType>
+	TUniquePtr<IOperator> TInputOperatorLiteralFactory<DataType>::CreateOperator(const FCreateOperatorParams& InParams, FBuildErrorArray& OutErrors) 
+	{
+		using FInputNodeType = TInputNode<DataType>;
+
+		// Create write reference by calling compatible constructor with literal.
+		FDataWriteReference DataRef = TDataWriteReferenceLiteralFactory<DataType>::CreateExplicitArgs(InParams.OperatorSettings, InitParam);
+
+		const FInputNodeType& InputNode = static_cast<const FInputNodeType&>(InParams.Node);
+
+		return MakeUnique<TInputOperator<DataType>>(InputNode.GetVertexName(), DataRef);
+	}
+
+	template<typename DataType>
+	TUniquePtr<IOperator> TInputOperatorFactory<DataType>::CreateOperator(const FCreateOperatorParams& InParams, FBuildErrorArray& OutErrors) 
+	{
+		using FInputNodeType = TInputNode<DataType>;
+
+		// Create write reference by calling copy constructor.
+		FDataWriteReference DataRef = TDataWriteReferenceFactory<DataType>::CreateExplicitArgs(InParams.OperatorSettings, Data);
+
+		const FInputNodeType& InputNode = static_cast<const FInputNodeType&>(InParams.Node);
+
+		return MakeUnique<TInputOperator<DataType>>(InputNode.GetVertexName(), DataRef);
+	}
+
 } // namespace Metasound
 #undef LOCTEXT_NAMESPACE //MetasoundOutputNode
diff --git a/Engine/Plugins/Runtime/Metasound/Source/MetasoundFrontend/Public/MetasoundOutputNode.h b/Engine/Plugins/Runtime/Metasound/Source/MetasoundFrontend/Public/MetasoundOutputNode.h
--- a/Engine/Plugins/Runtime/Metasound/Source/MetasoundFrontend/Public/MetasoundOutputNode.h
+++ b/Engine/Plugins/Runtime/Metasound/Source/MetasoundFrontend/Public/MetasoundOutputNode.h
@@ -14,8 +14,6 @@ namespace Metasound
 	template<typename DataType>
 	class TOutputNode : public FNode
 	{
-		static_assert(TDataReferenceTypeInfo<DataType>::bIsValidSpecialization, "Please use DECLARE_METASOUND_DATA_REFERENCE_TYPES with this class before trying to create an output node with it.");
-			
 		class FOutputOperator : public IOperator
 		{
 			public:
diff --git a/Engine/Plugins/Runtime/Metasound/Source/MetasoundFrontend/Public/MetasoundReceiveNode.h b/Engine/Plugins/Runtime/Metasound/Source/MetasoundFrontend/Public/MetasoundReceiveNode.h
--- a/Engine/Plugins/Runtime/Metasound/Source/MetasoundFrontend/Public/MetasoundReceiveNode.h
+++ b/Engine/Plugins/Runtime/Metasound/Source/MetasoundFrontend/Public/MetasoundReceiveNode.h
@@ -5,6 +5,7 @@
 #include "MetasoundBuilderInterface.h"
 #include "MetasoundNodeInterface.h"
 #include "MetasoundOperatorInterface.h"
+#include "MetasoundDataFactory.h"
 #include "MetasoundDataReference.h"
 #include "MetasoundExecutableOperator.h"
 #include "MetasoundRouter.h"
@@ -37,20 +38,17 @@ namespace Metasound
 	template<typename TDataType>
 	class TReceiveNode : public INode
 	{
-		static_assert(TDataReferenceTypeInfo<TDataType>::bIsValidSpecialization, "Please use DECLARE_METASOUND_DATA_REFERENCE_TYPES with this class before trying to create an converter node with it.");
-
-
 
 	public:
-		static FString& GetAddressInputName()
+		static const FString& GetAddressInputName()
 		{
-			static FString InputName = FString(TEXT("Address"));
+			static const FString InputName = FString(TEXT("Address"));
 			return InputName;
 		}
 
-		static FString& GetSendOutputName()
+		static const FString& GetSendOutputName()
 		{
-			static FString SendInput = FString(TDataReferenceTypeInfo<TDataType>::TypeName);
+			static const FString SendInput = GetMetasoundDataTypeString<TDataType>();
 			return SendInput;
 		}
 
@@ -117,7 +115,7 @@ namespace Metasound
 				virtual TUniquePtr<IOperator> CreateOperator(const FCreateOperatorParams& InParams, FBuildErrorArray& OutErrors) override
 				{
 					return MakeUnique<TReceiverOperator>(
-						TDataReferenceFactory<TDataType>::CreateNewWriteReference(InParams.OperatorSettings),
+						TDataWriteReferenceFactory<TDataType>::CreateAny(InParams.OperatorSettings),
 						InParams.InputDataReferences.GetDataReadReference<FSendAddress>(GetAddressInputName()),
 						InParams.OperatorSettings
 						);
diff --git a/Engine/Plugins/Runtime/Metasound/Source/MetasoundFrontend/Public/MetasoundRouter.h b/Engine/Plugins/Runtime/Metasound/Source/MetasoundFrontend/Public/MetasoundRouter.h
--- a/Engine/Plugins/Runtime/Metasound/Source/MetasoundFrontend/Public/MetasoundRouter.h
+++ b/Engine/Plugins/Runtime/Metasound/Source/MetasoundFrontend/Public/MetasoundRouter.h
@@ -2,15 +2,17 @@
 #pragma once
 
 #include "CoreMinimal.h"
-#include "UObject/NameTypes.h"
 
-#include "MetasoundOperatorSettings.h"
+#include "DSP/Dsp.h"
+#include "DSP/MultithreadedPatching.h"
+#include "MetasoundAudioBuffer.h"
+#include "MetasoundDataFactory.h"
 #include "MetasoundDataReference.h"
+#include "MetasoundLiteral.h"
 #include "MetasoundOperatorInterface.h"
-#include "MetasoundAudioBuffer.h"
+#include "MetasoundOperatorSettings.h"
 #include "Misc/Guid.h"
-#include "DSP/Dsp.h"
-#include "DSP/MultithreadedPatching.h"
+#include "UObject/NameTypes.h"
 
 #include <type_traits>
 
@@ -47,6 +49,9 @@ namespace Metasound
 		Global,
 		Extension // This means that this metasound router talks to an implentation of IMetasoundTransmissionListener.
 	};
+	/* TODO: Rename "Subsystem" to "Protocol".  A "Protocol" defines how data is moved around
+	 * and can be used by multiple APIs. 
+	 */
 
 	// This converts an ETransmissionScope into an FName that can be used in FSendAddress::Subsystem.
 	// For ETransmissionScope::Extension, the FName of that extension should be used instead.
@@ -520,9 +525,9 @@ namespace Metasound
 
 		virtual bool PushLiteral(const FDataTypeLiteralParam& InParam)
 		{
-			if (InParam.IsCompatibleWithType<TDataType>())
+			if (TLiteralTraits<TDataType>::IsParsable(InParam))
 			{
-				TDataType DataToPush = InParam.ParseTo<TDataType>(OperatorSettings);
+				TDataType DataToPush = TDataTypeLiteralFactory<TDataType>::CreateExplicitArgs(OperatorSettings, InParam);
 
 				FScopeLock ScopeLock(&AtomicDataLock);
 				if (!AtomicData.IsValid())
@@ -771,27 +776,27 @@ namespace Metasound
 
 	// Coalesce both our standard copying Sender/Receivers and our audio-specific sender/receiver implementations into TSenderPtr and TReceiverPtr:
 
-	template<typename TDataType, typename USenderType = typename std::conditional<TIsDerivedFrom<TDataType, IAudioDatatype>::Value, TAudioSender<TDataType>, TSender<TDataType>>::type>
+	template<typename TDataType, typename USenderType = typename std::conditional<TIsDerivedFrom<TDataType, IAudioDataType>::Value, TAudioSender<TDataType>, TSender<TDataType>>::type>
 	using TSenderPtr = TUniquePtr<USenderType>;
 
-	template<typename TDataType, typename UReceiverType = typename std::conditional<TIsDerivedFrom<TDataType, IAudioDatatype>::Value, TAudioReceiver<TDataType>, TReceiver<TDataType>>::type>
+	template<typename TDataType, typename UReceiverType = typename std::conditional<TIsDerivedFrom<TDataType, IAudioDataType>::Value, TAudioReceiver<TDataType>, TReceiver<TDataType>>::type>
 	using TReceiverPtr = TUniquePtr<UReceiverType>;
 
 	// SFINAE for creating the correct data channel type for the given datatype:
 
-	template<typename TDataType, typename TEnableIf<std::is_copy_constructible<TDataType>::value && !TIsDerivedFrom<TDataType, IAudioDatatype>::Value, bool>::Type = true>
+	template<typename TDataType, typename TEnableIf<std::is_copy_constructible<TDataType>::value && !TIsDerivedFrom<TDataType, IAudioDataType>::Value, bool>::Type = true>
 	TSharedRef<IDataChannel> MakeDataChannel(const FOperatorSettings& InSettings)
 	{
 		return MakeShareable(new TCopyableDataChannel<TDataType>(InSettings));
 	}
 
-	template<typename TDataType, typename TEnableIf<TIsDerivedFrom<TDataType, IAudioDatatype>::Value, bool>::Type = true>
+	template<typename TDataType, typename TEnableIf<TIsDerivedFrom<TDataType, IAudioDataType>::Value, bool>::Type = true>
 	TSharedRef<IDataChannel> MakeDataChannel(const FOperatorSettings& InSettings)
 	{
 		return MakeShareable(new TAudioDataChannel<TDataType>(InSettings));
 	}
 
-	template<typename TDataType, typename TEnableIf<!std::is_copy_constructible<TDataType>::value && !TIsDerivedFrom<TDataType, IAudioDatatype>::Value, bool>::Type = true>
+	template<typename TDataType, typename TEnableIf<!std::is_copy_constructible<TDataType>::value && !TIsDerivedFrom<TDataType, IAudioDataType>::Value, bool>::Type = true>
 	TSharedRef<IDataChannel> MakeDataChannel(const FOperatorSettings& InSettings)
 	{
 		return MakeShareable(new TNonOperationalDataChannel<TDataType>(InSettings));
@@ -800,7 +805,7 @@ namespace Metasound
 	// Utility function for properly downcasting to the correct type based on whether TDataType is an audio stream:
 
 	// For Audio Senders:
-	template <typename TDataType, typename TEnableIf<TIsDerivedFrom<TDataType, IAudioDatatype>::Value, bool>::Type = true>
+	template <typename TDataType, typename TEnableIf<TIsDerivedFrom<TDataType, IAudioDataType>::Value, bool>::Type = true>
 	TSenderPtr<TDataType> Downcast(TUniquePtr<ISender>&& InPtr)
 	{
 		checkf(InPtr->CheckType<TAudioSender<TDataType>>(), TEXT("Tried to downcast an ISender of type %s to a TSender of type %s!"), *(InPtr->GetDataType().ToString()), *(TAudioSender<TDataType>::GetDataTypeName().ToString()));
@@ -808,7 +813,7 @@ namespace Metasound
 	}
 
 	// For generic senders for copyable types:
-	template <typename TDataType, typename TEnableIf<std::is_copy_constructible<TDataType>::value && !TIsDerivedFrom<TDataType, IAudioDatatype>::Value, bool>::Type = true>
+	template <typename TDataType, typename TEnableIf<std::is_copy_constructible<TDataType>::value && !TIsDerivedFrom<TDataType, IAudioDataType>::Value, bool>::Type = true>
 	TSenderPtr<TDataType> Downcast(TUniquePtr<ISender>&& InPtr)
 	{
 		checkf(InPtr->CheckType<TSender<TDataType>>(), TEXT("Tried to downcast an ISender of type %s to a TSender of type %s!"), *(InPtr->GetDataType().ToString()), *(TSender<TDataType>::GetDataTypeName().ToString()));
@@ -816,14 +821,14 @@ namespace Metasound
 	}
 
 	// For invalid types (non-copyable AND not an audio type):
-	template <typename TDataType, typename TEnableIf<!std::is_copy_constructible<TDataType>::value && !TIsDerivedFrom<TDataType, IAudioDatatype>::Value, bool>::Type = true>
+	template <typename TDataType, typename TEnableIf<!std::is_copy_constructible<TDataType>::value && !TIsDerivedFrom<TDataType, IAudioDataType>::Value, bool>::Type = true>
 	TSenderPtr<TDataType> Downcast(TUniquePtr<ISender>&& InPtr)
 	{
 		return nullptr;
 	}
 
 	// For audio receivers:
-	template <typename TDataType, typename TEnableIf<TIsDerivedFrom<TDataType, IAudioDatatype>::Value, bool>::Type = true>
+	template <typename TDataType, typename TEnableIf<TIsDerivedFrom<TDataType, IAudioDataType>::Value, bool>::Type = true>
 	TReceiverPtr<TDataType> Downcast(TUniquePtr<IReceiver>&& InPtr)
 	{
 		checkf(InPtr->CheckType<TAudioReceiver<TDataType>>(), TEXT("Tried to downcast an IReceiver of type %s to a TSender of type %s!"), *(InPtr->GetDataType().ToString()), *(TAudioReceiver<TDataType>::GetDataTypeName().ToString()));
@@ -831,7 +836,7 @@ namespace Metasound
 	}
 
 	// For generic receivers for copyable types:
-	template <typename TDataType, typename TEnableIf<std::is_copy_constructible<TDataType>::value && !TIsDerivedFrom<TDataType, IAudioDatatype>::Value, bool>::Type = true>
+	template <typename TDataType, typename TEnableIf<std::is_copy_constructible<TDataType>::value && !TIsDerivedFrom<TDataType, IAudioDataType>::Value, bool>::Type = true>
 	TReceiverPtr<TDataType> Downcast(TUniquePtr<IReceiver>&& InPtr)
 	{
 		checkf(InPtr->CheckType<TReceiver<TDataType>>(), TEXT("Tried to downcast an IReceiver of type %s to a TSender of type %s!"), *(InPtr->GetDataType().ToString()), *(TReceiver<TDataType>::GetDataTypeName().ToString()));
@@ -839,7 +844,7 @@ namespace Metasound
 	}
 
 	// For invalid types (non-copyable AND not an audio type):
-	template <typename TDataType, typename TEnableIf<!std::is_copy_constructible<TDataType>::value && !TIsDerivedFrom<TDataType, IAudioDatatype>::Value, bool>::Type = true>
+	template <typename TDataType, typename TEnableIf<!std::is_copy_constructible<TDataType>::value && !TIsDerivedFrom<TDataType, IAudioDataType>::Value, bool>::Type = true>
 	TReceiverPtr<TDataType> Downcast(TUniquePtr<IReceiver>&& InPtr)
 	{
 		return nullptr;
diff --git a/Engine/Plugins/Runtime/Metasound/Source/MetasoundFrontend/Public/MetasoundSendNode.h b/Engine/Plugins/Runtime/Metasound/Source/MetasoundFrontend/Public/MetasoundSendNode.h
--- a/Engine/Plugins/Runtime/Metasound/Source/MetasoundFrontend/Public/MetasoundSendNode.h
+++ b/Engine/Plugins/Runtime/Metasound/Source/MetasoundFrontend/Public/MetasoundSendNode.h
@@ -20,18 +20,16 @@ namespace Metasound
 	template<typename TDataType>
 	class TSendNode : public INode
 	{
-		static_assert(TDataReferenceTypeInfo<TDataType>::bIsValidSpecialization, "Please use DECLARE_METASOUND_DATA_REFERENCE_TYPES with this class before trying to create an converter node with it.");
-
 	public:
-		static FString& GetAddressInputName()
+		static const FString& GetAddressInputName()
 		{
-			static FString InputName = FString(TEXT("Address"));
+			static const FString InputName = FString(TEXT("Address"));
 			return InputName;
 		}
 
-		static FString& GetSendInputName()
+		static const FString& GetSendInputName()
 		{
-			static FString SendInput = FString(TDataReferenceTypeInfo<TDataType>::TypeName);
+			static const FString SendInput = GetMetasoundDataTypeString<TDataType>();
 			return SendInput;
 		}
 
diff --git a/Engine/Plugins/Runtime/Metasound/Source/MetasoundGraphCore/Private/MetasoundDataReference.cpp b/Engine/Plugins/Runtime/Metasound/Source/MetasoundGraphCore/Private/MetasoundDataReference.cpp
--- a/Engine/Plugins/Runtime/Metasound/Source/MetasoundGraphCore/Private/MetasoundDataReference.cpp
+++ b/Engine/Plugins/Runtime/Metasound/Source/MetasoundGraphCore/Private/MetasoundDataReference.cpp
@@ -5,4 +5,6 @@
 namespace Metasound
 {
 	const TCHAR* TDataReferenceTypeInfo<void>::TypeName = TEXT("void");
+	const void* const TDataReferenceTypeInfo<void>::TypePtr = nullptr;
+	const void* const TDataReferenceTypeInfo<void>::TypeId = static_cast<const void* const>(TDataReferenceTypeInfo<void>::TypePtr);
 }
diff --git a/Engine/Plugins/Runtime/Metasound/Source/MetasoundGraphCore/Public/MetasoundDataFactory.h b/Engine/Plugins/Runtime/Metasound/Source/MetasoundGraphCore/Public/MetasoundDataFactory.h
new file mode 100644
--- /dev/null
+++ b/Engine/Plugins/Runtime/Metasound/Source/MetasoundGraphCore/Public/MetasoundDataFactory.h
@@ -0,0 +1,825 @@
+// Copyright Epic Games, Inc. All Rights Reserved.
+
+#pragma once
+
+#include "CoreMinimal.h"
+#include "MetasoundDataReference.h"
+#include "MetasoundLiteral.h"
+#include "MetasoundOperatorSettings.h"
+#include "Misc/Variant.h"
+#include <type_traits>
+
+namespace Metasound
+{
+	namespace DataFactoryPrivate
+	{
+		/** Description of available constructors for a registered Metasound Data Type. 
+		 *
+		 * @tparam DataType - The registered Metasound Data Type.
+		 * @tparam ArgTypes - A parameter pack of arguments to be passed to the data type's constructor.
+		 */
+		template<typename DataType, typename... ArgTypes>
+		struct TDataTypeConstructorTraits
+		{
+			static constexpr bool bIsDefaultConstructible = std::is_constructible<DataType>::value;
+			static constexpr bool bIsConstructibleWithSettings = std::is_constructible<DataType, const FOperatorSettings&>::value;
+			static constexpr bool bIsConstructibleWithArgs = std::is_constructible<DataType, ArgTypes...>::value;
+			static constexpr bool bIsConstructibleWithSettingsAndArgs = std::is_constructible<DataType, const FOperatorSettings&, ArgTypes...>::value;
+		};
+
+		/** Constructor traits for a for variant inputs. These traits test each
+		 * individual type supported by the variant and determines if *any* of the
+		 * types can be used as a single argument to construct the DataType.
+		 *
+		 * @tparam DataType - The registered Metasound Data Type.
+		 * @tparam Types - The types supported by the variant.
+		 */
+		template<typename DataType, typename... Types>
+		struct TDataTypeVariantConstructorTraits;
+
+		// Specialization for single arg.
+		template<typename DataType, typename ArgType>
+		struct TDataTypeVariantConstructorTraits<DataType, ArgType>
+		{
+			using FConstructorTraits = TDataTypeConstructorTraits<DataType, ArgType>;
+
+			public:
+
+			static constexpr bool bIsDefaultConstructible = FConstructorTraits::bIsDefaultConstructible;
+			static constexpr bool bIsConstructibleWithSettings = FConstructorTraits::bIsConstructibleWithSettings;
+			static constexpr bool bIsConstructibleWithArgs = FConstructorTraits::bIsConstructibleWithArgs;
+			static constexpr bool bIsConstructibleWithSettingsAndArgs = FConstructorTraits::bIsConstructibleWithSettingsAndArgs;
+		};
+
+		// Specialization for unrolling parameter packs. 
+		template<typename DataType, typename ArgType, typename... AdditionalTypes>
+		struct TDataTypeVariantConstructorTraits<DataType, ArgType, AdditionalTypes...>
+		{
+			private:
+
+			using FAdditionalConstructorTraits = TDataTypeVariantConstructorTraits<DataType, AdditionalTypes...>;
+			using FArgConstructorTraits = TDataTypeConstructorTraits<DataType, ArgType>;
+			
+			public:
+
+			static constexpr bool bIsDefaultConstructible = FAdditionalConstructorTraits::bIsDefaultConstructible;
+			static constexpr bool bIsConstructibleWithSettings = FAdditionalConstructorTraits::bIsConstructibleWithSettings;
+			static constexpr bool bIsConstructibleWithArgs = FArgConstructorTraits::bIsConstructibleWithArgs || FAdditionalConstructorTraits::bIsConstructibleWithArgs;
+			static constexpr bool bIsConstructibleWithSettingsAndArgs = FArgConstructorTraits::bIsConstructibleWithSettingsAndArgs || FAdditionalConstructorTraits::bIsConstructibleWithSettingsAndArgs;
+		};
+
+		// Specialization for unpacking the types supported by a TVariant. 
+		template<typename DataType, typename FirstVariantType, typename... AdditionalVariantTypes>
+		struct TDataTypeVariantConstructorTraits<DataType, TVariant<FirstVariantType, AdditionalVariantTypes...>>
+		{
+			private:
+
+			using FConstructorTraits = TDataTypeVariantConstructorTraits<DataType, FirstVariantType, AdditionalVariantTypes...>;
+			
+			public:
+
+			static constexpr bool bIsDefaultConstructible = FConstructorTraits::bIsDefaultConstructible;
+			static constexpr bool bIsConstructibleWithSettings = FConstructorTraits::bIsConstructibleWithSettings;
+			static constexpr bool bIsConstructibleWithArgs = FConstructorTraits::bIsConstructibleWithArgs;
+			static constexpr bool bIsConstructibleWithSettingsAndArgs = FConstructorTraits::bIsConstructibleWithSettingsAndArgs;
+		};
+
+
+		/** Denotes that both FOperatorSettings and parameter pack arguments must be
+		 * used in the constructor of the Metasound Data Type.
+		 *
+		 * @tparam DataType - The registered Metasound Data Type.
+		 * @tparam ArgTypes - A parameter pack of arguments to be passed to the data type's constructor.
+		 */
+		template<typename DataType, typename... ArgTypes>
+		struct TExplicitConstructorForwarding
+		{
+			using TConstructorTraits = TDataTypeConstructorTraits<DataType, ArgTypes...>;
+
+			static constexpr bool bForwardSettingsAndArgs = TConstructorTraits::bIsConstructibleWithSettingsAndArgs;
+			static constexpr bool bForwardArgs = false;
+			static constexpr bool bForwardSettings = false;
+			static constexpr bool bForwardNone = false;
+
+			static constexpr bool bCannotForwardToConstructor = !(bForwardSettingsAndArgs || bForwardArgs || bForwardSettings || bForwardNone);
+		};
+
+		/** Denotes that parameter pack arguments must be used in the constructor of
+		 * the Metasound Data Type. The use of the FOperatorSettings is optional.
+		 *
+		 * @tparam DataType - The registered Metasound Data Type.
+		 * @tparam ArgTypes - A parameter pack of arguments to be passed to the data type's constructor.
+		 */
+		template<typename DataType, typename... ArgTypes>
+		struct TExplicitArgsConstructorForwarding
+		{
+			using TConstructorTraits = TDataTypeConstructorTraits<DataType, ArgTypes...>;
+
+			static constexpr bool bForwardSettingsAndArgs = TConstructorTraits::bIsConstructibleWithSettingsAndArgs;
+			static constexpr bool bForwardArgs = !bForwardSettingsAndArgs && TConstructorTraits::bIsConstructibleWithArgs;
+			static constexpr bool bForwardSettings = false;
+			static constexpr bool bForwardNone = false;
+
+			static constexpr bool bCannotForwardToConstructor = !(bForwardSettingsAndArgs || bForwardArgs || bForwardSettings || bForwardNone);
+		};
+
+		/** Denotes that parameter pack arguments and/or FOperatorSettings is optional
+		 * when constructing the Metasound Data Type. 
+		 *
+		 * @tparam DataType - The registered Metasound Data Type.
+		 * @tparam ArgTypes - A parameter pack of arguments to be passed to the data type's constructor.
+		 */
+		template<typename DataType, typename... ArgTypes>
+		struct TAnyConstructorForwarding
+		{
+			using TConstructorTraits = TDataTypeConstructorTraits<DataType, ArgTypes...>;
+
+			static constexpr bool bForwardSettingsAndArgs = TConstructorTraits::bIsConstructibleWithSettingsAndArgs;
+			static constexpr bool bForwardArgs = !bForwardSettingsAndArgs && TConstructorTraits::bIsConstructibleWithArgs;
+			static constexpr bool bForwardSettings = !(bForwardArgs || bForwardSettingsAndArgs) && TConstructorTraits::bIsConstructibleWithSettings;
+			static constexpr bool bForwardNone = !(bForwardArgs || bForwardSettingsAndArgs || bForwardSettings) && TConstructorTraits::bIsDefaultConstructible;
+
+			static constexpr bool bCannotForwardToConstructor = !(bForwardSettingsAndArgs || bForwardArgs || bForwardSettings || bForwardNone);
+		};
+
+		/** Create a DataType. For use in TDataFactory. */
+		template<typename DataType>
+		struct TDataTypeCreator
+		{
+			template<typename... ArgTypes>
+			static DataType CreateNew(ArgTypes&&... Args)
+			{
+				return DataType(Forward<ArgTypes>(Args)...);
+			}
+		};
+
+	MSVC_PRAGMA(warning(push))
+	MSVC_PRAGMA(warning(disable : 4800)) // Disable warning when converting int to bool.
+
+		/** Core factory type for creating objects related to metasound data types. 
+		 *
+		 * TDataFactoryInternal provides a unified interface for constructing objects
+		 * using a FOperatorSettings and a parameter pack of arguments. The ArgForwardType
+		 * uses SFINAE to choose whether to forward all arguments, a subset of arguments
+		 * or no arguments to the underlying CreatorType.
+		 *
+		 * @tparam DataType - The registered metasound data type which defines the supported constructors.
+		 * @tparam CreatorType - A class providing a static "CreateNew(...)" function which
+		 *                       accepts the forwarded arguments and returns an object.
+		 */
+		template<typename DataType, typename CreatorType>
+		struct TDataFactoryInternal
+		{
+			/** Create a new object using FOperatorSettings and a parameter pack of
+			 * arguments. 
+			 *
+			 * Note: In this implementation, all arguments are forwarded to CreatorType::CreateNew(...)
+			 *
+			 * @tparam ArgForwardType - A type which describes which settings and arguments
+			 *                          should be forwarded to CreateType::CreateNew(...)
+			 * @tparam ArgTypes - A parameter pack of types to be forwarded to CreateType::CreateNew(...)
+			 */
+			template<
+				typename ArgForwardType,
+				typename... ArgTypes,
+				typename std::enable_if< ArgForwardType::bForwardSettingsAndArgs, int>::type = 0
+			>
+			static auto CreateNew(const FOperatorSettings& InSettings, ArgTypes&&... Args)
+			{
+				return CreatorType::CreateNew(InSettings, Forward<ArgTypes>(Args)...);
+			}
+
+			/** Create a new object using FOperatorSettings and a parameter pack of
+			 * arguments. 
+			 *
+			 * Note: In this implementation, the parameter pack is forwarded to CreateType::CreateNew(...), 
+			 * while the FOperatorSettings are ignored.
+			 *
+			 * @tparam ArgForwardType - A type which describes which settings and arguments
+			 *                          should be forwarded to CreateType::CreateNew(...)
+			 * @tparam ArgTypes - A parameter pack of types to be forwarded to CreateType::CreateNew(...)
+			 */
+			template<
+				typename ArgForwardType,
+				typename... ArgTypes,
+				typename std::enable_if< ArgForwardType::bForwardArgs, int>::type = 0
+			>
+			static auto CreateNew(const FOperatorSettings& InSettings, ArgTypes&&... Args)
+			{
+				return CreatorType::CreateNew(Forward<ArgTypes>(Args)...);
+			}
+
+			/** Create a new object using FOperatorSettings and a parameter pack of
+			 * arguments. 
+			 *
+			 * Note: In this implementation, the FOperatorSettings are forwarded to CreateType::CreateNew(...), 
+			 * while the paramter pack is ignored.
+			 *
+			 * @tparam ArgForwardType - A type which describes which settings and arguments
+			 *                          should be forwarded to CreateType::CreateNew(...)
+			 * @tparam ArgTypes - A parameter pack of types to be forwarded to CreateType::CreateNew(...)
+			 */
+			template<
+				typename ArgForwardType,
+				typename... ArgTypes,
+				typename std::enable_if< ArgForwardType::bForwardSettings, int>::type = 0
+			>
+			static auto CreateNew(const FOperatorSettings& InSettings, ArgTypes&&... Args)
+			{
+				return CreatorType::CreateNew(InSettings);
+			}
+
+			/** Create a new object using FOperatorSettings and a parameter pack of
+			 * arguments. 
+			 *
+			 * Note: In this implementation, the FOperatorSettings and parameter pack
+			 * are ignored.
+			 *
+			 * @tparam ArgForwardType - A type which describes which settings and arguments
+			 *                          should be forwarded to CreateType::CreateNew(...)
+			 * @tparam ArgTypes - A parameter pack of types to be forwarded to CreateType::CreateNew(...)
+			 */
+			template<
+				typename ArgForwardType,
+				typename... ArgTypes,
+				typename std::enable_if< ArgForwardType::bForwardNone, int>::type = 0
+			>
+			static auto CreateNew(const FOperatorSettings& InSettings, ArgTypes&&... Args)
+			{
+				return CreatorType::CreateNew();
+			}
+
+			/** Create a new object using FOperatorSettings and a parameter pack of
+			 * arguments. 
+			 *
+			 * Note: In this implementation, the FOperatorSettings and parameter pack
+			 * are ignored.
+			 *
+			 * @tparam ArgForwardType - A type which describes which settings and arguments
+			 *                          should be forwarded to CreateType::CreateNew(...)
+			 * @tparam ArgTypes - A parameter pack of types to be forwarded to CreateType::CreateNew(...)
+			 */
+			template<
+				typename ArgForwardType,
+				typename... ArgTypes,
+				typename std::enable_if< ArgForwardType::bCannotForwardToConstructor, int>::type = 0
+			>
+			static auto CreateNew(const FOperatorSettings& InSettings, ArgTypes&&... Args)
+			{
+				static_assert(!ArgForwardType::bCannotForwardToConstructor, "No constructor exists for the DataType which matches the given arguments and argument forwarding type.");
+			}
+		};
+
+	MSVC_PRAGMA(warning(pop))
+
+		/** Core factory type for creating objects related to Metasound DataTypes. 
+		 *
+		 * TDataVariantFactoryInternal provides a unified interface for constructing objects
+		 * using a FOperatorSettings and TVariant arguments. The VariantParseType
+		 * uses SFINAE to choose how to interpret the TVariant argument.
+		 *
+		 * @tparam DataType - The registered metasound data type which defines the supported constructors.
+		 * @tparam CreatorType - A class providing a static "CreateNewFromVariant(...)" function which
+		 *                       accepts the forwarded arguments and returns an object.
+		 */
+		template<typename DataType, typename CreatorType>
+		struct TDataVariantFactoryInternal
+		{
+			using FInternalFactory = TDataFactoryInternal<DataType, CreatorType>;
+
+			/** Create a new object using FOperatorSettings and TVariant.
+			 *
+			 * Note: In this implementation, all arguments are forwarded to CreatorType::CreateNew(...)
+			 * and the TVariant is parsed as the `VariantType`. FOperatorSettings are
+			 * forwarded to CreatorType::CreateNew(...) if it is supported.
+			 *
+			 * @tparam VariantParseType - A type which describes how to parse the InVariant
+			 *                            argument.
+			 * @tparam VariantType - The expected underlying type stored in InVariant.
+			 */
+			template<
+				typename VariantParseType,
+				typename VariantType,
+				typename std::enable_if< VariantParseType::bCreateWithArgType, int>::type = 0
+			>
+			static auto CreateNewFromVariant(const FOperatorSettings& InSettings, const VariantType& InVariant)
+			{
+				// The constructor must use the desired variant argument.
+				using TExplicitArgsConstructorForwarding = TExplicitArgsConstructorForwarding<DataType, typename VariantParseType::ArgType>;
+
+				return FInternalFactory::template CreateNew<TExplicitArgsConstructorForwarding>(InSettings, InVariant.template Get<typename VariantParseType::ArgType>());
+			}
+
+			/** Create a new object using FOperatorSettings and TVariant.
+			 *
+			 * Note: In this implementation, the InVariant object is ignore. FOperatorSettings 
+			 * are forwarded to CreatorType::CreateNew(...) if it is supported.
+			 *
+			 * @tparam VariantParseType - A type which describes how to parse the InVariant
+			 *                            argument.
+			 * @tparam VariantType - The expected underlying type stored in InVariant.
+			 */
+			template<
+				typename VariantParseType,
+				typename VariantType,
+				typename std::enable_if< VariantParseType::bCreateWithoutArg, int>::type = 0
+			>
+			static auto CreateNewFromVariant(const FOperatorSettings& InSettings, const VariantType& InVariant)
+			{
+				// The constructor must be the default constructor or accept a FOperatorSettings.
+				using TExplicitArgsConstructorForwarding = TExplicitArgsConstructorForwarding<DataType>;
+
+				// When constructing an object from a variant, callers expect the variant value
+				// to be used in the constructor. It is an error to ignore the variant value.
+				// This error can be fixed by changing the InVariant object passed in 
+				// or by adding a constructor to the DataType which accepts the VariantType.
+				checkf(false, TEXT("The value passed to the constructor is being ignored")); 
+
+				return FInternalFactory::template CreateNew<TExplicitArgsConstructorForwarding>(InSettings);
+			}
+
+			/** Create a new object using FOperatorSettings and TVariant.
+			 *
+			 * Note: In this implementation, the TVariant is parsed as a fallback type.
+			 * FOperatorSettings are forwarded to CreatorType::CreateNew(...) if it 
+			 * is supported.
+			 *
+			 * @tparam VariantParseType - A type which describes how to parse the InVariant
+			 *                            argument.
+			 * @tparam VariantType - The expected underlying type stored in InVariant.
+			 */
+			template<
+				typename VariantParseType,
+				typename VariantType,
+				typename std::enable_if< VariantParseType::bCreateWithFallbackArgType, int>::type = 0
+			>
+			static auto CreateNewFromVariant(const FOperatorSettings& InSettings, const VariantType& InVariant)
+			{
+				// The constructor must use the fallback variant argument.
+				using TExplicitArgsConstructorForwarding = TExplicitArgsConstructorForwarding<DataType, typename VariantParseType::FallbackArgType>;
+
+				// When constructing an object from a variant, callers expect the variant value
+				// to be used in the constructor. It is an error to interpret the variant value
+				// as something other than the expected type. Likely, this error is due
+				// to the DataType's constructor not supporting the VariantType as an argument.
+				// This error can be fixed by changing the InVariant object passed in 
+				// or by adding a constructor to the DataType which accepts the VariantType.
+				checkf(false, TEXT("The value passed to the constructor is parsed as an unrelated type.")); 
+
+				return FInternalFactory::template CreateNew<TExplicitArgsConstructorForwarding>(InSettings, InVariant.template Get<typename VariantParseType::FallbackArgType>());
+			}
+			
+			/** Create a new object using FOperatorSettings and TVariant.
+			 *
+			 * Note: In this implementation signals a static error as not possible ways
+			 * to construct that underlying object could be found.
+			 *
+			 * @tparam VariantParseType - A type which describes how to parse the InVariant
+			 *                            argument.
+			 * @tparam VariantType - The expected underlying type stored in InVariant.
+			 */
+			template<
+				typename VariantParseType,
+				typename VariantType,
+				typename std::enable_if< VariantParseType::bCannotForwardToConstructor, int>::type = 0
+			>
+			static auto CreateNewFromVariant(const FOperatorSettings& InSettings, const VariantType& InVariant)
+			{
+				// There is no DataType constructor which can be used to create the 
+				// DataType. This includes no default constructor, nor any constructor
+				// which accepts any of the variant types. 
+				static_assert(!VariantParseType::bCannotForwardToConstructor, "Constructor has no default constructor and does not support any supplied variant type");
+			}
+		};
+
+		// Promote which type to use for fallback 
+		template<bool ValidType, typename ThisType, typename PreviousType = void>
+		struct TTypePromoter
+		{
+			typedef PreviousType Type;
+			static constexpr bool bIsValidType = ValidType;
+		};
+
+		// TTypePromoter partial specialization for ValidType=true
+		template<typename ThisType, typename PreviousType>
+		struct TTypePromoter<true, ThisType, PreviousType>
+		{
+			typedef PreviousType Type;
+			static constexpr bool bIsValidType = true;
+		};
+
+		// TTypePromoter partial specialization for ValidType=false
+		template<typename ThisType, typename PreviousType>
+		struct TTypePromoter<false, ThisType, PreviousType>
+		{
+			typedef PreviousType Type;
+			static constexpr bool bIsValidType = !std::is_same<PreviousType, void>::value;
+		};
+
+		// TTypePromoter partial specialization for ValidType=true and no PreviousType.
+		template<typename ThisType>
+		struct TTypePromoter<true, ThisType>
+		{
+			typedef ThisType Type;
+			static constexpr bool bIsValidType = true;
+		};
+
+
+		// Helper for determine whether constructor accepts the variant type. Supports
+		// unrolling the parameter pack via recursive template.
+		template<typename DataType, typename ThisArgType, typename... VariantArgTypes>
+		struct TDataTypeVariantFallbackHelper
+		{
+		private:
+
+			static constexpr bool bIsConstructibleWithThisArg = std::is_constructible<DataType, ThisArgType>::value;
+			static constexpr bool bIsConstructibleWithSettingsAndThisArg = std::is_constructible<DataType, const FOperatorSettings&, ThisArgType>::value;
+
+			static constexpr bool bPromotThisArgType = bIsConstructibleWithThisArg || bIsConstructibleWithSettingsAndThisArg;
+			using FVariantPromoter = TTypePromoter<bPromotThisArgType, ThisArgType, typename TDataTypeVariantFallbackHelper<DataType, VariantArgTypes...>::FallbackType>;
+
+		public:
+			
+			// True if the DataType is default constructible or constructible with any 
+			// of the variant types.
+			static constexpr bool bIsConstructibleWithArgs = bIsConstructibleWithThisArg || TDataTypeVariantFallbackHelper<DataType, VariantArgTypes...>::bIsConstructibleWithArgs;
+
+			// True if the DataType is default constructible or constructible with any 
+			// of the variant types.
+			static constexpr bool bIsConstructibleWithSettingsAndArgs = bIsConstructibleWithSettingsAndThisArg || TDataTypeVariantFallbackHelper<DataType, VariantArgTypes...>::bIsConstructibleWithSettingsAndArgs;
+
+			// Fallback type to use if there is no other way to construct the object.
+			typedef typename FVariantPromoter::Type FallbackType;
+
+			// True if the constructor accepts the fallback type.
+			static constexpr bool bIsConstructibleWithFallbackArg = FVariantPromoter::bIsValidType; 
+		};
+
+		// Specialization of TDataTypeVariantFallbackHelper to handle end of recursion on 
+		// VariantArgTypes.
+		template<typename DataType, typename ThisArgType>
+		struct TDataTypeVariantFallbackHelper<DataType, ThisArgType>
+		{
+		private:
+			static constexpr bool bIsConstructibleWithThisArg = std::is_constructible<DataType, ThisArgType>::value;
+			static constexpr bool bIsConstructibleWithSettingsAndThisArg = std::is_constructible<DataType, const FOperatorSettings&, ThisArgType>::value;
+
+			static constexpr bool bPromotThisArgType = bIsConstructibleWithThisArg || bIsConstructibleWithSettingsAndThisArg;
+			using FVariantPromoter = TTypePromoter<bPromotThisArgType, ThisArgType>;
+
+		public:
+			// True if the DataType is default constructible or constructible with any 
+			// of the variant types.
+			static constexpr bool bIsConstructibleWithArgs = bIsConstructibleWithThisArg;
+
+			// True if the DataType is default constructible or constructible with any 
+			// of the variant types.
+			static constexpr bool bIsConstructibleWithSettingsAndArgs = bIsConstructibleWithSettingsAndThisArg;
+
+			// Fallback type to use if there is no other way to construct the object.
+			typedef typename FVariantPromoter::Type FallbackType;
+			// True if the constructor accepts the fallback type.
+			static constexpr bool bIsConstructibleWithFallbackArg = FVariantPromoter::bIsValidType; 
+		};
+
+		/** TDataTypeVariantParsing informs the TDataVariantFactoryInternal on which factory
+		 * method to instatiate.
+		 */
+		template<typename DataType, typename DesiredArgType, typename FirstVariantType, typename... AdditionalVariantTypes>
+		struct TDataTypeVariantParsing
+		{
+			using FDesiredConstructorTraits = TDataTypeConstructorTraits<DataType, DesiredArgType>; 
+			using FFallbackHelper = TDataTypeVariantFallbackHelper<DataType, FirstVariantType, AdditionalVariantTypes...>;
+
+			static constexpr bool bCreateWithArgType = FDesiredConstructorTraits::bIsConstructibleWithArgs || FDesiredConstructorTraits::bIsConstructibleWithSettingsAndArgs;
+			static constexpr bool bCreateWithoutArg = !bCreateWithArgType && (FDesiredConstructorTraits::bIsConstructibleWithSettings || FDesiredConstructorTraits::bIsDefaultConstructible);
+			static constexpr bool bCreateWithFallbackArgType = !(bCreateWithArgType || bCreateWithoutArg) && FFallbackHelper::bIsConstructibleWithFallbackArg;
+			static constexpr bool bCannotForwardToConstructor = !(bCreateWithArgType || bCreateWithoutArg || bCreateWithFallbackArgType);
+
+			typedef DesiredArgType ArgType;
+			typedef typename FFallbackHelper::FallbackType FallbackArgType;
+		};
+
+
+
+	} // End namespace DataFactoryPrivate
+
+	/** Determines whether a DataType supports a constructor which accepts and FOperatorSettings
+	 * with ArgTypes and/or just ArgTypes. 
+	 *
+	 * @tparam DataType - The registered Metasound Data Type.
+	 * @tparam ArgTypes - A parameter pack of arguments to be passed to the data type's constructor.
+	 */
+	template<typename DataType, typename... ArgTypes>
+	struct TIsParsable
+	{
+		using TConstructorTraits = DataFactoryPrivate::TDataTypeConstructorTraits<DataType, ArgTypes...>;
+
+		/* True if the DataType supports a constructor which accepts and FOperatorSettings 
+		 * with ArgTypes and/or just ArgTypes. False otherwise.
+		 */
+		static constexpr bool Value = TConstructorTraits::bIsConstructibleWithArgs || TConstructorTraits::bIsConstructibleWithSettingsAndArgs;
+	};
+
+	/** Determines whether a DataType supports construction using the given literal.
+	 *
+	 * @tparam DataType - The registered Metasound Data Type.
+	 */
+	template<typename DataType>
+	struct TLiteralTraits
+	{
+		using TVariantConstructorTraits = DataFactoryPrivate::TDataTypeVariantConstructorTraits<DataType, FDataTypeLiteralParam::FVariantType>;
+
+		static constexpr bool bIsParsableFromAnyLiteralType = TVariantConstructorTraits::bIsDefaultConstructible || TVariantConstructorTraits::bIsConstructibleWithSettings || TVariantConstructorTraits::bIsConstructibleWithArgs || TVariantConstructorTraits::bIsConstructibleWithSettingsAndArgs;
+
+		/** Determines if a constructor for the DataType exists which accepts 
+		 * an FOperatorSettings with the literals constructor arg type, and/or one that
+		 * accpets the literal constructor arg type. 
+		 *
+		 * @param InLiteral - The literal containing the constructor argument.
+		 *
+		 * @return True if a constructor for the DataType exists which accepts 
+		 * an FOperatorSettings with the literals constructor arg type, or one that
+		 * accpets the literal constructor arg type. It returns False otherwise.
+		 */
+		static const bool IsParsable(const FDataTypeLiteralParam& InLiteral)
+		{
+			switch (InLiteral.ConstructorArgType)
+			{
+				case ELiteralArgType::Boolean:
+
+					return TIsParsable<DataType, bool>::Value;
+
+				case ELiteralArgType::Integer:
+
+					return TIsParsable<DataType, int32>::Value;
+
+				case ELiteralArgType::Float:
+
+					return TIsParsable<DataType, float>::Value;
+
+				case ELiteralArgType::String:
+
+					return TIsParsable<DataType, FString>::Value;
+
+				case ELiteralArgType::UObjectProxy:
+
+					return TIsParsable<DataType, Audio::IProxyDataPtr>::Value;
+
+				case ELiteralArgType::UObjectProxyArray:
+
+					return TIsParsable<DataType, TArray<Audio::IProxyDataPtr>>::Value;
+
+				case ELiteralArgType::None:
+
+					return TIsParsable<DataType>::Value;
+
+				default:
+
+					checkNoEntry();
+					return false;
+			}
+		}
+	};
+
+	/** A base factory type for creating objects related to Metasound DataTypes. 
+	 *
+	 * TDataFactory provides a unified interface for constructing objects using 
+	 * a FOperatorSettings and a parameter pack of arguments. The various factory 
+	 * methods determine how strictly the DataType constructors must match the 
+	 * arguments to the factory method.
+	 *
+	 * @tparam DataType - The registered metasound data type which defines the supported constructors.
+	 * @tparam CreatorType - A class providing a static "CreateNew(...)" function which
+	 *                       accepts the forwarded arguments and returns an object.
+	 */
+	template<typename DataType, typename CreatorType>
+	struct TDataFactory
+	{
+        using FInternalFactory = DataFactoryPrivate::TDataFactoryInternal<DataType, CreatorType>;
+
+		/** Create the object using any supported constructor. */
+		template<typename... ArgTypes>
+		static auto CreateAny(const FOperatorSettings& InSettings, ArgTypes&&... Args)
+		{
+			using TForwarding = DataFactoryPrivate::TAnyConstructorForwarding<DataType, ArgTypes...>;
+
+			// CreateAny(...) tries to find a constructor which accepts the following
+			// signatures. This static assert is triggered if none of these constructors
+			// exist. 
+			//
+			// DataType()
+			// DataType(const FOperatorSettings&)
+			// DataType(ArgTypes...)
+			// DataType(const FOperatorSettings&, ArgTypes...);
+			//
+			static_assert(!TForwarding::bCannotForwardToConstructor, "No constructor exists for the DataType which accepts the forwarded arguments");
+
+			return FInternalFactory::template CreateNew<TForwarding>(InSettings, Forward<ArgTypes>(Args)...);
+		}
+
+		/** Create the object using only a constructor which exactly matches the
+		 * arguments to this factory method. 
+		 */
+		template<typename... ArgTypes>
+		static auto CreateExplicit(const FOperatorSettings& InSettings, ArgTypes&&... Args)
+		{
+			using TForwarding = DataFactoryPrivate::TExplicitConstructorForwarding<DataType, ArgTypes...>;
+
+			// CreateExplicit(...) tries to find a constructor which accepts the 
+			// exact same signature as this factory method. This static assert is 
+			// triggered because the constructor does not exist.
+			//
+			// DataType(const FOperatorSettings&, ArgTypes...);
+			//
+			static_assert(!TForwarding::bCannotForwardToConstructor, "No constructor exists for the DataType which accepts the forwarded arguments");
+
+			return FInternalFactory::template CreateNew<TForwarding>(InSettings, Forward<ArgTypes>(Args)...);
+		}
+
+		/** Create the object using only constructors which utilize all arguments
+		 * in the parameter pack (ArgTypes...).
+		 */
+		template<typename... ArgTypes>
+		static auto CreateExplicitArgs(const FOperatorSettings& InSettings, ArgTypes&&... Args)
+		{
+			using TForwarding = DataFactoryPrivate::TExplicitArgsConstructorForwarding<DataType, ArgTypes...>;
+
+			// CreateExplicitArgs(...) tries to find a constructor which accepts the following
+			// signatures. This static assert is triggered if none of these constructors
+			// exist. 
+			//
+			// DataType(ArgTypes...)
+			// DataType(const FOperatorSettings&, ArgTypes...);
+			//
+			static_assert(!TForwarding::bCannotForwardToConstructor, "No constructor exists for the DataType which accepts the forwarded arguments");
+
+			return FInternalFactory::template CreateNew<TForwarding>(InSettings, Forward<ArgTypes>(Args)...);
+		}
+	};
+
+
+	/** TDataTypeFactory creates a DataType.
+	 *
+	 * TDataTypeFactory provides several factory methods for forwarding 
+	 * arguments from the factory method to the DataType constructor.  See TDataFactory
+	 * for more information on the provided factory methods.
+	 *
+	 * @tparam DataType - The Metasound DataType.
+	 */
+	template<typename DataType>
+	struct TDataTypeFactory : TDataFactory<DataType, DataFactoryPrivate::TDataTypeCreator<DataType>>
+	{
+	};
+
+	/** TDataReadReferenceFactory creates TDataReadReferences for the given DataType.
+	 *
+	 * TDataReadReferenceFactory provides several factory methods for forwarding 
+	 * arguments from the factory method to the DataType constructor.  See TDataFactory
+	 * for more information on the provided factory methods.
+	 *
+	 * @tparam DataType - The Metasound DataType of the TDataReadReference.
+	 */
+	template<typename DataType>
+	struct TDataReadReferenceFactory : TDataFactory<DataType, TDataReadReference<DataType>>
+	{
+	};
+
+	/** TDataWriteReferenceFactory creates TDataWriteReferences for the given DataType.
+	 *
+	 * TDataWriteReferenceFactory provides several factory methods for forwarding 
+	 * arguments from the factory method to the DataType constructor.  See TDataFactory
+	 * for more information on the provided factory methods.
+	 *
+	 * @tparam DataType - The Metasound DataType of the TDataWriteReference.
+	 */
+	template<typename DataType>
+	struct TDataWriteReferenceFactory : TDataFactory<DataType, TDataWriteReference<DataType>>
+	{
+	};
+
+	
+	/** A base factory type for creating objects related to Metasound DataTypes. 
+	 *
+	 * TDataLiteralFactory provides a unified interface for constructing objects using 
+	 * a FOperatorSettings and FDataTypeLiteralParam. The various factory 
+	 * methods determine how strictly the DataType constructors must match the 
+	 * arguments to the factory method.
+	 *
+	 * @tparam DataType - The registered metasound data type which defines the supported constructors.
+	 * @tparam CreatorType - A class providing a static "CreateNew(...)" function which
+	 *                       accepts the forwarded arguments and returns an object.
+	 */
+	template<typename DataType, typename CreatorType>
+	struct TDataLiteralFactory
+	{
+		using FInternalFactory = DataFactoryPrivate::TDataVariantFactoryInternal<DataType, CreatorType>;
+
+		/** Create the object using any supported constructor. */
+		template<typename FirstVariantType, typename... AdditionalVariantTypes>
+		static auto CreateAny(const FOperatorSettings& InSettings, const TVariant<FirstVariantType, AdditionalVariantTypes...>& InVariant)
+		{
+			using FDataVariantParsing = DataFactoryPrivate::TDataTypeVariantParsing<DataType, void, FirstVariantType, AdditionalVariantTypes...>;
+
+			return FInternalFactory::template CreateNewFromVariant<FDataVariantParsing>(InSettings, InVariant);
+		}
+
+		/** Create the object using only constructors which utilize the InVariant. */
+		template<typename DesiredArgType, typename FirstVariantType, typename... AdditionalVariantTypes>
+		static auto CreateExplicitArgs(const FOperatorSettings& InSettings, const TVariant<FirstVariantType, AdditionalVariantTypes...>& InVariant)
+		{
+			using FDataVariantParsing = DataFactoryPrivate::TDataTypeVariantParsing<DataType, DesiredArgType, FirstVariantType, AdditionalVariantTypes...>;
+
+			return FInternalFactory::template CreateNewFromVariant<FDataVariantParsing>(InSettings, InVariant);
+		}
+
+		/** Create the object using only constructors which utilize the InLiteral. 
+		 *
+		 * @param InSettings - The FOperatorSettings to be passed to the DataType 
+		 *                     constructor if an appropriate constructor exists.
+		 * @param InLiteral - The literal to be passed to the constructor.
+		 *
+		 * @return An object related to the DataType. The exact type determined by
+		 *         the CreatorType class template parameter.
+		 */
+		static auto CreateExplicitArgs(const FOperatorSettings& InSettings, const FDataTypeLiteralParam& InLiteral)
+		{
+			switch (InLiteral.ConstructorArgType)
+			{
+				case ELiteralArgType::Boolean:
+
+					return CreateExplicitArgs<bool>(InSettings, InLiteral.ConstructorArg);
+
+				case ELiteralArgType::Integer:
+
+					return CreateExplicitArgs<int>(InSettings, InLiteral.ConstructorArg);
+
+				case ELiteralArgType::Float:
+
+					return CreateExplicitArgs<float>(InSettings, InLiteral.ConstructorArg);
+
+				case ELiteralArgType::String:
+
+					return CreateExplicitArgs<FString>(InSettings, InLiteral.ConstructorArg);
+
+				case ELiteralArgType::UObjectProxy:
+
+					return CreateExplicitArgs<Audio::IProxyDataPtr>(InSettings, InLiteral.ConstructorArg);
+
+				case ELiteralArgType::UObjectProxyArray:
+
+					return CreateExplicitArgs<TArray<Audio::IProxyDataPtr>>(InSettings, InLiteral.ConstructorArg);
+
+				case ELiteralArgType::None:
+
+					return CreateExplicitArgs<void>(InSettings, InLiteral.ConstructorArg);
+
+				case ELiteralArgType::Invalid:
+				default:
+
+					checkNoEntry();
+
+					return CreateAny(InSettings, InLiteral.ConstructorArg);
+			}
+		}
+	};
+
+	/** TDataTypeLiteralFactory creates a DataType.
+	 *
+	 * TDataTypeLiteralFactory provides several factory methods for forwarding 
+	 * arguments from the factory method to the DataType constructor.  See TDataFactory
+	 * for more information on the provided factory methods.
+	 *
+	 * @tparam DataType - The Metasound DataType.
+	 */
+	template<typename DataType>
+	struct TDataTypeLiteralFactory : TDataLiteralFactory<DataType, DataFactoryPrivate::TDataTypeCreator<DataType>>
+	{
+	};
+
+	/** TDataReadReferenceLiteralFactory creates TDataReadReferences for the given DataType.
+	 *
+	 * TDataReadReferenceLiteralFactory provides several factory methods for forwarding 
+	 * arguments from the factory method to the DataType constructor.  See TDataFactory
+	 * for more information on the provided factory methods.
+	 *
+	 * @tparam DataType - The Metasound DataType of the TDataReadReference.
+	 */
+	template<typename DataType>
+	struct TDataReadReferenceLiteralFactory : TDataLiteralFactory<DataType, TDataReadReference<DataType>>
+	{
+	};
+
+	/** TDataWriteReferenceLiteralFactory creates TDataWriteReferences for the given DataType.
+	 *
+	 * TDataWriteReferenceLiteralFactory provides several factory methods for forwarding 
+	 * arguments from the factory method to the DataType constructor.  See TDataFactory
+	 * for more information on the provided factory methods.
+	 *
+	 * @tparam DataType - The Metasound DataType of the TDataWriteReference.
+	 */
+	template<typename DataType>
+	struct TDataWriteReferenceLiteralFactory : TDataLiteralFactory<DataType, TDataWriteReference<DataType>>
+	{
+	};
+}
diff --git a/Engine/Plugins/Runtime/Metasound/Source/MetasoundGraphCore/Public/MetasoundDataReference.h b/Engine/Plugins/Runtime/Metasound/Source/MetasoundGraphCore/Public/MetasoundDataReference.h
--- a/Engine/Plugins/Runtime/Metasound/Source/MetasoundGraphCore/Public/MetasoundDataReference.h
+++ b/Engine/Plugins/Runtime/Metasound/Source/MetasoundGraphCore/Public/MetasoundDataReference.h
@@ -15,16 +15,11 @@
 	struct ::Metasound::TDataReferenceTypeInfo<DataType> \
 	{ \
 		static ModuleApi const TCHAR* TypeName; \
-		static constexpr ::Metasound::FMetasoundDataTypeId TypeId = static_cast<::Metasound::FMetasoundDataTypeId>(::Metasound::GetMetasoundDataTypeIdPtr<DataType>()); \
-		static constexpr bool bIsStringParsable = TTestIfDataTypeCtorIsImplemented<DataType, const FString&>::Value; \
-		static constexpr bool bIsBoolParsable = ::Metasound::TTestIfDataTypeCtorIsImplemented<DataType, bool>::Value; \
-		static constexpr bool bIsIntParsable = ::Metasound::TTestIfDataTypeCtorIsImplemented<DataType, int32>::Value; \
-		static constexpr bool bIsFloatParsable = ::Metasound::TTestIfDataTypeCtorIsImplemented<DataType, float>::Value; \
-		static constexpr bool bIsProxyParsable =  ::Metasound::TTestIfDataTypeCtorIsImplemented<DataType, const Audio::IProxyData&>::Value; \
-		static constexpr bool bIsProxyArrayParsable =  ::Metasound::TTestIfDataTypeCtorIsImplemented<DataType, const Audio::IProxyData&>::Value; \
-		static constexpr bool bIsConstructableWithSettings = ::Metasound::TTestIfDataTypeSettingsCtorIsImplemented<DataType>::Value; \
-		static constexpr bool bCanUseDefaultConstructor = std::is_constructible<DataType>::value; \
-		static constexpr bool bIsValidSpecialization = true; \
+		static ModuleApi const FMetasoundDataTypeId TypeId; \
+		\
+		private: \
+		\
+		static const DataType* const TypePtr; \
 	}; \
 	\
 	typedef ::Metasound::TDataReferenceTypeInfo<DataType> DataTypeInfoTypeName; \
@@ -33,8 +28,10 @@
 	typedef ::Metasound::TDataWriteReference<DataType> DataWriteReferenceTypeName; \
 
 // This only needs to be called if you don't plan on calling REGISTER_METASOUND_DATATYPE.
-#define IMPL_METASOUND_DATA_TYPE(DataType, DataTypeName) \
-	const TCHAR* ::Metasound::TDataReferenceTypeInfo<DataType>::TypeName = TEXT(DataTypeName);
+#define DEFINE_METASOUND_DATA_TYPE(DataType, DataTypeName) \
+	const TCHAR*  				::Metasound::TDataReferenceTypeInfo<DataType>::TypeName = TEXT(DataTypeName); \
+	const DataType* const ::Metasound::TDataReferenceTypeInfo<DataType>::TypePtr = nullptr; \
+	const void* const ::Metasound::TDataReferenceTypeInfo<DataType>::TypeId = static_cast<const FMetasoundDataTypeId>(&::Metasound::TDataReferenceTypeInfo<DataType>::TypePtr);
 
 namespace Audio
 {
@@ -42,58 +39,10 @@ namespace Audio
 	class IProxyData;
 }
 
+using FMetasoundDataTypeId = void const * ;
+
 namespace Metasound
 {
-	/** Used to generate a unique ID for each registered data type. */
-	template<typename DataType>
-	struct TMetasoundDataTypeIdPtr
-	{
-		/** static const address of specific type as an ID. */
-		static const DataType* const Ptr;
-	};
-
-	template<typename DataType>
-	const DataType* const TMetasoundDataTypeIdPtr<DataType>::Ptr = nullptr;
-
-	/** Returns an address of a pointer which is defined once per a Metasound
-	 * data type. 
-	 */
-	template <typename DataType>
-	constexpr const DataType* const* GetMetasoundDataTypeIdPtr() noexcept 
-	{
-		return &TMetasoundDataTypeIdPtr<DataType>::Ptr;
-	}
-
-	/** ID for parameter type. Should be unique for different C++ types. */
-	using FMetasoundDataTypeId = const void*;
-
-	/** Test if a data type has a constructor which accepts FOperatorSettings */
-	template <typename TDataType>
-	struct TTestIfDataTypeSettingsCtorIsImplemented
-	{
-	public:
-		static constexpr bool Value = std::is_constructible<TDataType, const ::Metasound::FOperatorSettings&>::value;
-	};
-
-	/** Test if a stat type has a constructor which accepts the given data type 
-	 * and a FOperatorSettings.
-	 */
-	template <typename TDataType, typename TTypeToParse>
-	struct TTestIfDataTypeCtorIsImplemented
-	{
-	private:
-		static constexpr bool bSupportsConstructionWithSettings = 
-			std::is_constructible<TDataType, TTypeToParse, const ::Metasound::FOperatorSettings&>::value
-			|| std::is_constructible<TDataType, TTypeToParse, ::Metasound::FOperatorSettings>::value;
-
-		static constexpr bool bSupportsConstructionWithoutSettings = std::is_constructible<TDataType, TTypeToParse>::value;
-
-	public:
-
-		static constexpr bool Value = bSupportsConstructionWithSettings || bSupportsConstructionWithoutSettings;
-	};
-
-
 	/** Helper class to enforce specialization of TDataReferenceTypeInfo */
 	template<typename DataType>
 	struct TSpecializationHelper 
@@ -107,32 +56,13 @@ namespace Metasound
 	template<typename DataType>
 	struct TDataReferenceTypeInfo
 	{
-		/** Name of the data type. Internally, these are converted to FNames.
-		 * Users should be aware of FName case-insensitivity and other 
-		 * limitiations. 
-		 */
-		static constexpr const TCHAR TypeName[] = TEXT("");
-
-		/** Magic number used to check the data type when casting. */
-		static constexpr const FMetasoundDataTypeId TypeId = nullptr;
-
 		// This static assert is triggered if TDataReferenceTypeInfo is used 
 		// without specialization.
 		static_assert(TSpecializationHelper<DataType>::Value, "TDataReferenceTypeInfo must be specialized.  Use macro DECLARE_METASOUND_DATA_REFERENCE_TYPES");
-	
-		static constexpr bool bIsStringParsable = false;
-		static constexpr bool bIsBoolParsable = false;
-		static constexpr bool bIsIntParsable = false;
-		static constexpr bool bIsFloatParsable = false;
-		static constexpr bool bIsProxyParsable = false;
-		static constexpr bool bIsProxyArrayParsable = false;
-		static constexpr bool bIsConstructableWithSettings = false;
-		static constexpr bool bCanUseDefaultConstructor = false;
-		static constexpr bool bIsValidSpecialization = false;
 	};
 
 
-	/** Return the data type name for a registered data type. */
+	/** Return the data type FName for a registered data type. */
 	template<typename DataType>
 	const FName GetMetasoundDataTypeName() 
 	{
@@ -141,12 +71,21 @@ namespace Metasound
 		return TypeName;
 	}
 
+	/** Return the data type string for a registered data type. */
+	template<typename DataType>
+	const FString GetMetasoundDataTypeString() 
+	{
+		static const FString TypeString = FString(TDataReferenceTypeInfo<std::decay_t<DataType>>::TypeName);
+
+		return TypeString;
+	}
+
 	/** Return the data type ID for a registered data type. 
 	 *
 	 * This ID is runtime constant but may change between executions and builds.
 	 */
 	template<typename DataType>
-	const FMetasoundDataTypeId GetMetasoundDataTypeId() 
+	const void* const GetMetasoundDataTypeId() 
 	{
 		return TDataReferenceTypeInfo<std::decay_t<DataType>>::TypeId;
 	}
@@ -156,7 +95,7 @@ namespace Metasound
 	struct TDataReferenceTypeInfo<void>
 	{
 		static METASOUNDGRAPHCORE_API const TCHAR* TypeName;
-		static constexpr FMetasoundDataTypeId TypeId = static_cast<FMetasoundDataTypeId>(GetMetasoundDataTypeIdPtr<void>());
+		static METASOUNDGRAPHCORE_API const void* const TypeId;
 		static constexpr bool bIsStringParsable = false;
 		static constexpr bool bIsBoolParsable = false;
 		static constexpr bool bIsIntParsable = false;
@@ -165,9 +104,13 @@ namespace Metasound
 		static constexpr bool bIsProxyArrayParsable = false;
 		static constexpr bool bIsConstructableWithSettings = false;
 		static constexpr bool bCanUseDefaultConstructor = false;
-		static constexpr bool bIsValidSpecialization = false;
+
+		private:
+
+		static const void* const TypePtr;
 	};
 
+
 	/** A Data Reference Interface.
 	 *
 	 * A parameter references provides information and access to a shared object in the graph. 
@@ -184,7 +127,7 @@ namespace Metasound
 		virtual const FName& GetDataTypeName() const = 0;
 
 		/** Returns the ID of the parameter type. */
-		virtual FMetasoundDataTypeId GetDataTypeId() const = 0;
+		virtual const void* const GetDataTypeId() const = 0;
 
 		/** Creates a copy of the parameter type. */
 		virtual TUniquePtr<IDataReference> Clone() const = 0;
@@ -199,7 +142,7 @@ namespace Metasound
 	bool IsDataReferenceOfType(const IDataReference& InReference)
 	{
 		static const FName TypeName = GetMetasoundDataTypeName<DataType>();
-		static const FMetasoundDataTypeId TypeId = GetMetasoundDataTypeId<DataType>();
+		static const void* const TypeId = GetMetasoundDataTypeId<DataType>();
 
 		bool bEqualTypeName = InReference.GetDataTypeName() == TypeName;
 		// TODO: need to move data type ID definition to traits implementation.
@@ -270,7 +213,7 @@ namespace Metasound
 			}
 
 			/** Return the ID of the underlying type. */
-			virtual FMetasoundDataTypeId GetDataTypeId() const override
+			virtual const void* const GetDataTypeId() const override
 			{
 				return GetMetasoundDataTypeId<DataType>();
 			}
@@ -401,7 +344,7 @@ namespace Metasound
 			template <typename... ArgTypes>
 			static TDataReadReference<DataType> CreateNew(ArgTypes&&... Args)
 			{
-				static_assert(std::is_constructible<DataType, ArgTypes...>::value, "Tried to call TDataReadReference::CreateNew with args that don't match any constructor for an underlying type!");
+				static_assert(std::is_constructible<DataType, ArgTypes...>::value, "DataType constructor does not support provided types.");
 				return TDataReadReference<DataType>(EDataRefShouldConstruct::NewObject, Forward<ArgTypes>(Args)...);
 			}
 
@@ -449,163 +392,5 @@ namespace Metasound
 				return MakeUnique< FDataReadReference >(*this);
 			}
 	};
-
-	enum class ELiteralArgType : uint8
-	{
-		Boolean,
-		Integer,
-		Float,
-		String,
-		UObjectProxy,
-		UObjectProxyArray,
-		None, // If this is set, we will invoke TType(const FOperatorSettings&) If that constructor exists, or the default constructor if not.
-		Invalid,
-	};
-
-	// This empty base class is used so that we can specialize various nodes (Send, Receive, etc.) for subclasses of IAudioDatatype.
-	class METASOUNDGRAPHCORE_API IAudioDatatype
-	{
-		/**
-		 * Audio datatypes require the following functions:
-		 * int32 GetNumChannels() const { return NumChannels; }
-		 *
-		 * int32 GetMaxNumChannels() const { return MaxNumChannels; }
-
-		 * const TArrayView<const FAudioBufferReadRef> GetBuffers() const { return ReadableBuffers; }
-
-		 * const TArrayView<const FAudioBufferWriteRef> GetBuffers() { return WritableBuffers; }
-		 */
-	};
-
-
-	// TODO: add creation of data type from literal to TDataReferenceFactory by introducing TVariant interface. 
-	// TODO: swap FOperatorSettings and ArgType order in code which uses literals. 
-	
-	/** TDataReferenceFactory
-	 *
-	 * This factory provides a convenient interface for creating new data references. 
-	 * The factory prioritizes the available data reference constructors and chooses
-	 * the highest prioritized constructor which satisfies the provided constructor
-	 * arguments. This allows callers to construct with a consistent interface for 
-	 * varying data reference types. 
-	 *
-	 * The constructors are prioritized as follows
-	 * 1. A constructor which accepts an FOperatorSettings and provided arguments.
-	 * 2. A constructor which accepts the provided arguments.
-	 * 3. A constructor which accepts an FOperatorSettings.
-	 * 4. The default constructor.
-	 *
-	 * If none of these constructors exist for the provided arguments, then the
-	 * factory creation method will not compile. 
-	 */
-	template<typename DataType>
-	class TDataReferenceFactory
-	{
-		private:
-
-			/** TDataReferenceFactoryHelper determines which constructors are available 
-			 * for a given data type and arguments.
-			 */
-			template<typename T, typename... ArgTypes>
-			struct TDataReferenceFactoryHelper
-			{
-				static constexpr bool bIsDefaultConstructible = std::is_constructible<T>::value;
-				static constexpr bool bIsConstructibleWithSettings = std::is_constructible<T, const FOperatorSettings&>::value;
-				static constexpr bool bIsConstructibleWithArgs = std::is_constructible<T, ArgTypes...>::value;
-				static constexpr bool bIsConstructibleWithSettingsAndArgs = std::is_constructible<T, const FOperatorSettings&, ArgTypes...>::value;
-			};
-
-			/** TDataReferenceFactoryHelper specialization which handles an empty parameter pack. */
-			template<typename T>
-			struct TDataReferenceFactoryHelper<T>
-			{
-				static constexpr bool bIsDefaultConstructible = std::is_constructible<T>::value;
-				static constexpr bool bIsConstructibleWithSettings = std::is_constructible<T, const FOperatorSettings&>::value;
-				static constexpr bool bIsConstructibleWithArgs = false;
-				static constexpr bool bIsConstructibleWithSettingsAndArgs = false;
-			};
-
-			/** TDataReferenceFactoryChooser prioritizes the available constructors. */
-			template<typename T, typename... ArgTypes>
-			struct TDataReferenceFactoryChooser
-			{
-				using FHelper = TDataReferenceFactoryHelper<T, ArgTypes...>;
-
-				static constexpr bool bIsConstructibleWithoutArg = FHelper::bIsDefaultConstructible || FHelper::bIsConstructibleWithSettings;
-
-				// Only one of these booleans should be true. These enable the appropriate
-				// "CreateNewReference(...)" function. If multiple are true, than a
-				// duplicate symbol will exist. 
-				static constexpr bool bCreateWithSettingsAndArgs = FHelper::bIsConstructibleWithSettingsAndArgs;
-				static constexpr bool bCreateWithArgs = FHelper::bIsConstructibleWithArgs && !bCreateWithSettingsAndArgs;
-				static constexpr bool bCreateWithSettings = FHelper::bIsConstructibleWithSettings && !bCreateWithArgs;
-				static constexpr bool bCreateWithDefaultConstructor = FHelper::bIsDefaultConstructible && !bCreateWithSettings;
-
-				static constexpr bool bCreateWithAnyConstructor = bCreateWithSettingsAndArgs || bCreateWithArgs || bCreateWithSettings || bCreateWithDefaultConstructor;
-
-				// TODO: 
-				static_assert(bCreateWithAnyConstructor, "No constructor possible");
-			};
-
-			/** Create a reference using FOperatorSettings and arguments. */
-			template<
-				typename ReferenceType, 
-				typename... ArgTypes,
-				typename std::enable_if< TDataReferenceFactoryChooser<DataType, ArgTypes...>::bCreateWithSettingsAndArgs, int>::type = 0
-			>
-			static ReferenceType CreateNewReference(const FOperatorSettings& InSettings, ArgTypes&&... Args)
-			{
-				return ReferenceType::CreateNew(InSettings, Forward<ArgTypes>(Args)...);
-			}
-
-			/** Create a reference using arguments. */
-			template<
-				typename ReferenceType, 
-				typename... ArgTypes,
-				typename std::enable_if< TDataReferenceFactoryChooser<DataType, ArgTypes...>::bCreateWithArgs, int>::type = 0
-			>
-			static ReferenceType CreateNewReference(const FOperatorSettings& InSettings, ArgTypes&&... Args)
-			{
-				return ReferenceType::CreateNew(Forward<ArgTypes>(Args)...);
-			}
-
-			/** Create a reference using FOperatorSettings. */
-			template<
-				typename ReferenceType, 
-				typename... ArgTypes,
-				typename std::enable_if< TDataReferenceFactoryChooser<DataType, ArgTypes...>::bCreateWithSettings, int>::type = 0
-			>
-			static ReferenceType CreateNewReference(const FOperatorSettings& InSettings, ArgTypes&&... Args)
-			{
-				return ReferenceType::CreateNew(InSettings);
-			}
-
-			/** Create a reference using the default constructor */
-			template<
-				typename ReferenceType, 
-				typename... ArgTypes,
-				typename std::enable_if< TDataReferenceFactoryChooser<DataType, ArgTypes...>::bCreateWithDefaultConstructor, int>::type = 0
-			>
-			static ReferenceType CreateNewReference(const FOperatorSettings& InSettings, ArgTypes&&... Args)
-			{
-				return ReferenceType::CreateNew();
-			}
-
-		public:
-
-			/** Create a new TDataWriteReference<DataType> using the given FOperatorSettings and ArgTypes. */
-			template<typename... ArgTypes>
-			static TDataWriteReference<DataType> CreateNewWriteReference(const FOperatorSettings& InSettings, ArgTypes&&... Args)
-			{
-				return CreateNewReference<TDataWriteReference<DataType>>(InSettings, Forward<ArgTypes>(Args)...);
-			}
-
-			/** Create a new TDataReadReference<DataType> using the given FOperatorSettings and ArgTypes. */
-			template<typename... ArgTypes>
-			static TDataReadReference<DataType> CreateNewReadReference(const FOperatorSettings& InSettings, ArgTypes&&... Args)
-			{
-				return CreateNewReference<TDataReadReference<DataType>>(InSettings, Forward<ArgTypes>(Args)...);
-			}
-	};
 }
 
diff --git a/Engine/Plugins/Runtime/Metasound/Source/MetasoundGraphCore/Public/MetasoundLiteral.h b/Engine/Plugins/Runtime/Metasound/Source/MetasoundGraphCore/Public/MetasoundLiteral.h
new file mode 100644
--- /dev/null
+++ b/Engine/Plugins/Runtime/Metasound/Source/MetasoundGraphCore/Public/MetasoundLiteral.h
@@ -0,0 +1,158 @@
+// Copyright Epic Games, Inc. All Rights Reserved.
+
+#pragma once
+
+#include "CoreMinimal.h"
+
+#include "IAudioProxyInitializer.h"
+#include "MetasoundOperatorSettings.h" 
+#include "Misc/TVariant.h"
+
+namespace Metasound
+{
+	/*
+	const TCHAR* ToString(EConstructorArgType InType)
+	{
+		switch (InType)
+		{
+		case Metasound::EConstructorArgType::Boolean:
+			return TEXT("Boolean");
+			break;
+		case Metasound::EConstructorArgType::Integer:
+			return TEXT("Integer");
+			break;
+		case Metasound::EConstructorArgType::Float:
+			return TEXT("Float");
+			break;
+		case Metasound::EConstructorArgType::String:
+			return TEXT("String");
+			break;
+		default:
+			return TEXT("");
+			break;
+		}
+	}
+	*/
+
+	enum class ELiteralArgType : uint8
+	{
+		Boolean,
+		Integer,
+		Float,
+		String,
+		UObjectProxy,
+		UObjectProxyArray,
+		None, // If this is set, we will invoke TType(const FOperatorSettings&) If that constructor exists, or the default constructor if not.
+		Invalid,
+	};
+
+	/**
+	 * Convenience wrapper for safely invoking the correct constructor.
+	 * This can be used for int32, float, bool, or FString.
+	 *
+	 * For example:
+	 * // Somewhere before DECLARE_METASOUND_DATA_REFERENCE_TYPES(FAudioBuffer...), ParseFrom<FAudioBuffer>(int32) is defined
+	 * // which means this data type can be created from an integer.
+	 * Metasound::FAudioBuffer ::Metasound::ParseFrom<Metasound::FAudioBuffer>(int32 InNumFrames, const ::Metasound::FOperatorSettings&)
+	 * {
+	 *     return Metasound::FAudioBuffer(InNumFrames);
+	 * }
+	 *
+	 * //...
+	 * // In the frontend, we know that int and pass it to somewhere in MetasoundGraphCore...
+	 * FDataTypeLiteralParam InitParam(512);
+	 *
+	 * //...
+	 * // In the backend, we can safely construct an FAudioBuffer.
+	 * if(InitParam.IsCompatibleWithType<Metasound::FAudioBuffer>())
+	 * {
+	 *     Metasound::FAudioBuffer AudioBuffer = InitParam.ParseTo<Metasound::FAudioBuffer>();
+	 *     //...
+	 * }
+	 */
+	struct FDataTypeLiteralParam
+	{
+		FDataTypeLiteralParam(const FDataTypeLiteralParam& Other) = delete;
+		FDataTypeLiteralParam(FDataTypeLiteralParam&& Other) = default;
+		FDataTypeLiteralParam& operator=(FDataTypeLiteralParam&& Other) = default;
+
+		typedef TVariant<bool, int, float, FString, Audio::IProxyDataPtr, TArray<Audio::IProxyDataPtr>> FVariantType;
+
+		ELiteralArgType ConstructorArgType;
+		FVariantType ConstructorArg;
+
+		// builds an invalid FDataTypeLiteralParam.
+		static FDataTypeLiteralParam InvalidParam()
+		{
+			FDataTypeLiteralParam InitParam;
+			InitParam.ConstructorArgType = ELiteralArgType::Invalid;
+			return MoveTemp(InitParam);
+		}
+
+		bool IsValid() const
+		{
+			return ConstructorArgType != ELiteralArgType::Invalid;
+		}
+
+		FDataTypeLiteralParam()
+			: ConstructorArgType(ELiteralArgType::None)
+		{}
+
+		FDataTypeLiteralParam(bool bInValue)
+			: ConstructorArgType(ELiteralArgType::Boolean)
+		{
+			ConstructorArg.Set<bool>(bInValue);
+		}
+
+		FDataTypeLiteralParam(int32 InValue)
+			: ConstructorArgType(ELiteralArgType::Integer)
+		{
+			ConstructorArg.Set<int32>(InValue);
+		}
+
+		FDataTypeLiteralParam(float InValue)
+			: ConstructorArgType(ELiteralArgType::Float)
+		{
+			ConstructorArg.Set<float>(InValue);
+		}
+
+		FDataTypeLiteralParam(const FString& InString)
+			: ConstructorArgType(ELiteralArgType::String)
+		{
+			ConstructorArg.Set<FString>(InString);
+		}
+
+		FDataTypeLiteralParam(Audio::IProxyDataPtr&& InProxy)
+			: ConstructorArgType(ELiteralArgType::UObjectProxy)
+		{
+			ConstructorArg.Set<Audio::IProxyDataPtr>(MoveTemp(InProxy));
+		}
+
+		FDataTypeLiteralParam(TArray<Audio::IProxyDataPtr>&& InProxyArray)
+			: ConstructorArgType(ELiteralArgType::UObjectProxyArray)
+		{
+			ConstructorArg.Set<TArray<Audio::IProxyDataPtr>>(MoveTemp(InProxyArray));
+		}
+
+	};
+
+	// This struct is passed to FInputNodeConstructorCallback and FOutputNodeConstructorCallback.
+	struct FInputNodeConstructorParams
+	{
+		// the instance name and name of the specific connection that we should use.
+		const FString& InNodeName;
+		const FString& InVertexName;
+
+		// the operator settings that we will compile with
+		const FOperatorSettings& InSettings;
+
+		FDataTypeLiteralParam InitParam;
+	};
+
+	struct FOutputNodeConstrutorParams
+	{
+		// the instance name and name of the specific connection that we should use.
+		const FString& InNodeName;
+		const FString& InVertexName;
+	};
+}
diff --git a/Engine/Plugins/Runtime/Metasound/Source/MetasoundGraphCore/Public/MetasoundNodeInterface.h b/Engine/Plugins/Runtime/Metasound/Source/MetasoundGraphCore/Public/MetasoundNodeInterface.h
--- a/Engine/Plugins/Runtime/Metasound/Source/MetasoundGraphCore/Public/MetasoundNodeInterface.h
+++ b/Engine/Plugins/Runtime/Metasound/Source/MetasoundGraphCore/Public/MetasoundNodeInterface.h
@@ -4,9 +4,11 @@
 
 #include "CoreMinimal.h"
 #include "Internationalization/Text.h"
+#include "MetasoundDataFactory.h"
 #include "MetasoundDataReference.h"
 #include "MetasoundOperatorInterface.h"
 #include "MetasoundVertex.h"
+#include "MetasoundLiteral.h"
 
 namespace Metasound
 {
@@ -29,8 +31,8 @@ namespace Metasound
 			checkf(ParamMap.Contains(ParamName), TEXT("Tried to use node initialization parameter that didn't exist!"));
 
 			// If this hits, check that you have a constructor for DataType that can interpret the parameter.
-			checkf(ParamMap[ParamName].IsCompatibleWithType<ParamType>(), TEXT("Tried to parse initialization parameter to an invalid type!"));
-			return ParamMap[ParamName].ParseTo<ParamType>();
+			//checkf(ParamMap[ParamName].IsCompatibleWithType<ParamType>(), TEXT("Tried to parse initialization parameter to an invalid type!"));
+			return TDataTypeLiteralFactory<ParamType>::CreateAny(ParamMap[ParamName]);
 		}
 	};
 
diff --git a/Engine/Plugins/Runtime/Metasound/Source/MetasoundGraphCore/Public/MetasoundOperatorInterface.h b/Engine/Plugins/Runtime/Metasound/Source/MetasoundGraphCore/Public/MetasoundOperatorInterface.h
--- a/Engine/Plugins/Runtime/Metasound/Source/MetasoundGraphCore/Public/MetasoundOperatorInterface.h
+++ b/Engine/Plugins/Runtime/Metasound/Source/MetasoundGraphCore/Public/MetasoundOperatorInterface.h
@@ -6,10 +6,6 @@
 
 #include "MetasoundDataReference.h"
 #include "MetasoundDataReferenceCollection.h"
-#include "MetasoundOperatorSettings.h" // < TODO: For literal param.  Should remove when literal param is moved. 
-#include "Misc/TVariant.h" // < TODO: For literal param.  Should remove when literal param is moved. 
-#include "IAudioProxyInitializer.h" // < TODO: For literal param.  Should remove when literal param is moved. 
-
 
 namespace Metasound
 {
@@ -57,435 +53,3 @@ namespace Metasound
 		virtual FExecuteFunction GetExecuteFunction() = 0;
 	};
 }
-
-
-
-namespace Metasound
-{
-	/*
-	class IGraphOperator : public IOperator
-	{
-		public:
-
-			virtual ~IGraphOperator() {}
-	};
-	*/
-
-	/*
-	const TCHAR* ToString(EConstructorArgType InType)
-	{
-		switch (InType)
-		{
-		case Metasound::EConstructorArgType::Boolean:
-			return TEXT("Boolean");
-			break;
-		case Metasound::EConstructorArgType::Integer:
-			return TEXT("Integer");
-			break;
-		case Metasound::EConstructorArgType::Float:
-			return TEXT("Float");
-			break;
-		case Metasound::EConstructorArgType::String:
-			return TEXT("String");
-			break;
-		default:
-			return TEXT("");
-			break;
-		}
-	}
-	*/
-
-	// TODO: move data type literal to own file.
-	/**
-	 * Convenience wrapper for safely invoking the correct constructor.
-	 * This can be used for int32, float, bool, or FString.
-	 *
-	 * For example:
-	 * // Somewhere before DECLARE_METASOUND_DATA_REFERENCE_TYPES(FAudioBuffer...), ParseFrom<FAudioBuffer>(int32) is defined
-	 * // which means this data type can be created from an integer.
-	 * Metasound::FAudioBuffer ::Metasound::ParseFrom<Metasound::FAudioBuffer>(int32 InNumFrames, const ::Metasound::FOperatorSettings&)
-	 * {
-	 *     return Metasound::FAudioBuffer(InNumFrames);
-	 * }
-	 *
-	 * //...
-	 * // In the frontend, we know that int and pass it to somewhere in MetasoundGraphCore...
-	 * FDataTypeLiteralParam InitParam(512);
-	 *
-	 * //...
-	 * // In the backend, we can safely construct an FAudioBuffer.
-	 * if(InitParam.IsCompatibleWithType<Metasound::FAudioBuffer>())
-	 * {
-	 *     Metasound::FAudioBuffer AudioBuffer = InitParam.ParseTo<Metasound::FAudioBuffer>();
-	 *     //...
-	 * }
-	 */
-	struct FDataTypeLiteralParam
-	{
-		FDataTypeLiteralParam(const FDataTypeLiteralParam& Other) = delete;
-		FDataTypeLiteralParam(FDataTypeLiteralParam&& Other) = default;
-		FDataTypeLiteralParam& operator=(FDataTypeLiteralParam&& Other) = default;
-
-		ELiteralArgType ConstructorArgType;
-		TVariant<bool, int, float, FString, Audio::IProxyDataPtr, TArray<Audio::IProxyDataPtr>> ConstructorArg;
-
-		// builds an invalid FDataTypeLiteralParam.
-		static FDataTypeLiteralParam InvalidParam()
-		{
-			FDataTypeLiteralParam InitParam;
-			InitParam.ConstructorArgType = ELiteralArgType::Invalid;
-			return MoveTemp(InitParam);
-		}
-
-		bool IsValid() const
-		{
-			return ConstructorArgType != ELiteralArgType::Invalid;
-		}
-
-		FDataTypeLiteralParam()
-			: ConstructorArgType(ELiteralArgType::None)
-		{}
-
-		FDataTypeLiteralParam(bool bInValue)
-			: ConstructorArgType(ELiteralArgType::Boolean)
-		{
-			ConstructorArg.Set<bool>(bInValue);
-		}
-
-		FDataTypeLiteralParam(int32 InValue)
-			: ConstructorArgType(ELiteralArgType::Integer)
-		{
-			ConstructorArg.Set<int32>(InValue);
-		}
-
-		FDataTypeLiteralParam(float InValue)
-			: ConstructorArgType(ELiteralArgType::Float)
-		{
-			ConstructorArg.Set<float>(InValue);
-		}
-
-		FDataTypeLiteralParam(const FString& InString)
-			: ConstructorArgType(ELiteralArgType::String)
-		{
-			ConstructorArg.Set<FString>(InString);
-		}
-
-		FDataTypeLiteralParam(Audio::IProxyDataPtr&& InProxy)
-			: ConstructorArgType(ELiteralArgType::UObjectProxy)
-		{
-			ConstructorArg.Set<Audio::IProxyDataPtr>(MoveTemp(InProxy));
-		}
-
-		FDataTypeLiteralParam(TArray<Audio::IProxyDataPtr>&& InProxyArray)
-			: ConstructorArgType(ELiteralArgType::UObjectProxyArray)
-		{
-			ConstructorArg.Set<TArray<Audio::IProxyDataPtr>>(MoveTemp(InProxyArray));
-		}
-
-		template<typename DataType>
-		bool IsCompatibleWithType() const
-		{
-			switch (ConstructorArgType)
-			{
-				case ELiteralArgType::Float:
-				{
-					return TTestIfDataTypeCtorIsImplemented<DataType, float>::Value;
-				}
-				case ELiteralArgType::Integer:
-				{
-					return TTestIfDataTypeCtorIsImplemented<DataType, int32>::Value;
-				}
-				case ELiteralArgType::Boolean:
-				{
-					return TTestIfDataTypeCtorIsImplemented<DataType, bool>::Value;
-				}
-				case ELiteralArgType::String:
-				{
-					return TTestIfDataTypeCtorIsImplemented<DataType, const FString&>::Value;
-				}
-				case ELiteralArgType::UObjectProxy:
-				{
-					return TTestIfDataTypeCtorIsImplemented<DataType, const Audio::IProxyData&>::Value;
-				}
-				case ELiteralArgType::UObjectProxyArray:
-				{
-					return TTestIfDataTypeCtorIsImplemented<DataType, const TArray<Audio::IProxyDataPtr>&>::Value;
-				}
-				case ELiteralArgType::None:
-				{
-					return TTestIfDataTypeSettingsCtorIsImplemented<DataType>::Value || std::is_constructible<DataType>::value;
-				}
-				default:
-				{
-					checkNoEntry();
-					return false;
-				}
-			}
-		}
-
-	private:
-		// This is a bucket constructor for when ConstructorArgType is unsupported by the data type we're trying to generate.
-		// This will only be invoked if something has gone terribly wrong in the frontend.
-		template<typename DataType>
-		DataType InvalidConstructor() const
-		{
-			checkNoEntry();
-
-			// We don't know which constructors DataType supports, and we should never enter this code,
-			// so we just dereference garbage memory here.
-			return *reinterpret_cast<DataType*>(0x1);
-		}
-
-		// These are SFINAE-specialized to ensure we safely invoke the correct function,
-		// or invoke InvalidConstructor if the correct function doesn't exist.
-		// For each literal type, there are three SFINAE variants from which we generate a function at compile time:
-		// 1: If there's a constructor that can take the literal type and FOperatorSettings, we use that.
-		// 2: Otherwise, if there's a constructor that just takes the literal type, we'll use that.
-		// 3: If neither constructor exists, we crash. Ideally this has already been caught by the static assert in the REGISTER_METASOUND_DATATYPE macro.
-
-		template<typename DataType, typename TEnableIf<
-			TIsConstructible<DataType, const TArray<Audio::IProxyDataPtr>&, const FOperatorSettings&>::Value
-			, bool >::Type = true >
-			DataType ParseAsAudioProxyArray(const FOperatorSettings& WithOperatorSettings) const
-		{
-			check(ConstructorArgType == ELiteralArgType::UObjectProxyArray);
-			const TArray<Audio::IProxyDataPtr>& ProxyArrayToParse = ConstructorArg.Get<TArray<Audio::IProxyDataPtr>>();
-			return DataType(ProxyArrayToParse, WithOperatorSettings);
-		}
-
-		template<typename DataType, typename TEnableIf<
-			TIsConstructible<DataType, const TArray<Audio::IProxyDataPtr>&>::Value && !TIsConstructible<DataType, const TArray<Audio::IProxyDataPtr>&, const FOperatorSettings&>::Value
-			, bool >::Type = true >
-			DataType ParseAsAudioProxyArray(const FOperatorSettings& WithOperatorSettings) const
-		{
-			check(ConstructorArgType == ELiteralArgType::UObjectProxyArray);
-			const TArray<Audio::IProxyDataPtr>& ProxyArrayToParse = ConstructorArg.Get<TArray<Audio::IProxyDataPtr>>();
-			return DataType(ProxyArrayToParse);
-		}
-
-		template<typename DataType>
-		DataType ParseAsAudioProxyArray(...) const
-		{
-			return InvalidConstructor<DataType>();
-		}
-
-		template<typename DataType, typename TEnableIf<
-			TIsConstructible<DataType, const Audio::IProxyData&, const FOperatorSettings&>::Value
-			, bool >::Type = true >
-		DataType ParseAsAudioProxy(const FOperatorSettings& WithOperatorSettings) const
-		{
-			check(ConstructorArgType == ELiteralArgType::UObjectProxy);
-			const Audio::IProxyDataPtr& ProxyToParse = ConstructorArg.Get<Audio::IProxyDataPtr>();
-			return DataType(ProxyToParse, WithOperatorSettings);
-		}
-
-		template<typename DataType, typename TEnableIf<
-			TIsConstructible<DataType, const Audio::IProxyData&>::Value && !TIsConstructible<DataType, const Audio::IProxyData&, const FOperatorSettings&>::Value
-			, bool >::Type = true >
-		DataType ParseAsAudioProxy(const FOperatorSettings& WithOperatorSettings) const
-		{
-			check(ConstructorArgType == ELiteralArgType::UObjectProxy);
-			const Audio::IProxyDataPtr& ProxyToParse = ConstructorArg.Get<Audio::IProxyDataPtr>();
-			check(ProxyToParse.IsValid());
-			return DataType(*ProxyToParse);
-		}
-
-		template<typename DataType>
-		DataType ParseAsAudioProxy(...) const
-		{
-			return InvalidConstructor<DataType>();
-		}
-
-		template<typename DataType, typename TEnableIf< 
-			TIsConstructible<DataType, const FString&, const FOperatorSettings&>::Value
-			, bool >::Type = true >
-		DataType ParseAsString(const FOperatorSettings& WithOperatorSettings) const
-		{
-			check(ConstructorArgType == ELiteralArgType::String);
-			FString StringToParse = ConstructorArg.Get<FString>();
-			return DataType(StringToParse, WithOperatorSettings);
-		}
-
-		template<typename DataType, typename TEnableIf< 
-			TIsConstructible<DataType, const FString&>::Value && !TIsConstructible<DataType, const FString&, const FOperatorSettings&>::Value
-			, bool >::Type = true >
-		DataType ParseAsString(const FOperatorSettings& WithOperatorSettings) const
-		{
-			check(ConstructorArgType == ELiteralArgType::String);
-			FString StringToParse = ConstructorArg.Get<FString>();
-			return DataType(StringToParse);
-		}
-
-		template<typename DataType>
-		DataType ParseAsString(...) const
-		{
-			return InvalidConstructor<DataType>();
-		}
-
-		template<typename DataType, typename TEnableIf< 
-			TIsConstructible<DataType, float, const FOperatorSettings&>::Value
-			, bool >::Type = true >
-		DataType ParseAsFloat(const FOperatorSettings& WithOperatorSettings) const
-		{
-			check(ConstructorArgType == ELiteralArgType::Float);
-			float ValueToParse = ConstructorArg.Get<float>();
-			return DataType(ValueToParse, WithOperatorSettings);
-		}
-
-		template<typename DataType, typename TEnableIf< 
-			TIsConstructible<DataType, float >::Value && !TIsConstructible<DataType, float, const FOperatorSettings&>::Value
-			, bool >::Type = true >
-		DataType ParseAsFloat(const FOperatorSettings& WithOperatorSettings) const
-		{
-			check(ConstructorArgType == ELiteralArgType::Float);
-			float ValueToParse = ConstructorArg.Get<float>();
-			return DataType(ValueToParse);
-		}
-
-		template<typename DataType>
-		DataType ParseAsFloat(...) const
-		{
-			return InvalidConstructor<DataType>();
-		}
-
-		template<typename DataType, typename TEnableIf< 
-			TIsConstructible<DataType, int32, const FOperatorSettings&>::Value
-			, bool >::Type = true >
-		DataType ParseAsInt(const FOperatorSettings& WithOperatorSettings) const
-		{
-			check(ConstructorArgType == ELiteralArgType::Integer);
-			int32 ValueToParse = ConstructorArg.Get<int32>();
-			return DataType(ValueToParse, WithOperatorSettings);
-		}
-
-		template<typename DataType, typename TEnableIf< 
-			TIsConstructible<DataType, int32>::Value && !TIsConstructible<DataType, int32, const FOperatorSettings&>::Value
-			, bool >::Type = true >
-		DataType ParseAsInt(const FOperatorSettings& WithOperatorSettings) const
-		{
-			check(ConstructorArgType == ELiteralArgType::Integer);
-			int32 ValueToParse = ConstructorArg.Get<int32>();
-			return DataType(ValueToParse);
-		}
-
-		template<typename DataType>
-		DataType ParseAsInt(...) const
-		{
-			return InvalidConstructor<DataType>();
-		}
-
-		template<typename DataType, typename TEnableIf< 
-			TIsConstructible<DataType, bool, const FOperatorSettings&>::Value
-			, bool >::Type = true >
-		DataType ParseAsBool(const FOperatorSettings& WithOperatorSettings) const
-		{
-			check(ConstructorArgType == ELiteralArgType::Boolean);
-			bool ValueToParse = ConstructorArg.Get<bool>();
-			return DataType(ValueToParse, WithOperatorSettings);
-		}
-
-		template<typename DataType, typename TEnableIf< 
-			TIsConstructible<DataType, bool>::Value && !TIsConstructible<DataType, bool, const FOperatorSettings&>::Value
-			, bool >::Type = true >
-		DataType ParseAsBool(const FOperatorSettings& WithOperatorSettings) const
-		{
-			check(ConstructorArgType == ELiteralArgType::Boolean);
-			bool ValueToParse = ConstructorArg.Get<bool>();
-			return DataType(ValueToParse);
-		}
-
-		template<typename DataType>
-		DataType ParseAsBool(...) const
-		{
-			return InvalidConstructor<DataType>();
-		}
-
-
-		template<typename DataType, typename TEnableIf< 
-			TIsConstructible<DataType, const FOperatorSettings&>::Value
-			, bool >::Type = true >
-		DataType ParseWithDefaultConstructor(const FOperatorSettings& WithOperatorSettings) const
-		{
-			check(ConstructorArgType == ELiteralArgType::None);
-			return DataType(WithOperatorSettings);
-		}
-
-		template<typename DataType, typename TEnableIf< 
-			TIsConstructible<DataType>::Value && !TIsConstructible<DataType, const FOperatorSettings&>::Value
-			, bool>::Type = true >
-		DataType ParseWithDefaultConstructor(const FOperatorSettings& WithOperatorSettings) const
-		{
-			check(ConstructorArgType == ELiteralArgType::None);
-			return DataType();
-		}
-
-		template<typename DataType>
-		DataType ParseWithDefaultConstructor(...) const
-		{
-			return InvalidConstructor<DataType>();
-		}
-
-	public:
-
-		// This can be used to invoke the correct constructor for this init param for a data type.
-		template<typename DataType>
-		DataType ParseTo(const FOperatorSettings& WithOperatorSettings) const
-		{
-			switch (ConstructorArgType)
-			{
-				case ELiteralArgType::Float:
-				{
-					return ParseAsFloat<DataType>(WithOperatorSettings);
-				}
-				case ELiteralArgType::Integer:
-				{
-					return ParseAsInt<DataType>(WithOperatorSettings);
-				}
-				case ELiteralArgType::Boolean:
-				{
-					return ParseAsBool<DataType>(WithOperatorSettings);
-				}
-				case ELiteralArgType::String:
-				{
-					return ParseAsString<DataType>(WithOperatorSettings);
-				}
-				case ELiteralArgType::UObjectProxy:
-				{
-					return ParseAsAudioProxy<DataType>(WithOperatorSettings);
-				}
-				case ELiteralArgType::UObjectProxyArray:
-				{
-					return ParseAsAudioProxyArray<DataType>(WithOperatorSettings);
-				}
-				case ELiteralArgType::None:
-				{
-					return ParseWithDefaultConstructor<DataType>(WithOperatorSettings);
-				}
-				default:
-				{
-					return InvalidConstructor<DataType>();
-				}
-			}
-		}
-	};
-
-	// This struct is passed to FInputNodeConstructorCallback and FOutputNodeConstructorCallback.
-	struct FInputNodeConstructorParams
-	{
-		// the instance name and name of the specific connection that we should use.
-		const FString& InNodeName;
-		const FString& InVertexName;
-
-		// the operator settings that we will compile with
-		const FOperatorSettings& InSettings;
-
-		FDataTypeLiteralParam InitParam;
-	};
-
-	struct FOutputNodeConstrutorParams
-	{
-		// the instance name and name of the specific connection that we should use.
-		const FString& InNodeName;
-		const FString& InVertexName;
-	};
-}
diff --git a/Engine/Plugins/Runtime/Metasound/Source/MetasoundStandardNodes/Private/MetasoundAudioFormats.cpp b/Engine/Plugins/Runtime/Metasound/Source/MetasoundStandardNodes/Private/MetasoundAudioFormats.cpp
--- a/Engine/Plugins/Runtime/Metasound/Source/MetasoundStandardNodes/Private/MetasoundAudioFormats.cpp
+++ b/Engine/Plugins/Runtime/Metasound/Source/MetasoundStandardNodes/Private/MetasoundAudioFormats.cpp
@@ -6,15 +6,16 @@
 #include "MetasoundDataReference.h"
 #include "MetasoundAudioBuffer.h"
 
+REGISTER_METASOUND_DATATYPE(Metasound::FUnformattedAudio, "Audio:Unformatted")
+REGISTER_METASOUND_DATATYPE(Metasound::FMonoAudioFormat, "Audio:Mono")
+REGISTER_METASOUND_DATATYPE(Metasound::FStereoAudioFormat, "Audio:Stereo")
+
+// FMultichannelAudio cannot be used as an input type because it's channel count
+// cannot be changed at runtime. Hence, it is not registered.
+DEFINE_METASOUND_DATA_TYPE(Metasound::FMultichannelAudioFormat, "Audio:Multichannel")
+
 namespace Metasound
 {
-	REGISTER_METASOUND_DATATYPE(FUnformattedAudio, "Audio:Unformatted")
-	REGISTER_METASOUND_DATATYPE(FMonoAudioFormat, "Audio:Mono")
-	REGISTER_METASOUND_DATATYPE(FStereoAudioFormat, "Audio:Stereo")
-
-	// FMultichannelAudio cannot be used as an input type because it's channel count
-	// cannot be changed at runtime. Hence, it is not registered.
-	IMPL_METASOUND_DATA_TYPE(FMultichannelAudioFormat, "Audio:Multichannel")
 
 	/* FUnformattedAudio */
 	FUnformattedAudio::FUnformattedAudio(int32 InNumFrames, int32 InNumChannels, int32 InMaxNumChannels)
@@ -37,7 +38,7 @@ namespace Metasound
 		SetNumChannels(InNumChannels);
 	}
 
-	FUnformattedAudio::FUnformattedAudio(int32 InInitialNumChannels, const FOperatorSettings& InSettings)
+	FUnformattedAudio::FUnformattedAudio(const FOperatorSettings& InSettings, int32 InInitialNumChannels)
 		: FUnformattedAudio(InSettings.GetNumFramesPerBlock(), InInitialNumChannels, 8)
 	{
 	}
@@ -79,7 +80,7 @@ namespace Metasound
 		ReadableBuffers = ReadableBufferStorage;
 	}
 
-	FMultichannelAudioFormat::FMultichannelAudioFormat(int32 InNumChannels, const FOperatorSettings& InSettings)
+	FMultichannelAudioFormat::FMultichannelAudioFormat(const FOperatorSettings& InSettings, int32 InNumChannels)
 		: FMultichannelAudioFormat(InSettings.GetNumFramesPerBlock(), InNumChannels)
 	{}
 
diff --git a/Engine/Plugins/Runtime/Metasound/Source/MetasoundStandardNodes/Private/MetasoundBop.cpp b/Engine/Plugins/Runtime/Metasound/Source/MetasoundStandardNodes/Private/MetasoundBop.cpp
--- a/Engine/Plugins/Runtime/Metasound/Source/MetasoundStandardNodes/Private/MetasoundBop.cpp
+++ b/Engine/Plugins/Runtime/Metasound/Source/MetasoundStandardNodes/Private/MetasoundBop.cpp
@@ -8,7 +8,7 @@
 
 namespace Metasound
 {
-	FBop::FBop(bool bShouldBop, const FOperatorSettings& InSettings)
+	FBop::FBop(const FOperatorSettings& InSettings, bool bShouldBop)
 	:	NumFramesPerBlock(InSettings.GetNumFramesPerBlock())
 	,	SampleRate(InSettings.GetSampleRate())
 	{
@@ -18,7 +18,7 @@ namespace Metasound
 		}
 	}
 
-	FBop::FBop(int32 InFrameToBop, const FOperatorSettings& InSettings)
+	FBop::FBop(const FOperatorSettings& InSettings, int32 InFrameToBop)
 	:	NumFramesPerBlock(InSettings.GetNumFramesPerBlock())
 	,	SampleRate(InSettings.GetSampleRate())
 	{
@@ -26,7 +26,7 @@ namespace Metasound
 	}
 
 	FBop::FBop(const FOperatorSettings& InSettings)
-		: FBop(false, InSettings)
+		: FBop(InSettings, false)
 	{}
 
 	void FBop::BopFrame(int32 InFrameToBop)
diff --git a/Engine/Plugins/Runtime/Metasound/Source/MetasoundStandardNodes/Private/MetasoundStandardNodesModule.cpp b/Engine/Plugins/Runtime/Metasound/Source/MetasoundStandardNodes/Private/MetasoundStandardNodesModule.cpp
--- a/Engine/Plugins/Runtime/Metasound/Source/MetasoundStandardNodes/Private/MetasoundStandardNodesModule.cpp
+++ b/Engine/Plugins/Runtime/Metasound/Source/MetasoundStandardNodes/Private/MetasoundStandardNodesModule.cpp
@@ -10,14 +10,11 @@
 #include "MetasoundTime.h"
 #include "MetasoundFrequency.h"
 
-namespace Metasound
-{
-	REGISTER_METASOUND_DATATYPE(FBop, "Primitive:Bop", ::Metasound::ELiteralArgType::Boolean)
-	REGISTER_METASOUND_DATATYPE(FFloatTime, "Primitive:Time", ::Metasound::ELiteralArgType::Float)
-	REGISTER_METASOUND_DATATYPE(FDoubleTime, "Primitive:Time:HighResolution", ::Metasound::ELiteralArgType::Float)
-	REGISTER_METASOUND_DATATYPE(FSampleTime, "Primitive:Time:SampleResolution", ::Metasound::ELiteralArgType::Float)
-	REGISTER_METASOUND_DATATYPE(FFrequency, "Primitive:Frequency", ::Metasound::ELiteralArgType::Float)
-}
+REGISTER_METASOUND_DATATYPE(Metasound::FBop, "Primitive:Bop", ::Metasound::ELiteralArgType::Boolean)
+REGISTER_METASOUND_DATATYPE(Metasound::FFloatTime, "Primitive:Time", ::Metasound::ELiteralArgType::Float)
+REGISTER_METASOUND_DATATYPE(Metasound::FDoubleTime, "Primitive:Time:HighResolution", ::Metasound::ELiteralArgType::Float)
+REGISTER_METASOUND_DATATYPE(Metasound::FSampleTime, "Primitive:Time:SampleResolution", ::Metasound::ELiteralArgType::Float)
+REGISTER_METASOUND_DATATYPE(Metasound::FFrequency, "Primitive:Frequency", ::Metasound::ELiteralArgType::Float)
 
 namespace Metasound 
 {
diff --git a/Engine/Plugins/Runtime/Metasound/Source/MetasoundStandardNodes/Public/MetasoundAudioFormats.h b/Engine/Plugins/Runtime/Metasound/Source/MetasoundStandardNodes/Public/MetasoundAudioFormats.h
--- a/Engine/Plugins/Runtime/Metasound/Source/MetasoundStandardNodes/Public/MetasoundAudioFormats.h
+++ b/Engine/Plugins/Runtime/Metasound/Source/MetasoundStandardNodes/Public/MetasoundAudioFormats.h
@@ -19,7 +19,7 @@ namespace Metasound
 	 * of the FUnformattedAudio. All audio buffers within a FUnformattedAudio object must contain the 
 	 * same number of audio frames.
 	 */
-	class METASOUNDSTANDARDNODES_API FUnformattedAudio : public IAudioDatatype
+	class METASOUNDSTANDARDNODES_API FUnformattedAudio : public IAudioDataType
 	{
 		public:
 			/** FUnformattedAudio Constructor.
@@ -37,10 +37,10 @@ namespace Metasound
 			/**
 			 * FUnformattedAudio Constructor used by the metasound frontend.
 			 *
-			 * @param InNumChannels - initial number of audio channels.
 			 * @param InSettings - Operator Settings passed in on construction.
+			 * @param InNumChannels - initial number of audio channels.
 			 */
-			explicit FUnformattedAudio(int32 InNumChannels, const FOperatorSettings& InSettings);
+			explicit FUnformattedAudio(const FOperatorSettings& InSettings, int32 InNumChannels);
 
 
 			/** Sets the number of audio channels.
@@ -107,7 +107,7 @@ namespace Metasound
 	 * of the FMultichannelAudioFormat. All audio buffers within a FMultichannelAudioFormat object must contain the 
 	 * same number of audio frames.
 	 */
-	class METASOUNDSTANDARDNODES_API FMultichannelAudioFormat : public IAudioDatatype
+	class METASOUNDSTANDARDNODES_API FMultichannelAudioFormat : public IAudioDataType
 	{
 		public:
 			FMultichannelAudioFormat();
@@ -122,10 +122,10 @@ namespace Metasound
 			/**
 			 * FMultichannelAudioFormat Constructor used by the metasound frontend.
 			 *
-			 * @param InNumChannels - initial number of audio channels.
 			 * @param InSettings - Operator Settings passed in on construction.
+			 * @param InNumChannels - initial number of audio channels.
 			 */
-			explicit FMultichannelAudioFormat(int32 InNumChannels, const FOperatorSettings& InSettings);
+			explicit FMultichannelAudioFormat(const FOperatorSettings& InSettings, int32 InNumChannels);
 
 			/** FMultichannelAudioFormat Constructor.
 			 *
@@ -212,7 +212,7 @@ namespace Metasound
 	 * number of audio frames.
 	 */
 	template<int32 NumChannels>
-	class TStaticChannelAudioFormat : public IAudioDatatype
+	class TStaticChannelAudioFormat : public IAudioDataType
 	{
 		public:
 			/** TStaticChannelAudioFormat Constructor
diff --git a/Engine/Plugins/Runtime/Metasound/Source/MetasoundStandardNodes/Public/MetasoundBop.h b/Engine/Plugins/Runtime/Metasound/Source/MetasoundStandardNodes/Public/MetasoundBop.h
--- a/Engine/Plugins/Runtime/Metasound/Source/MetasoundStandardNodes/Public/MetasoundBop.h
+++ b/Engine/Plugins/Runtime/Metasound/Source/MetasoundStandardNodes/Public/MetasoundBop.h
@@ -22,17 +22,17 @@ namespace Metasound
 		public:
 			/** FBop constructor. 
 			 *
-			 * @param bShouldBop - If true, bops first sample.
 			 * @param InSettings - Operator settings.
+			 * @param bShouldBop - If true, bops first sample.
 			 */
-			explicit FBop(bool bShouldBop, const FOperatorSettings& InSettings);
+			explicit FBop(const FOperatorSettings& InSettings, bool bShouldBop);
 
 			/** FBop constructor. 
 			 *
 			 * @param InFrameToBop - Set specific frame to bop.
 			 * @param InSettings - Operator settings.
 			 */
-			explicit FBop(int32 InFrameToBop, const FOperatorSettings& InSettings);
+			explicit FBop(const FOperatorSettings& InSettings, int32 InFrameToBop);
 
 			/** FBop constructor. By default it is not bopped.
 			 *
diff --git a/Engine/Plugins/Runtime/Metasound/Source/MetasoundStandardNodes/Public/MetasoundTime.h b/Engine/Plugins/Runtime/Metasound/Source/MetasoundStandardNodes/Public/MetasoundTime.h
--- a/Engine/Plugins/Runtime/Metasound/Source/MetasoundStandardNodes/Public/MetasoundTime.h
+++ b/Engine/Plugins/Runtime/Metasound/Source/MetasoundStandardNodes/Public/MetasoundTime.h
@@ -299,10 +299,10 @@ namespace Metasound
 			/**
 			 * Constructor used by the Metasound Frontend.
 			 *
-			 * @param InSeconds The number in time to initialize to.
 			 * @param InSettings the operator settings to use.
+			 * @param InSeconds The number in time to initialize to.
 			 */
-			TTime(float InSeconds, const FOperatorSettings& InSettings)
+			TTime(const FOperatorSettings& InSettings, float InSeconds)
 				: TTime(InSeconds, ETimeResolution::Seconds, InSettings.GetSampleRate())
 			{}
 
