Metasound Frontent Query

diff --git a/Engine/Plugins/Runtime/Metasound/Source/MetasoundFrontend/Private/MetasoundFrontendQuery.cpp b/Engine/Plugins/Runtime/Metasound/Source/MetasoundFrontend/Private/MetasoundFrontendQuery.cpp
new file mode 100644
--- /dev/null
+++ b/Engine/Plugins/Runtime/Metasound/Source/MetasoundFrontend/Private/MetasoundFrontendQuery.cpp
@@ -0,0 +1,591 @@
+// Copyright Epic Games, Inc. All Rights Reserved.
+
+#include "MetasoundFrontendQuery.h"
+
+#include "Algo/Sort.h"
+#include "CoreMinimal.h"
+#include "MetasoundFrontendDataLayout.h"
+
+namespace Metasound
+{
+
+	namespace FrontendQueryPrivate
+	{
+		struct FGenerateFunctionFrontendQueryStep: IFrontendQueryGenerateStep
+		{
+			using FGenerateFunction = FFrontendQueryStep::FGenerateFunction;
+
+			FGenerateFunctionFrontendQueryStep(FGenerateFunction InFunc)
+			:	Func(InFunc)
+			{
+			}
+
+			void Generate(TArray<FFrontendQueryEntry>& OutEntries) const override
+			{
+				Func(OutEntries);
+			}
+
+			private:
+
+			FGenerateFunction Func;
+		};
+
+		struct FMapFunctionFrontendQueryStep: IFrontendQueryMapStep
+		{
+			using FMapFunction = FFrontendQueryStep::FMapFunction;
+
+			FMapFunctionFrontendQueryStep(FMapFunction InFunc)
+			:	Func(InFunc)
+			{
+			}
+
+			FFrontendQueryEntry::FKey Map(const FFrontendQueryEntry& InEntry) const override
+			{
+				return Func(InEntry);
+			}
+
+			FMapFunction Func;
+		};
+
+		struct FReduceFunctionFrontendQueryStep: IFrontendQueryReduceStep
+		{
+			using FReduceFunction = FFrontendQueryStep::FReduceFunction;
+
+			FReduceFunctionFrontendQueryStep(FReduceFunction InFunc)
+			:	Func(InFunc)
+			{
+			}
+
+			void Reduce(FFrontendQueryEntry::FKey InKey, TArrayView<FFrontendQueryEntry*>& InEntries, FReduceOutputView& OutResult) const override
+			{
+				return Func(InKey, InEntries, OutResult);
+			}
+
+			FReduceFunction Func;
+		};
+
+		struct FFilterFunctionFrontendQueryStep: IFrontendQueryFilterStep
+		{
+			using FFilterFunction = FFrontendQueryStep::FFilterFunction;
+
+			FFilterFunctionFrontendQueryStep(FFilterFunction InFunc)
+			:	Func(InFunc)
+			{
+			}
+
+			bool Filter(const FFrontendQueryEntry& InEntry) const override
+			{
+				return Func(InEntry);
+			}
+
+			FFilterFunction Func;
+		};
+
+		struct FScoreFunctionFrontendQueryStep: IFrontendQueryScoreStep
+		{
+			using FScoreFunction = FFrontendQueryStep::FScoreFunction;
+
+			FScoreFunctionFrontendQueryStep(FScoreFunction InFunc)
+			:	Func(InFunc)
+			{
+			}
+
+			float Score(const FFrontendQueryEntry& InEntry) const override
+			{
+				return Func(InEntry);
+			}
+
+			FScoreFunction Func;
+		};
+
+		struct FSortFunctionFrontendQueryStep: IFrontendQuerySortStep
+		{
+			using FSortFunction = FFrontendQueryStep::FSortFunction;
+
+			FSortFunctionFrontendQueryStep(FSortFunction InFunc)
+			:	Func(InFunc)
+			{
+			}
+
+			bool Sort(const FFrontendQueryEntry& InEntryLHS, const FFrontendQueryEntry& InEntryRHS) const override
+			{
+				return Func(InEntryLHS, InEntryRHS);
+			}
+
+			FSortFunction Func;
+		};
+
+		struct FLimitFunctionFrontendQueryStep: IFrontendQueryLimitStep
+		{
+			using FLimitFunction = FFrontendQueryStep::FLimitFunction;
+
+			FLimitFunctionFrontendQueryStep(FLimitFunction InFunc)
+			:	Func(InFunc)
+			{
+			}
+
+			int32 Limit() const override
+			{
+				return Func();
+			}
+
+			FLimitFunction Func;
+		};
+
+		template<typename StepType>
+		struct TStepExecuter : public FFrontendQueryStep::IStepExecuter
+		{
+			TStepExecuter(TUniquePtr<StepType>&& InStep)
+			:	Step(MoveTemp(InStep))
+			{
+			}
+
+			TUniquePtr<StepType> Step;
+		};
+
+
+		struct FGenerateStepExecuter : TStepExecuter<IFrontendQueryGenerateStep>
+		{
+			using TStepExecuter<IFrontendQueryGenerateStep>::TStepExecuter;
+
+			void ExecuteStep(FFrontendQuerySelection& InOutResult) const override
+			{
+				if (Step.IsValid())
+				{
+					TArray<FFrontendQueryEntry> NewEntries;
+
+					Step->Generate(NewEntries);
+
+					InOutResult.AppendToStorageAndSelection(NewEntries);
+				}
+			}
+		};
+
+		struct FMapStepExecuter : TStepExecuter<IFrontendQueryMapStep>
+		{
+			using TStepExecuter<IFrontendQueryMapStep>::TStepExecuter;
+
+			void ExecuteStep(FFrontendQuerySelection& InOutResult) const override
+			{
+				if (Step.IsValid())
+				{
+					for (FFrontendQueryEntry* Entry : InOutResult.GetSelection())
+					{
+						if (nullptr != Entry)
+						{
+							Entry->Key = Step->Map(*Entry);
+						}
+					}
+				}
+			}
+		};
+
+		struct FReduceStepExecuter : TStepExecuter<IFrontendQueryReduceStep>
+		{
+			using TStepExecuter<IFrontendQueryReduceStep>::TStepExecuter;
+
+			void ExecuteStep(FFrontendQuerySelection& InOutResult) const override
+			{
+				if (Step.IsValid())
+				{
+					if (InOutResult.GetSelection().Num() > 0)
+					{
+						TArrayView<FFrontendQueryEntry* > SelectionView = InOutResult.GetSelection();
+
+						TArray<FFrontendQueryEntry*> SortedSelection(SelectionView.GetData(), SelectionView.Num()); 
+
+						InOutResult.ResetSelection();
+
+						Algo::SortBy(SortedSelection, [](const FFrontendQueryEntry* InEntry)
+							{
+								if (nullptr == InEntry)
+								{
+									return FFrontendQueryEntry::InvalidKey;
+								}
+
+								return InEntry->Key;
+							}
+						);
+
+
+						int32 Num = SortedSelection.Num();
+						int32 StartIndex = 0;
+
+						FFrontendQueryEntry::FKey CurrentKey = FFrontendQueryEntry::InvalidKey;
+
+						auto ReduceFunc = [&](int32 Index)
+						{
+							FFrontendQueryEntry::FKey ThisKey = nullptr == SortedSelection[Index] ? FFrontendQueryEntry::InvalidKey : SortedSelection[Index]->Key;
+
+							if (CurrentKey != ThisKey)
+							{
+								if (Index > StartIndex)
+								{
+									TArrayView<FFrontendQueryEntry*> ResultsToReduce(&SortedSelection[StartIndex], Index - StartIndex);
+
+									FFrontendQuerySelection::FReduceOutputView ReduceResult(ResultsToReduce);
+
+									Step->Reduce(CurrentKey, ResultsToReduce, ReduceResult);
+
+									InOutResult.AppendToStorageAndSelection(ReduceResult);
+								}
+
+								CurrentKey = ThisKey;
+								StartIndex = Index;
+							}
+						};
+
+						for (int32 ResultIndex = 0; ResultIndex < Num; ResultIndex++)
+						{
+							ReduceFunc(ResultIndex);
+						}
+
+						ReduceFunc(Num - 1);
+					}
+				}
+			}
+		};
+
+		struct FFilterStepExecuter : TStepExecuter<IFrontendQueryFilterStep>
+		{
+			using TStepExecuter<IFrontendQueryFilterStep>::TStepExecuter;
+
+			void ExecuteStep(FFrontendQuerySelection& InOutResult) const override
+			{
+				if (Step.IsValid())
+				{
+					TArrayView<FFrontendQueryEntry*> Selection = InOutResult.GetSelection();
+
+					InOutResult.ResetSelection();
+
+					InOutResult.AppendToSelection(Selection.FilterByPredicate(
+						[&](const FFrontendQueryEntry* Entry)
+						{
+							return (nullptr != Entry) && Step->Filter(*Entry);
+						}
+					));
+				}
+			}
+		};
+
+		struct FScoreStepExecuter : TStepExecuter<IFrontendQueryScoreStep>
+		{
+			using TStepExecuter<IFrontendQueryScoreStep>::TStepExecuter;
+
+			void ExecuteStep(FFrontendQuerySelection& InOutResult) const override
+			{
+				if (Step.IsValid())
+				{
+					for (FFrontendQueryEntry* Entry : InOutResult.GetSelection())
+					{
+						if (nullptr != Entry)
+						{
+							Entry->Score = Step->Score(*Entry);
+						}
+					}
+				}
+			}
+		};
+		
+		struct FSortStepExecuter : TStepExecuter<IFrontendQuerySortStep>
+		{
+			using TStepExecuter<IFrontendQuerySortStep>::TStepExecuter;
+
+			void ExecuteStep(FFrontendQuerySelection& InOutResult) const override
+			{
+				if (Step.IsValid())
+				{
+					InOutResult.GetSelection().Sort(
+						[&](const FFrontendQueryEntry& InLHS, const FFrontendQueryEntry& InRHS)
+						{
+							return Step->Sort(InLHS, InRHS);
+						}
+					);
+				}
+			}
+		};
+
+		struct FLimitStepExecuter : TStepExecuter<IFrontendQueryLimitStep>
+		{
+			using TStepExecuter<IFrontendQueryLimitStep>::TStepExecuter;
+
+			void ExecuteStep(FFrontendQuerySelection& InOutResult) const override
+			{
+				if (Step.IsValid())
+				{
+					int32 Limit = Step->Limit();
+					if (InOutResult.GetSelection().Num() > Limit)
+					{
+						InOutResult.SetSelection(InOutResult.GetSelection().Slice(0, Limit));
+					}
+				}
+			}
+		};
+	}
+
+	FFrontendQuerySelection::FReduceOutputView::FReduceOutputView(TArrayView<FFrontendQueryEntry*> InEntries)
+	:	InitialEntries(InEntries)
+	{
+	}
+
+	void FFrontendQuerySelection::FReduceOutputView::Add(FFrontendQueryEntry& InResult)
+	{
+		if (!InitialEntries.Contains(&InResult))
+		{
+			NewEntryStorage.Add(InResult);
+		}
+		else
+		{
+			ExistingEntryPointers.Add(&InResult);
+		}
+	}
+
+	TArrayView<FFrontendQueryEntry*> FFrontendQuerySelection::FReduceOutputView::GetSelectedInitialEntries()
+	{
+		return ExistingEntryPointers;
+	}
+
+	TArrayView<FFrontendQueryEntry> FFrontendQuerySelection::FReduceOutputView::GetSelectedNewEntries() 
+	{
+		return NewEntryStorage;
+	}
+
+	TArrayView<FFrontendQueryEntry> FFrontendQuerySelection::GetStorage()
+	{
+		return Storage;
+	}
+
+	TArrayView<const FFrontendQueryEntry> FFrontendQuerySelection::GetStorage() const
+	{
+		return Storage;
+	}
+
+	TArrayView<FFrontendQueryEntry*> FFrontendQuerySelection::GetSelection()
+	{
+		return Selection;
+	}
+
+	TArrayView<const FFrontendQueryEntry * const> FFrontendQuerySelection::GetSelection() const
+	{
+		return MakeArrayView(Selection);
+	}
+
+	void FFrontendQuerySelection::AppendToStorageAndSelection(TArrayView<const FFrontendQueryEntry> InEntries)
+	{
+		const int32 Num = InEntries.Num();
+
+		if (Num >  0)
+		{
+			int32 StorageIndex = Storage.Num();
+			int32 SelectionIndex = Selection.Num();
+
+			Storage.Append(InEntries.GetData(), InEntries.Num());
+			Selection.AddZeroed(Num);
+
+			for (int32 i = 0; i < Num; i++)
+			{
+				Selection[SelectionIndex] = &Storage[StorageIndex];
+				SelectionIndex++;
+				StorageIndex++;
+			}
+		}
+	}
+
+	void FFrontendQuerySelection::AppendToStorageAndSelection(FReduceOutputView& InReduceView)
+	{
+		AppendToStorageAndSelection(InReduceView.GetSelectedNewEntries());
+		AppendToSelection(InReduceView.GetSelectedInitialEntries());
+	}
+
+	void FFrontendQuerySelection::ResetStorageAndSelection()
+	{
+		Storage.Reset();
+		Selection.Reset();
+	}
+
+	void FFrontendQuerySelection::ResetSelection()
+	{
+		Selection.Reset();
+	}
+
+	void FFrontendQuerySelection::SetSelection(TArrayView<FFrontendQueryEntry*> InEntries)
+	{
+		Selection = InEntries;
+	}
+
+	void FFrontendQuerySelection::AddToSelection(FFrontendQueryEntry* InEntry)
+	{
+		if (nullptr != InEntry)
+		{
+			Selection.Add(InEntry);
+		}
+	}
+
+	void FFrontendQuerySelection::AppendToSelection(TArrayView<FFrontendQueryEntry*> InEntries)
+	{
+		Selection.Append(InEntries.GetData(), InEntries.Num());
+	}
+
+
+
+	FFrontendQuerySelectionView::FFrontendQuerySelectionView(TUniquePtr<FFrontendQuerySelection>&& InResult)
+	:	Result(MoveTemp(InResult))
+	{
+		if (!Result.IsValid())
+		{
+			Result = MakeUnique<FFrontendQuerySelection>();
+		}
+	}
+
+	TArrayView<const FFrontendQueryEntry> FFrontendQuerySelectionView::GetStorage() const
+	{
+		return Result->GetStorage();
+	}
+
+	TArrayView<const FFrontendQueryEntry* const> FFrontendQuerySelectionView::GetSelection() const
+	{
+		return Result->GetSelection();
+	}
+
+	FFrontendQueryStep::FFrontendQueryStep(FGenerateFunction InFunc)
+	:	StepExecuter(MakeUnique<FrontendQueryPrivate::FGenerateStepExecuter>(MakeUnique<FrontendQueryPrivate::FGenerateFunctionFrontendQueryStep>(InFunc)))
+	{
+	}
+
+	FFrontendQueryStep::FFrontendQueryStep(FMapFunction InFunc)
+	:	StepExecuter(MakeUnique<FrontendQueryPrivate::FMapStepExecuter>(MakeUnique<FrontendQueryPrivate::FMapFunctionFrontendQueryStep>(InFunc)))
+	{
+	}
+
+	FFrontendQueryStep::FFrontendQueryStep(FReduceFunction InFunc)
+	:	StepExecuter(MakeUnique<FrontendQueryPrivate::FReduceStepExecuter>(MakeUnique<FrontendQueryPrivate::FReduceFunctionFrontendQueryStep>(InFunc)))
+	{
+	}
+
+	FFrontendQueryStep::FFrontendQueryStep(FFilterFunction InFunc)
+	:	StepExecuter(MakeUnique<FrontendQueryPrivate::FFilterStepExecuter>(MakeUnique<FrontendQueryPrivate::FFilterFunctionFrontendQueryStep>(InFunc)))
+	{
+	}
+	
+	FFrontendQueryStep::FFrontendQueryStep(FScoreFunction InFunc)
+	:	StepExecuter(MakeUnique<FrontendQueryPrivate::FScoreStepExecuter>(MakeUnique<FrontendQueryPrivate::FScoreFunctionFrontendQueryStep>(InFunc)))
+	{
+	}
+
+	FFrontendQueryStep::FFrontendQueryStep(FSortFunction InFunc)
+	:	StepExecuter(MakeUnique<FrontendQueryPrivate::FSortStepExecuter>(MakeUnique<FrontendQueryPrivate::FSortFunctionFrontendQueryStep>(InFunc)))
+	{
+	}
+
+	FFrontendQueryStep::FFrontendQueryStep(FLimitFunction InFunc)
+	:	StepExecuter(MakeUnique<FrontendQueryPrivate::FLimitStepExecuter>(MakeUnique<FrontendQueryPrivate::FLimitFunctionFrontendQueryStep>(InFunc)))
+	{
+	}
+
+	FFrontendQueryStep::FFrontendQueryStep(TUniquePtr<IFrontendQueryGenerateStep>&& InStep)
+	:	StepExecuter(MakeUnique<FrontendQueryPrivate::FGenerateStepExecuter>(MoveTemp(InStep)))
+	{
+	}
+	
+	FFrontendQueryStep::FFrontendQueryStep(TUniquePtr<IFrontendQueryMapStep>&& InStep)
+	:	StepExecuter(MakeUnique<FrontendQueryPrivate::FMapStepExecuter>(MoveTemp(InStep)))
+	{
+	}
+
+	FFrontendQueryStep::FFrontendQueryStep(TUniquePtr<IFrontendQueryReduceStep>&& InStep)
+	:	StepExecuter(MakeUnique<FrontendQueryPrivate::FReduceStepExecuter>(MoveTemp(InStep)))
+	{
+	}
+
+	FFrontendQueryStep::FFrontendQueryStep(TUniquePtr<IFrontendQueryFilterStep>&& InStep)
+	:	StepExecuter(MakeUnique<FrontendQueryPrivate::FFilterStepExecuter>(MoveTemp(InStep)))
+	{
+	}
+
+	FFrontendQueryStep::FFrontendQueryStep(TUniquePtr<IFrontendQueryScoreStep>&& InStep)
+	:	StepExecuter(MakeUnique<FrontendQueryPrivate::FScoreStepExecuter>(MoveTemp(InStep)))
+	{
+	}
+
+	FFrontendQueryStep::FFrontendQueryStep(TUniquePtr<IFrontendQuerySortStep>&& InStep)
+	:	StepExecuter(MakeUnique<FrontendQueryPrivate::FSortStepExecuter>(MoveTemp(InStep)))
+	{
+	}
+
+	FFrontendQueryStep::FFrontendQueryStep(TUniquePtr<IFrontendQueryLimitStep>&& InStep)
+	:	StepExecuter(MakeUnique<FrontendQueryPrivate::FLimitStepExecuter>(MoveTemp(InStep)))
+	{
+	}
+
+	void FFrontendQueryStep::ExecuteStep(FFrontendQuerySelection& InOutResult)
+	{
+		if (StepExecuter.IsValid())
+		{
+			StepExecuter->ExecuteStep(InOutResult);
+		}
+	}
+
+	const TArray<TUniquePtr<FFrontendQueryStep>>& FFrontendQuery::GetSteps() const
+	{
+		return Steps;
+	}
+
+	FFrontendQuery& FFrontendQuery::AddGenerateLambdaStep(FGenerateFunction InFunc)
+	{
+		return AddFunctionStep(InFunc);
+	}
+
+	FFrontendQuery& FFrontendQuery::AddMapLambdaStep(FMapFunction InFunc)
+	{
+		return AddFunctionStep(InFunc);
+	}
+
+	FFrontendQuery& FFrontendQuery::AddReduceLambdaStep(FReduceFunction InFunc)
+	{
+		return AddFunctionStep(InFunc);
+	}
+
+	FFrontendQuery& FFrontendQuery::AddFilterLambdaStep(FFilterFunction InFunc)
+	{
+		return AddFunctionStep(InFunc);
+	}
+
+	FFrontendQuery& FFrontendQuery::AddScoreLambdaStep(FScoreFunction InFunc)
+	{
+		return AddFunctionStep(InFunc);
+	}
+
+	FFrontendQuery& FFrontendQuery::AddSortLambdaStep(FSortFunction InFunc)
+	{
+		return AddFunctionStep(InFunc);
+	}
+
+	FFrontendQuery& FFrontendQuery::AddLimitLambdaStep(FLimitFunction InFunc)
+	{
+		return AddFunctionStep(InFunc);
+	}
+
+	FFrontendQuery& FFrontendQuery::AddStep(TUniquePtr<FFrontendQueryStep>&& InStep)
+	{
+		Steps.Add(MoveTemp(InStep));
+		return *this;
+	}
+
+	FFrontendQuerySelectionView FFrontendQuery::ExecuteQuery()
+	{
+		Result = MakeUnique<FFrontendQuerySelection>();
+
+		for (int32 StepIndex = 0; StepIndex < Steps.Num(); StepIndex++)
+		{
+			if (!Steps[StepIndex].IsValid())
+			{
+				continue;
+			}
+
+			Steps[StepIndex]->ExecuteStep(*Result);
+		}
+
+		return FFrontendQuerySelectionView(MoveTemp(Result));
+	}
+}
+
diff --git a/Engine/Plugins/Runtime/Metasound/Source/MetasoundFrontend/Private/MetasoundFrontendQuerySteps.cpp b/Engine/Plugins/Runtime/Metasound/Source/MetasoundFrontend/Private/MetasoundFrontendQuerySteps.cpp
new file mode 100644
--- /dev/null
+++ b/Engine/Plugins/Runtime/Metasound/Source/MetasoundFrontend/Private/MetasoundFrontendQuerySteps.cpp
@@ -0,0 +1,54 @@
+// Copyright Epic Games, Inc. All Rights Reserved.
+
+#include "MetasoundFrontendQuerySteps.h"
+
+#include "CoreMinimal.h"
+#include "MetasoundFrontend.h"
+#include "MetasoundFrontendDataLayout.h"
+
+namespace Metasound
+{
+	void FGenerateAllAvailableNodeClasses::Generate(TArray<FFrontendQueryEntry>& OutEntries) const
+	{
+		const TArray<Frontend::FNodeClassInfo> ClassInfos = Frontend::GetAllAvailableNodeClasses();
+
+		for (const Frontend::FNodeClassInfo& ClassInfo : ClassInfos)
+		{
+			OutEntries.Emplace(TInPlaceType<FMetasoundClassDescription>(), Frontend::GenerateClassDescriptionForNode(ClassInfo));
+		}
+	}
+
+	FFilterClassesByInputVertexDataType::FFilterClassesByInputVertexDataType(const FName& InTypeName)
+	:	InputVertexTypeName(InTypeName)
+	{
+	}
+
+	bool FFilterClassesByInputVertexDataType::Filter(const FFrontendQueryEntry& InEntry) const
+	{
+		check(InEntry.Value.IsType<FMetasoundClassDescription>());
+
+		return InEntry.Value.Get<FMetasoundClassDescription>().Inputs.ContainsByPredicate(
+			[this](const FMetasoundInputDescription& InDesc)
+			{
+				return InDesc.TypeName == InputVertexTypeName;
+			}
+		);
+	}
+
+	FFilterClassesByOutputVertexDataType::FFilterClassesByOutputVertexDataType(const FName& InTypeName)
+	:	OutputVertexTypeName(InTypeName)
+	{
+	}
+
+	bool FFilterClassesByOutputVertexDataType::Filter(const FFrontendQueryEntry& InEntry) const
+	{
+		return InEntry.Value.Get<FMetasoundClassDescription>().Outputs.ContainsByPredicate(
+			[this](const FMetasoundOutputDescription& InDesc)
+			{
+				return InDesc.TypeName == OutputVertexTypeName;
+			}
+		);
+	}
+
+}
+
diff --git a/Engine/Plugins/Runtime/Metasound/Source/MetasoundFrontend/Public/MetasoundFrontendQuery.h b/Engine/Plugins/Runtime/Metasound/Source/MetasoundFrontend/Public/MetasoundFrontendQuery.h
new file mode 100644
--- /dev/null
+++ b/Engine/Plugins/Runtime/Metasound/Source/MetasoundFrontend/Public/MetasoundFrontendQuery.h
@@ -0,0 +1,302 @@
+// Copyright Epic Games, Inc. All Rights Reserved.
+
+#pragma once
+
+#include "CoreMinimal.h"
+#include "MetasoundFrontendDataLayout.h"
+#include "Misc/Variant.h"
+
+namespace Metasound
+{
+	/** FFrontendQueryEntry represents one value in the query. It contains a key,
+	 * value and score. 
+	 */
+	struct FFrontendQueryEntry
+	{
+		using FValue = TVariant<FMetasoundDocument, FMetasoundClassDescription, FMetasoundNodeDescription, FMetasoundGraphDescription>;
+		using FKey = int32;
+
+		static constexpr FKey InvalidKey = 0xFFFFFFFF;
+
+		FKey Key = InvalidKey;
+		FValue Value;
+		float Score = 0.f;
+
+		template<typename... ArgTypes>
+		FFrontendQueryEntry(ArgTypes&&... Args)
+		:	Value(Forward<ArgTypes>(Args)...)
+		{
+		}
+	};
+
+	/** FFrontendQuerySelection represents a selection utilized during a query.
+	 *
+	 * Note: FFrontendQuerySelection owns the memory for the related entries. 
+	 */
+	class METASOUNDFRONTEND_API FFrontendQuerySelection
+	{
+	public:
+		/** FReduceOutputView provides an interface for storing results during
+		 * a "Reduce" query step.  It is primarily used to only expose operations
+		 * which are acceptable to call on a FFrontendQuerySelection during a 
+		 * call to "Reduce".
+		 */
+		class FReduceOutputView
+		{
+			// deleted constructors
+			FReduceOutputView(const FReduceOutputView&) = delete;
+			FReduceOutputView(FReduceOutputView&&) = delete;
+
+			// deleted operators
+			FReduceOutputView& operator=(const FReduceOutputView&) = delete;
+			FReduceOutputView& operator=(FReduceOutputView&&) = delete;
+
+		public:
+			// Set initial entries which are stored elsewhere
+			FReduceOutputView(TArrayView<FFrontendQueryEntry*> InEntries);
+			~FReduceOutputView() = default;
+
+			/** Add a result to the selection */
+			void Add(FFrontendQueryEntry& InResult);
+
+			/** Get selection which was also in initial entries. */
+			TArrayView<FFrontendQueryEntry*> GetSelectedInitialEntries();
+
+			/** Get selection which was not in the initial entries. */
+			TArrayView<FFrontendQueryEntry> GetSelectedNewEntries();
+
+		private:
+
+			TArray<FFrontendQueryEntry*> InitialEntries;
+
+			TArray<FFrontendQueryEntry*> ExistingEntryPointers;
+			TArray<FFrontendQueryEntry> NewEntryStorage;
+		};
+
+		/** Get an array of all entries stored in this selection. */
+		TArrayView<FFrontendQueryEntry> GetStorage();
+
+		/** Get an array of all entries stored in this selection. */
+		TArrayView<const FFrontendQueryEntry> GetStorage() const;
+
+		/** Get an array of all entries selected. */
+		TArrayView<FFrontendQueryEntry*> GetSelection();
+
+		/** Get an array of all entries selected. */
+		TArrayView<const FFrontendQueryEntry * const> GetSelection() const;
+
+		/** Add entries to both storage and current selection. */
+		void AppendToStorageAndSelection(TArrayView<const FFrontendQueryEntry> InEntries);
+
+		/** Add entries to both storage and current selection. */
+		void AppendToStorageAndSelection(FReduceOutputView& InReduceView);
+
+		/** Reset both storage and selection. */
+		void ResetStorageAndSelection();
+
+		/** Reset selection, but keep storage. */
+		void ResetSelection();
+
+		/** Set selection. All entries must already be in storage. */
+		void SetSelection(TArrayView<FFrontendQueryEntry*> InEntries);
+
+		/** Add to selection. The entry must already be in storage. */
+		void AddToSelection(FFrontendQueryEntry* InEntry);
+
+		/** Append entries to selection. Each selection must already be in storage. */
+		void AppendToSelection(TArrayView<FFrontendQueryEntry*> InEntries);
+
+	private:
+		TArray<FFrontendQueryEntry*> Selection;
+		TArray<FFrontendQueryEntry> Storage;
+	};
+
+	/** FFrontendQuerySelectionView holds the result of query. */
+	class METASOUNDFRONTEND_API FFrontendQuerySelectionView
+	{
+		FFrontendQuerySelectionView() = delete;
+	public:
+
+		/** The FFrontendQuerySelectionView is constructed with an existing
+		 * FFrontendQuerySelection.
+		 */
+		FFrontendQuerySelectionView(TUniquePtr<FFrontendQuerySelection>&& InResult);
+
+		/** Get all the stored results. */
+		TArrayView<const FFrontendQueryEntry> GetStorage() const;
+
+		/** Get the selected results. */
+		TArrayView<const FFrontendQueryEntry* const> GetSelection() const;
+
+	private:
+		TUniquePtr<FFrontendQuerySelection> Result;
+	};
+
+	/** Interface for an individual step in a query */
+	class IFrontendQueryStep
+	{
+		public:
+			virtual ~IFrontendQueryStep() = default;
+	};
+
+	/** Interface for a query step which generates entries. */
+	class IFrontendQueryGenerateStep : public IFrontendQueryStep
+	{
+		public:
+			virtual ~IFrontendQueryGenerateStep() = default;
+			virtual void Generate(TArray<FFrontendQueryEntry>& OutEntries) const = 0;
+	};
+
+	/** Interface for a query step which maps entries to keys. */
+	class IFrontendQueryMapStep : public IFrontendQueryStep
+	{
+		public:
+			virtual ~IFrontendQueryMapStep() = default;
+			virtual FFrontendQueryEntry::FKey Map(const FFrontendQueryEntry& InEntry) const = 0;
+	};
+
+	/** Interface for a query step which reduces entries with the same key. */
+	class IFrontendQueryReduceStep : public IFrontendQueryStep
+	{
+		public:
+			using FReduceOutputView = FFrontendQuerySelection::FReduceOutputView;
+
+			virtual ~IFrontendQueryReduceStep() = default;
+			virtual void Reduce(FFrontendQueryEntry::FKey InKey, TArrayView<FFrontendQueryEntry*>& InEntries, FReduceOutputView& OutResult) const = 0;
+	};
+
+	/** Interface for a query step which filters entries. */
+	class IFrontendQueryFilterStep : public IFrontendQueryStep
+	{
+		public:
+			virtual ~IFrontendQueryFilterStep() = default;
+			virtual bool Filter(const FFrontendQueryEntry& InEntry) const = 0;
+	};
+
+	/** Interface for a query step which scores entries. */
+	class IFrontendQueryScoreStep : public IFrontendQueryStep
+	{
+		public:
+			virtual ~IFrontendQueryScoreStep() = default;
+			virtual float Score(const FFrontendQueryEntry& InEntry) const = 0;
+	};
+
+	/** Interface for a query step which sorts entries. */
+	class IFrontendQuerySortStep : public IFrontendQueryStep
+	{
+		public:
+			virtual ~IFrontendQuerySortStep() = default;
+			virtual bool Sort(const FFrontendQueryEntry& InEntryLHS, const FFrontendQueryEntry& InEntryRHS) const = 0;
+	};
+
+	class IFrontendQueryLimitStep : public IFrontendQueryStep
+	{
+		public:
+			virtual ~IFrontendQueryLimitStep() = default;
+			virtual int32 Limit() const = 0;
+	};
+
+	/** FFrontendQueryStep wraps all the support IFrontenQueryStep interfaces 
+	 * and supplies unified `ExecuteStep(...)` member function.
+	 */
+	class METASOUNDFRONTEND_API FFrontendQueryStep
+	{
+		FFrontendQueryStep() = delete;
+	public:
+
+		using FReduceOutputView = FFrontendQuerySelection::FReduceOutputView;
+
+		using FGenerateFunction = TFunction<void (TArray<FFrontendQueryEntry>&)>;
+		using FMapFunction = TFunction<FFrontendQueryEntry::FKey (const FFrontendQueryEntry&)>;
+		using FReduceFunction = TFunction<void (FFrontendQueryEntry::FKey, TArrayView<FFrontendQueryEntry*>, FReduceOutputView& )>;
+		using FFilterFunction = TFunction<bool (const FFrontendQueryEntry&)>;
+		using FScoreFunction = TFunction<float (const FFrontendQueryEntry&)>;
+		using FSortFunction = TFunction<bool (const FFrontendQueryEntry& InEntryLHS, const FFrontendQueryEntry& InEntryRHS)>;
+		using FLimitFunction = TFunction<int32 ()>;
+
+		/** Create query step using TFunction or lambda. */
+		FFrontendQueryStep(FGenerateFunction InFunc);
+		FFrontendQueryStep(FMapFunction InFunc);
+		FFrontendQueryStep(FReduceFunction InFunc);
+		FFrontendQueryStep(FFilterFunction InFilt);
+		FFrontendQueryStep(FScoreFunction InScore);
+		FFrontendQueryStep(FSortFunction InSort);
+		FFrontendQueryStep(FLimitFunction InLimit);
+
+		/** Create a query step using a IFrotnedQueryStep */
+		FFrontendQueryStep(TUniquePtr<IFrontendQueryGenerateStep>&& InStep);
+		FFrontendQueryStep(TUniquePtr<IFrontendQueryMapStep>&& InStep);
+		FFrontendQueryStep(TUniquePtr<IFrontendQueryReduceStep>&& InStep);
+		FFrontendQueryStep(TUniquePtr<IFrontendQueryFilterStep>&& InStep);
+		FFrontendQueryStep(TUniquePtr<IFrontendQueryScoreStep>&& InStep);
+		FFrontendQueryStep(TUniquePtr<IFrontendQuerySortStep>&& InStep);
+		FFrontendQueryStep(TUniquePtr<IFrontendQueryLimitStep>&& InStep);
+
+		/** Execute a step on a selection. */
+		void ExecuteStep(FFrontendQuerySelection& InOutResult);
+
+		struct IStepExecuter
+		{
+			virtual ~IStepExecuter() = default;
+			virtual void ExecuteStep(FFrontendQuerySelection& InOutResult) const = 0;
+		};
+
+	private:
+
+		TUniquePtr<IStepExecuter> StepExecuter;
+	};
+
+	/** FFrontendQuery contains a set of query steps which can be executed 
+	 * to produce a FFrontendQuerySelectionView
+	 */
+	class METASOUNDFRONTEND_API FFrontendQuery
+	{
+	public:
+
+		using FGenerateFunction = FFrontendQueryStep::FGenerateFunction;
+		using FMapFunction = FFrontendQueryStep::FMapFunction;
+		using FReduceFunction = FFrontendQueryStep::FReduceFunction;
+		using FFilterFunction = FFrontendQueryStep::FFilterFunction;
+		using FScoreFunction = FFrontendQueryStep::FScoreFunction;
+		using FSortFunction = FFrontendQueryStep::FSortFunction;
+		using FLimitFunction = FFrontendQueryStep::FLimitFunction;
+
+		FFrontendQuery() = default;
+
+		/** Return all steps in a query. */
+		const TArray<TUniquePtr<FFrontendQueryStep>>& GetSteps() const;
+
+		/** Add a step to the query. */
+		template<typename StepType, typename... ArgTypes>
+		FFrontendQuery& AddStep(ArgTypes&&... Args)
+		{
+			return AddStep(MakeUnique<FFrontendQueryStep>(MakeUnique<StepType>(Forward<ArgTypes>(Args)...)));
+		}
+
+		template<typename FuncType>
+		FFrontendQuery& AddFunctionStep(TFunction<FuncType> InFunc)
+		{
+			return AddStep(MakeUnique<FFrontendQueryStep>(InFunc));
+		}
+
+		FFrontendQuery& AddGenerateLambdaStep(FGenerateFunction InFunc);
+		FFrontendQuery& AddMapLambdaStep(FMapFunction InFunc);
+		FFrontendQuery& AddReduceLambdaStep(FReduceFunction InFunc);
+		FFrontendQuery& AddFilterLambdaStep(FFilterFunction InFunc);
+		FFrontendQuery& AddScoreLambdaStep(FScoreFunction InFunc);
+		FFrontendQuery& AddSortLambdaStep(FSortFunction InFunc);
+		FFrontendQuery& AddLimitLambdaStep(FLimitFunction InFunc);
+
+		/** Add a step to the query. */
+		FFrontendQuery& AddStep(TUniquePtr<FFrontendQueryStep>&& InStep);
+
+		/** Calls all steps in the query and returns the selection. */
+		FFrontendQuerySelectionView ExecuteQuery();
+
+	private:
+
+		TUniquePtr<FFrontendQuerySelection> Result;
+
+		TArray<TUniquePtr<FFrontendQueryStep>> Steps;
+	};
+}
diff --git a/Engine/Plugins/Runtime/Metasound/Source/MetasoundFrontend/Public/MetasoundFrontendQuerySteps.h b/Engine/Plugins/Runtime/Metasound/Source/MetasoundFrontend/Public/MetasoundFrontendQuerySteps.h
new file mode 100644
--- /dev/null
+++ b/Engine/Plugins/Runtime/Metasound/Source/MetasoundFrontend/Public/MetasoundFrontendQuerySteps.h
@@ -0,0 +1,53 @@
+// Copyright Epic Games, Inc. All Rights Reserved.
+
+#pragma once
+
+#include "CoreMinimal.h"
+#include "MetasoundDataReference.h"
+#include "MetasoundFrontend.h"
+#include "MetasoundFrontendQuery.h"
+
+namespace Metasound
+{
+	class METASOUNDFRONTEND_API FGenerateAllAvailableNodeClasses : public IFrontendQueryGenerateStep
+	{
+	public:
+
+		void Generate(TArray<FFrontendQueryEntry>& OutEntries) const override;
+	};
+
+	class METASOUNDFRONTEND_API FFilterClassesByInputVertexDataType : public IFrontendQueryFilterStep
+	{
+	public: 
+		template<typename DataType>
+		FFilterClassesByInputVertexDataType()
+		:	FFilterClassesByInputVertexDataType(GetMetasoundDataTypeName<DataType>())
+		{
+		}
+
+		FFilterClassesByInputVertexDataType(const FName& InTypeName);
+
+		bool Filter(const FFrontendQueryEntry& InEntry) const override;
+
+	private:
+		FName InputVertexTypeName;
+	};
+
+	class METASOUNDFRONTEND_API FFilterClassesByOutputVertexDataType : public IFrontendQueryFilterStep
+	{
+	public: 
+		template<typename DataType>
+		FFilterClassesByOutputVertexDataType()
+		:	FFilterClassesByOutputVertexDataType(GetMetasoundDataTypeName<DataType>())
+		{
+		}
+
+		FFilterClassesByOutputVertexDataType(const FName& InTypeName);
+
+		bool Filter(const FFrontendQueryEntry& InEntry) const override;
+
+	private:
+		FName OutputVertexTypeName;
+	};
+}
+
