Add default parsable to metasounds

diff --git a/Engine/Plugins/Runtime/Metasound/Source/MetasoundFrontend/Private/MetasoundFrontendRegistries.cpp b/Engine/Plugins/Runtime/Metasound/Source/MetasoundFrontend/Private/MetasoundFrontendRegistries.cpp
--- a/Engine/Plugins/Runtime/Metasound/Source/MetasoundFrontend/Private/MetasoundFrontendRegistries.cpp
+++ b/Engine/Plugins/Runtime/Metasound/Source/MetasoundFrontend/Private/MetasoundFrontendRegistries.cpp
@@ -238,7 +238,7 @@ Metasound::ELiteralArgType FMetasoundFrontendRegistryContainer::GetDesiredLitera
 	{
 		return Metasound::ELiteralArgType::Boolean;
 	}
-	else if (DataTypeInfo.Info.bIsConstructableWithSettings || DataTypeInfo.Info.bIsDefaultConstructible)
+	else if (DataTypeInfo.Info.bIsDefaultParsable)
 	{
 		return Metasound::ELiteralArgType::None;
 	}
@@ -302,7 +302,7 @@ bool FMetasoundFrontendRegistryContainer::DoesDataTypeSupportLiteralType(FName I
 		}
 		case Metasound::ELiteralArgType::None:
 		{
-			return DataTypeInfo.Info.bIsConstructableWithSettings || DataTypeInfo.Info.bIsDefaultConstructible;
+			return DataTypeInfo.Info.bIsDefaultParsable;
 		}
 		case Metasound::ELiteralArgType::Invalid:
 		default:
diff --git a/Engine/Plugins/Runtime/Metasound/Source/MetasoundFrontend/Public/MetasoundDataTypeRegistrationMacro.h b/Engine/Plugins/Runtime/Metasound/Source/MetasoundFrontend/Public/MetasoundDataTypeRegistrationMacro.h
--- a/Engine/Plugins/Runtime/Metasound/Source/MetasoundFrontend/Public/MetasoundDataTypeRegistrationMacro.h
+++ b/Engine/Plugins/Runtime/Metasound/Source/MetasoundFrontend/Public/MetasoundDataTypeRegistrationMacro.h
@@ -174,24 +174,23 @@ bool RegisterDataTypeWithFrontend()
 	::Metasound::FDataTypeConstructorCallbacks Callbacks = { MoveTemp(InputNodeConstructor), MoveTemp(OutputNodeConstructor), MoveTemp(ProxyGenerator) };
 
 
-	::Metasound::FDataTypeRegistryInfo RegistryInfo;
+	Metasound::FDataTypeRegistryInfo RegistryInfo;
 
 	RegistryInfo.DataTypeName = DataTypeName;
 	RegistryInfo.PreferredLiteralType = PreferredArgType;
 
-	RegistryInfo.bIsBoolParsable = ::Metasound::TIsParsable<TDataType, bool>::Value;
-	RegistryInfo.bIsIntParsable = ::Metasound::TIsParsable<TDataType, int32>::Value;
-	RegistryInfo.bIsFloatParsable = ::Metasound::TIsParsable<TDataType, float>::Value;
-	RegistryInfo.bIsStringParsable = ::Metasound::TIsParsable<TDataType, FString>::Value;
-	RegistryInfo.bIsProxyParsable = ::Metasound::TIsParsable<TDataType, const Audio::IProxyDataPtr&>::Value;
-	RegistryInfo.bIsProxyArrayParsable = ::Metasound::TIsParsable<TDataType, const TArray<Audio::IProxyDataPtr>& >::Value;
-	RegistryInfo.bIsConstructableWithSettings = std::is_constructible<TDataType, const ::Metasound::FOperatorSettings&>::value;
-	RegistryInfo.bIsDefaultConstructible = std::is_constructible<TDataType>::value;
+	RegistryInfo.bIsBoolParsable = Metasound::TIsParsable<TDataType, bool>::Value;
+	RegistryInfo.bIsIntParsable = Metasound::TIsParsable<TDataType, int32>::Value;
+	RegistryInfo.bIsFloatParsable = Metasound::TIsParsable<TDataType, float>::Value;
+	RegistryInfo.bIsStringParsable = Metasound::TIsParsable<TDataType, FString>::Value;
+	RegistryInfo.bIsProxyParsable = Metasound::TIsParsable<TDataType, const Audio::IProxyDataPtr&>::Value;
+	RegistryInfo.bIsProxyArrayParsable = Metasound::TIsParsable<TDataType, const TArray<Audio::IProxyDataPtr>& >::Value;
+	RegistryInfo.bIsDefaultParsable = Metasound::TIsParsable<TDataType>::Value;
 	
 	RegistryInfo.ProxyGeneratorClass = UClassToUse::StaticClass();
 
 	bool bSucceeded = FMetasoundFrontendRegistryContainer::Get()->RegisterDataType(RegistryInfo, MoveTemp(Callbacks));
-	ensureAlwaysMsgf(bSucceeded, TEXT("Failed to register data type %s in the node registry!"), *::Metasound::GetMetasoundDataTypeString<TDataType>());
+	ensureAlwaysMsgf(bSucceeded, TEXT("Failed to register data type %s in the node registry!"), *Metasound::GetMetasoundDataTypeString<TDataType>());
 	
 	RegisterConverterNodes<TDataType>();
 	AttemptToRegisterSendAndReceiveNodes<TDataType>();
diff --git a/Engine/Plugins/Runtime/Metasound/Source/MetasoundFrontend/Public/MetasoundFrontendRegistries.h b/Engine/Plugins/Runtime/Metasound/Source/MetasoundFrontend/Public/MetasoundFrontendRegistries.h
--- a/Engine/Plugins/Runtime/Metasound/Source/MetasoundFrontend/Public/MetasoundFrontendRegistries.h
+++ b/Engine/Plugins/Runtime/Metasound/Source/MetasoundFrontend/Public/MetasoundFrontendRegistries.h
@@ -25,38 +25,25 @@ namespace Metasound
 		FName DataTypeName;
 
 		// What type we should default to using for literals.
-		ELiteralArgType PreferredLiteralType;
+		ELiteralArgType PreferredLiteralType = ELiteralArgType::Invalid;
+
+		// This indicates the type can only be constructed with FOperatorSettings or the default constructor.
+		bool bIsDefaultParsable = false;
 
 		// These bools signify what basic
 		// UProperty primitives we can use to describe this data type as a literal in a document.
-		bool bIsBoolParsable;
-		bool bIsIntParsable;
-		bool bIsFloatParsable;
-		bool bIsStringParsable;
+		bool bIsBoolParsable = false;
+		bool bIsIntParsable = false;
+		bool bIsFloatParsable = false;
+		bool bIsStringParsable = false;
 
 		// these are used for using UObjects, or arrays of UObjects.
-		bool bIsProxyParsable;
-		bool bIsProxyArrayParsable;
+		bool bIsProxyParsable = false;
+		bool bIsProxyArrayParsable = false;
 
 		// If this datatype was registered with a specific UClass to use to filter with, that will be used here:
-		UClass* ProxyGeneratorClass;
-
-		// This indicates the type can only be constructed with FOperatorSettings and no other args.
-		// TODO: these can be consolidated to a single bool, since FDataTypeLiteralParam automatically falls back to not using the settings.
-		bool bIsConstructableWithSettings;
-		bool bIsDefaultConstructible;
-
-		FDataTypeRegistryInfo()
-			: bIsBoolParsable(false)
-			, bIsIntParsable(false)
-			, bIsFloatParsable(false)
-			, bIsStringParsable(false)
-			, bIsProxyParsable(false)
-			, bIsProxyArrayParsable(false)
-			, ProxyGeneratorClass(nullptr)
-			, bIsConstructableWithSettings(false)
-			, bIsDefaultConstructible(false)
-		{}
+		UClass* ProxyGeneratorClass = nullptr;
+
 	};
 
 	namespace Frontend
