Metasound Envorinment Variable

diff --git a/Engine/Plugins/Runtime/Metasound/Source/MetasoundEngine/Private/MetasoundEngineEnvironment.cpp b/Engine/Plugins/Runtime/Metasound/Source/MetasoundEngine/Private/MetasoundEngineEnvironment.cpp
new file mode 100644
--- /dev/null
+++ b/Engine/Plugins/Runtime/Metasound/Source/MetasoundEngine/Private/MetasoundEngineEnvironment.cpp
@@ -0,0 +1,9 @@
+// Copyright Epic Games, Inc. All Rights Reserved.
+
+#include "MetasoundEngineEnvironment.h"
+
+#include "AudioDeviceManager.h"
+#include "MetasoundEnvironment.h"
+
+DEFINE_METASOUND_ENVIRONMENT_VARIABLE_TYPE(FAudioDeviceHandle);
+
diff --git a/Engine/Plugins/Runtime/Metasound/Source/MetasoundEngine/Private/MetasoundSource.cpp b/Engine/Plugins/Runtime/Metasound/Source/MetasoundEngine/Private/MetasoundSource.cpp
--- a/Engine/Plugins/Runtime/Metasound/Source/MetasoundEngine/Private/MetasoundSource.cpp
+++ b/Engine/Plugins/Runtime/Metasound/Source/MetasoundEngine/Private/MetasoundSource.cpp
@@ -8,6 +8,7 @@
 
 #include "MetasoundAssetBase.h"
 #include "MetasoundBop.h"
+#include "MetasoundEngineEnvironment.h"
 #include "MetasoundGenerator.h"
 #include "MetasoundAudioFormats.h"
 #include "MetasoundPrimitives.h"
@@ -60,32 +61,50 @@ float UMetasoundSource::GetDuration()
 
 ISoundGeneratorPtr UMetasoundSource::CreateSoundGenerator(const FSoundGeneratorInitParams& InParams)
 {
+	using namespace Metasound;
+
 	NumChannels = 1;
 	Duration = INDEFINITELY_LOOPING_DURATION;
 	bLooping = true;
 	SampleRate = InParams.SampleRate;
 	const float BlockRate = 100.f; // Metasound graph gets evaluated 100 times per second.
 
-	Metasound::FOperatorSettings InSettings(InParams.SampleRate, BlockRate);
+	FOperatorSettings InSettings(InParams.SampleRate, BlockRate);
+	FMetasoundEnvironment Environment;
 
-	TArray<Metasound::IOperatorBuilder::FBuildErrorPtr> BuildErrors;
+	// Add audio device ID to environment.
+	FAudioDeviceHandle DeviceHandle;
+	if (UWorld* World = GetWorld())
+	{
+		DeviceHandle = World->GetAudioDevice();
+	}
+	if (!DeviceHandle.IsValid())
+	{
+		if (FAudioDeviceManager* DeviceManager = FAudioDeviceManager::Get())
+		{
+			DeviceHandle = DeviceManager->GetMainAudioDeviceHandle();
+		}
+	}
+	Environment.SetValue<FAudioDeviceHandle>(GetAudioDeviceHandleVariableName(), DeviceHandle);
 
-	Metasound::Frontend::FGraphHandle RootGraph = GetRootGraphHandle();
+	TArray<IOperatorBuilder::FBuildErrorPtr> BuildErrors;
+
+	Frontend::FGraphHandle RootGraph = GetRootGraphHandle();
 	ensureAlways(RootGraph.IsValid());
 
-	TUniquePtr<Metasound::IOperator> Operator = RootGraph.BuildOperator(InSettings, BuildErrors);
+	TUniquePtr<IOperator> Operator = RootGraph.BuildOperator(InSettings, Environment, BuildErrors);
 	if (ensureAlways(Operator.IsValid()))
 	{
-		Metasound::FDataReferenceCollection Outputs = Operator->GetOutputs();
+		FDataReferenceCollection Outputs = Operator->GetOutputs();
 
-		Metasound::FMetasoundGeneratorInitParams InitParams =
+		FMetasoundGeneratorInitParams InitParams =
 		{
 			MoveTemp(Operator),
-			Outputs.GetDataReadReferenceOrConstruct<Metasound::FAudioBuffer>(GetAudioOutputName(), InSettings.GetNumFramesPerBlock()),
-			Outputs.GetDataReadReferenceOrConstruct<Metasound::FBop>(GetIsFinishedOutputName(), InSettings, false)
+			Outputs.GetDataReadReferenceOrConstruct<FAudioBuffer>(GetAudioOutputName(), InSettings.GetNumFramesPerBlock()),
+			Outputs.GetDataReadReferenceOrConstruct<FBop>(GetIsFinishedOutputName(), InSettings, false)
 		};
 
-		return ISoundGeneratorPtr(new Metasound::FMetasoundGenerator(MoveTemp(InitParams)));
+		return ISoundGeneratorPtr(new FMetasoundGenerator(MoveTemp(InitParams)));
 	}
 	else
 	{
@@ -118,6 +137,12 @@ const FString& UMetasoundSource::GetIsFinishedOutputName()
 	return OnFinishedOutputName;
 }
 
+const FString& UMetasoundSource::GetAudioDeviceHandleVariableName()
+{
+	static FString AudioDeviceHandleVarName = FString(TEXT("AudioDeviceHandle"));
+	return AudioDeviceHandleVarName;
+}
+
 FMetasoundArchetype UMetasoundSource::GetArchetype() const
 {
 	FMetasoundArchetype Archetype;
@@ -145,6 +170,12 @@ FMetasoundArchetype UMetasoundSource::GetArchetype() const
 	OnFinished.ToolTip = LOCTEXT("OnFinishedToolTip", "Bop executed to initiate stopping the source.");
 	Archetype.RequiredOutputs.Add(OnFinished);
 
+	FMetasoundEnvironmentVariableDescription AudioDeviceHandle;
+	AudioDeviceHandle.Name = GetAudioDeviceHandleVariableName();
+	AudioDeviceHandle.DisplayName = FText::FromString(AudioDeviceHandle.Name);
+	AudioDeviceHandle.ToolTip = LOCTEXT("AudioDeviceHandleToolTip", "Audio device handle");
+	Archetype.EnvironmentVariables.Add(AudioDeviceHandle);
+
 	return Archetype;
 }
 
diff --git a/Engine/Plugins/Runtime/Metasound/Source/MetasoundEngine/Public/MetasoundEngineEnvironment.h b/Engine/Plugins/Runtime/Metasound/Source/MetasoundEngine/Public/MetasoundEngineEnvironment.h
new file mode 100644
--- /dev/null
+++ b/Engine/Plugins/Runtime/Metasound/Source/MetasoundEngine/Public/MetasoundEngineEnvironment.h
@@ -0,0 +1,9 @@
+// Copyright Epic Games, Inc. All Rights Reserved.
+
+#pragma once
+
+#include "AudioDeviceManager.h"
+#include "MetasoundEnvironment.h"
+
+DECLARE_METASOUND_ENVIRONMENT_VARIABLE_TYPE(FAudioDeviceHandle);
+
diff --git a/Engine/Plugins/Runtime/Metasound/Source/MetasoundEngine/Public/MetasoundSource.h b/Engine/Plugins/Runtime/Metasound/Source/MetasoundEngine/Public/MetasoundSource.h
--- a/Engine/Plugins/Runtime/Metasound/Source/MetasoundEngine/Public/MetasoundSource.h
+++ b/Engine/Plugins/Runtime/Metasound/Source/MetasoundEngine/Public/MetasoundSource.h
@@ -114,4 +114,5 @@ private:
 	static const FString& GetOnPlayInputName();
 	static const FString& GetAudioOutputName();
 	static const FString& GetIsFinishedOutputName();
-};
\ No newline at end of file
+	static const FString& GetAudioDeviceHandleVariableName();
+};
diff --git a/Engine/Plugins/Runtime/Metasound/Source/MetasoundFrontend/Private/MetasoundFrontend.cpp b/Engine/Plugins/Runtime/Metasound/Source/MetasoundFrontend/Private/MetasoundFrontend.cpp
--- a/Engine/Plugins/Runtime/Metasound/Source/MetasoundFrontend/Private/MetasoundFrontend.cpp
+++ b/Engine/Plugins/Runtime/Metasound/Source/MetasoundFrontend/Private/MetasoundFrontend.cpp
@@ -163,9 +163,6 @@ namespace Metasound
 				}
 
 				const FName NodeName = DummyNode->GetClassName();
-				const FVertexInterface& DummyNodeInterface = DummyNode->GetDefaultVertexInterface();
-				const FInputVertexInterface& DummyInputInterface = DummyNodeInterface.GetInputInterface();
-				const FOutputVertexInterface& DummyOutputInterface = DummyNodeInterface.GetOutputInterface();
 
 				FMetasoundClassMetadata NodeMetadata;
 				NodeMetadata.NodeName = NodeName.ToString();
@@ -177,9 +174,11 @@ namespace Metasound
 				FMetasoundClassDescription ClassDescription;
 				ClassDescription.Metadata = NodeMetadata;
 
+				const FVertexInterface& DummyNodeInterface = DummyNode->GetDefaultVertexInterface();
+
 				// External metasounds aren't dependent on any other nodes by definition, so all we need to do
 				// is populate the Input and Output sets.
-				for (auto& InputTuple : DummyInputInterface)
+				for (auto& InputTuple : DummyNodeInterface.GetInputInterface())
 				{
 					FMetasoundInputDescription InputDescription;
 					InputDescription.Name = InputTuple.Value.GetVertexName();
@@ -189,7 +188,7 @@ namespace Metasound
 					ClassDescription.Inputs.Add(InputDescription);
 				}
 
-				for (auto& OutputTuple : DummyOutputInterface)
+				for (auto& OutputTuple : DummyNodeInterface.GetOutputInterface())
 				{
 					FMetasoundOutputDescription OutputDescription;
 					OutputDescription.Name = OutputTuple.Value.GetVertexName();
@@ -199,6 +198,15 @@ namespace Metasound
 					ClassDescription.Outputs.Add(OutputDescription);
 				}
 
+				for (auto& EnvTuple : DummyNodeInterface.GetEnvironmentInterface())
+				{
+					FMetasoundEnvironmentVariableDescription EnvironmentDescription;
+					EnvironmentDescription.Name = EnvTuple.Value.GetVertexName();
+					EnvironmentDescription.ToolTip = EnvTuple.Value.GetDescription();
+
+					ClassDescription.EnvironmentVariables.Add(EnvironmentDescription);
+				}
+
 				// Populate lookup data.
 				ClassDescription.ExternalNodeClassLookupInfo.ExternalNodeClassName = InInfo.LookupKey.NodeName;
 				ClassDescription.ExternalNodeClassLookupInfo.ExternalNodeClassHash = InInfo.LookupKey.NodeHash;
@@ -2327,7 +2335,7 @@ namespace Metasound
 			return TTuple<FGraphHandle, FNodeHandle>(SubgraphHandle, SubgraphNode);
 		}
 
-		TUniquePtr<IOperator> FGraphHandle::BuildOperator(const FOperatorSettings& InSettings, TArray<IOperatorBuilder::FBuildErrorPtr>& OutBuildErrors) const
+		TUniquePtr<IOperator> FGraphHandle::BuildOperator(const FOperatorSettings& InSettings, const FMetasoundEnvironment& InEnvironment, TArray<IOperatorBuilder::FBuildErrorPtr>& OutBuildErrors) const
 		{
 			if (!IsValid())
 			{
@@ -2621,9 +2629,9 @@ namespace Metasound
 			/** NOTE: In the future- we should split steps 1-4 above, and 5 below. GraphToBuild and the node map can be cached on the graph handle. */
 
 			// Step 5: Invoke Operator Builder
-			FOperatorBuilder Builder(InSettings, FOperatorBuilderSettings::GetDefaultSettings());
+			FOperatorBuilder Builder(FOperatorBuilderSettings::GetDefaultSettings());
 
-			return Builder.BuildGraphOperator(GraphToBuild, OutBuildErrors);
+			return Builder.BuildGraphOperator(GraphToBuild, InSettings, InEnvironment, OutBuildErrors);
 		}
 
 		const FText& FGraphHandle::GetInputDisplayName(FString InputName) const
diff --git a/Engine/Plugins/Runtime/Metasound/Source/MetasoundFrontend/Public/MetasoundFrontend.h b/Engine/Plugins/Runtime/Metasound/Source/MetasoundFrontend/Public/MetasoundFrontend.h
--- a/Engine/Plugins/Runtime/Metasound/Source/MetasoundFrontend/Public/MetasoundFrontend.h
+++ b/Engine/Plugins/Runtime/Metasound/Source/MetasoundFrontend/Public/MetasoundFrontend.h
@@ -460,7 +460,7 @@ namespace Metasound
 			// metasound operator, which can then be used for playback.
 			// @returns nullptr on failure.
 			// @todo: make an API for doing this async.
-			TUniquePtr<IOperator> BuildOperator(const FOperatorSettings& InSettings, TArray<IOperatorBuilder::FBuildErrorPtr>& OutBuildErrors) const;
+			TUniquePtr<IOperator> BuildOperator(const FOperatorSettings& InSettings, const FMetasoundEnvironment& InEnvironment, TArray<IOperatorBuilder::FBuildErrorPtr>& OutBuildErrors) const;
 
 			// This function will ensure that the document's root class has the inputs and outputs required by the archetype.
 			void FixDocumentToMatchArchetype();
diff --git a/Engine/Plugins/Runtime/Metasound/Source/MetasoundFrontend/Public/MetasoundFrontendDataLayout.h b/Engine/Plugins/Runtime/Metasound/Source/MetasoundFrontend/Public/MetasoundFrontendDataLayout.h
--- a/Engine/Plugins/Runtime/Metasound/Source/MetasoundFrontend/Public/MetasoundFrontendDataLayout.h
+++ b/Engine/Plugins/Runtime/Metasound/Source/MetasoundFrontend/Public/MetasoundFrontendDataLayout.h
@@ -143,6 +143,24 @@ struct FMetasoundOutputDescription
 	FText ToolTip;
 };
 
+USTRUCT() 
+struct FMetasoundEnvironmentVariableDescription
+{
+	GENERATED_BODY()
+
+	// The name of this environment variable.
+	UPROPERTY(VisibleAnywhere, Category = CustomView)
+	FString Name;
+
+	// The name of this input displayed to users.
+	UPROPERTY(EditAnywhere, Category = Parameters, meta = (DisplayName = "Name"))
+	FText DisplayName;
+
+	// Optional description text about this input.
+	UPROPERTY(EditAnywhere, Category = Parameters)
+	FText ToolTip;
+};
+
 // This is used to describe a single incoming connection or literal for a given input on a FMetasoundNodeDescription.
 USTRUCT()
 struct FMetasoundNodeConnectionDescription
@@ -256,6 +274,9 @@ struct FMetasoundClassDescription
 	UPROPERTY(EditAnywhere, Category = CustomView)
 	TArray<FMetasoundOutputDescription> Outputs;
 
+	UPROPERTY(EditAnywhere, Category = CustomView)
+	TArray<FMetasoundEnvironmentVariableDescription> EnvironmentVariables;
+
 	// If this object is itself of the Metasound type, here we list the unique IDs of other objects it depends on to be built.
 	// These unique IDs are then found in FMetasoundDocument::Dependencies.
 	UPROPERTY()
@@ -292,6 +313,10 @@ struct FMetasoundArchetype
 	// call Special Outputs
 	UPROPERTY()
 	TArray<FMetasoundOutputDescription> RequiredOutputs;
+
+	// The environment variables supplied by the archetype.
+	UPROPERTY()
+	TArray<FMetasoundEnvironmentVariableDescription> EnvironmentVariables;
 };
 
 // Outermost description of a single metasound.
@@ -315,4 +340,4 @@ struct FMetasoundDocument
 
 	UPROPERTY()
 	FMetasoundEditorData EditorData;
-};
\ No newline at end of file
+};
diff --git a/Engine/Plugins/Runtime/Metasound/Source/MetasoundGenerator/Private/SynthComponentMetasoundExample.cpp b/Engine/Plugins/Runtime/Metasound/Source/MetasoundGenerator/Private/SynthComponentMetasoundExample.cpp
--- a/Engine/Plugins/Runtime/Metasound/Source/MetasoundGenerator/Private/SynthComponentMetasoundExample.cpp
+++ b/Engine/Plugins/Runtime/Metasound/Source/MetasoundGenerator/Private/SynthComponentMetasoundExample.cpp
@@ -2,10 +2,11 @@
 
 #include "SynthComponentMetasoundExample.h"
 
-#include "MetasoundOperatorInterface.h"
-#include "MetasoundNodeInterface.h"
 #include "MetasoundAudioFormats.h"
+#include "MetasoundEnvironment.h"
 #include "MetasoundGenerator.h"
+#include "MetasoundNodeInterface.h"
+#include "MetasoundOperatorInterface.h"
 
 #include "MetasoundOscNode.h"
 #include "MetasoundAudioMultiplyNode.h"
@@ -219,11 +220,12 @@ ISoundGeneratorPtr USynthComponentMetasoundExample::CreateSoundGenerator(const F
 
 	TUniquePtr<FSynthComponentMetasoundExampleGraph> ExampleGraph = CreateSynthComponentMetasoundExampleGraph();
 
-	FOperatorSettings Settings(InParams.SampleRate, InParams.NumFramesPerCallback);
-	FOperatorBuilder Builder(Settings, FOperatorBuilderSettings::GetDefaultSettings());
+	FOperatorBuilder Builder(FOperatorBuilderSettings::GetDefaultSettings());
 
+	FOperatorSettings Settings(InParams.SampleRate, InParams.NumFramesPerCallback);
+	FMetasoundEnvironment Environment;
 	TArray<IOperatorBuilder::FBuildErrorPtr> Errors;
-	FOperatorUniquePtr Operator = Builder.BuildGraphOperator(*ExampleGraph->Graph, Errors);
+	FOperatorUniquePtr Operator = Builder.BuildGraphOperator(*ExampleGraph->Graph, Settings, Environment, Errors);
 
 	SetGraphOperator(MoveTemp(Operator), Settings, TEXT("Audio"));
 
diff --git a/Engine/Plugins/Runtime/Metasound/Source/MetasoundGraphCore/Private/MetasoundEnvironment.cpp b/Engine/Plugins/Runtime/Metasound/Source/MetasoundGraphCore/Private/MetasoundEnvironment.cpp
new file mode 100644
--- /dev/null
+++ b/Engine/Plugins/Runtime/Metasound/Source/MetasoundGraphCore/Private/MetasoundEnvironment.cpp
@@ -0,0 +1,42 @@
+// Copyright Epic Games, Inc. All Rights Reserved.
+
+#include "MetasoundEnvironment.h"
+
+
+DEFINE_METASOUND_ENVIRONMENT_VARIABLE_TYPE(void);
+DEFINE_METASOUND_ENVIRONMENT_VARIABLE_TYPE(bool);
+DEFINE_METASOUND_ENVIRONMENT_VARIABLE_TYPE(int8);
+DEFINE_METASOUND_ENVIRONMENT_VARIABLE_TYPE(uint8);
+DEFINE_METASOUND_ENVIRONMENT_VARIABLE_TYPE(int16);
+DEFINE_METASOUND_ENVIRONMENT_VARIABLE_TYPE(uint16);
+DEFINE_METASOUND_ENVIRONMENT_VARIABLE_TYPE(int32);
+DEFINE_METASOUND_ENVIRONMENT_VARIABLE_TYPE(uint32);
+DEFINE_METASOUND_ENVIRONMENT_VARIABLE_TYPE(int64);
+DEFINE_METASOUND_ENVIRONMENT_VARIABLE_TYPE(uint64);
+DEFINE_METASOUND_ENVIRONMENT_VARIABLE_TYPE(float);
+DEFINE_METASOUND_ENVIRONMENT_VARIABLE_TYPE(double);
+DEFINE_METASOUND_ENVIRONMENT_VARIABLE_TYPE(FString);
+DEFINE_METASOUND_ENVIRONMENT_VARIABLE_TYPE(FName);
+
+namespace Metasound
+{
+	FMetasoundEnvironment::FMetasoundEnvironment(const FMetasoundEnvironment& InOther)
+	{
+		*this = InOther;
+	}
+
+	FMetasoundEnvironment& FMetasoundEnvironment::operator=(const FMetasoundEnvironment& InOther)
+	{
+		Variables.Reset();
+
+		for (auto& VarTuple : InOther.Variables)
+		{
+			if (ensure(VarTuple.Value.IsValid()))
+			{
+				Variables.Add(VarTuple.Key, VarTuple.Value->Clone());
+			}
+		}
+
+		return *this;
+	}
+}
diff --git a/Engine/Plugins/Runtime/Metasound/Source/MetasoundGraphCore/Private/MetasoundOperatorBuilder.cpp b/Engine/Plugins/Runtime/Metasound/Source/MetasoundGraphCore/Private/MetasoundOperatorBuilder.cpp
--- a/Engine/Plugins/Runtime/Metasound/Source/MetasoundGraphCore/Private/MetasoundOperatorBuilder.cpp
+++ b/Engine/Plugins/Runtime/Metasound/Source/MetasoundGraphCore/Private/MetasoundOperatorBuilder.cpp
@@ -110,9 +110,8 @@ namespace Metasound
 	}
 
 
-	FOperatorBuilder::FOperatorBuilder(const FOperatorSettings& InSettings, const FOperatorBuilderSettings& InBuilderSettings)
-	:	OperatorSettings(InSettings)
-	,	BuilderSettings(InBuilderSettings)
+	FOperatorBuilder::FOperatorBuilder(const FOperatorBuilderSettings& InBuilderSettings)
+	:	BuilderSettings(InBuilderSettings)
 	{
 		MaxBuildStatusErrorLevel = BuilderSettings.bFailOnAnyError ? FBuildStatus::NoError : FBuildStatus::NonFatalError;
 	}
@@ -121,12 +120,13 @@ namespace Metasound
 	{
 	}
 
-	TUniquePtr<IOperator> FOperatorBuilder::BuildGraphOperator(const IGraph& InGraph, TArray<FBuildErrorPtr>& OutErrors)
+	TUniquePtr<IOperator> FOperatorBuilder::BuildGraphOperator(const IGraph& InGraph, const FOperatorSettings& InOperatorSettings, const FMetasoundEnvironment& InEnvironment, TArray<FBuildErrorPtr>& OutErrors) 
 	{
 		using namespace OperatorBuilderPrivate;
 
 		FBuildStatus BuildStatus = FBuildStatus::NoError;
 
+
 		// Validate that the sources and destinations declared in an edge actually
 		// exist in the node.
 		if (BuilderSettings.bValidateVerticesExist)
@@ -161,6 +161,7 @@ namespace Metasound
 			return TUniquePtr<IOperator>(nullptr);
 		}
 
+
 		// Create algo adapter view of graph to cache graph operations.
 		TPimplPtr<FDirectedGraphAlgoAdapter> AlgoAdapter = FDirectedGraphAlgo::CreateDirectedGraphAlgoAdapter(InGraph);
 		
@@ -170,13 +171,15 @@ namespace Metasound
 			return TUniquePtr<IOperator>(nullptr);
 		}
 
+		FBuildContext BuildContext(InGraph, *AlgoAdapter, InOperatorSettings, InEnvironment, OutErrors);
+
 		TArray<const INode*> SortedNodes;
 
 		// Sort the nodes in a valid execution order
-		BuildStatus |= DepthFirstTopologicalSort(*AlgoAdapter, SortedNodes, OutErrors);
+		BuildStatus |= DepthFirstTopologicalSort(BuildContext, SortedNodes);
 
 		// Prune unneeded nodes. 
-		BuildStatus |= PruneNodes(*AlgoAdapter, SortedNodes, OutErrors);
+		BuildStatus |= PruneNodes(BuildContext, SortedNodes);
 
 		// Check build status in case build routine should be exited early.
 		if (BuildStatus > MaxBuildStatusErrorLevel)
@@ -188,7 +191,7 @@ namespace Metasound
 		FNodeDataReferenceMap NodeData;
 
 		// Create node operators from factories. 
-		BuildStatus |= CreateOperators(InGraph, SortedNodes, Operators, NodeData, OutErrors);
+		BuildStatus |= CreateOperators(BuildContext, SortedNodes, Operators, NodeData);
 
 		if (BuildStatus > MaxBuildStatusErrorLevel)
 		{
@@ -196,19 +199,19 @@ namespace Metasound
 		}
 
 		// Create graph operator from collection of node operators.
-		return CreateGraphOperator(InGraph, Operators, NodeData, OutErrors);
+		return CreateGraphOperator(BuildContext, Operators, NodeData);
 	}
 
-	FOperatorBuilder::FBuildStatus FOperatorBuilder::DepthFirstTopologicalSort(const FDirectedGraphAlgoAdapter& InAdapter, TArray<const INode*>& OutNodes, TArray<FBuildErrorPtr>& OutErrors) const
+	FOperatorBuilder::FBuildStatus FOperatorBuilder::DepthFirstTopologicalSort(FBuildContext& InContext, TArray<const INode*>& OutNodes) const
 	{
 		using namespace OperatorBuilderPrivate;
 
-		bool bSuccess = FDirectedGraphAlgo::DepthFirstTopologicalSort(InAdapter, OutNodes);
+		bool bSuccess = FDirectedGraphAlgo::DepthFirstTopologicalSort(InContext.AlgoAdapter, OutNodes);
 
 		if (!bSuccess)
 		{
 			// If there was an error, there is likely a cycle in the graph.
-			AddBuildErrorsForCycles(InAdapter, OutErrors);
+			AddBuildErrorsForCycles(InContext.AlgoAdapter, InContext.Errors);
 
 			return FBuildStatus::FatalError;
 		}
@@ -216,16 +219,16 @@ namespace Metasound
 		return FBuildStatus::NoError;
 	}
 
-	FOperatorBuilder::FBuildStatus FOperatorBuilder::KahnsTopologicalSort(const FDirectedGraphAlgoAdapter& InAdapter, TArray<const INode*>& OutNodes, TArray<FBuildErrorPtr>& OutErrors) const
+	FOperatorBuilder::FBuildStatus FOperatorBuilder::KahnsTopologicalSort(FBuildContext& InContext, TArray<const INode*>& OutNodes) const
 	{
 		using namespace OperatorBuilderPrivate;
 
-		bool bSuccess = FDirectedGraphAlgo::KahnTopologicalSort(InAdapter, OutNodes);
+		bool bSuccess = FDirectedGraphAlgo::KahnTopologicalSort(InContext.AlgoAdapter, OutNodes);
 
 		if (!bSuccess)
 		{
 			// If there was an error, there is likely a cycle in the graph.
-			AddBuildErrorsForCycles(InAdapter, OutErrors);
+			AddBuildErrorsForCycles(InContext.AlgoAdapter, InContext.Errors);
 
 			return FBuildStatus::FatalError;
 		}
@@ -233,7 +236,7 @@ namespace Metasound
 		return FBuildStatus::NoError;
 	}
 
-	FOperatorBuilder::FBuildStatus FOperatorBuilder::PruneNodes(const FDirectedGraphAlgoAdapter& InAdapter, TArray<const INode*>& InOutNodes, TArray<FBuildErrorPtr>& OutErrors) const
+	FOperatorBuilder::FBuildStatus FOperatorBuilder::PruneNodes(FBuildContext& InContext, TArray<const INode*>& InOutNodes) const
 	{
 		using namespace OperatorBuilderPrivate;
 
@@ -244,15 +247,15 @@ namespace Metasound
 		switch (BuilderSettings.PruningMode)
 		{
 			case EOperatorBuilderNodePruning::PruneNodesWithoutExternalDependency:
-				FDirectedGraphAlgo::FindReachableNodes(InAdapter, ReachableNodes);
+				FDirectedGraphAlgo::FindReachableNodes(InContext.AlgoAdapter, ReachableNodes);
 				break;
 
 			case EOperatorBuilderNodePruning::PruneNodesWithoutOutputDependency:
-				FDirectedGraphAlgo::FindReachableNodesFromOutput(InAdapter, ReachableNodes);
+				FDirectedGraphAlgo::FindReachableNodesFromOutput(InContext.AlgoAdapter, ReachableNodes);
 				break;
 
 			case EOperatorBuilderNodePruning::PruneNodesWithoutInputDependency:
-				FDirectedGraphAlgo::FindReachableNodesFromInput(InAdapter, ReachableNodes);
+				FDirectedGraphAlgo::FindReachableNodesFromInput(InContext.AlgoAdapter, ReachableNodes);
 				break;
 
 			case EOperatorBuilderNodePruning::None:
@@ -272,7 +275,7 @@ namespace Metasound
 			// Pruning all nodes. 
 			for (const INode* Node : InOutNodes)
 			{
-				AddBuildError<FNodePrunedError>(OutErrors, Node);
+				AddBuildError<FNodePrunedError>(InContext.Errors, Node);
 			}
 
 			InOutNodes.Reset();
@@ -295,7 +298,7 @@ namespace Metasound
 			}
 			else
 			{
-				AddBuildError<FNodePrunedError>(OutErrors, Node);
+				AddBuildError<FNodePrunedError>(InContext.Errors, Node);
 
 				// Denote a pruned node as a non-fatal error. In the future this
 				// may be simply a warning as some nodes are required to conform
@@ -309,7 +312,7 @@ namespace Metasound
 		return BuildStatus;
 	}
 
-	FOperatorBuilder::FBuildStatus FOperatorBuilder::GatherInputDataReferences(const INode* InNode, const FNodeEdgeMultiMap& InEdgeMap, const FNodeDataReferenceMap& InDataReferenceMap, FDataReferenceCollection& OutCollection, TArray<FBuildErrorPtr>& OutErrors) const
+	FOperatorBuilder::FBuildStatus FOperatorBuilder::GatherInputDataReferences(FBuildContext& InContext, const INode* InNode, const FNodeEdgeMultiMap& InEdgeMap, const FNodeDataReferenceMap& InDataReferenceMap, FDataReferenceCollection& OutCollection) const
 	{
 		using namespace OperatorBuilderPrivate;
 
@@ -326,7 +329,7 @@ namespace Metasound
 			{
 				// This is likely due to a failed topological sort and is more of an internal error
 				// than a user error.
-				AddBuildError<FInternalError>(OutErrors, TEXT(__FILE__), __LINE__);
+				AddBuildError<FInternalError>(InContext.Errors, TEXT(__FILE__), __LINE__);
 
 				return FBuildStatus::NonFatalError;
 			}
@@ -339,7 +342,7 @@ namespace Metasound
 				// did not match the readable parameter refs created by the operators outputs. Or, the edge description is invalid.
 
 				// TODO: consider checking that outputs of node operators match output descriptions of nodes. 
-				AddBuildError<FMissingOutputDataReferenceError>(OutErrors, Edge->From);
+				AddBuildError<FMissingOutputDataReferenceError>(InContext.Errors, Edge->From);
 
 				BuildStatus |= FBuildStatus::NonFatalError;
 				continue;
@@ -350,7 +353,7 @@ namespace Metasound
 
 			if (!bSuccess)
 			{
-				AddBuildError<FMissingOutputDataReferenceError>(OutErrors, Edge->From);
+				AddBuildError<FMissingOutputDataReferenceError>(InContext.Errors, Edge->From);
 				BuildStatus |= FBuildStatus::NonFatalError;
 			}
 		}
@@ -358,14 +361,14 @@ namespace Metasound
 		return BuildStatus;
 	}
 
-	FOperatorBuilder::FBuildStatus FOperatorBuilder::CreateOperators(const IGraph& InGraph, const TArray<const INode*>& InSortedNodes, TArray<TUniquePtr<IOperator>>& OutOperators, FNodeDataReferenceMap& OutDataReferences, TArray<FBuildErrorPtr>& OutErrors) const
+	FOperatorBuilder::FBuildStatus FOperatorBuilder::CreateOperators(FBuildContext& InContext, const TArray<const INode*>& InSortedNodes, TArray<TUniquePtr<IOperator>>& OutOperators, FNodeDataReferenceMap& OutDataReferences) const
 	{
 		FBuildStatus BuildStatus;
 
 		FNodeEdgeMultiMap NodeInputEdges;
 
 		// Gather input edges for each node.
-		for (const FDataEdge& Edge : InGraph.GetDataEdges())
+		for (const FDataEdge& Edge : InContext.Graph.GetDataEdges())
 		{
 			NodeInputEdges.Add(Edge.To.Node, &Edge);
 		}
@@ -377,19 +380,18 @@ namespace Metasound
 			FDataReferenceCollection InputCollection;
 
 			// Gather the input parameters for this IOperator from the output parameters of already created IOperators. 
-			BuildStatus |= GatherInputDataReferences(Node, NodeInputEdges, OutDataReferences, InputCollection, OutErrors);
+			BuildStatus |= GatherInputDataReferences(InContext, Node, NodeInputEdges, OutDataReferences, InputCollection);
 
 			if (BuildStatus >= FBuildStatus::FatalError)
 			{
 				return BuildStatus;
 			}
 
-			// TODO: make factory retrieval const 
 			FOperatorFactorySharedRef Factory = Node->GetDefaultOperatorFactory();
 
-			FCreateOperatorParams CreateParams(*Node, OperatorSettings, InputCollection);
+			FCreateOperatorParams CreateParams(*Node, InContext.Settings, InputCollection, InContext.Environment);
 
-			FOperatorPtr Operator = Factory->CreateOperator(CreateParams, OutErrors);
+			FOperatorPtr Operator = Factory->CreateOperator(CreateParams, InContext.Errors);
 
 			if (!Operator.IsValid())
 			{
@@ -406,7 +408,7 @@ namespace Metasound
 		return BuildStatus;
 	}
 
-	FOperatorBuilder::FBuildStatus FOperatorBuilder::GatherGraphDataReferences(const IGraph& InGraph, FNodeDataReferenceMap& InNodeDataReferences, FDataReferenceCollection& OutGraphInputs, FDataReferenceCollection& OutGraphOutputs, TArray<FBuildErrorPtr>& OutErrors) const
+	FOperatorBuilder::FBuildStatus FOperatorBuilder::GatherGraphDataReferences(FBuildContext& InContext, FNodeDataReferenceMap& InNodeDataReferences, FDataReferenceCollection& OutGraphInputs, FDataReferenceCollection& OutGraphOutputs) const
 	{
 		using namespace OperatorBuilderPrivate;
 		using FDestinationElement = FInputDataDestinationCollection::ElementType;
@@ -415,14 +417,14 @@ namespace Metasound
 		FBuildStatus BuildStatus;
 
 		// Gather graph inputs
-		for (const FDestinationElement& Element : InGraph.GetInputDataDestinations())
+		for (const FDestinationElement& Element : InContext.Graph.GetInputDataDestinations())
 		{
 			const FInputDataDestination& InputDestination = Element.Value;
 
 			if (!InNodeDataReferences.Contains(InputDestination.Node))
 			{
 				// An input node was likely pruned.
-				AddBuildError<FMissingInputDataReferenceError>(OutErrors, InputDestination);
+				AddBuildError<FMissingInputDataReferenceError>(InContext.Errors, InputDestination);
 
 				BuildStatus |= FBuildStatus::NonFatalError;
 				continue;
@@ -433,7 +435,7 @@ namespace Metasound
 
 			if (!Collection.ContainsDataWriteReference(InputDestination.Vertex.GetVertexName(), InputDestination.Vertex.GetDataTypeName()))
 			{
-				AddBuildError<FMissingInputDataReferenceError>(OutErrors, InputDestination);
+				AddBuildError<FMissingInputDataReferenceError>(InContext.Errors, InputDestination);
 
 				BuildStatus |= FBuildStatus::NonFatalError;
 				continue;
@@ -443,7 +445,7 @@ namespace Metasound
 
 			if (!bSuccess)
 			{
-				AddBuildError<FMissingInputDataReferenceError>(OutErrors, InputDestination);
+				AddBuildError<FMissingInputDataReferenceError>(InContext.Errors, InputDestination);
 
 				BuildStatus |= FBuildStatus::NonFatalError;
 				continue;
@@ -451,14 +453,14 @@ namespace Metasound
 		}
 
 		// Gather graph outputs.
-		for (const FSourceElement& Element : InGraph.GetOutputDataSources())
+		for (const FSourceElement& Element : InContext.Graph.GetOutputDataSources())
 		{
 			const FOutputDataSource& OutputSource = Element.Value;
 
 			if (!InNodeDataReferences.Contains(OutputSource.Node))
 			{
 				// An output node was likely pruned.
-				AddBuildError<FMissingOutputDataReferenceError>(OutErrors, OutputSource);
+				AddBuildError<FMissingOutputDataReferenceError>(InContext.Errors, OutputSource);
 
 				BuildStatus |= FBuildStatus::NonFatalError;
 
@@ -471,7 +473,7 @@ namespace Metasound
 			{
 				// This will likely produce an IOperator which does not work as
 				// expected.
-				AddBuildError<FMissingOutputDataReferenceError>(OutErrors, OutputSource);
+				AddBuildError<FMissingOutputDataReferenceError>(InContext.Errors, OutputSource);
 
 				BuildStatus |= FBuildStatus::NonFatalError;
 
@@ -484,7 +486,7 @@ namespace Metasound
 			{
 				// This will likely produce an IOperator which does not work as
 				// expected.
-				AddBuildError<FMissingOutputDataReferenceError>(OutErrors, OutputSource);
+				AddBuildError<FMissingOutputDataReferenceError>(InContext.Errors, OutputSource);
 
 				BuildStatus |= FBuildStatus::NonFatalError;
 
@@ -495,12 +497,12 @@ namespace Metasound
 		return BuildStatus;
 	}
 
-	TUniquePtr<IOperator> FOperatorBuilder::CreateGraphOperator(const IGraph& InGraph, TArray<FOperatorPtr>& InOperators, FNodeDataReferenceMap& InNodeDataReferences, TArray<FBuildErrorPtr>& OutErrors) const
+	TUniquePtr<IOperator> FOperatorBuilder::CreateGraphOperator(FBuildContext& InContext, TArray<FOperatorPtr>& InOperators, FNodeDataReferenceMap& InNodeDataReferences) const
 	{
 		FDataReferenceCollection GraphInputs;
 		FDataReferenceCollection GraphOutputs;
 
-		FBuildStatus BuildStatus = GatherGraphDataReferences(InGraph, InNodeDataReferences, GraphInputs, GraphOutputs, OutErrors);
+		FBuildStatus BuildStatus = GatherGraphDataReferences(InContext, InNodeDataReferences, GraphInputs, GraphOutputs);
 
 		if (BuildStatus > MaxBuildStatusErrorLevel)
 		{
diff --git a/Engine/Plugins/Runtime/Metasound/Source/MetasoundGraphCore/Private/MetasoundVertex.cpp b/Engine/Plugins/Runtime/Metasound/Source/MetasoundGraphCore/Private/MetasoundVertex.cpp
--- a/Engine/Plugins/Runtime/Metasound/Source/MetasoundGraphCore/Private/MetasoundVertex.cpp
+++ b/Engine/Plugins/Runtime/Metasound/Source/MetasoundGraphCore/Private/MetasoundVertex.cpp
@@ -185,13 +185,90 @@ namespace Metasound
 		}
 	}
 
-			/** Construct with an input and output interface. */
+	FEnvironmentVertex::FEnvironmentVertex()
+	:	VertexModel(MakeUnique<FEmptyVertexModel>(TEXT(""), FText::GetEmpty()))
+	{
+	}
+
+	/** Copy constructor */
+	FEnvironmentVertex::FEnvironmentVertex(const FEnvironmentVertex& InOther)
+	{
+		if (InOther.VertexModel.IsValid())
+		{
+			// call underlying model's clone to copy model.
+			VertexModel = InOther.VertexModel->Clone();
+		}
+		else
+		{
+			// Make an empty model if no valid model exists.
+			VertexModel = MakeUnique<FEmptyVertexModel>(TEXT(""), FText::GetEmpty());
+		}
+	}
+
+	FEnvironmentVertex& FEnvironmentVertex::operator=(const FEnvironmentVertex& InOther)
+	{
+		if (InOther.VertexModel.IsValid())
+		{
+			// call underlying model's clone to copy model.
+			VertexModel = InOther.VertexModel->Clone();
+		}
+		else
+		{
+			// Make an empty model if no valid model exists.
+			VertexModel = MakeUnique<FEmptyVertexModel>(TEXT(""), FText::GetEmpty());
+		}
+
+		return *this;
+	}
+
+	const FString& FEnvironmentVertex::GetVertexName() const
+	{
+		return VertexModel->VertexName;
+	}
+
+	const FText& FEnvironmentVertex::GetDescription() const
+	{
+		return VertexModel->Description;
+	}
+
+	bool FEnvironmentVertex::IsVariableOfSameType(const IMetasoundEnvironmentVariable& InVariable) const
+	{
+		return VertexModel->IsVariableOfSameType(InVariable);
+	}
+
+	bool operator==(const FEnvironmentVertex& InLHS, const FEnvironmentVertex& InRHS)
+	{
+		return (InLHS.VertexModel->IsEqual(*InRHS.VertexModel) && InRHS.VertexModel->IsEqual(*InLHS.VertexModel));
+	}
+
+	bool operator!=(const FEnvironmentVertex& InLHS, const FEnvironmentVertex& InRHS)
+	{
+		return !(InLHS == InRHS);
+	}
+
+	bool operator<(const FEnvironmentVertex& InLHS, const FEnvironmentVertex& InRHS)
+	{
+		if (InLHS == InRHS)
+		{
+			return false;
+		}
+
+		return InLHS.GetVertexName() < InRHS.GetVertexName();
+	}
+
 	FVertexInterface::FVertexInterface(const FInputVertexInterface& InInputs, const FOutputVertexInterface& InOutputs)
 	:	InputInterface(InInputs)
 	,	OutputInterface(InOutputs)
 	{
 	}
 
+	FVertexInterface::FVertexInterface(const FInputVertexInterface& InInputs, const FOutputVertexInterface& InOutputs, const FEnvironmentVertexInterface& InEnvironmentVariables)
+	:	InputInterface(InInputs)
+	,	OutputInterface(InOutputs)
+	,	EnvironmentInterface(InEnvironmentVariables)
+	{
+	}
+
 	const FInputVertexInterface& FVertexInterface::GetInputInterface() const
 	{
 		return InputInterface;
@@ -202,12 +279,12 @@ namespace Metasound
 		return InputInterface;
 	}
 
-	const FInputDataVertex& FVertexInterface::GetInputVertex(const FDataVertexKey& InKey) const
+	const FInputDataVertex& FVertexInterface::GetInputVertex(const FVertexKey& InKey) const
 	{
 		return InputInterface[InKey];
 	}
 
-	bool FVertexInterface::ContainsInputVertex(const FDataVertexKey& InKey) const
+	bool FVertexInterface::ContainsInputVertex(const FVertexKey& InKey) const
 	{
 		return InputInterface.Contains(InKey);
 	}
@@ -222,19 +299,43 @@ namespace Metasound
 		return OutputInterface;
 	}
 
-	const FOutputDataVertex& FVertexInterface::GetOutputVertex(const FDataVertexKey& InKey) const
+	const FOutputDataVertex& FVertexInterface::GetOutputVertex(const FVertexKey& InKey) const
 	{
 		return OutputInterface[InKey];
 	}
 
-	bool FVertexInterface::ContainsOutputVertex(const FDataVertexKey& InKey) const
+	bool FVertexInterface::ContainsOutputVertex(const FVertexKey& InKey) const
 	{
 		return OutputInterface.Contains(InKey);
 	}
 
+	const FEnvironmentVertexInterface& FVertexInterface::GetEnvironmentInterface() const
+	{
+		return EnvironmentInterface;
+	}
+
+	FEnvironmentVertexInterface& FVertexInterface::GetEnvironmentInterface()
+	{
+		return EnvironmentInterface;
+	}
+
+	const FEnvironmentVertex& FVertexInterface::GetEnvironmentVertex(const FVertexKey& InKey) const
+	{
+		return EnvironmentInterface[InKey];
+	}
+
+	bool FVertexInterface::ContainsEnvironmentVertex(const FVertexKey& InKey) const
+	{
+		return EnvironmentInterface.Contains(InKey);
+	}
+
 	bool operator==(const FVertexInterface& InLHS, const FVertexInterface& InRHS)
 	{
-		return (InLHS.InputInterface == InRHS.InputInterface) && (InLHS.OutputInterface == InRHS.OutputInterface);
+		const bool bIsEqual = (InLHS.InputInterface == InRHS.InputInterface) && 
+			(InLHS.OutputInterface == InRHS.OutputInterface) && 
+			(InLHS.EnvironmentInterface == InRHS.EnvironmentInterface);
+
+		return bIsEqual;
 	}
 
 	bool operator!=(const FVertexInterface& InLHS, const FVertexInterface& InRHS)
diff --git a/Engine/Plugins/Runtime/Metasound/Source/MetasoundGraphCore/Public/MetasoundBuilderInterface.h b/Engine/Plugins/Runtime/Metasound/Source/MetasoundGraphCore/Public/MetasoundBuilderInterface.h
--- a/Engine/Plugins/Runtime/Metasound/Source/MetasoundGraphCore/Public/MetasoundBuilderInterface.h
+++ b/Engine/Plugins/Runtime/Metasound/Source/MetasoundGraphCore/Public/MetasoundBuilderInterface.h
@@ -3,6 +3,7 @@
 #pragma once
 
 #include "CoreMinimal.h"
+#include "MetasoundEnvironment.h"
 #include "MetasoundNodeInterface.h"
 #include "MetasoundOperatorInterface.h"
 
@@ -50,10 +51,14 @@ namespace Metasound
 		/** Collection of input parameters available for to an IOperator. */
 		const FDataReferenceCollection& InputDataReferences;
 
-		FCreateOperatorParams(const INode& InNode, const FOperatorSettings& InOperatorSettings, const FDataReferenceCollection& InInputDataReferences)
+		/** Environment settings available. */
+		const FMetasoundEnvironment& Environment;
+
+		FCreateOperatorParams(const INode& InNode, const FOperatorSettings& InOperatorSettings, const FDataReferenceCollection& InInputDataReferences, const FMetasoundEnvironment& InEnvironment)
 		:	Node(InNode)
 		,	OperatorSettings(InOperatorSettings)
 		,	InputDataReferences(InInputDataReferences)
+		,	Environment(InEnvironment)
 		{
 		}
 	};
@@ -113,11 +118,13 @@ namespace Metasound
 			/** Build a graph operator from a graph. 
 			 *
 			 * @params InGraph - The input graph object containing edges, input vertices and output vertices.
+			 * @param InOperatorSettings - Settings to be passed to all operators on creation.
+			 * @param InEnvironment - The environment variables to use during construction. 
 			 * @param OutErrors - An array of errors. Errors can be added if issues occur while creating the IOperator.
 			 *
 			 * @return A unique pointer to an IOperator. 
 			 */
-			virtual TUniquePtr<IOperator> BuildGraphOperator(const IGraph& InGraph, TArray<FBuildErrorPtr>& OutErrors) = 0;
+			virtual TUniquePtr<IOperator> BuildGraphOperator(const IGraph& InGraph, const FOperatorSettings& InOperatorSettings, const FMetasoundEnvironment& InEnvironment, TArray<FBuildErrorPtr>& OutErrors) = 0;
 	};
 }
 
diff --git a/Engine/Plugins/Runtime/Metasound/Source/MetasoundGraphCore/Public/MetasoundEnvironment.h b/Engine/Plugins/Runtime/Metasound/Source/MetasoundGraphCore/Public/MetasoundEnvironment.h
new file mode 100644
--- /dev/null
+++ b/Engine/Plugins/Runtime/Metasound/Source/MetasoundGraphCore/Public/MetasoundEnvironment.h
@@ -0,0 +1,268 @@
+// Copyright Epic Games, Inc. All Rights Reserved.
+
+#pragma once
+
+#define DECLARE_METASOUND_ENVIRONMENT_VARIABLE_TYPE(VarType) \
+	template<> \
+	struct ::Metasound::TEnvironmentVariableTypeInfo<VarType> \
+	{ \
+		typedef VarType Type; \
+		static const FMetasoundEnvironmentVariableTypeId TypeId; \
+		\
+	private: \
+		\
+		static const VarType* const TypePtr; \
+	};
+
+#define DEFINE_METASOUND_ENVIRONMENT_VARIABLE_TYPE(VarType) \
+	const VarType* const ::Metasound::TEnvironmentVariableTypeInfo<VarType>::TypePtr = nullptr; \
+	const void* const ::Metasound::TEnvironmentVariableTypeInfo<VarType>::TypeId = static_cast<const FMetasoundEnvironmentVariableTypeId>(&::Metasound::TEnvironmentVariableTypeInfo<VarType>::TypePtr);
+	
+using FMetasoundEnvironmentVariableTypeId = void const*;
+
+#include "Containers/SortedMap.h"
+#include "CoreMinimal.h"
+#include <type_traits>
+
+namespace Metasound
+{
+	namespace MetasoundEnvironmentPrivate
+	{
+		// Helper class to ensure that DECLARE_METASOUND_ENVIRONMENT_VARIABLE is used. 
+		template<typename DataType>
+		struct TSpecializationHelper 
+		{
+			enum { Value = false };
+		};
+	}
+
+
+	template<typename VarType>
+	struct TEnvironmentVariableTypeInfo
+	{
+		typedef VarType Type;
+
+		static const FMetasoundEnvironmentVariableTypeId TypeId;
+
+		// Trigger an assert if this template has not been specialized. 
+		static_assert(MetasoundEnvironmentPrivate::TSpecializationHelper<VarType>::Value, "TEnvironmentVariableTypeInfo must be specialized.  Use macro DECLARE_METASOUND_ENVIRONMENT_VARIABLE_TYPE");
+	};
+
+
+	/** Return the type ID for a environment variable type. */
+	template<typename Type>
+	FMetasoundEnvironmentVariableTypeId GetMetasoundEnvironmentVariableTypeId()
+	{
+		using FTypeInfo = TEnvironmentVariableTypeInfo< std::decay_t<Type> >;
+
+		return FTypeInfo::TypeId;
+	}
+
+	/** Interface for a metasound environment variable which supports
+	 * a name, a runtime type ID, and a clone method.
+	 */
+	class IMetasoundEnvironmentVariable
+	{
+	private:
+		// Only allow TMetasoundEnvironmentVariable to implement this class
+		// by making constructor private so that casting can be effectively checked. 
+		template<typename Type> 
+        friend class TMetasoundEnvironmentVariable;
+
+		IMetasoundEnvironmentVariable() = default;
+
+	public:
+		virtual ~IMetasoundEnvironmentVariable() = default;
+
+		/** The name of this environment variable. */
+		virtual const FString& GetName() const = 0;
+
+		/** The type id of this environment variable. */
+		virtual FMetasoundEnvironmentVariableTypeId GetTypeId() const = 0;
+
+		/** Create a copy of this environment variable. */
+		virtual TUniquePtr<IMetasoundEnvironmentVariable> Clone() const = 0;
+	};
+
+	template<typename Type>
+	class TMetasoundEnvironmentVariable : public IMetasoundEnvironmentVariable
+	{
+		// Make sure that the `Type` is not a reference or CV qualified. 
+		static_assert(std::is_same<Type, typename std::decay<Type>::type>::value, "Type must not change when decayed");
+
+		// Make sure the `Type` has appropriate construction and copy characteristics. 
+		static_assert(std::is_constructible<Type>::value, "Environment variable types must be default constructible");
+		static_assert(std::is_copy_constructible<Type>::value, "Environment variable types must be copy constructible");
+		static_assert(std::is_copy_assignable<Type>::value, "Environment variable types must be copy assignable");
+
+	public:
+		/** Construct an environment variable 
+		 *
+		 * @param InName - Name of the environment variable.
+		 * @param ...Args - Parameter pack to be forwarded to underlying types constructor.
+		 */
+		template<typename... ArgTypes>
+		TMetasoundEnvironmentVariable(const FString& InName, ArgTypes&&... Args)
+		:	Name(InName)
+		,	Value(Forward<ArgTypes>(Args)...)
+		{
+		}
+
+		/** Get the current value. */
+		const Type& GetValue() const
+		{
+			return Value;
+		}
+
+		/** Set the current value. */
+		void SetValue(const Type& InValue)
+		{
+			Value = InValue;
+		}
+
+		/** Get the name of this environment variable. */
+		const FString& GetName() const override
+		{
+			return Name;
+		}
+
+		/** Get the type ID of this environment variable. */
+		FMetasoundEnvironmentVariableTypeId GetTypeId() const override
+		{
+			return GetMetasoundEnvironmentVariableTypeId<Type>();
+		}
+
+		/** Create a clone of this environment variable. */
+		TUniquePtr<IMetasoundEnvironmentVariable> Clone() const override
+		{
+			return MakeUnique<TMetasoundEnvironmentVariable<Type>>(Name, Value);
+		}
+
+	private:
+		FString Name;
+		Type Value;
+	};
+
+
+	/** Returns true if the environment variable contains the same type as the `VarType` */
+	template<typename VarType>
+	bool IsEnvironmentVariableOfType(const IMetasoundEnvironmentVariable& InVar)
+	{
+		return GetMetasoundEnvironmentVariableTypeId<VarType>() == InVar.GetTypeId();
+	}
+
+	/** Casts an environment variable to the derived type. */
+	template<typename Type>
+	TMetasoundEnvironmentVariable<Type>& CastMetasoundEnvironmentVariableChecked(IMetasoundEnvironmentVariable& InVar)
+	{
+		using FDerivedType = TMetasoundEnvironmentVariable<Type>;
+
+		check(IsEnvironmentVariableOfType<Type>(InVar));
+		return static_cast<FDerivedType&>(InVar);
+	}
+
+	/** Casts an environment variable to the derived type. */
+	template<typename Type>
+	const TMetasoundEnvironmentVariable<Type>& CastMetasoundEnvironmentVariableChecked(const IMetasoundEnvironmentVariable& InVar)
+	{
+		using FDerivedType = TMetasoundEnvironmentVariable<Type>;
+
+		check(IsEnvironmentVariableOfType<Type>(InVar));
+		return static_cast<const FDerivedType&>(InVar);
+	}
+
+	/** FMetasoundEnvironment contains a set of TMetasoundEnvironmentVariables requiring
+	 * that each environment variable has a unique name. 
+	 */
+	class FMetasoundEnvironment
+	{
+	public:
+		FMetasoundEnvironment() = default;
+		~FMetasoundEnvironment() = default;
+
+		FMetasoundEnvironment(const FMetasoundEnvironment& InOther);
+
+		FMetasoundEnvironment& operator=(const FMetasoundEnvironment& InOther);
+
+		/** Returns true if the environment variable with the given name contains
+		 * the data of the same type as `VarType`
+		 *
+		 * @param InVariableName - Name of environment variable. 
+		 * @tparam VarType - Type of the underlying data stored in the environment variable. 
+		 */
+		template<typename VarType>
+		bool IsType(const FString& InVariableName) const
+		{
+			return IsEnvironmentVariableOfType<VarType>(*Variables[InVariableName]);
+		}
+
+		/** Returns true if the environment contains a variable with the name `InVariableName`
+		 * and the type `VarType`.
+		 *
+		 * @param InVariableName - Name of environment variable. 
+		 * @tparam VarType - Type of the underlying data stored in the environment variable. 
+		 */
+		template<typename VarType>
+		bool Contains(const FString& InVariableName) const
+		{
+			if (Variables.Contains(InVariableName))
+			{
+				return IsType<VarType>(InVariableName);
+			}
+			return false;
+		}
+
+		/** Returns the environment variable data. 
+		 *
+		 * @param InVariableName - Name of environment variable. 
+		 * @tparam VarType - Type of the underlying data stored in the environment variable. 
+		 *
+		 * @return copy of the underlying data.
+		 */
+		template<typename VarType>
+		VarType GetValue(const FString& InVariableName) const
+		{
+			if (ensure(Contains<VarType>(InVariableName)))
+			{
+				const TMetasoundEnvironmentVariable<VarType>& Var = CastMetasoundEnvironmentVariableChecked<VarType>(*Variables[InVariableName]);
+
+				return Var.GetValue();
+			}
+			return VarType{};
+		}
+
+		/** Sets the environment variable data
+		 *
+		 * @param InVariableName - Name of environment variable. 
+		 * @tparam VarType - Type of the underlying data stored in the environment variable. 
+		 *
+		 * @return const ref to the underlying data.
+		 */
+		template<typename VarType>
+		void SetValue(const FString& InVariableName, const VarType& InValue)
+		{
+			Variables.Add(InVariableName, MakeUnique<TMetasoundEnvironmentVariable<VarType>>(InVariableName, InValue));
+		}
+
+
+	private:
+		TSortedMap<FString, TUniquePtr<IMetasoundEnvironmentVariable>> Variables;
+	};
+}
+
+/** Declare basic set of variable types. */
+DECLARE_METASOUND_ENVIRONMENT_VARIABLE_TYPE(void);
+DECLARE_METASOUND_ENVIRONMENT_VARIABLE_TYPE(bool);
+DECLARE_METASOUND_ENVIRONMENT_VARIABLE_TYPE(int8);
+DECLARE_METASOUND_ENVIRONMENT_VARIABLE_TYPE(uint8);
+DECLARE_METASOUND_ENVIRONMENT_VARIABLE_TYPE(int16);
+DECLARE_METASOUND_ENVIRONMENT_VARIABLE_TYPE(uint16);
+DECLARE_METASOUND_ENVIRONMENT_VARIABLE_TYPE(int32);
+DECLARE_METASOUND_ENVIRONMENT_VARIABLE_TYPE(uint32);
+DECLARE_METASOUND_ENVIRONMENT_VARIABLE_TYPE(int64);
+DECLARE_METASOUND_ENVIRONMENT_VARIABLE_TYPE(uint64);
+DECLARE_METASOUND_ENVIRONMENT_VARIABLE_TYPE(float);
+DECLARE_METASOUND_ENVIRONMENT_VARIABLE_TYPE(double);
+DECLARE_METASOUND_ENVIRONMENT_VARIABLE_TYPE(FString);
+DECLARE_METASOUND_ENVIRONMENT_VARIABLE_TYPE(FName);
+
diff --git a/Engine/Plugins/Runtime/Metasound/Source/MetasoundGraphCore/Public/MetasoundOperatorBuilder.h b/Engine/Plugins/Runtime/Metasound/Source/MetasoundGraphCore/Public/MetasoundOperatorBuilder.h
--- a/Engine/Plugins/Runtime/Metasound/Source/MetasoundGraphCore/Public/MetasoundOperatorBuilder.h
+++ b/Engine/Plugins/Runtime/Metasound/Source/MetasoundGraphCore/Public/MetasoundOperatorBuilder.h
@@ -3,6 +3,7 @@
 #pragma once
 
 #include "CoreMinimal.h"
+#include "MetasoundEnvironment.h"
 #include "MetasoundNodeInterface.h"
 #include "MetasoundOperatorInterface.h"
 #include "MetasoundBuilderInterface.h"
@@ -96,24 +97,23 @@ namespace Metasound
 
 			/** FOperatorBuilder constructor.
 			 *
-			 * @param InOperatorSettings - Settings to be passed to all operators
-			 *                             on creation.
 			 * @param InBuilderSettings  - Settings to configure builder options.
 			 */
-			FOperatorBuilder(const FOperatorSettings& InOperatorSettings, const FOperatorBuilderSettings& InBuilderSettings);
+			FOperatorBuilder(const FOperatorBuilderSettings& InBuilderSettings);
 
 			virtual ~FOperatorBuilder();
 
 			/** Create an IOperator from an IGraph.
 			 *
-			 * @param InGraph   - The graph containing input, output and edge information.
-			 * @param OutErrors - An array of build errors that will be populated 
-			 *                    with any issues encountered during the build process.
+			 * @param InGraph            - The graph containing input, output and edge information.
+			 * @param InOperatorSettings - Settings to be passed to all operators on creation.
+			 * @param InEnvironment      - The environment variables to use during construction. 
+			 * @param OutErrors          - An array of build errors that will be populated with any issues encountered during the build process.
 			 *
 			 * @return A TUniquePtr to an IOperator. If the processes was unsuccessful, 
 			 *         the returned pointer will contain a nullptr and be invalid.
 			 */
-			virtual TUniquePtr<IOperator> BuildGraphOperator(const IGraph& InGraph, TArray<FBuildErrorPtr>& OutErrors) override;
+			virtual TUniquePtr<IOperator> BuildGraphOperator(const IGraph& InGraph, const FOperatorSettings& InOperatorSettings, const FMetasoundEnvironment& InEnvironment, TArray<FBuildErrorPtr>& OutErrors) override;
 			
 
 		private:
@@ -185,28 +185,45 @@ namespace Metasound
 				EStatus Value = NoError;
 			};
 
+			struct FBuildContext
+			{
+				const IGraph& Graph;
+				const FDirectedGraphAlgoAdapter& AlgoAdapter;
+				const FOperatorSettings& Settings;
+				const FMetasoundEnvironment& Environment;
+
+				TArray<FBuildErrorPtr>& Errors;
+
+				FBuildContext(const IGraph& InGraph, FDirectedGraphAlgoAdapter& InAlgoAdapter, const FOperatorSettings& InSettings, const FMetasoundEnvironment& InEnvironment, TArray<FBuildErrorPtr>& InOutErrors)
+				:	Graph(InGraph)
+				,	AlgoAdapter(InAlgoAdapter)
+				,	Settings(InSettings)
+				,	Environment(InEnvironment)
+				,	Errors(InOutErrors)
+				{
+				}
+			};
+
 			// Perform topological sort using depth first algorithm.
-			FBuildStatus DepthFirstTopologicalSort(const FDirectedGraphAlgoAdapter& InAdapter, TArray<const INode*>& OutNodes, TArray<FBuildErrorPtr>& OutErrors) const;
+			FBuildStatus DepthFirstTopologicalSort(FBuildContext& InContext, TArray<const INode*>& OutNodes) const;
 
 			// Perform topological sort using kahns algorithm.
-			FBuildStatus KahnsTopologicalSort(const FDirectedGraphAlgoAdapter& InAdapter, TArray<const INode*>& OutNodes, TArray<FBuildErrorPtr>& OutErrors) const;
+			FBuildStatus KahnsTopologicalSort(FBuildContext& InContext, TArray<const INode*>& OutNodes) const;
 
 			// Prune unreachable nodes from InOutNodes
-			FBuildStatus PruneNodes(const FDirectedGraphAlgoAdapter& InAdapter, TArray<const INode*>& InOutNodes, TArray<FBuildErrorPtr>& OutErrors) const;
+			FBuildStatus PruneNodes(FBuildContext& InContext, TArray<const INode*>& InOutNodes) const;
 
 			// Get all input data references for a given node.
-			FBuildStatus GatherInputDataReferences(const INode* InNode, const FNodeEdgeMultiMap& InEdgeMap, const FNodeDataReferenceMap& InDataReferenceMap, FDataReferenceCollection& OutCollection, TArray<FBuildErrorPtr>& OutErrors) const;
+			FBuildStatus GatherInputDataReferences(FBuildContext& InContext, const INode* InNode, const FNodeEdgeMultiMap& InEdgeMap, const FNodeDataReferenceMap& InDataReferenceMap, FDataReferenceCollection& OutCollection) const;
 
 			// Get all input/output data references for a given graph.
-			FBuildStatus GatherGraphDataReferences(const IGraph& InGraph, FNodeDataReferenceMap& InNodeDataReferences, FDataReferenceCollection& OutGraphInputs, FDataReferenceCollection& OutGraphOutputs, TArray<FBuildErrorPtr>& OutErrors) const;
+			FBuildStatus GatherGraphDataReferences(FBuildContext& InContext, FNodeDataReferenceMap& InNodeDataReferences, FDataReferenceCollection& OutGraphInputs, FDataReferenceCollection& OutGraphOutputs) const;
 
 			// Call the operator factories for the nodes
-			FBuildStatus CreateOperators(const IGraph& InGraph, const TArray<const INode*>& InSortedNodes, TArray<FOperatorPtr>& OutOperators, FNodeDataReferenceMap& OutDataReferences, TArray<FBuildErrorPtr>& OutErrors) const;
+			FBuildStatus CreateOperators(FBuildContext& InContext, const TArray<const INode*>& InSortedNodes, TArray<FOperatorPtr>& OutOperators, FNodeDataReferenceMap& OutDataReferences) const;
 
 			// Create the final graph operator from the individual operators.
-			TUniquePtr<IOperator> CreateGraphOperator(const IGraph& InGraph, TArray<FOperatorPtr>& InOperators, FNodeDataReferenceMap& InNodeDataReferences, TArray<FBuildErrorPtr>& OutErrors) const;
-
-			FOperatorSettings OperatorSettings;
+			TUniquePtr<IOperator> CreateGraphOperator(FBuildContext& InContext, TArray<FOperatorPtr>& InOperators, FNodeDataReferenceMap& InNodeDataReferences) const;
 
 			FOperatorBuilderSettings BuilderSettings;
 
diff --git a/Engine/Plugins/Runtime/Metasound/Source/MetasoundGraphCore/Public/MetasoundVertex.h b/Engine/Plugins/Runtime/Metasound/Source/MetasoundGraphCore/Public/MetasoundVertex.h
--- a/Engine/Plugins/Runtime/Metasound/Source/MetasoundGraphCore/Public/MetasoundVertex.h
+++ b/Engine/Plugins/Runtime/Metasound/Source/MetasoundGraphCore/Public/MetasoundVertex.h
@@ -5,6 +5,7 @@
 #include "CoreMinimal.h"
 #include "Containers/SortedMap.h"
 #include "MetasoundDataReference.h"
+#include "MetasoundEnvironment.h"
 
 #include <type_traits>
 
@@ -13,7 +14,8 @@ namespace Metasound
 	/** Key type for an FInputDataVertexColletion or 
 	 * FOutputDataVertexCollection. 
 	 */
-	typedef FString FDataVertexKey;
+	typedef FString FVertexKey;
+	typedef FString FDataVertexKey; // TODO: Remove
 
 	// Forward declare.
 	class FInputDataVertex;
@@ -21,6 +23,9 @@ namespace Metasound
 	// Forward declare.
 	class FOutputDataVertex;
 
+	// Forward declare
+	class FEnvironmentVertex;
+
 	/** FDataVertexModel
 	 *
 	 * An FDataVertexModel implements various vertex behaviors. 
@@ -350,28 +355,191 @@ namespace Metasound
 			TUniquePtr<FOutputDataVertexModel> VertexModel;
 	};
 
+	/** FEnvironmentVertexModel
+	 *
+	 * An FEnvironmentVertexModel implements various vertex behaviors. 
+	 *
+	 * The FEnvironmentVertex supplies copy constructors and assignment operators 
+	 * by utilizing the Clone() operator of the FEnvironmentVertexModel. 
+	 * This allows FEnvironmentVertexModel behavior to be passed to other objects 
+	 * without having to know the concrete implementation of the FEnvironmentVertexModel.
+	 */
+	struct FEnvironmentVertexModel
+	{
+		/** FEnvironmentVertexModel Construtor
+		 *
+		 * @InVertexName - Name of vertex.
+		 * @InDescription - Human readable vertex description.
+		 */
+		FEnvironmentVertexModel(const FString& InVertexName, const FText& InDescription)
+		:	VertexName(InVertexName)
+		,	Description(InDescription)
+		{
+		}
+
+		virtual ~FEnvironmentVertexModel() = default;
+
+		/** Name of vertex. */
+		const FString VertexName;
+
+		/** Description of the vertex. */
+		const FText Description;
+
+		/** Test if a IMetasoundEnvironmentVariable matches the DataType of this vertex model.
+		 *
+		 * @param InVariable - Environment variable object.
+		 *
+		 * @return True if the types are equal, false otherwise.
+		 */
+		virtual bool IsVariableOfSameType(const IMetasoundEnvironmentVariable& InVariable) const = 0;
+
+		virtual FMetasoundEnvironmentVariableTypeId GetVariableTypeId() const = 0;
+
+		/** Create a clone of this FEnvironmentVertexModel */
+		virtual TUniquePtr<FEnvironmentVertexModel> Clone() const = 0;
+
+		friend bool METASOUNDGRAPHCORE_API operator==(const FEnvironmentVertex& InLHS, const FEnvironmentVertex& InRHS);
+		friend bool METASOUNDGRAPHCORE_API operator!=(const FEnvironmentVertex& InLHS, const FEnvironmentVertex& InRHS);
 
-	/** Create an FDataVertexKey from a FDataVertex. */
-	FORCEINLINE FDataVertexKey MakeDataVertexKey(const FInputDataVertex& InVertex)
+		private:
+
+		/** Compare another FEnvironmentVertexModel is equal to this
+		 *
+		 * @return True if models are equal, false otherwise.
+		 */
+		virtual bool IsEqual(const FEnvironmentVertexModel& InOther) const = 0;
+	};
+
+	/** TEnvironmentVertexModel creates a simple, unchanging, environment vertex. */
+	template<typename VarType>
+	struct TEnvironmentVertexModel : FEnvironmentVertexModel
+	{
+		/** Inherited constructor. */
+		using FEnvironmentVertexModel::FEnvironmentVertexModel;
+
+		/** Test if a IMetasoundEnvironmentVariable matches the DataType of this vertex model.
+		 *
+		 * @param InVariable - Environment variable object.
+		 *
+		 * @return True if the types are equal, false otherwise.
+		 */
+		bool IsVariableOfSameType(const IMetasoundEnvironmentVariable& InVariable) const override
+		{
+			return IsEnvironmentVariableOfType<VarType>(InVariable);
+		}
+		
+		/** Return the type id of the vertex. */
+		FMetasoundEnvironmentVariableTypeId GetVariableTypeId() const override
+		{
+			return GetMetasoundEnvironmentVariableTypeId<VarType>();
+		}
+
+		/** Create a clone of this TEnvironmentVertexModel*/
+		virtual TUniquePtr<FEnvironmentVertexModel> Clone() const override
+		{
+			return MakeUnique<TEnvironmentVertexModel<VarType>>(this->VertexName, this->Description);
+		}
+
+		friend bool METASOUNDGRAPHCORE_API operator==(const FEnvironmentVertex& InLHS, const FEnvironmentVertex& InRHS);
+		friend bool METASOUNDGRAPHCORE_API operator!=(const FEnvironmentVertex& InLHS, const FEnvironmentVertex& InRHS);
+	private:
+
+		bool IsEqual(const FEnvironmentVertexModel& InOther) const override
+		{
+			return (GetVariableTypeId() == InOther.GetVariableTypeId()) && (VertexName == InOther.VertexName);
+		}
+	};
+
+	/** FEnvironmentVertex represents an input vertex to an interface. It uses
+	 * a FEnvironmentVertexModel to determine its behavior.
+	 */
+	class METASOUNDGRAPHCORE_API FEnvironmentVertex
+	{
+		using FEmptyVertexModel = TEnvironmentVertexModel<void>;
+
+		public:
+
+			/** Construct an FEnvironmentVertex with a given vertex model.
+			 *
+			 * @param InVertexModel - A model subclass of FEnvironmentVertexModel.
+			 */
+			template<typename VertexModelType>
+			FEnvironmentVertex(const VertexModelType& InVertexModel)
+			:	VertexModel(InVertexModel.Clone())
+			{
+				static_assert(TIsDerivedFrom<VertexModelType, FEnvironmentVertexModel>::Value, "Vertex implementation must be a subclass of FEnvironmentVertexModel");
+
+				if (!VertexModel.IsValid())
+				{
+					VertexModel = MakeUnique<FEmptyVertexModel>(TEXT(""), FText::GetEmpty());
+				}
+			}
+
+			/** Construct an empty FInputDataVertex. */
+			FEnvironmentVertex();
+
+			/** Copy constructor */
+			FEnvironmentVertex(const FEnvironmentVertex& InOther);
+
+			/** Assignment operator. */
+			FEnvironmentVertex& operator=(const FEnvironmentVertex& InOther);
+
+			/** Name of vertex. */
+			const FString& GetVertexName() const;
+
+			/** Description of the vertex. */
+			const FText& GetDescription() const;
+
+			/** Test if a IMetasoundEnvironmentVariable matches the DataType of this vertex.
+			 *
+			 * @param InVariable - Environment variable object.
+			 *
+			 * @return True if the types are equal, false otherwise.
+			 */
+			bool IsVariableOfSameType(const IMetasoundEnvironmentVariable& InVariable) const;
+
+			friend bool METASOUNDGRAPHCORE_API operator==(const FEnvironmentVertex& InLHS, const FEnvironmentVertex& InRHS);
+			friend bool METASOUNDGRAPHCORE_API operator!=(const FEnvironmentVertex& InLHS, const FEnvironmentVertex& InRHS);
+			friend bool METASOUNDGRAPHCORE_API operator<(const FEnvironmentVertex& InLHS, const FEnvironmentVertex& InRHS);
+
+		private:
+
+			TUniquePtr<FEnvironmentVertexModel> VertexModel;
+	};
+
+
+	/** Create an FVertexKey from a FDataVertex. */
+	FORCEINLINE FVertexKey MakeDataVertexKey(const FInputDataVertex& InVertex)
 	{
 		return InVertex.GetVertexName();
 	}
 
-	/** Create an FDataVertexKey from a FDataVertex. */
-	FORCEINLINE FDataVertexKey MakeDataVertexKey(const FOutputDataVertex& InVertex)
+	/** Create an FVertexKey from a FDataVertex. */
+	FORCEINLINE FVertexKey MakeDataVertexKey(const FOutputDataVertex& InVertex)
+	{
+		return InVertex.GetVertexName();
+	}
+
+	/** Create an FVertexKey from a FEnvironmentVariableVertex. */
+	FORCEINLINE FVertexKey MakeEnvironmentVertexKey(const FEnvironmentVertex& InVertex)
 	{
 		return InVertex.GetVertexName();
 	}
 
 	/** TVertexInterfaceGroups encapsulates multiple related data vertices. It 
-	 * requires that each vertex in the group have a unique FDataVertexKey.
+	 * requires that each vertex in the group have a unique FVertexKey.
 	 */
 	template<typename VertexType>
 	class TVertexInterfaceGroup
 	{
-			static_assert(std::is_base_of<FInputDataVertex, VertexType>::value || std::is_base_of<FOutputDataVertex, VertexType>::value, "VertexType must be derived from either FInputDataVertex or FOutputDataVertex");
+			static constexpr bool bIsDerivedFromFInputDataVertex = std::is_base_of<FInputDataVertex, VertexType>::value;
+			static constexpr bool bIsDerivedFromFOutputDataVertex = std::is_base_of<FOutputDataVertex, VertexType>::value;
+			static constexpr bool bIsDerivedFromFEnvironmentVertex = std::is_base_of<FEnvironmentVertex, VertexType>::value;
+			static constexpr bool bIsSupportedVertexType = bIsDerivedFromFInputDataVertex || bIsDerivedFromFOutputDataVertex || bIsDerivedFromFEnvironmentVertex;
 
-			using FContainerType = TSortedMap<FDataVertexKey, VertexType>;
+			static_assert(bIsSupportedVertexType, "VertexType must be derived from FInputDataVertex, FOutputDataVertex, or FEnvironmentVertex");
+
+			using FContainerType = TSortedMap<FVertexKey, VertexType>;
 
 			// Required for end of recursion.
 			static void CopyInputs(FContainerType& InStorage)
@@ -396,7 +564,7 @@ namespace Metasound
 
 		public:
 
-			using RangedForConstIteratorType = typename TSortedMap<FDataVertexKey, VertexType>::RangedForConstIteratorType;
+			using RangedForConstIteratorType = typename TSortedMap<FVertexKey, VertexType>::RangedForConstIteratorType;
 
 			/** TVertexInterfaceGroup constructor with variadic list of vertex
 			 * models.
@@ -422,20 +590,20 @@ namespace Metasound
 			}
 
 			/** Remove a vertex by key. */
-			bool Remove(const FDataVertexKey& InKey)
+			bool Remove(const FVertexKey& InKey)
 			{
 				int32 NumRemoved = Vertices.Remove();
 				return (NumRemoved > 0);
 			}
 
 			/** Returns true if the group contains a vertex with a matching key. */
-			bool Contains(const FDataVertexKey& InKey) const
+			bool Contains(const FVertexKey& InKey) const
 			{
 				return Vertices.Contains(InKey);
 			}
 
 			/** Return the vertex for a given vertex key. */
-			const VertexType& operator[](const FDataVertexKey& InKey) const
+			const VertexType& operator[](const FVertexKey& InKey) const
 			{
 				return Vertices[InKey];
 			}
@@ -482,6 +650,9 @@ namespace Metasound
 	/** FOutputVertexInterface is a TVertexInterfaceGroup which holds FOutputDataVertexes. */
 	typedef TVertexInterfaceGroup<FOutputDataVertex> FOutputVertexInterface;
 
+	/** FEnvironmentVertexInterface is a TVertexInterfaceGroup which holds FEnvironmentVertexes. */
+	typedef TVertexInterfaceGroup<FEnvironmentVertex> FEnvironmentVertexInterface;
+
 	/** FVertexInterface provides access to a collection of input and output vertex
 	 * interfaces. 
 	 */
@@ -495,6 +666,9 @@ namespace Metasound
 			/** Construct with an input and output interface. */
 			FVertexInterface(const FInputVertexInterface& InInputs, const FOutputVertexInterface& InOutputs);
 
+			/** Construct with input, output and environment interface. */
+			FVertexInterface(const FInputVertexInterface& InInputs, const FOutputVertexInterface& InOutputs, const FEnvironmentVertexInterface& InEnvironmentVariables);
+
 			/** Return the input interface. */
 			const FInputVertexInterface& GetInputInterface() const;
 
@@ -502,10 +676,10 @@ namespace Metasound
 			FInputVertexInterface& GetInputInterface();
 
 			/** Return an input vertex. */
-			const FInputDataVertex& GetInputVertex(const FDataVertexKey& InKey) const;
+			const FInputDataVertex& GetInputVertex(const FVertexKey& InKey) const;
 
 			/** Returns true if an input vertex with the given key exists. */
-			bool ContainsInputVertex(const FDataVertexKey& InKey) const;
+			bool ContainsInputVertex(const FVertexKey& InKey) const;
 
 			/** Return the output interface. */
 			const FOutputVertexInterface& GetOutputInterface() const;
@@ -514,10 +688,22 @@ namespace Metasound
 			FOutputVertexInterface& GetOutputInterface();
 
 			/** Return an output vertex. */
-			const FOutputDataVertex& GetOutputVertex(const FDataVertexKey& InKey) const;
+			const FOutputDataVertex& GetOutputVertex(const FVertexKey& InKey) const;
+
+			/** Returns true if an outptu vertex with the given key exists. */
+			bool ContainsOutputVertex(const FVertexKey& InKey) const;
+
+			/** Return the output interface. */
+			const FEnvironmentVertexInterface& GetEnvironmentInterface() const;
+
+			/** Return the output interface. */
+			FEnvironmentVertexInterface& GetEnvironmentInterface();
+
+			/** Return an output vertex. */
+			const FEnvironmentVertex& GetEnvironmentVertex(const FVertexKey& InKey) const;
 
 			/** Returns true if an outptu vertex with the given key exists. */
-			bool ContainsOutputVertex(const FDataVertexKey& InKey) const;
+			bool ContainsEnvironmentVertex(const FVertexKey& InKey) const;
 
 			/** Test for equality between two interfaces. */
 			friend bool METASOUNDGRAPHCORE_API operator==(const FVertexInterface& InLHS, const FVertexInterface& InRHS);
@@ -529,5 +715,6 @@ namespace Metasound
 
 			FInputVertexInterface InputInterface;
 			FOutputVertexInterface OutputInterface;
+			FEnvironmentVertexInterface EnvironmentInterface;
 	};
 }
