Fix static registration issue (has to be done after class implementation.  TODO: add static analysis warnings to circumvent this)

diff --git a/Engine/Plugins/Runtime/Metasound/Source/MetasoundEngine/Private/MetasoundWavePlayerNode.cpp b/Engine/Plugins/Runtime/Metasound/Source/MetasoundEngine/Private/MetasoundWavePlayerNode.cpp
--- a/Engine/Plugins/Runtime/Metasound/Source/MetasoundEngine/Private/MetasoundWavePlayerNode.cpp
+++ b/Engine/Plugins/Runtime/Metasound/Source/MetasoundEngine/Private/MetasoundWavePlayerNode.cpp
@@ -15,8 +15,6 @@
 
 namespace Metasound
 {
-	METASOUND_REGISTER_NODE(FWavePlayerNode)
-	
 	// WavePlayer custom error 
 	class FWavePlayerError : public FBuildErrorBase
 	{
@@ -225,5 +223,7 @@ namespace Metasound
 	{
 		return InInterface == Interface;
 	}
-}
-#undef LOCTEXT_NAMESPACE //MetasoundWaveNode
+
+	METASOUND_REGISTER_NODE(FWavePlayerNode)
+} // namespace Metasound
+#undef LOCTEXT_NAMESPACE // MetasoundWaveNode
