Replace FDataVertexKey with FVertexKey

diff --git a/Engine/Plugins/Runtime/Metasound/Source/MetasoundGraphCore/Private/MetasoundGraph.cpp b/Engine/Plugins/Runtime/Metasound/Source/MetasoundGraphCore/Private/MetasoundGraph.cpp
--- a/Engine/Plugins/Runtime/Metasound/Source/MetasoundGraphCore/Private/MetasoundGraph.cpp
+++ b/Engine/Plugins/Runtime/Metasound/Source/MetasoundGraphCore/Private/MetasoundGraph.cpp
@@ -37,7 +37,7 @@ namespace Metasound
 		return PluginNodeMissingPrompt;
 	}
 
-	bool FGraph::AddInputDataDestination(const INode& InNode, const FDataVertexKey& InVertexKey)
+	bool FGraph::AddInputDataDestination(const INode& InNode, const FVertexKey& InVertexKey)
 	{
 		if (!InNode.GetVertexInterface().ContainsInputVertex(InVertexKey))
 		{
@@ -62,7 +62,7 @@ namespace Metasound
 		return InputDestinations;
 	}
 
-	bool FGraph::AddOutputDataSource(const INode& InNode, const FDataVertexKey& InVertexKey)
+	bool FGraph::AddOutputDataSource(const INode& InNode, const FVertexKey& InVertexKey)
 	{
 		if (!InNode.GetVertexInterface().ContainsOutputVertex(InVertexKey))
 		{
@@ -92,7 +92,7 @@ namespace Metasound
 		Edges.Add(InEdge);
 	}
 
-	bool FGraph::AddDataEdge(const INode& FromNode, const FDataVertexKey& FromKey, const INode& ToNode, const FDataVertexKey& ToKey)
+	bool FGraph::AddDataEdge(const INode& FromNode, const FVertexKey& FromKey, const INode& ToNode, const FVertexKey& ToKey)
 	{
 		const FVertexInterface& FromVertexInterface = FromNode.GetVertexInterface();
 		const FVertexInterface& ToVertexInterface = ToNode.GetVertexInterface();
diff --git a/Engine/Plugins/Runtime/Metasound/Source/MetasoundGraphCore/Private/MetasoundGraphLinter.cpp b/Engine/Plugins/Runtime/Metasound/Source/MetasoundGraphCore/Private/MetasoundGraphLinter.cpp
--- a/Engine/Plugins/Runtime/Metasound/Source/MetasoundGraphCore/Private/MetasoundGraphLinter.cpp
+++ b/Engine/Plugins/Runtime/Metasound/Source/MetasoundGraphCore/Private/MetasoundGraphLinter.cpp
@@ -27,7 +27,7 @@ namespace Metasound
 
 			const FInputVertexInterface& InputInterface = InDest.Node->GetVertexInterface().GetInputInterface();
 
-			FDataVertexKey VertexKey = MakeDataVertexKey(InDest.Vertex);
+			FVertexKey VertexKey = MakeDataVertexKey(InDest.Vertex);
 
 			if (!InputInterface.Contains(VertexKey))
 			{
@@ -52,7 +52,7 @@ namespace Metasound
 
 			const FOutputVertexInterface& OutputInterface = InSource.Node->GetVertexInterface().GetOutputInterface();
 
-			FDataVertexKey VertexKey = MakeDataVertexKey(InSource.Vertex);
+			FVertexKey VertexKey = MakeDataVertexKey(InSource.Vertex);
 
 			if (!OutputInterface.Contains(VertexKey))
 			{
diff --git a/Engine/Plugins/Runtime/Metasound/Source/MetasoundGraphCore/Public/MetasoundGraph.h b/Engine/Plugins/Runtime/Metasound/Source/MetasoundGraphCore/Public/MetasoundGraph.h
--- a/Engine/Plugins/Runtime/Metasound/Source/MetasoundGraphCore/Public/MetasoundGraph.h
+++ b/Engine/Plugins/Runtime/Metasound/Source/MetasoundGraphCore/Public/MetasoundGraph.h
@@ -78,7 +78,7 @@ namespace Metasound
 			 *
 			 * @return True if the edge was successfully added. False otherwise.
 			 */
-			bool AddDataEdge(const INode& FromNode, const FDataVertexKey& FromVertexKey, const INode& ToNode, const FDataVertexKey& ToVertexKey);
+			bool AddDataEdge(const INode& FromNode, const FVertexKey& FromVertexKey, const INode& ToNode, const FVertexKey& ToVertexKey);
 
 			/** Add an input data destination to describe how data provided 
 			 * outside this graph should be routed internally.
@@ -89,7 +89,7 @@ namespace Metasound
 			 * @return True if the destination was successfully added. False 
 			 * otherwise.
 			 */
-			bool AddInputDataDestination(const INode& InNode, const FDataVertexKey& InVertexKey);
+			bool AddInputDataDestination(const INode& InNode, const FVertexKey& InVertexKey);
 
 
 			/** Add an input data destination to describe how data provided 
@@ -106,7 +106,7 @@ namespace Metasound
 			 * @return True if the source was successfully added. False 
 			 * otherwise.
 			 */
-			bool AddOutputDataSource(const INode& InNode, const FDataVertexKey& InVertexKey);
+			bool AddOutputDataSource(const INode& InNode, const FVertexKey& InVertexKey);
 
 			/** Add an output data source which describes routing of data which is 
 			 * owned this graph and exposed externally.
diff --git a/Engine/Plugins/Runtime/Metasound/Source/MetasoundGraphCore/Public/MetasoundVertex.h b/Engine/Plugins/Runtime/Metasound/Source/MetasoundGraphCore/Public/MetasoundVertex.h
--- a/Engine/Plugins/Runtime/Metasound/Source/MetasoundGraphCore/Public/MetasoundVertex.h
+++ b/Engine/Plugins/Runtime/Metasound/Source/MetasoundGraphCore/Public/MetasoundVertex.h
@@ -15,7 +15,6 @@ namespace Metasound
 	 * FOutputDataVertexCollection. 
 	 */
 	typedef FString FVertexKey;
-	typedef FString FDataVertexKey; // TODO: Remove
 
 	// Forward declare.
 	class FInputDataVertex;
