Metasound metadata. Adding metadata to node info and adding static registration

diff --git a/Engine/Plugins/Runtime/Metasound/Source/MetasoundEditor/Private/MetasoundDetailCustomization.cpp b/Engine/Plugins/Runtime/Metasound/Source/MetasoundEditor/Private/MetasoundDetailCustomization.cpp
--- a/Engine/Plugins/Runtime/Metasound/Source/MetasoundEditor/Private/MetasoundDetailCustomization.cpp
+++ b/Engine/Plugins/Runtime/Metasound/Source/MetasoundEditor/Private/MetasoundDetailCustomization.cpp
@@ -155,14 +155,20 @@ namespace Metasound
 			const FName AuthorPropertyPath = BuildChildPath(GetMetadataPropertyPath(), GET_MEMBER_NAME_CHECKED(FMetasoundClassMetadata, AuthorName));
 			const FName DescPropertyPath = BuildChildPath(GetMetadataPropertyPath(), GET_MEMBER_NAME_CHECKED(FMetasoundClassMetadata, MetasoundDescription));
 			const FName NodeTypePropertyPath = BuildChildPath(GetMetadataPropertyPath(), GET_MEMBER_NAME_CHECKED(FMetasoundClassMetadata, NodeType));
+			const FName MajorVersionPropertyPath = BuildChildPath(GetMetadataPropertyPath(), GET_MEMBER_NAME_CHECKED(FMetasoundClassMetadata, MajorVersion));
+			const FName MinorVersionPropertyPath = BuildChildPath(GetMetadataPropertyPath(), GET_MEMBER_NAME_CHECKED(FMetasoundClassMetadata, MinorVersion));
 
 			TSharedPtr<IPropertyHandle> AuthorHandle = DetailLayout.GetProperty(AuthorPropertyPath);
 			TSharedPtr<IPropertyHandle> DescHandle = DetailLayout.GetProperty(DescPropertyPath);
 			TSharedPtr<IPropertyHandle> NodeTypeHandle = DetailLayout.GetProperty(NodeTypePropertyPath);
+			TSharedPtr<IPropertyHandle> MajorVersionHandle = DetailLayout.GetProperty(MajorVersionPropertyPath);
+			TSharedPtr<IPropertyHandle> MinorVersionHandle = DetailLayout.GetProperty(MinorVersionPropertyPath);
 
 			GeneralCategoryBuilder.AddProperty(NodeTypeHandle);
 			GeneralCategoryBuilder.AddProperty(AuthorHandle);
 			GeneralCategoryBuilder.AddProperty(DescHandle);
+			GeneralCategoryBuilder.AddProperty(MajorVersionHandle);
+			GeneralCategoryBuilder.AddProperty(MinorVersionHandle);
 
 			// Input/Output Categories
 
diff --git a/Engine/Plugins/Runtime/Metasound/Source/MetasoundEditor/Private/MetasoundEditorGraphNode.cpp b/Engine/Plugins/Runtime/Metasound/Source/MetasoundEditor/Private/MetasoundEditorGraphNode.cpp
--- a/Engine/Plugins/Runtime/Metasound/Source/MetasoundEditor/Private/MetasoundEditorGraphNode.cpp
+++ b/Engine/Plugins/Runtime/Metasound/Source/MetasoundEditor/Private/MetasoundEditorGraphNode.cpp
@@ -333,7 +333,7 @@ FText UMetasoundEditorGraphNode::GetTooltipText() const
 		case EMetasoundClassType::External:
 		case EMetasoundClassType::MetasoundGraph:
 		{
-			return GenerateClassDescriptionForNode(NodeHandle.GetClassInfo()).Metadata.MetasoundDescription;
+			return GenerateClassDescription(NodeHandle.GetClassInfo()).Metadata.MetasoundDescription;
 		}
 		break;
 	}
diff --git a/Engine/Plugins/Runtime/Metasound/Source/MetasoundEditor/Private/MetasoundEditorGraphSchema.cpp b/Engine/Plugins/Runtime/Metasound/Source/MetasoundEditor/Private/MetasoundEditorGraphSchema.cpp
--- a/Engine/Plugins/Runtime/Metasound/Source/MetasoundEditor/Private/MetasoundEditorGraphSchema.cpp
+++ b/Engine/Plugins/Runtime/Metasound/Source/MetasoundEditor/Private/MetasoundEditorGraphSchema.cpp
@@ -575,7 +575,7 @@ void UMetasoundEditorGraphSchema::GetAllMetasoundActions(FGraphActionMenuBuilder
 	const TArray<FNodeClassInfo> ClassInfos = GetAllAvailableNodeClasses();
 	for (const FNodeClassInfo& ClassInfo : ClassInfos)
 	{
-		const FMetasoundClassDescription ClassDescription = GenerateClassDescriptionForNode(ClassInfo);
+		const FMetasoundClassDescription ClassDescription = GenerateClassDescription(ClassInfo);
 		const FMetasoundClassMetadata Metadata = ClassDescription.Metadata;
 		const FText Tooltip = Metadata.AuthorName.IsEmpty()
 			? Metadata.MetasoundDescription
diff --git a/Engine/Plugins/Runtime/Metasound/Source/MetasoundEditor/Private/MetasoundFactory.cpp b/Engine/Plugins/Runtime/Metasound/Source/MetasoundEditor/Private/MetasoundFactory.cpp
--- a/Engine/Plugins/Runtime/Metasound/Source/MetasoundEditor/Private/MetasoundFactory.cpp
+++ b/Engine/Plugins/Runtime/Metasound/Source/MetasoundEditor/Private/MetasoundFactory.cpp
@@ -20,10 +20,15 @@ UMetasoundFactory::UMetasoundFactory(const FObjectInitializer& ObjectInitializer
 UObject* UMetasoundFactory::FactoryCreateNew(UClass* InClass, UObject* InParent, FName Name, EObjectFlags Flags, UObject* InContext, FFeedbackContext* InFeedbackContext)
 {
 	UMetasound* NewMetasound = NewObject<UMetasound>(InParent, Name, Flags);
+
 	FMetasoundClassMetadata Metadata;
+
 	Metadata.NodeName = NewMetasound->GetName();
+	Metadata.MajorVersion = 1;
+	Metadata.MinorVersion = 0;
 	Metadata.NodeType = EMetasoundClassType::MetasoundGraph;
 	Metadata.AuthorName = FText::FromString(UKismetSystemLibrary::GetPlatformUserName());
+
 	NewMetasound->SetMetadata(Metadata);
 
 	NewMetasound->ConformDocumentToArchetype();
@@ -44,12 +49,19 @@ UMetasoundSourceFactory::UMetasoundSourceFactory(const FObjectInitializer& Objec
 UObject* UMetasoundSourceFactory::FactoryCreateNew(UClass* InClass, UObject* InParent, FName Name, EObjectFlags Flags, UObject* InContext, FFeedbackContext* InFeedbackContext)
 {
 	UMetasoundSource* MetasoundSource = NewObject<UMetasoundSource>(InParent, Name, Flags);
+
 	FMetasoundClassMetadata Metadata;
+
 	Metadata.NodeName = MetasoundSource->GetName();
+	Metadata.MajorVersion = 1;
+	Metadata.MinorVersion = 0;
 	Metadata.NodeType = EMetasoundClassType::MetasoundGraph;
 	Metadata.AuthorName = FText::FromString(UKismetSystemLibrary::GetPlatformUserName());
+
 	MetasoundSource->SetMetadata(Metadata);
+
 	MetasoundSource->ConformDocumentToArchetype();
 
 	return MetasoundSource;
 }
+
diff --git a/Engine/Plugins/Runtime/Metasound/Source/MetasoundEngine/Private/MetasoundWavePlayerNode.cpp b/Engine/Plugins/Runtime/Metasound/Source/MetasoundEngine/Private/MetasoundWavePlayerNode.cpp
--- a/Engine/Plugins/Runtime/Metasound/Source/MetasoundEngine/Private/MetasoundWavePlayerNode.cpp
+++ b/Engine/Plugins/Runtime/Metasound/Source/MetasoundEngine/Private/MetasoundWavePlayerNode.cpp
@@ -122,13 +122,6 @@ namespace Metasound
 		TUniquePtr<Audio::IDecoderOutput> DecoderOutput;
 	};
 
-	const FNodeInfo FWavePlayerNode::Info = FNodeInfo(
-	{
-		FName(TEXT("Wave Player")),
-		LOCTEXT("Metasound_WavePlayerNodeDescription", "Plays a supplied Wave"),
-		PluginAuthor,
-		PluginNodeMissingPrompt
-	});
 
 	TUniquePtr<IOperator> FWavePlayerNode::FOperatorFactory::CreateOperator(
 		const FCreateOperatorParams& InParams, 
@@ -186,12 +179,44 @@ namespace Metasound
 		return MakeUnique<FWavePlayerOperator>(InParams.OperatorSettings, Wave);
 	}
 
+	FVertexInterface FWavePlayerNode::DeclareVertexInterface()
+	{
+		return FVertexInterface(
+			FInputVertexInterface(
+				TInputDataVertexModel<FWaveAsset>(TEXT("Wave"), LOCTEXT("WaveTooltip", "The Wave to be decoded"))
+			),
+			FOutputVertexInterface(
+				TOutputDataVertexModel<FAudioBuffer>(TEXT("Audio"), LOCTEXT("AudioTooltip", "The output audio"))
+			)
+		);
+	}
+
+	const FNodeInfo& FWavePlayerNode::GetNodeInfo()
+	{
+		auto InitNodeInfo = []() -> FNodeInfo
+		{
+			FNodeInfo Info;
+			Info.ClassName = FName(TEXT("Wave Player"));
+			Info.MajorVersion = 1;
+			Info.MinorVersion = 0;
+			Info.Description = LOCTEXT("Metasound_WavePlayerNodeDescription", "Plays a supplied Wave");
+			Info.Author = PluginAuthor;
+			Info.PromptIfMissing = PluginNodeMissingPrompt;
+			Info.DefaultInterface = DeclareVertexInterface();
+
+			return Info;
+		};
+
+		static const FNodeInfo Info = InitNodeInfo();
+
+		return Info;
+	}
+
 	FWavePlayerNode::FWavePlayerNode(const FString& InName)
-		:	FNode(InName, FWavePlayerNode::Info)
+		:	FNode(InName, GetNodeInfo())
 		,	Factory(MakeOperatorFactoryRef<FWavePlayerNode::FOperatorFactory>())
+		,	Interface(DeclareVertexInterface())
 	{
-		Interface.GetInputInterface().Add(TInputDataVertexModel<FWaveAsset>(TEXT("Wave"), LOCTEXT("WaveTooltip", "The Wave to be decoded")));
-		Interface.GetOutputInterface().Add(TOutputDataVertexModel<FAudioBuffer>(TEXT("Audio"), LOCTEXT("AudioTooltip", "The output audio")));
 	}
 
 	FWavePlayerNode::FWavePlayerNode(const FNodeInitData& InInitData)
@@ -199,17 +224,14 @@ namespace Metasound
 	{
 	}
 
-	FOperatorFactorySharedRef FWavePlayerNode::GetDefaultOperatorFactory() const
+	FOperatorFactorySharedRef FWavePlayerNode::GetDefaultOperatorFactory() const 
 	{
 		return Factory;
 	}
 
-	const FVertexInterface& FWavePlayerNode::GetVertexInterface() const
-	{
-		return Interface;
-	}
 
-	const FVertexInterface& FWavePlayerNode::GetDefaultVertexInterface() const
+
+	const FVertexInterface& FWavePlayerNode::GetVertexInterface() const
 	{
 		return Interface;
 	}
diff --git a/Engine/Plugins/Runtime/Metasound/Source/MetasoundEngine/Public/MetasoundWavePlayerNode.h b/Engine/Plugins/Runtime/Metasound/Source/MetasoundEngine/Public/MetasoundWavePlayerNode.h
--- a/Engine/Plugins/Runtime/Metasound/Source/MetasoundEngine/Public/MetasoundWavePlayerNode.h
+++ b/Engine/Plugins/Runtime/Metasound/Source/MetasoundEngine/Public/MetasoundWavePlayerNode.h
@@ -17,8 +17,11 @@ namespace Metasound
 			virtual TUniquePtr<IOperator> CreateOperator(const FCreateOperatorParams& InParams, FBuildErrorArray& OutErrors) override;
 		};
 
+		
+
 	public:
-		static const FNodeInfo Info;
+		static FVertexInterface DeclareVertexInterface();
+		static const FNodeInfo& GetNodeInfo();
 
 		FWavePlayerNode(const FString& InName);
 
@@ -32,9 +35,6 @@ namespace Metasound
 		/** Return the current vertex interface. */
 		virtual const FVertexInterface& GetVertexInterface() const override;
 
-		/** Return the default vertex interface. */
-		virtual const FVertexInterface& GetDefaultVertexInterface() const override;
-
 		/** Set the vertex interface. If the vertex was successfully changed, returns true. 
 		 *
 		 * @param InInterface - New interface for node. 
diff --git a/Engine/Plugins/Runtime/Metasound/Source/MetasoundFrontend/Private/MetasoundFrontend.cpp b/Engine/Plugins/Runtime/Metasound/Source/MetasoundFrontend/Private/MetasoundFrontend.cpp
--- a/Engine/Plugins/Runtime/Metasound/Source/MetasoundFrontend/Private/MetasoundFrontend.cpp
+++ b/Engine/Plugins/Runtime/Metasound/Source/MetasoundFrontend/Private/MetasoundFrontend.cpp
@@ -51,167 +51,76 @@ namespace Metasound
 			return OutClasses;
 		}
 
-		TArray<FNodeClassInfo> GetAllNodeClassesInNamespace(const FString& InNamespace)
-		{
-			TArray<FNodeClassInfo> OutClasses;
-
-			auto& Registry = GetExternalNodeRegistry();
-			for (auto& NodeClassTuple : Registry)
-			{
-				// TODO: Build fname namespace tree. as we register nodes and types.
-				// for now we just parse the string here.
-				FString NodeName = NodeClassTuple.Key.NodeName.ToString();
-				if (NodeName.StartsWith(InNamespace))
-				{
-					FNodeClassInfo ClassInfo;
-					ClassInfo.NodeName = NodeName;
-					ClassInfo.NodeType = EMetasoundClassType::External;
-					ClassInfo.LookupKey = NodeClassTuple.Key;
-
-					OutClasses.Add(ClassInfo);
-				}
-			}
 
-			return OutClasses;
+		// gets all metadata (name, description, author, what to say if it's missing) for a given node.
+		FMetasoundClassMetadata GenerateClassMetadata(const FNodeClassInfo& InInfo)
+		{
+			return GenerateClassDescription(InInfo).Metadata;
 		}
 
-		TArray<FNodeClassInfo> GetAllNodesWhoseNameContains(const FString& InSubstring)
+		FMetasoundClassDescription GenerateClassDescription(const FNodeInfo& InNodeMetadata)
 		{
-			TArray<FNodeClassInfo> OutClasses;
+			FMetasoundClassDescription ClassDescription;
 
-			auto& Registry = GetExternalNodeRegistry();
-			for (auto& NodeClassTuple : Registry)
+			// Set metadata for class description.
+			ClassDescription.Metadata = FMetasoundClassMetadata
 			{
-				FString NodeName = NodeClassTuple.Key.NodeName.ToString();
-				if (NodeName.Contains(InSubstring))
-				{
-					FNodeClassInfo ClassInfo;
-					ClassInfo.NodeName = NodeName;
-					ClassInfo.NodeType = EMetasoundClassType::External;
-					ClassInfo.LookupKey = NodeClassTuple.Key;
-
-					OutClasses.Add(ClassInfo);
-				}
-			}
-
-			return OutClasses;
-		}
-
-		TArray<FNodeClassInfo> GetAllNodesWithAnOutputOfType(const FName& InType)
-		{
-			TArray<FNodeClassInfo> OutClasses;
+				InNodeMetadata.ClassName.ToString(),
+				InNodeMetadata.MajorVersion,
+				InNodeMetadata.MinorVersion,
+				EMetasoundClassType::External,
+				InNodeMetadata.Description,
+				InNodeMetadata.PromptIfMissing,
+				InNodeMetadata.Author
+			};
 
-			auto& Registry = GetExternalNodeRegistry();
-			for (auto& NodeClassTuple : Registry)
+			// External metasounds aren't dependent on any other nodes by definition, so all we need to do
+			// is populate the Input and Output sets.
+			for (auto& InputTuple : InNodeMetadata.DefaultInterface.GetInputInterface())
 			{
-				if (NodeClassTuple.Value.OutputTypes.Contains(InType))
-				{
-					FNodeClassInfo ClassInfo;
-					ClassInfo.NodeName = NodeClassTuple.Key.NodeName.ToString();
-					ClassInfo.NodeType = EMetasoundClassType::External;
-					ClassInfo.LookupKey = NodeClassTuple.Key;
+				FMetasoundInputDescription InputDescription;
+				InputDescription.Name = InputTuple.Value.GetVertexName();
+				InputDescription.TypeName = InputTuple.Value.GetDataTypeName();
+				InputDescription.ToolTip = InputTuple.Value.GetDescription();
 
-					OutClasses.Add(ClassInfo);
-					break;
-				}
+				ClassDescription.Inputs.Add(InputDescription);
 			}
 
-			return OutClasses;
-		}
+			for (auto& OutputTuple : InNodeMetadata.DefaultInterface.GetOutputInterface())
+			{
+				FMetasoundOutputDescription OutputDescription;
+				OutputDescription.Name = OutputTuple.Value.GetVertexName();
+				OutputDescription.TypeName = OutputTuple.Value.GetDataTypeName();
+				OutputDescription.ToolTip = OutputTuple.Value.GetDescription();
 
-		TArray<FNodeClassInfo> GetAllNodesWithAnInputOfType(const FName& InType)
-		{
-			TArray<FNodeClassInfo> OutClasses;
+				ClassDescription.Outputs.Add(OutputDescription);
+			}
 
-			auto& Registry = GetExternalNodeRegistry();
-			for (auto& NodeClassTuple : Registry)
+			for (auto& EnvTuple : InNodeMetadata.DefaultInterface.GetEnvironmentInterface())
 			{
-				if (NodeClassTuple.Value.InputTypes.Contains(InType))
-				{
-					FNodeClassInfo ClassInfo;
-					ClassInfo.NodeName = NodeClassTuple.Key.NodeName.ToString();
-					ClassInfo.NodeType = EMetasoundClassType::External;
-					ClassInfo.LookupKey = NodeClassTuple.Key;
+				FMetasoundEnvironmentVariableDescription EnvironmentDescription;
+				EnvironmentDescription.Name = EnvTuple.Value.GetVertexName();
+				EnvironmentDescription.ToolTip = EnvTuple.Value.GetDescription();
 
-					OutClasses.Add(ClassInfo);
-					break;
-				}
+				ClassDescription.EnvironmentVariables.Add(EnvironmentDescription);
 			}
 
-			return OutClasses;
-		}
 
-		// gets all metadata (name, description, author, what to say if it's missing) for a given node.
-		FMetasoundClassMetadata GenerateMetadataForNode(const FNodeClassInfo& InInfo)
-		{
-			return GenerateClassDescriptionForNode(InInfo).Metadata;
+			// Populate lookup data.
+			FNodeRegistryKey Key = FMetasoundFrontendRegistryContainer::GetRegistryKey(InNodeMetadata);
+
+			ClassDescription.ExternalNodeClassLookupInfo.ExternalNodeClassName = Key.NodeName;
+			ClassDescription.ExternalNodeClassLookupInfo.ExternalNodeClassHash = Key.NodeHash;
+
+			return ClassDescription;
 		}
 
-		FMetasoundClassDescription GenerateClassDescriptionForNode(const FNodeClassInfo& InInfo)
+		FMetasoundClassDescription GenerateClassDescription(const FNodeClassInfo& InInfo)
 		{
-			FNodeInitData DummyInitData;
-			DummyInitData.InstanceName = FString(TEXT("Unused node for registration"));
-
 			auto& Registry = GetExternalNodeRegistry();
 			if (Registry.Contains(InInfo.LookupKey))
 			{
-				TUniquePtr<Metasound::INode> DummyNode = Registry[InInfo.LookupKey].GetterCallback(DummyInitData);
-
-				if (!ensureAlwaysMsgf(DummyNode.IsValid(), TEXT("Node class %s failed to return a valid node. Likely something is wrong with the METASOUND_REGISTER_NODE macro.")))
-				{
-					return FMetasoundClassDescription();
-				}
-
-				const FName NodeName = DummyNode->GetClassName();
-
-				FMetasoundClassMetadata NodeMetadata;
-				NodeMetadata.NodeName = NodeName.ToString();
-				NodeMetadata.NodeType = EMetasoundClassType::External;
-				NodeMetadata.AuthorName = DummyNode->GetAuthorName();
-				NodeMetadata.MetasoundDescription = DummyNode->GetDescription();
-				NodeMetadata.PromptIfMissing = DummyNode->GetPromptIfMissing();
-
-				FMetasoundClassDescription ClassDescription;
-				ClassDescription.Metadata = NodeMetadata;
-
-				const FVertexInterface& DummyNodeInterface = DummyNode->GetDefaultVertexInterface();
-
-				// External metasounds aren't dependent on any other nodes by definition, so all we need to do
-				// is populate the Input and Output sets.
-				for (auto& InputTuple : DummyNodeInterface.GetInputInterface())
-				{
-					FMetasoundInputDescription InputDescription;
-					InputDescription.Name = InputTuple.Value.GetVertexName();
-					InputDescription.TypeName = InputTuple.Value.GetDataTypeName();
-					InputDescription.ToolTip = InputTuple.Value.GetDescription();
-
-					ClassDescription.Inputs.Add(InputDescription);
-				}
-
-				for (auto& OutputTuple : DummyNodeInterface.GetOutputInterface())
-				{
-					FMetasoundOutputDescription OutputDescription;
-					OutputDescription.Name = OutputTuple.Value.GetVertexName();
-					OutputDescription.TypeName = OutputTuple.Value.GetDataTypeName();
-					OutputDescription.ToolTip = OutputTuple.Value.GetDescription();
-
-					ClassDescription.Outputs.Add(OutputDescription);
-				}
-
-				for (auto& EnvTuple : DummyNodeInterface.GetEnvironmentInterface())
-				{
-					FMetasoundEnvironmentVariableDescription EnvironmentDescription;
-					EnvironmentDescription.Name = EnvTuple.Value.GetVertexName();
-					EnvironmentDescription.ToolTip = EnvTuple.Value.GetDescription();
-
-					ClassDescription.EnvironmentVariables.Add(EnvironmentDescription);
-				}
-
-				// Populate lookup data.
-				ClassDescription.ExternalNodeClassLookupInfo.ExternalNodeClassName = InInfo.LookupKey.NodeName;
-				ClassDescription.ExternalNodeClassLookupInfo.ExternalNodeClassHash = InInfo.LookupKey.NodeHash;
-
-				return ClassDescription;
+				return Registry[InInfo.LookupKey].CreateClassDescription();
 			}
 			else
 			{
@@ -2020,7 +1929,7 @@ namespace Metasound
 			// If we haven't added a node of this class to the graph yet, add it to the dependencies for this class.
 			if (!bFoundMatchingDependencyInDocument)
 			{
-				FMetasoundClassDescription NewDependencyClassDescription = GenerateClassDescriptionForNode(InNodeClass);
+				FMetasoundClassDescription NewDependencyClassDescription = GenerateClassDescription(InNodeClass);
 				NewDependencyClassDescription.UniqueID = FindNewUniqueDependencyId();
 				Dependencies.Add(NewDependencyClassDescription);
 				GraphsClassDeclaration->DependencyIDs.Add(NewDependencyClassDescription.UniqueID);
diff --git a/Engine/Plugins/Runtime/Metasound/Source/MetasoundFrontend/Private/MetasoundFrontendQuerySteps.cpp b/Engine/Plugins/Runtime/Metasound/Source/MetasoundFrontend/Private/MetasoundFrontendQuerySteps.cpp
--- a/Engine/Plugins/Runtime/Metasound/Source/MetasoundFrontend/Private/MetasoundFrontendQuerySteps.cpp
+++ b/Engine/Plugins/Runtime/Metasound/Source/MetasoundFrontend/Private/MetasoundFrontendQuerySteps.cpp
@@ -14,7 +14,7 @@ namespace Metasound
 
 		for (const Frontend::FNodeClassInfo& ClassInfo : ClassInfos)
 		{
-			OutEntries.Emplace(TInPlaceType<FMetasoundClassDescription>(), Frontend::GenerateClassDescriptionForNode(ClassInfo));
+			OutEntries.Emplace(TInPlaceType<FMetasoundClassDescription>(), Frontend::GenerateClassDescription(ClassInfo));
 		}
 	}
 
diff --git a/Engine/Plugins/Runtime/Metasound/Source/MetasoundFrontend/Private/MetasoundFrontendRegistries.cpp b/Engine/Plugins/Runtime/Metasound/Source/MetasoundFrontend/Private/MetasoundFrontendRegistries.cpp
--- a/Engine/Plugins/Runtime/Metasound/Source/MetasoundFrontend/Private/MetasoundFrontendRegistries.cpp
+++ b/Engine/Plugins/Runtime/Metasound/Source/MetasoundFrontend/Private/MetasoundFrontendRegistries.cpp
@@ -9,34 +9,26 @@
 #define WITH_METASOUND_FRONTEND 0
 #endif
 
-FMetasoundFrontendRegistryContainer::FMetasoundFrontendRegistryContainer()
-	: bHasModuleBeenInitialized(false)
-{
-}
-
-FMetasoundFrontendRegistryContainer::FNodeRegistryKey FMetasoundFrontendRegistryContainer::GetRegistryKeyForNodeInternal(const Metasound::INode& InNode)
+namespace MetasoundFrontendRegistryPrivate
 {
-	using FInputVertexInterface = ::Metasound::FInputVertexInterface;
-	using FOutputVertexInterface = ::Metasound::FOutputVertexInterface;
+	// All registry keys should be created through this function to ensure consistency.
+	Metasound::Frontend::FNodeRegistryKey GetRegistryKey(const FName& InClassName, int32 InMajorVersion)
+	{
+		Metasound::Frontend::FNodeRegistryKey Key;
 
-	const FName NodeName = InNode.GetClassName();
-	const FInputVertexInterface& Inputs = InNode.GetDefaultVertexInterface().GetInputInterface();
-	const FOutputVertexInterface& Outputs = InNode.GetDefaultVertexInterface().GetOutputInterface();
+		Key.NodeName = InClassName;
 
-	// Construct a hash using a combination of the class name, input names and output names.
-	uint32 NodeHash = FCrc::StrCrc32(*NodeName.ToString());
+		// NodeHash is hash of node name and major version.
+		Key.NodeHash = FCrc::StrCrc32(*Key.NodeName.ToString());
+		Key.NodeHash = HashCombine(Key.NodeHash, FCrc::TypeCrc32(InMajorVersion));
 
-	for (auto& InputTuple : Inputs)
-	{
-		NodeHash = HashCombine(NodeHash, FCrc::StrCrc32(*InputTuple.Value.GetVertexTypeName().ToString()));
-	}
-
-	for (auto& OutputTuple : Outputs)
-	{
-		NodeHash = HashCombine(NodeHash, FCrc::StrCrc32(*OutputTuple.Value.GetVertexTypeName().ToString()));
+		return Key;
 	}
+}
 
-	return { NodeName, NodeHash };
+FMetasoundFrontendRegistryContainer::FMetasoundFrontendRegistryContainer()
+	: bHasModuleBeenInitialized(false)
+{
 }
 
 FMetasoundFrontendRegistryContainer* FMetasoundFrontendRegistryContainer::LazySingleton = nullptr;
@@ -113,7 +105,7 @@ TUniquePtr<Metasound::INode> FMetasoundFrontendRegistryContainer::ConstructInput
 {
 	if (ensureAlwaysMsgf(DataTypeRegistry.Contains(InInputType), TEXT("Couldn't find data type %s!"), *InInputType.ToString()))
 	{
-		return DataTypeRegistry[InInputType].Callbacks.InputNodeConstructor(MoveTemp(InParams));
+		return DataTypeRegistry[InInputType].Callbacks.CreateInputNode(MoveTemp(InParams));
 	}
 	else
 	{
@@ -125,7 +117,7 @@ TUniquePtr<Metasound::INode> FMetasoundFrontendRegistryContainer::ConstructOutpu
 {
 	if (ensureAlwaysMsgf(DataTypeRegistry.Contains(InOutputType), TEXT("Couldn't find data type %s!"), *InOutputType.ToString()))
 	{
-		return DataTypeRegistry[InOutputType].Callbacks.OutputNodeConstructor(InParams);
+		return DataTypeRegistry[InOutputType].Callbacks.CreateOutputNode(InParams);
 	}
 	else
 	{
@@ -137,7 +129,7 @@ Metasound::FDataTypeLiteralParam FMetasoundFrontendRegistryContainer::GenerateLi
 {
 	if (ensureAlwaysMsgf(DataTypeRegistry.Contains(InDataType), TEXT("Couldn't find data type %s!"), *InDataType.ToString()))
 	{
-		 Audio::IProxyDataPtr ProxyPtr = DataTypeRegistry[InDataType].Callbacks.ProxyConstructor(InObject);
+		 Audio::IProxyDataPtr ProxyPtr = DataTypeRegistry[InDataType].Callbacks.CreateAudioProxy(InObject);
 		 if (ensureAlwaysMsgf(ProxyPtr.IsValid(), TEXT("UObject failed to create a valid proxy!")))
 		 {
 			 return Metasound::FDataTypeLiteralParam(MoveTemp(ProxyPtr));
@@ -163,7 +155,7 @@ Metasound::FDataTypeLiteralParam FMetasoundFrontendRegistryContainer::GenerateLi
 		{
 			if (InObject)
 			{
-				Audio::IProxyDataPtr ProxyPtr = DataTypeRegistry[InDataType].Callbacks.ProxyConstructor(InObject);
+				Audio::IProxyDataPtr ProxyPtr = DataTypeRegistry[InDataType].Callbacks.CreateAudioProxy(InObject);
 				ensureAlwaysMsgf(ProxyPtr.IsValid(), TEXT("UObject failed to create a valid proxy!"));
 				ProxyArray.Add(MoveTemp(ProxyPtr));
 			}
@@ -189,7 +181,7 @@ TUniquePtr<Metasound::INode> FMetasoundFrontendRegistryContainer::ConstructExter
 	}
 	else
 	{
-		return ExternalNodeRegistry[RegistryKey].GetterCallback(InInitData);
+		return ExternalNodeRegistry[RegistryKey].CreateNode(InInitData);
 	}
 }
 
@@ -328,72 +320,42 @@ bool FMetasoundFrontendRegistryContainer::RegisterDataType(const ::Metasound::FD
 	}
 }
 
-bool FMetasoundFrontendRegistryContainer::RegisterExternalNode(FNodeGetterCallback&& InCallback)
+bool FMetasoundFrontendRegistryContainer::RegisterExternalNode(Metasound::FCreateMetasoundNodeFunction&& InCreateNode, Metasound::FCreateMetasoundClassDescriptionFunction&& InCreateDescription)
 {
-	using FVertexInterface = Metasound::FVertexInterface;
-	using FInputVertexInterface = Metasound::FInputVertexInterface;
-	using FOutputVertexInterface = Metasound::FOutputVertexInterface;
-
-	Metasound::FNodeInitData DummyInitData;
-	TUniquePtr<Metasound::INode> DummyNodePtr = InCallback(DummyInitData);
-	if (!ensureAlwaysMsgf(DummyNodePtr.IsValid(), TEXT("Invalid getter registered!")))
-	{
-		return false;
-	}
-
-	Metasound::INode& DummyNode = *DummyNodePtr;
-
-	FNodeRegistryKey InKey = GetRegistryKeyForNodeInternal(DummyNode);
+	FNodeRegistryElement RegistryElement = FNodeRegistryElement(MoveTemp(InCreateNode), MoveTemp(InCreateDescription));
 
-	const FName NodeName = DummyNode.GetClassName();
+	FNodeRegistryKey Key;
 
-	// check to see if an identical node was already registered, and log
-	ensureAlwaysMsgf(!ExternalNodeRegistry.Contains(InKey), TEXT("Node with identical name, inputs and outputs to node %s was already registered. The previously registered node will be overwritten. This could also happen because METASOUND_REGISTER_NODE is in a public header."), *NodeName.ToString());
-
-	UE_LOG(LogTemp, Display, TEXT("Registered Metasound Node %s"), *NodeName.ToString());
-
-	using FInputVertexInterface = ::Metasound::FInputVertexInterface;
-	using FOutputVertexInterface = ::Metasound::FOutputVertexInterface;
-
-	const FInputVertexInterface& Inputs = DummyNode.GetDefaultVertexInterface().GetInputInterface();
-	const FOutputVertexInterface& Outputs = DummyNode.GetDefaultVertexInterface().GetOutputInterface();
-
-	const bool bShouldLogInputsAndOutputs = true;
-	if (bShouldLogInputsAndOutputs)
+	if (ensure(FMetasoundFrontendRegistryContainer::GetRegistryKey(RegistryElement, Key)))
 	{
-		UE_LOG(LogTemp, Display, TEXT("    %d inputs:"), Inputs.Num());
-		for (auto& InputTuple : Inputs)
-		{
-			UE_LOG(LogTemp, Display, TEXT("      %s (of type %s)"), *InputTuple.Value.GetVertexName(), *InputTuple.Value.GetVertexTypeName().ToString());
-		}
+		// check to see if an identical node was already registered, and log
+		ensureAlwaysMsgf(!ExternalNodeRegistry.Contains(Key), TEXT("Node with identical name, inputs and outputs to node %s was already registered. The previously registered node will be overwritten. This could also happen because METASOUND_REGISTER_NODE is in a public header."), *Key.NodeName.ToString());
 
-		UE_LOG(LogTemp, Display, TEXT("    %d outputs:"), Outputs.Num());
-		for (auto& OutputTuple : Outputs)
-		{
-			UE_LOG(LogTemp, Display, TEXT("      %s (of type %s)"), *OutputTuple.Value.GetVertexName(), *OutputTuple.Value.GetVertexTypeName().ToString());
-		}
+		ExternalNodeRegistry.Add(MoveTemp(Key), MoveTemp(RegistryElement));
+
+		return true;
 	}
 
-	TArray<FName> InputTypes;
-	TArray<FName> OutputTypes;
+	return false;
+}
 
-	for (auto& InputTuple : Inputs)
+bool FMetasoundFrontendRegistryContainer::GetRegistryKey(const Metasound::Frontend::FNodeRegistryElement& InElement, Metasound::Frontend::FNodeRegistryKey& OutKey)
+{
+	if (InElement.CreateClassDescription)
 	{
-		InputTypes.Add(InputTuple.Value.GetVertexTypeName());
-	}
+		FMetasoundClassDescription Description = InElement.CreateClassDescription();
 
-	for (auto& OutputTuple : Outputs)
-	{
-		OutputTypes.Add(OutputTuple.Value.GetVertexTypeName());
-	}
+		OutKey = MetasoundFrontendRegistryPrivate::GetRegistryKey(FName(Description.Metadata.NodeName), Description.Metadata.MajorVersion);
 
-	FNodeRegistryElement RegistryElement = FNodeRegistryElement(MoveTemp(InCallback));
-	RegistryElement.InputTypes = MoveTemp(InputTypes);
-	RegistryElement.OutputTypes = MoveTemp(OutputTypes);
+		return true;
+	}
 
-	ExternalNodeRegistry.Add(MoveTemp(InKey), MoveTemp(RegistryElement));
+	return false;
+}
 
-	return true;
+Metasound::Frontend::FNodeRegistryKey FMetasoundFrontendRegistryContainer::GetRegistryKey(const FNodeInfo& InNodeMetadata)
+{
+	return MetasoundFrontendRegistryPrivate::GetRegistryKey(InNodeMetadata.ClassName, InNodeMetadata.MajorVersion);
 }
 
 bool FMetasoundFrontendRegistryContainer::RegisterConversionNode(const FConverterNodeRegistryKey& InNodeKey, const FConverterNodeInfo& InNodeInfo)
diff --git a/Engine/Plugins/Runtime/Metasound/Source/MetasoundFrontend/Public/MetasoundAutoConverterNode.h b/Engine/Plugins/Runtime/Metasound/Source/MetasoundFrontend/Public/MetasoundAutoConverterNode.h
--- a/Engine/Plugins/Runtime/Metasound/Source/MetasoundFrontend/Public/MetasoundAutoConverterNode.h
+++ b/Engine/Plugins/Runtime/Metasound/Source/MetasoundFrontend/Public/MetasoundAutoConverterNode.h
@@ -4,6 +4,7 @@
 #include "CoreMinimal.h"
 #include "MetasoundBuilderInterface.h"
 #include "MetasoundBuildError.h"
+#include "MetasoundNode.h"
 #include "MetasoundNodeInterface.h"
 #include "MetasoundOperatorInterface.h"
 #include "MetasoundDataFactory.h"
@@ -12,14 +13,13 @@
 
 #include <type_traits>
 
-#define LOCTEXT_NAMESPACE "MetasoundGraphCore"
-
+#define LOCTEXT_NAMESPACE "MetasoundFrontend"
 
 namespace Metasound
 {
 	// This convenience node can be registered and will invoke static_cast<ToDataType>(FromDataType) every time it is executed.
 	template<typename FromDataType, typename ToDataType>
-	class TAutoConverterNode : public INode
+	class TAutoConverterNode : public FNode
 	{
 		static_assert(std::is_convertible<FromDataType, ToDataType>::value, "Tried to create an auto converter node between two types we can't static_cast between.");
 
@@ -36,6 +36,42 @@ namespace Metasound
 			return OutputName;
 		}
 
+		static FVertexInterface DeclareVertexInterface()
+		{
+			return FVertexInterface(
+				FInputVertexInterface(
+					TInputDataVertexModel<FromDataType>(GetInputName(), FText::GetEmpty())
+				),
+				FOutputVertexInterface(
+					TOutputDataVertexModel<ToDataType>(GetOutputName(), FText::GetEmpty())
+				)
+			);
+		}
+
+		static const FNodeInfo& GetAutoConverterNodeMetadata()
+		{
+			auto InitNodeInfo = []() -> FNodeInfo
+			{
+				FNodeInfo Info;
+				Info.ClassName = FName(*FString::Printf(TEXT("%s To %s"), *GetInputName(), *GetOutputName()));
+				Info.MajorVersion = 1;
+				Info.MinorVersion = 0;
+				Info.Description = LOCTEXT("Metasound_ConverterNodeDescription", "Converts between two different data types.");
+				Info.Author = PluginAuthor;
+				Info.PromptIfMissing = PluginNodeMissingPrompt;
+				Info.DefaultInterface = DeclareVertexInterface();
+
+				Info.CategoryHierarchy = {LOCTEXT("Metasound_ConvertNodeCategory", "Convert"), FText::AsCultureInvariant(GetMetasoundDataTypeString<FromDataType>())};
+				Info.Keywords = {TEXT("convert"), GetMetasoundDataTypeName<FromDataType>(), GetMetasoundDataTypeName<ToDataType>()};
+
+				return Info;
+			};
+
+			static const FNodeInfo Info = InitNodeInfo();
+
+			return Info;
+		}
+
 	private:
 		/** FConverterOperator converts from "FromDataType" to "ToDataType" using
 		 * a implicit conversion operators. 
@@ -109,20 +145,9 @@ namespace Metasound
 		};
 
 		public:
-			static FVertexInterface DeclareVertexInterface()
-			{
-				return FVertexInterface(
-					FInputVertexInterface(
-						TInputDataVertexModel<FromDataType>(GetInputName(), FText::GetEmpty())
-					),
-					FOutputVertexInterface(
-						TOutputDataVertexModel<ToDataType>(GetOutputName(), FText::GetEmpty())
-					)
-				);
-			}
 
 			TAutoConverterNode(const FNodeInitData& InInitData)
-				: NodeDescription(InInitData.InstanceName)
+				: FNode(InInitData.InstanceName, GetAutoConverterNodeMetadata())
 				, Interface(DeclareVertexInterface())
 				, Factory(MakeOperatorFactoryRef<FCoverterOperatorFactory>())
 			{
@@ -130,43 +155,11 @@ namespace Metasound
 
 			virtual ~TAutoConverterNode() = default;
 
-			virtual const FName& GetClassName() const override
-			{
-				static const FName ClassName(*FString::Printf(TEXT("%s To %s"), *GetInputName(), *GetOutputName()));
-				return ClassName;
-			}
-
-			virtual const FString& GetInstanceName() const override
-			{
-				return NodeDescription;
-			}
-
-			virtual const FText& GetDescription() const override
-			{
-				static const FText Description = LOCTEXT("Metasound_ConverterNodeDescription", "Converts between two different data types.");
-				return Description;
-			}
-
-			virtual const FText& GetAuthorName() const override
-			{
-				return PluginAuthor;
-			}
-
-			virtual const FText& GetPromptIfMissing() const override
-			{
-				return PluginNodeMissingPrompt;
-			}
-
 			virtual const FVertexInterface& GetVertexInterface() const override
 			{
 				return Interface;
 			}
 
-			virtual const FVertexInterface& GetDefaultVertexInterface() const override
-			{
-				return Interface;
-			}
-
 			virtual bool SetVertexInterface(const FVertexInterface& InInterface) override
 			{
 				return Interface == InInterface;
@@ -183,10 +176,8 @@ namespace Metasound
 			}
 
 		private:
-			FString NodeDescription;
 			FVertexInterface Interface;
 			FOperatorFactorySharedRef Factory;
-			
 	};
 }
 #undef LOCTEXT_NAMESPACE
diff --git a/Engine/Plugins/Runtime/Metasound/Source/MetasoundFrontend/Public/MetasoundConverterNodeRegistrationMacro.h b/Engine/Plugins/Runtime/Metasound/Source/MetasoundFrontend/Public/MetasoundConverterNodeRegistrationMacro.h
--- a/Engine/Plugins/Runtime/Metasound/Source/MetasoundFrontend/Public/MetasoundConverterNodeRegistrationMacro.h
+++ b/Engine/Plugins/Runtime/Metasound/Source/MetasoundFrontend/Public/MetasoundConverterNodeRegistrationMacro.h
@@ -12,7 +12,7 @@
 #include "MetasoundFrontendRegistries.h"
 
 template<typename TNodeType, typename TFromType, typename TToType>
-bool RegisterConversionNode(const FString& FromPin, const FString& ToPin)
+bool RegisterConversionNode(const FString& FromPin, const FString& ToPin, const Metasound::FNodeInfo& InNodeMetadata)
 {
 	using FConverterNodeRegistryKey = ::Metasound::Frontend::FConverterNodeRegistryKey;
 	using FConverterNodeInfo = ::Metasound::Frontend::FConverterNodeInfo;
@@ -38,7 +38,7 @@ bool RegisterConversionNode(const FString& FromPin, const FString& ToPin)
 	{
 		FromPin,
 		ToPin,
-		FMetasoundFrontendRegistryContainer::GetRegistryKeyForNode<TNodeType>()
+		FMetasoundFrontendRegistryContainer::GetRegistryKey(InNodeMetadata)
 	};
 
 	return FMetasoundFrontendRegistryContainer::Get()->RegisterConversionNode(RegistryKey, ConverterNodeInfo);
diff --git a/Engine/Plugins/Runtime/Metasound/Source/MetasoundFrontend/Public/MetasoundDataTypeRegistrationMacro.h b/Engine/Plugins/Runtime/Metasound/Source/MetasoundFrontend/Public/MetasoundDataTypeRegistrationMacro.h
--- a/Engine/Plugins/Runtime/Metasound/Source/MetasoundFrontend/Public/MetasoundDataTypeRegistrationMacro.h
+++ b/Engine/Plugins/Runtime/Metasound/Source/MetasoundFrontend/Public/MetasoundDataTypeRegistrationMacro.h
@@ -21,200 +21,206 @@
 
 #include <type_traits>
 
-// SFINAE used to optionally invoke subclasses of IAudioProxyDataFactory when we can.
-template<typename UClassToUse, typename TEnableIf<TIsDerivedFrom<UClassToUse, IAudioProxyDataFactory>::Value, bool>::Type = true>
-IAudioProxyDataFactory* CastToAudioProxyDataFactory(UObject* InObject)
+namespace Metasound
 {
-	UClassToUse* DowncastObject = Cast<UClassToUse>(InObject);
-	if (ensureAlways(DowncastObject))
+	// SFINAE used to optionally invoke subclasses of IAudioProxyDataFactory when we can.
+	template<typename UClassToUse, typename TEnableIf<TIsDerivedFrom<UClassToUse, IAudioProxyDataFactory>::Value, bool>::Type = true>
+	IAudioProxyDataFactory* CastToAudioProxyDataFactory(UObject* InObject)
 	{
-		return static_cast<IAudioProxyDataFactory*>(DowncastObject);
+		UClassToUse* DowncastObject = Cast<UClassToUse>(InObject);
+		if (ensureAlways(DowncastObject))
+		{
+			return static_cast<IAudioProxyDataFactory*>(DowncastObject);
+		}
+		else
+		{
+			return nullptr;
+		}
 	}
-	else
+
+	template<typename UClassToUse, typename TEnableIf<!TIsDerivedFrom<UClassToUse, IAudioProxyDataFactory>::Value, bool>::Type = true>
+	IAudioProxyDataFactory* CastToAudioProxyDataFactory(UObject* InObject)
 	{
 		return nullptr;
 	}
-}
-
-template<typename UClassToUse, typename TEnableIf<!TIsDerivedFrom<UClassToUse, IAudioProxyDataFactory>::Value, bool>::Type = true>
-IAudioProxyDataFactory* CastToAudioProxyDataFactory(UObject* InObject)
-{
-	return nullptr;
-}
-
-// Helper utility to test if we can transmit a datatype between a send and a receive node.
-template <typename TDataType>
-struct TIsTransmittable
-{
-private:
-	static constexpr bool bCanBeTransmitted =
-		std::is_copy_constructible<TDataType>::value|| TIsDerivedFrom<TDataType, ::Metasound::IAudioDataType>::Value;
-
-public:
-
-	static constexpr bool Value = bCanBeTransmitted;
-};
-
-// This utility function can be used to optionally check to see if we can transmit a data type, and autogenerate send and receive nodes for that datatype.
-template<typename TDataType, typename TEnableIf<TIsTransmittable<TDataType>::Value, bool>::Type = true>
-void AttemptToRegisterSendAndReceiveNodes()
-{
-	ensureAlways(RegisterNodeWithFrontend<Metasound::TSendNode<TDataType>>());
-	ensureAlways(RegisterNodeWithFrontend<Metasound::TReceiveNode<TDataType>>());
-}
 
-template<typename TDataType, typename TEnableIf<!TIsTransmittable<TDataType>::Value, bool>::Type = true>
-void AttemptToRegisterSendAndReceiveNodes()
-{
-	// This implementation intentionally noops, because Metasound::TIsTransmittable is false for this datatype.
-	// This is either because the datatype is not trivially copyable, and thus can't be buffered between threads,
-	// or it's not an audio buffer type, which we use Audio::FPatchMixerSplitter instances for.
-}
+	// Helper utility to test if we can transmit a datatype between a send and a receive node.
+	template <typename TDataType>
+	struct TIsTransmittable
+	{
+	private:
+		static constexpr bool bCanBeTransmitted =
+			std::is_copy_constructible<TDataType>::value|| TIsDerivedFrom<TDataType, IAudioDataType>::Value;
 
+	public:
 
+		static constexpr bool Value = bCanBeTransmitted;
+	};
 
-// This utility function can be used to check to see if we can static cast between two types, and autogenerate a node for that static cast.
-template<typename TFromDataType, typename TToDataType, typename std::enable_if<std::is_convertible<TFromDataType, TToDataType>::value, int>::type = 0>
-void AttemptToRegisterConverter()
-{
-	using FConverterNode = Metasound::TAutoConverterNode<TFromDataType, TToDataType>;
+	// This utility function can be used to optionally check to see if we can transmit a data type, and autogenerate send and receive nodes for that datatype.
+	template<typename TDataType, typename TEnableIf<TIsTransmittable<TDataType>::Value, bool>::Type = true>
+	void AttemptToRegisterSendAndReceiveNodes()
+	{
+		ensureAlways(RegisterNodeWithFrontend<Metasound::TSendNode<TDataType>>());
+		ensureAlways(RegisterNodeWithFrontend<Metasound::TReceiveNode<TDataType>>());
+	}
 
-	if (!TIsSame<TFromDataType, TToDataType>::Value && !FMetasoundFrontendRegistryContainer::Get()->IsNodeRegistered<FConverterNode>())
+	template<typename TDataType, typename TEnableIf<!TIsTransmittable<TDataType>::Value, bool>::Type = true>
+	void AttemptToRegisterSendAndReceiveNodes()
 	{
-		ensureAlways(RegisterNodeWithFrontend<FConverterNode>());
-		
-		bool bSucessfullyRegisteredConversionNode = RegisterConversionNode<FConverterNode, TFromDataType, TToDataType>(FConverterNode::GetInputName(), FConverterNode::GetOutputName());
-		ensureAlways(bSucessfullyRegisteredConversionNode);
+		// This implementation intentionally noops, because Metasound::TIsTransmittable is false for this datatype.
+		// This is either because the datatype is not trivially copyable, and thus can't be buffered between threads,
+		// or it's not an audio buffer type, which we use Audio::FPatchMixerSplitter instances for.
 	}
-}
 
-template<typename TFromDataType, typename TToDataType, typename std::enable_if<!std::is_convertible<TFromDataType, TToDataType>::value, int>::type = 0>
-void AttemptToRegisterConverter()
-{
-	// This implementation intentionally noops, because static_cast<TFromDataType>(TToDataType&) is invalid.
-}
 
-// Here we attempt to infer and autogenerate conversions for basic datatypes.
-template<typename TDataType>
-void RegisterConverterNodes()
-{
-	// Conversions to this data type:
-	AttemptToRegisterConverter<bool, TDataType>();
-	AttemptToRegisterConverter<int32, TDataType>();
-	AttemptToRegisterConverter<float, TDataType>();
-	AttemptToRegisterConverter<FString, TDataType>();
-
-	// Conversions from this data type:
-	AttemptToRegisterConverter<TDataType, bool>();
-	AttemptToRegisterConverter<TDataType, int32>();
-	AttemptToRegisterConverter<TDataType, float>();
-	AttemptToRegisterConverter<TDataType, FString>();
-}
 
-template<typename TDataType, ::Metasound::ELiteralArgType PreferredArgType = ::Metasound::ELiteralArgType::None, typename UClassToUse = UObject>
-bool RegisterDataTypeWithFrontend()
-{
-	// if we reenter this code (because DECLARE_METASOUND_DATA_REFERENCE_TYPES was called twice with the same type),
-	// we catch it here.
-	static bool bAlreadyRegisteredThisDataType = false;
-	if (bAlreadyRegisteredThisDataType)
+	// This utility function can be used to check to see if we can static cast between two types, and autogenerate a node for that static cast.
+	template<typename TFromDataType, typename TToDataType, typename std::enable_if<std::is_convertible<TFromDataType, TToDataType>::value, int>::type = 0>
+	void AttemptToRegisterConverter()
 	{
-		UE_LOG(LogTemp, Display, TEXT("Tried to call REGISTER_METASOUND_DATATYPE twice with the same class %s. ignoring the second call. Likely because REGISTER_METASOUND_DATATYPE is in a header that's used in multiple modules. Consider moving it to a private header or cpp file."), ::Metasound::TDataReferenceTypeInfo<TDataType>::TypeName)
-		return false;
-	}
+		using FConverterNode = Metasound::TAutoConverterNode<TFromDataType, TToDataType>;
 
-	bAlreadyRegisteredThisDataType = true;
+		const FNodeInfo& Metadata = FConverterNode::GetAutoConverterNodeMetadata();
+		const Metasound::Frontend::FNodeRegistryKey Key = FMetasoundFrontendRegistryContainer::GetRegistryKey(Metadata);
 
-	// Lambdas that generate our template-instantiated input and output nodes:
-	FInputNodeConstructorCallback InputNodeConstructor = [](::Metasound::FInputNodeConstructorParams&& InParams) -> TUniquePtr<::Metasound::INode>
-	{
-		return TUniquePtr<::Metasound::INode>(new ::Metasound::TInputNode<TDataType>(InParams.InNodeName, InParams.InVertexName, MoveTemp(InParams.InitParam)));
-	};
+		if (!std::is_same<TFromDataType, TToDataType>::value && !FMetasoundFrontendRegistryContainer::Get()->IsNodeRegistered(Key))
+		{
+			ensureAlways(RegisterNodeWithFrontend<FConverterNode>(Metadata));
+			
+			bool bSucessfullyRegisteredConversionNode = RegisterConversionNode<FConverterNode, TFromDataType, TToDataType>(FConverterNode::GetInputName(), FConverterNode::GetOutputName(), Metadata);
+			ensureAlways(bSucessfullyRegisteredConversionNode);
+		}
+	}
 
-	FOutputNodeConstructorCallback OutputNodeConstructor = [](const ::Metasound::FOutputNodeConstrutorParams& InParams) -> TUniquePtr<::Metasound::INode>
+	template<typename TFromDataType, typename TToDataType, typename std::enable_if<!std::is_convertible<TFromDataType, TToDataType>::value, int>::type = 0>
+	void AttemptToRegisterConverter()
 	{
-		return TUniquePtr<::Metasound::INode>(new ::Metasound::TOutputNode<TDataType>(InParams.InNodeName, InParams.InVertexName));
-	};
-	
-	// By default, this function should not be used, unless the preferred arg type is UObjectProxy or UObjectProxyArray, and UClassToUse should be specified.
-	FProxyGetterCallback ProxyGenerator = [](UObject* InObject) -> Audio::IProxyDataPtr
+		// This implementation intentionally noops, because static_cast<TFromDataType>(TToDataType&) is invalid.
+	}
+
+	// Here we attempt to infer and autogenerate conversions for basic datatypes.
+	template<typename TDataType>
+	void RegisterConverterNodes()
 	{
-		checkNoEntry();
-		return Audio::IProxyDataPtr(nullptr);
-	};
+		// Conversions to this data type:
+		AttemptToRegisterConverter<bool, TDataType>();
+		AttemptToRegisterConverter<int32, TDataType>();
+		AttemptToRegisterConverter<float, TDataType>();
+		AttemptToRegisterConverter<FString, TDataType>();
+
+		// Conversions from this data type:
+		AttemptToRegisterConverter<TDataType, bool>();
+		AttemptToRegisterConverter<TDataType, int32>();
+		AttemptToRegisterConverter<TDataType, float>();
+		AttemptToRegisterConverter<TDataType, FString>();
+	}
 
-	// If this datatype uses a UObject or UObject array literal, we generate a lambda to build a proxy here:
-	constexpr bool bSpecifiedUClassForProxy = !TIsSame<UClassToUse, UObject>::Value;
-	if (bSpecifiedUClassForProxy)
+	template<typename TDataType, ELiteralArgType PreferredArgType = ELiteralArgType::None, typename UClassToUse = UObject>
+	bool RegisterDataTypeWithFrontend()
 	{
-		static_assert(!bSpecifiedUClassForProxy || std::is_base_of<IAudioProxyDataFactory, UClassToUse>::value, "If a Metasound Datatype uses a UObject as a literal, the UClass of that object needs to also derive from Audio::IProxyDataFactory. See USoundWave as an example.");
-		ProxyGenerator = [](UObject* InObject) -> Audio::IProxyDataPtr
+		// if we reenter this code (because DECLARE_METASOUND_DATA_REFERENCE_TYPES was called twice with the same type),
+		// we catch it here.
+		static bool bAlreadyRegisteredThisDataType = false;
+		if (bAlreadyRegisteredThisDataType)
 		{
-			IAudioProxyDataFactory* ObjectAsFactory = CastToAudioProxyDataFactory<UClassToUse>(InObject);
-
-			if (ensureAlways(ObjectAsFactory))
-			{
-				static FName ProxySubsystemName = TEXT("Metasound");
-				
-				Audio::FProxyDataInitParams ProxyInitParams;
-				ProxyInitParams.NameOfFeatureRequestingProxy = ProxySubsystemName;
-
-				return ObjectAsFactory->CreateNewProxyData(ProxyInitParams);
-			}
-			else
-			{
-				return Audio::IProxyDataPtr(nullptr);
-			}
-		};
-	}
+			UE_LOG(LogTemp, Display, TEXT("Tried to call REGISTER_METASOUND_DATATYPE twice with the same class %s. ignoring the second call. Likely because REGISTER_METASOUND_DATATYPE is in a header that's used in multiple modules. Consider moving it to a private header or cpp file."), TDataReferenceTypeInfo<TDataType>::TypeName)
+			return false;
+		}
 
-	static const FName DataTypeName = ::Metasound::GetMetasoundDataTypeName<TDataType>();
+		bAlreadyRegisteredThisDataType = true;
 
-	// Pack all of our various constructor lambdas to a single struct.
-	::Metasound::FDataTypeConstructorCallbacks Callbacks = { MoveTemp(InputNodeConstructor), MoveTemp(OutputNodeConstructor), MoveTemp(ProxyGenerator) };
+		// Lambdas that generate our template-instantiated input and output nodes:
+		FCreateInputNodeFunction InputNodeConstructor = [](FInputNodeConstructorParams&& InParams) -> TUniquePtr<INode>
+		{
+			return TUniquePtr<INode>(new TInputNode<TDataType>(InParams.InNodeName, InParams.InVertexName, MoveTemp(InParams.InitParam)));
+		};
 
+		FCreateOutputNodeFunction OutputNodeConstructor = [](const FOutputNodeConstrutorParams& InParams) -> TUniquePtr<INode>
+		{
+			return TUniquePtr<INode>(new TOutputNode<TDataType>(InParams.InNodeName, InParams.InVertexName));
+		};
+		
+		// By default, this function should not be used, unless the preferred arg type is UObjectProxy or UObjectProxyArray, and UClassToUse should be specified.
+		FCreateAudioProxyFunction ProxyGenerator = [](UObject* InObject) -> Audio::IProxyDataPtr
+		{
+			checkNoEntry();
+			return Audio::IProxyDataPtr(nullptr);
+		};
 
-	Metasound::FDataTypeRegistryInfo RegistryInfo;
+		// If this datatype uses a UObject or UObject array literal, we generate a lambda to build a proxy here:
+		constexpr bool bSpecifiedUClassForProxy = !TIsSame<UClassToUse, UObject>::Value;
+		if (bSpecifiedUClassForProxy)
+		{
+			static_assert(!bSpecifiedUClassForProxy || std::is_base_of<IAudioProxyDataFactory, UClassToUse>::value, "If a Metasound Datatype uses a UObject as a literal, the UClass of that object needs to also derive from Audio::IProxyDataFactory. See USoundWave as an example.");
+			ProxyGenerator = [](UObject* InObject) -> Audio::IProxyDataPtr
+			{
+				IAudioProxyDataFactory* ObjectAsFactory = CastToAudioProxyDataFactory<UClassToUse>(InObject);
+
+				if (ensureAlways(ObjectAsFactory))
+				{
+					static FName ProxySubsystemName = TEXT("Metasound");
+					
+					Audio::FProxyDataInitParams ProxyInitParams;
+					ProxyInitParams.NameOfFeatureRequestingProxy = ProxySubsystemName;
+
+					return ObjectAsFactory->CreateNewProxyData(ProxyInitParams);
+				}
+				else
+				{
+					return Audio::IProxyDataPtr(nullptr);
+				}
+			};
+		}
+
+		static const FName DataTypeName = GetMetasoundDataTypeName<TDataType>();
+
+		// Pack all of our various constructor lambdas to a single struct.
+		FDataTypeConstructorCallbacks Callbacks = { MoveTemp(InputNodeConstructor), MoveTemp(OutputNodeConstructor), MoveTemp(ProxyGenerator) };
+
+
+		FDataTypeRegistryInfo RegistryInfo;
+
+		RegistryInfo.DataTypeName = DataTypeName;
+		RegistryInfo.PreferredLiteralType = PreferredArgType;
+
+		RegistryInfo.bIsBoolParsable = TIsParsable<TDataType, bool>::Value;
+		RegistryInfo.bIsIntParsable = TIsParsable<TDataType, int32>::Value;
+		RegistryInfo.bIsFloatParsable = TIsParsable<TDataType, float>::Value;
+		RegistryInfo.bIsStringParsable = TIsParsable<TDataType, FString>::Value;
+		RegistryInfo.bIsProxyParsable = TIsParsable<TDataType, const Audio::IProxyDataPtr&>::Value;
+		RegistryInfo.bIsProxyArrayParsable = TIsParsable<TDataType, const TArray<Audio::IProxyDataPtr>& >::Value;
+		RegistryInfo.bIsDefaultParsable = TIsParsable<TDataType>::Value;
+		
+		RegistryInfo.ProxyGeneratorClass = UClassToUse::StaticClass();
 
-	RegistryInfo.DataTypeName = DataTypeName;
-	RegistryInfo.PreferredLiteralType = PreferredArgType;
+		bool bSucceeded = FMetasoundFrontendRegistryContainer::Get()->RegisterDataType(RegistryInfo, MoveTemp(Callbacks));
+		ensureAlwaysMsgf(bSucceeded, TEXT("Failed to register data type %s in the node registry!"), *GetMetasoundDataTypeString<TDataType>());
+		
+		RegisterConverterNodes<TDataType>();
+		AttemptToRegisterSendAndReceiveNodes<TDataType>();
+		
+		return bSucceeded;
+	}
 
-	RegistryInfo.bIsBoolParsable = Metasound::TIsParsable<TDataType, bool>::Value;
-	RegistryInfo.bIsIntParsable = Metasound::TIsParsable<TDataType, int32>::Value;
-	RegistryInfo.bIsFloatParsable = Metasound::TIsParsable<TDataType, float>::Value;
-	RegistryInfo.bIsStringParsable = Metasound::TIsParsable<TDataType, FString>::Value;
-	RegistryInfo.bIsProxyParsable = Metasound::TIsParsable<TDataType, const Audio::IProxyDataPtr&>::Value;
-	RegistryInfo.bIsProxyArrayParsable = Metasound::TIsParsable<TDataType, const TArray<Audio::IProxyDataPtr>& >::Value;
-	RegistryInfo.bIsDefaultParsable = Metasound::TIsParsable<TDataType>::Value;
-	
-	RegistryInfo.ProxyGeneratorClass = UClassToUse::StaticClass();
+	template<typename DataType>
+	struct TMetasoundDataTypeRegistration
+	{
+		static_assert(std::is_same<DataType, typename std::decay<DataType>::type>::value, "DataType and decayed DataType must be the same");
+		
+		static constexpr bool bCanRegister = TInputNode<DataType>::bCanRegister;
 
-	bool bSucceeded = FMetasoundFrontendRegistryContainer::Get()->RegisterDataType(RegistryInfo, MoveTemp(Callbacks));
-	ensureAlwaysMsgf(bSucceeded, TEXT("Failed to register data type %s in the node registry!"), *Metasound::GetMetasoundDataTypeString<TDataType>());
-	
-	RegisterConverterNodes<TDataType>();
-	AttemptToRegisterSendAndReceiveNodes<TDataType>();
-	
-	return bSucceeded;
+		static const bool bSuccessfullyRegistered;
+	};
 }
 
-template<typename DataType>
-struct TMetasoundDataTypeRegistration
-{
-	static_assert(std::is_same<DataType, typename std::decay<DataType>::type>::value, "DataType and decayed DataType must be the same");
-	
-	static constexpr bool bCanRegister = ::Metasound::TInputNode<DataType>::bCanRegister;
-
-	static const bool bSuccessfullyRegistered;
-};
-
 // This should be used to expose a datatype as a potential input or output for a metasound graph.
 // The first argument to the macro is the class to expose.
 // the second argument is the display name of that type in the Metasound editor.
 // Optionally, a Metasound::ELiteralArgType can be passed in to designate a preferred literal type-
 // For example, if Metasound::ELiteralArgType::Float is passed in, we will default to using a float parameter to create this datatype.
 // If no argument is passed in, we will infer a literal type to use.
-// If ::Metasound::
+// If 
 // Metasound::ELiteralArgType::Invalid can be used to enforce that we don't provide space for a literal, in which case you should have a default constructor or a constructor that takes [const FOperatorSettings&] implemented.
 // If you pass in a preferred arg type, please make sure that the passed in datatype has a matching constructor, since we won't check this until runtime.
 
@@ -230,6 +236,6 @@ struct TMetasoundDataTypeRegistration
 
 #define REGISTER_METASOUND_DATATYPE(DataType, DataTypeName, ...) \
 	DEFINE_METASOUND_DATA_TYPE(DataType, DataTypeName); \
-	static_assert(::TMetasoundDataTypeRegistration<DataType>::bCanRegister, CANNOT_REGISTER_METASOUND_DATA_TYPE_ASSERT_STRING(DataType)); \
-	template<> const bool ::TMetasoundDataTypeRegistration<DataType>::bSuccessfullyRegistered = ::FMetasoundFrontendRegistryContainer::Get()->EnqueueInitCommand([](){ ::RegisterDataTypeWithFrontend<DataType, ##__VA_ARGS__>(); }); // This static bool is useful for debugging, but also is the only way the compiler will let us call this function outside of an expression.
+	static_assert(::Metasound::TMetasoundDataTypeRegistration<DataType>::bCanRegister, CANNOT_REGISTER_METASOUND_DATA_TYPE_ASSERT_STRING(DataType)); \
+	template<> const bool ::Metasound::TMetasoundDataTypeRegistration<DataType>::bSuccessfullyRegistered = ::FMetasoundFrontendRegistryContainer::Get()->EnqueueInitCommand([](){ ::Metasound::RegisterDataTypeWithFrontend<DataType, ##__VA_ARGS__>(); }); // This static bool is useful for debugging, but also is the only way the compiler will let us call this function outside of an expression.
 
diff --git a/Engine/Plugins/Runtime/Metasound/Source/MetasoundFrontend/Public/MetasoundFrontend.h b/Engine/Plugins/Runtime/Metasound/Source/MetasoundFrontend/Public/MetasoundFrontend.h
--- a/Engine/Plugins/Runtime/Metasound/Source/MetasoundFrontend/Public/MetasoundFrontend.h
+++ b/Engine/Plugins/Runtime/Metasound/Source/MetasoundFrontend/Public/MetasoundFrontend.h
@@ -55,22 +55,64 @@ namespace Metasound
 		// Get all available nodes of any type.
 		METASOUNDFRONTEND_API TArray<FNodeClassInfo> GetAllAvailableNodeClasses();
 
-		// Get all possible nodes whose name begins with a specific namespace.
-		METASOUNDFRONTEND_API TArray<FNodeClassInfo> GetAllNodeClassesInNamespace(const FString& InNamespace);
-
-		// Similar to GetAllNodeClassesInNamespace, except searches all nodes for a match with a substring.
-		METASOUNDFRONTEND_API TArray<FNodeClassInfo> GetAllNodesWhoseNameContains(const FString& InSubstring);
+		/** Returns all metadata (name, description, author, what to say if it's missing) for a given node.
+		 *
+		 * @param InInfo - Class info for a already registered external node.
+		 *
+		 * @return Metadata for node.
+		 */
+		METASOUNDFRONTEND_API FMetasoundClassMetadata GenerateClassMetadata(const FNodeClassInfo& InInfo);
+
+		/** Generates a new FMetasoundClassDescription from Node Metadata 
+		 *
+		 * @param InNodeMetadata - Metadata describing an external node.
+		 *
+		 * @return Class description for external node.
+		 */
+		METASOUNDFRONTEND_API FMetasoundClassDescription GenerateClassDescription(const FNodeInfo& InNodeMetadata);
+
+		/** Generates a new FMetasoundClassDescription from node lookup info.
+		 *
+		 * @param InInfo - Class info for a already registered external node.
+		 *
+		 * @return Class description for external node.
+		 */
+		METASOUNDFRONTEND_API FMetasoundClassDescription GenerateClassDescription(const FNodeClassInfo& InInfo);
+
+		/** Generates a new FMetasoundClassDescription from Node init data
+		 *
+		 * @tparam NodeType - Type of node to instantiate.
+		 * @param InNodeInitData - Data used to call constructor of node.
+		 *
+		 * @return Class description for external node.
+		 */
+		template<typename NodeType>
+		FMetasoundClassDescription GenerateClassDescription(const FNodeInitData& InNodeInitData)
+		{
+			TUniquePtr<INode> Node = MakeUnique<NodeType>(InNodeInitData);
 
-		// Searches for any node types that can output a specific type.
-		METASOUNDFRONTEND_API TArray<FNodeClassInfo> GetAllNodesWithAnOutputOfType(const FName& InType);
+			if (ensure(Node.IsValid()))
+			{
+				return GenerateClassDescription(Node->GetMetadata());
+			}
 
-		METASOUNDFRONTEND_API TArray<FNodeClassInfo> GetAllNodesWithAnInputOfType(const FName& InType);
+			return FMetasoundClassDescription();
+		}
 
-		// gets all metadata (name, description, author, what to say if it's missing) for a given node.
-		METASOUNDFRONTEND_API FMetasoundClassMetadata GenerateMetadataForNode(const FNodeClassInfo& InInfo);
+		/** Generates a new FMetasoundClassDescription from a NodeType
+		 *
+		 * @tparam NodeType - Type of node.
+		 *
+		 * @return Class description for external node.
+		 */
+		template<typename NodeType>
+		FMetasoundClassDescription GenerateClassDescription()
+		{
+			FNodeInitData InitData;
+			InitData.InstanceName = FString(TEXT("GenerateClassDescriptionForNode"));
 
-		// Generates a new FMetasoundClassDescription for a given node class. Only used by classes that manipulate Metasound Description data directly.
-		METASOUNDFRONTEND_API FMetasoundClassDescription GenerateClassDescriptionForNode(const FNodeClassInfo& InInfo);
+			return GenerateClassDescription<NodeType>(InitData);
+		}
 
 		// Convenience function to convert Metasound::ELiteralArgType to EMetasoundLiteralType, the UEnum used for metasound documents.
 		METASOUNDFRONTEND_API EMetasoundLiteralType GetMetasoundLiteralType(Metasound::ELiteralArgType InLiteralType);
diff --git a/Engine/Plugins/Runtime/Metasound/Source/MetasoundFrontend/Public/MetasoundFrontendDataLayout.h b/Engine/Plugins/Runtime/Metasound/Source/MetasoundFrontend/Public/MetasoundFrontendDataLayout.h
--- a/Engine/Plugins/Runtime/Metasound/Source/MetasoundFrontend/Public/MetasoundFrontendDataLayout.h
+++ b/Engine/Plugins/Runtime/Metasound/Source/MetasoundFrontend/Public/MetasoundFrontendDataLayout.h
@@ -45,6 +45,12 @@ struct FMetasoundClassMetadata
 	UPROPERTY()
 	FString NodeName;
 
+	UPROPERTY(EditAnywhere, Category = General)
+	int32 MajorVersion = 1;
+
+	UPROPERTY(EditAnywhere, Category = General)
+	int32 MinorVersion = 0;
+
 	// This will always be set to EMetasoundObjectType::Metasound.
 	UPROPERTY(VisibleAnywhere, Category = General, meta = (DisplayName = "Type"))
 	EMetasoundClassType NodeType;
@@ -194,6 +200,7 @@ struct FMetasoundNodeDescription
 {
 	GENERATED_BODY()
 
+
 	// Unique integer given to this node. Only has to be unique within the Nodes array.
 	UPROPERTY()
 	uint32 UniqueID;
@@ -224,6 +231,12 @@ struct FMetasoundGraphDescription
 {
 	GENERATED_BODY()
 
+	// TODO: pp - Break out FMetasoundNodeDescription to be a list of connections 
+	// and a list of nodes. Currently FMetasoudnNodeConnectionDescription lives 
+	// on the input node, but would be better living on the graph. 
+	//
+	// NodeConnections TMap<Tuple(FromNodeId, ToNodeId), {VertexName, VertexName}>?
+	//
 	UPROPERTY()
 	TArray<FMetasoundNodeDescription> Nodes;
 };
diff --git a/Engine/Plugins/Runtime/Metasound/Source/MetasoundFrontend/Public/MetasoundFrontendRegistries.h b/Engine/Plugins/Runtime/Metasound/Source/MetasoundFrontend/Public/MetasoundFrontendRegistries.h
--- a/Engine/Plugins/Runtime/Metasound/Source/MetasoundFrontend/Public/MetasoundFrontendRegistries.h
+++ b/Engine/Plugins/Runtime/Metasound/Source/MetasoundFrontend/Public/MetasoundFrontendRegistries.h
@@ -3,21 +3,23 @@
 #pragma once
 
 #include "CoreMinimal.h"
-#include "MetasoundNodeInterface.h"
+#include "IAudioProxyInitializer.h"
 #include "MetasoundDataReference.h"
+#include "MetasoundFrontendDataLayout.h"
+#include "MetasoundNodeInterface.h"
 #include "MetasoundOperatorInterface.h"
-#include "IAudioProxyInitializer.h"
 
-typedef TUniqueFunction<TUniquePtr<Metasound::INode>(::Metasound::FInputNodeConstructorParams&&)> FInputNodeConstructorCallback;
-typedef TUniqueFunction<TUniquePtr<Metasound::INode>(const ::Metasound::FOutputNodeConstrutorParams&)> FOutputNodeConstructorCallback;
+namespace Metasound
+{
+	typedef TFunction<TUniquePtr<Metasound::INode>(::Metasound::FInputNodeConstructorParams&&)> FCreateInputNodeFunction;
+	typedef TFunction<TUniquePtr<Metasound::INode>(const ::Metasound::FOutputNodeConstrutorParams&)> FCreateOutputNodeFunction;
 
-// This function is used to create a proxy from a datatype's base uclass.
-typedef TUniqueFunction<Audio::IProxyDataPtr(UObject*)> FProxyGetterCallback;
+	// This function is used to create a proxy from a datatype's base uclass.
+	typedef TFunction<Audio::IProxyDataPtr(UObject*)> FCreateAudioProxyFunction;
 
-typedef TUniqueFunction<TUniquePtr<Metasound::INode>(const Metasound::FNodeInitData&)> FNodeGetterCallback;
+	typedef TFunction<TUniquePtr<Metasound::INode>(const Metasound::FNodeInitData&)> FCreateMetasoundNodeFunction;
+	typedef TFunction<FMetasoundClassDescription()> FCreateMetasoundClassDescriptionFunction;
 
-namespace Metasound
-{
 	// Various elements that we pass to the frontend registry based on templated type traits.
 	struct FDataTypeRegistryInfo
 	{
@@ -55,7 +57,7 @@ namespace Metasound
 
 			// A hash generated from the input types and output types for this node.
 			// TODO: Write up some tests to ensure this is deterministic.
-			uint32 NodeHash;
+			uint32 NodeHash = 0;
 
 			FORCEINLINE bool operator==(const FNodeRegistryKey& Other) const
 			{
@@ -72,13 +74,13 @@ namespace Metasound
 		struct METASOUNDFRONTEND_API FNodeRegistryElement
 		{
 			// This lambda can be used to get an INodeBase for this specific node class.
-			FNodeGetterCallback GetterCallback;
+			FCreateMetasoundNodeFunction CreateNode;
 
-			TArray<FName> InputTypes;
-			TArray<FName> OutputTypes;
+			FCreateMetasoundClassDescriptionFunction CreateClassDescription;
 
-			FNodeRegistryElement(FNodeGetterCallback&& InCallback)
-				: GetterCallback(MoveTemp(InCallback))
+			FNodeRegistryElement(FCreateMetasoundNodeFunction&& InCreateNodeFunction, FCreateMetasoundClassDescriptionFunction&& InCreateDescriptionFunction)
+				: CreateNode(MoveTemp(InCreateNodeFunction))
+				, CreateClassDescription(InCreateDescriptionFunction)
 			{
 			}
 		};
@@ -105,10 +107,10 @@ namespace Metasound
 		struct METASOUNDFRONTEND_API FConverterNodeInfo
 		{
 			// If this node has multiple input pins, we use this to designate which pin should be used.
-			FString PreferredConverterInputPin;
+			FVertexKey PreferredConverterInputPin;
 
 			// If this node has multiple output pins, we use this to designate which pin should be used.
-			FString PreferredConverterOutputPin;
+			FVertexKey PreferredConverterOutputPin;
 
 			// The key for this node in the node registry.
 			FNodeRegistryKey NodeKey;
@@ -129,13 +131,13 @@ namespace Metasound
 	struct FDataTypeConstructorCallbacks
 	{
 		// This constructs a TInputNode<> with the corresponding datatype.
-		FInputNodeConstructorCallback InputNodeConstructor;
+		FCreateInputNodeFunction CreateInputNode;
 
 		// This constructs a TOutputNode<> with the corresponding datatype.
-		FOutputNodeConstructorCallback OutputNodeConstructor;
+		FCreateOutputNodeFunction CreateOutputNode;
 
 		// For datatypes that use a UObject literal or a UObject literal array, this lambda generates a literal from the corresponding UObject.
-		FProxyGetterCallback ProxyConstructor;
+		FCreateAudioProxyFunction CreateAudioProxy;
 	};
 }
 
@@ -153,6 +155,7 @@ class METASOUNDFRONTEND_API FMetasoundFrontendRegistryContainer
 
 	using FDataTypeRegistryInfo = Metasound::FDataTypeRegistryInfo;
 	using FDataTypeConstructorCallbacks = ::Metasound::FDataTypeConstructorCallbacks;
+	using FNodeInfo = Metasound::FNodeInfo;
 
 public:
 	static FMetasoundFrontendRegistryContainer* Get();
@@ -191,25 +194,26 @@ public:
 	bool DoesDataTypeSupportLiteralType(FName InDataType, Metasound::ELiteralArgType InLiteralType) const;
 
 	bool RegisterDataType(const FDataTypeRegistryInfo& InDataInfo, FDataTypeConstructorCallbacks&& InCallbacks);
-	bool RegisterExternalNode(FNodeGetterCallback&& InCallback);
-	bool RegisterConversionNode(const FConverterNodeRegistryKey& InNodeKey, const FConverterNodeInfo& InNodeInfo);
-
-	template <typename TNodeType>
-	static FNodeRegistryKey GetRegistryKeyForNode()
-	{
-		Metasound::FNodeInitData DummyInitData;
 
-		TNodeType DummyNode = TNodeType(DummyInitData);
+	/** Register external node with the frontend.
+	 *
+	 * @param InCreateNode - Function for creating node from FNodeInitData.
+	 * @param InCreateDescription - Function for creating a FMetasoundClassDescription.
+	 *
+	 * @return True on success.
+	 */
+	bool RegisterExternalNode(Metasound::FCreateMetasoundNodeFunction&& InCreateNode, Metasound::FCreateMetasoundClassDescriptionFunction&& InCreateDescription);
 
-		return GetRegistryKeyForNodeInternal(DummyNode);
-	}
+	bool RegisterConversionNode(const FConverterNodeRegistryKey& InNodeKey, const FConverterNodeInfo& InNodeInfo);
 
-	template <typename TNodeType>
-	bool IsNodeRegistered()
+	bool IsNodeRegistered(const FNodeRegistryKey& InKey) const
 	{
-		return ExternalNodeRegistry.Contains(GetRegistryKeyForNode<TNodeType>());
+		return ExternalNodeRegistry.Contains(InKey);
 	}
 
+	static FNodeRegistryKey GetRegistryKey(const FNodeInfo& InNodeMetadata);
+	static bool GetRegistryKey(const FNodeRegistryElement& InElement, FNodeRegistryKey& OutKey);
+
 	// Return any data types that can be used as a metasound input type or output type.
 	TArray<FName> GetAllValidDataTypes();
 
@@ -220,9 +224,6 @@ public:
 private:
 	FMetasoundFrontendRegistryContainer();
 
-	// Internal convenience function for generating a hash for an INode.
-	static FNodeRegistryKey GetRegistryKeyForNodeInternal(const Metasound::INode& InNode);
-
 	static FMetasoundFrontendRegistryContainer* LazySingleton;
 
 	// This buffer is used to enqueue nodes and datatypes to register when the module has been initialized,
diff --git a/Engine/Plugins/Runtime/Metasound/Source/MetasoundFrontend/Public/MetasoundInputNode.h b/Engine/Plugins/Runtime/Metasound/Source/MetasoundFrontend/Public/MetasoundInputNode.h
--- a/Engine/Plugins/Runtime/Metasound/Source/MetasoundFrontend/Public/MetasoundInputNode.h
+++ b/Engine/Plugins/Runtime/Metasound/Source/MetasoundFrontend/Public/MetasoundInputNode.h
@@ -105,21 +105,6 @@ namespace Metasound
 			// If true, this node can be instantiated by the FrontEnd.
 			static constexpr bool bCanRegister = TInputOperatorLiteralFactory<DataType>::bCanRegister;
 
-			static const FNodeInfo& GetNodeInfo()
-			{
-				static const FString ClassNameString = FString(TEXT("Input_")) + GetMetasoundDataTypeName<DataType>().ToString();
-				static const FName ClassName(*ClassNameString);
-
-				static const FNodeInfo Info = {
-					ClassName,
-					LOCTEXT("Metasound_InputNodeDescription", "Input into the parent Metasound graph."),
-					PluginAuthor,
-					PluginNodeMissingPrompt
-				};
-
-				return Info;
-			}
-
 			static FVertexInterface DeclareVertexInterface(const FString& InVertexName)
 			{
 				return FVertexInterface(
@@ -132,11 +117,30 @@ namespace Metasound
 				);
 			}
 
+			static FNodeInfo GetNodeInfo(const FString& InVertexName)
+			{
+				static const FString ClassNameString = FString(TEXT("Input_")) + GetMetasoundDataTypeName<DataType>().ToString();
+				static const FName ClassName(*ClassNameString);
+
+				FNodeInfo Info;
+
+				Info.ClassName = ClassName;
+				Info.MajorVersion = 1;
+				Info.MinorVersion = 0;
+				Info.Description = LOCTEXT("Metasound_InputNodeDescription", "Input into the parent Metasound graph.");
+				Info.Author = PluginAuthor;
+				Info.PromptIfMissing = PluginNodeMissingPrompt;
+				Info.DefaultInterface = DeclareVertexInterface(InVertexName);
+
+				return Info;
+			}
+
+
 			/* Construct a TInputNode using the TInputOperatorFactory<> and forwarding 
 			 * Args to the TInputOperatorFactory constructor.*/
 			template<typename... ArgTypes>
 			TInputNode(const FString& InNodeDescription, const FString& InVertexName, ArgTypes&&... Args)
-			:	FNode(InNodeDescription, GetNodeInfo())
+			:	FNode(InNodeDescription, GetNodeInfo(InVertexName))
 			,	VertexName(InVertexName)
 			,	Interface(DeclareVertexInterface(InVertexName))
 			,	Factory(MakeOperatorFactoryRef<TInputOperatorFactory<DataType>>(Forward<ArgTypes>(Args)...))
@@ -146,7 +150,7 @@ namespace Metasound
 			/* Construct a TInputNode using the TInputOperatorLiteralFactory<> and moving
 			 * InParam to the TInputOperatorLiteralFactory constructor.*/
 			explicit TInputNode(const FString& InNodeDescription, const FString& InVertexName, FDataTypeLiteralParam&& InParam)
-			:	FNode(InNodeDescription, GetNodeInfo())
+			:	FNode(InNodeDescription, GetNodeInfo(InVertexName))
 			,	VertexName(InVertexName)
 			,	Interface(DeclareVertexInterface(InVertexName))
 			, 	Factory(MakeOperatorFactoryRef<TInputOperatorLiteralFactory<DataType>>(MoveTemp(InParam)))
@@ -163,11 +167,6 @@ namespace Metasound
 				return Interface;
 			}
 
-			virtual const FVertexInterface& GetDefaultVertexInterface() const override
-			{
-				return Interface;
-			}
-
 			virtual bool SetVertexInterface(const FVertexInterface& InInterface) override
 			{
 				return Interface == InInterface;
diff --git a/Engine/Plugins/Runtime/Metasound/Source/MetasoundFrontend/Public/MetasoundNodeRegistrationMacro.h b/Engine/Plugins/Runtime/Metasound/Source/MetasoundFrontend/Public/MetasoundNodeRegistrationMacro.h
--- a/Engine/Plugins/Runtime/Metasound/Source/MetasoundFrontend/Public/MetasoundNodeRegistrationMacro.h
+++ b/Engine/Plugins/Runtime/Metasound/Source/MetasoundFrontend/Public/MetasoundNodeRegistrationMacro.h
@@ -5,6 +5,7 @@
 #include "CoreMinimal.h"
 #include "MetasoundNodeInterface.h"
 #include "Templates/IntegralConstant.h"
+#include "MetasoundFrontend.h"
 #include "MetasoundFrontendRegistries.h"
 
 // Utility class to ensure that a node class can use the constructor the frontend uses.
@@ -21,8 +22,9 @@ struct ConstructorTakesNodeInitData {
 	static const bool Value = sizeof(TestForConstructor<NodeClass>(nullptr)) == sizeof(uint16);
 };
 
+
 template <typename FNodeType>
-bool RegisterNodeWithFrontend()
+bool RegisterNodeWithFrontend(const Metasound::FNodeInfo& InMetadata)
 {
 	// if we reenter this code (because DECLARE_METASOUND_DATA_REFERENCE_TYPES was called twice with the same type),
 	// we catch it here.
@@ -39,14 +41,29 @@ bool RegisterNodeWithFrontend()
 		[](const Metasound::FNodeInitData& InInitData) -> TUniquePtr<Metasound::INode>
 		{
 			return TUniquePtr<Metasound::INode>(new FNodeType(InInitData));
-		});
+		},
+		[=]() -> FMetasoundClassDescription
+		{
+			return Metasound::Frontend::GenerateClassDescription(InMetadata);
+		}
+	);
 
 	ensureAlwaysMsgf(bSuccessfullyRegisteredNode, TEXT("Registering node class failed. Please check the logs."));
 	return bSuccessfullyRegisteredNode;
 }
 
-#define METASOUND_REGISTER_NODE(NodeClass) \
+template <typename FNodeType>
+bool RegisterNodeWithFrontend()
+{
+	// Register a node using a prototype node.
+	// TODO: may want to add a warning here since we don't want to use this registration pathway.
+	Metasound::FNodeInitData InitData;
+	TUniquePtr<Metasound::INode> Node = MakeUnique<FNodeType>(InitData);
+	return RegisterNodeWithFrontend<FNodeType>(Node->GetMetadata());
+}
+
+#define METASOUND_REGISTER_NODE(NodeClass, ...) \
 	 static_assert(std::is_base_of<::Metasound::INodeBase, NodeClass>::value, "To be registered as a  Metasound Node," #NodeClass "need to be a derived class from Metasound::INodeBase, Metasound::INode, or Metasound::FNode."); \
 	 static_assert(::ConstructorTakesNodeInitData<NodeClass>::Value, "In order to be registered as a Metasound Node, " #NodeClass " needs to implement the following public constructor: " #NodeClass "(const Metasound::FNodeInitData& InInitData);"); \
-	 static bool bSuccessfullyRegistered##NodeClass  = FMetasoundFrontendRegistryContainer::Get()->EnqueueInitCommand([](){ ::RegisterNodeWithFrontend<NodeClass>(); }); // This static bool is useful for debugging, but also is the only way the compiler will let us call this function outside of an expression.
+	 static bool bSuccessfullyRegistered##NodeClass  = FMetasoundFrontendRegistryContainer::Get()->EnqueueInitCommand([](){ ::RegisterNodeWithFrontend<NodeClass>(__VA_ARGS__); }); // This static bool is useful for debugging, but also is the only way the compiler will let us call this function outside of an expression.
 
diff --git a/Engine/Plugins/Runtime/Metasound/Source/MetasoundFrontend/Public/MetasoundOutputNode.h b/Engine/Plugins/Runtime/Metasound/Source/MetasoundFrontend/Public/MetasoundOutputNode.h
--- a/Engine/Plugins/Runtime/Metasound/Source/MetasoundFrontend/Public/MetasoundOutputNode.h
+++ b/Engine/Plugins/Runtime/Metasound/Source/MetasoundFrontend/Public/MetasoundOutputNode.h
@@ -69,20 +69,6 @@ namespace Metasound
 				FString DataReferenceName;
 		};
 
-		static const FNodeInfo& GetNodeInfo()
-		{
-			static const FString ClassNameString = FString(TEXT("Output_")) + GetMetasoundDataTypeName<DataType>().ToString();
-
-			static const FNodeInfo NodeInfo = {
-				FName(*ClassNameString),
-				LOCTEXT("Metasound_OutputNodeDescription", "Output from the parent Metasound graph."),
-				PluginAuthor,
-				PluginNodeMissingPrompt
-			};
-
-			return NodeInfo;
-		}
-
 		static FVertexInterface GetVertexInterface(const FString& InVertexName)
 		{
 			return FVertexInterface(
@@ -95,9 +81,28 @@ namespace Metasound
 			);
 		}
 
+		static FNodeInfo GetNodeInfo(const FString& InVertexName)
+		{
+			static const FString ClassNameString = FString(TEXT("Output_")) + GetMetasoundDataTypeName<DataType>().ToString();
+
+			FNodeInfo Info;
+
+			Info.ClassName = FName(*ClassNameString);
+			Info.MajorVersion = 1;
+			Info.MinorVersion = 0;
+			Info.Description = LOCTEXT("Metasound_OutputNodeDescription", "Output from the parent Metasound graph.");
+			Info.Author = PluginAuthor;
+			Info.PromptIfMissing = PluginNodeMissingPrompt;
+			Info.DefaultInterface = GetVertexInterface(InVertexName);
+
+			return Info;
+		};
+
+
+
 		public:
 			TOutputNode(const FString& InInstanceName, const FString& InVertexName)
-			:	FNode(InInstanceName, GetNodeInfo())
+			:	FNode(InInstanceName, GetNodeInfo(InVertexName))
 			,	VertexInterface(GetVertexInterface(InVertexName))
 			,	Factory(MakeShared<FOutputOperatorFactory, ESPMode::ThreadSafe>(InVertexName))
 			{
@@ -109,12 +114,6 @@ namespace Metasound
 				return VertexInterface;
 			}
 
-			/** Return the default vertex interface. */
-			virtual const FVertexInterface& GetDefaultVertexInterface() const override
-			{
-				return VertexInterface;
-			}
-
 			/** Set the vertex interface. If the vertex was successfully changed, returns true. 
 			 *
 			 * @param InInterface - New interface for node. 
diff --git a/Engine/Plugins/Runtime/Metasound/Source/MetasoundFrontend/Public/MetasoundReceiveNode.h b/Engine/Plugins/Runtime/Metasound/Source/MetasoundFrontend/Public/MetasoundReceiveNode.h
--- a/Engine/Plugins/Runtime/Metasound/Source/MetasoundFrontend/Public/MetasoundReceiveNode.h
+++ b/Engine/Plugins/Runtime/Metasound/Source/MetasoundFrontend/Public/MetasoundReceiveNode.h
@@ -3,6 +3,7 @@
 
 #include "CoreMinimal.h"
 #include "MetasoundBuilderInterface.h"
+#include "MetasoundNode.h"
 #include "MetasoundNodeInterface.h"
 #include "MetasoundOperatorInterface.h"
 #include "MetasoundDataFactory.h"
@@ -12,31 +13,13 @@
 
 #include <type_traits>
 
-#define LOCTEXT_NAMESPACE "MetasoundGraphCore"
+#define LOCTEXT_NAMESPACE "MetasoundFrontend"
 
 
 namespace Metasound
 {
-	// TODO: should consider making this a subclass of TInputNode. Both share requirement
-	// of needing a default DataReference. The construction of that gets tricky so 
-	// would be nice to consolidate all that logic. 
-	//
-	//
-	// TODO: Maybe it could work like this...
-	// Can instantiate a factory with default args. Factory is valid as long as constructor of data type supports
-	// parameters (w/ or w/o operator settings.) Factory makes a lambda for creating data type references.
-	// Factory gets instantiated when node is declared.  Basically lets us forward constructor 
-	// args until the creation of the lambda, which requires a copy to capture "[=]". 
-	//
-	// But that still leaves the issue of setting those values on the node. Would be best to have that in NodeInitParams
-	// but its not clear how that is passed about. Maybe something that checks whether it can construct from literal. 
-	//
-	// Or you could go another route where you check whether the Operator constructor takes a specific set of references and 
-	// use that to auto-generate factories.  
-	
-	// This convenience node can be registered and will invoke static_cast<ToDataType>(FromDataType) every time it is executed.
 	template<typename TDataType>
-	class TReceiveNode : public INode
+	class TReceiveNode : public FNode
 	{
 
 	public:
@@ -52,6 +35,41 @@ namespace Metasound
 			return SendInput;
 		}
 
+		static FVertexInterface DeclareVertexInterface()
+		{
+			return FVertexInterface(
+				FInputVertexInterface(
+					TInputDataVertexModel<FSendAddress>(GetAddressInputName(), FText::GetEmpty())
+				),
+				FOutputVertexInterface(
+					TOutputDataVertexModel<TDataType>(GetSendOutputName(), FText::GetEmpty())
+				)
+			);
+		}
+
+		static const FNodeInfo& GetNodeInfo()
+		{
+			auto InitNodeInfo = []() -> FNodeInfo
+			{
+				FNodeInfo Info;
+				Info.ClassName = FName(*FString::Printf(TEXT("Receive %s"), *GetSendOutputName()));
+				Info.MajorVersion = 1;
+				Info.MinorVersion = 0;
+				Info.Description = LOCTEXT("Metasound_ReceiveNodeDescription", "Receives data from a send node with the same name.");
+				Info.Author = PluginAuthor;
+				Info.PromptIfMissing = PluginNodeMissingPrompt;
+				Info.DefaultInterface = DeclareVertexInterface();
+				Info.CategoryHierarchy = {LOCTEXT("Metasound_ReceiveNodeCategory", "Receive"), FText::AsCultureInvariant(GetMetasoundDataTypeString<TDataType>())};
+				Info.Keywords = {TEXT("receive"), GetMetasoundDataTypeName<TDataType>()};
+
+				return Info;
+			};
+
+			static const FNodeInfo Info = InitNodeInfo();
+
+			return Info;
+		}
+
 	private:
 		class TReceiverOperator : public TExecutableOperator<TReceiverOperator>
 		{
@@ -124,20 +142,9 @@ namespace Metasound
 
 		public:
 
-			static FVertexInterface DeclareVertexInterface()
-			{
-				return FVertexInterface(
-					FInputVertexInterface(
-						TInputDataVertexModel<FSendAddress>(GetAddressInputName(), FText::GetEmpty())
-					),
-					FOutputVertexInterface(
-						TOutputDataVertexModel<TDataType>(GetSendOutputName(), FText::GetEmpty())
-					)
-				);
-			}
-
+			// TODO: default value of received object.
 			TReceiveNode(const FNodeInitData& InInitData)
-				: NodeDescription(InInitData.InstanceName)
+				: FNode(InInitData.InstanceName, GetNodeInfo())
 				, Interface(DeclareVertexInterface())
 				, Factory(MakeOperatorFactoryRef<FReceiverOperatorFactory>())
 			{
@@ -145,43 +152,11 @@ namespace Metasound
 
 			virtual ~TReceiveNode() = default;
 
-			virtual const FName& GetClassName() const override
-			{
-				static const FName ClassName(*FString::Printf(TEXT("Receive %s"), *GetSendOutputName()));
-				return ClassName;
-			}
-
-			virtual const FString& GetInstanceName() const override
-			{
-				return NodeDescription;
-			}
-
-			virtual const FText& GetDescription() const override
-			{
-				static const FText Description = LOCTEXT("Metasound_ReceiveNodeDescription", "receives data from a send node with the same name.");
-				return Description;
-			}
-
-			virtual const FText& GetAuthorName() const override
-			{
-				return PluginAuthor;
-			}
-
-			virtual const FText& GetPromptIfMissing() const override
-			{
-				return PluginNodeMissingPrompt;
-			}
-
 			virtual const FVertexInterface& GetVertexInterface() const override
 			{
 				return Interface;
 			}
 
-			virtual const FVertexInterface& GetDefaultVertexInterface() const override
-			{
-				return Interface;
-			}
-
 			virtual bool SetVertexInterface(const FVertexInterface& InInterface) override
 			{
 				return Interface == InInterface;
@@ -198,7 +173,6 @@ namespace Metasound
 			}
 
 		private:
-			FString NodeDescription;
 			FVertexInterface Interface;
 			FOperatorFactorySharedRef Factory;
 	};
diff --git a/Engine/Plugins/Runtime/Metasound/Source/MetasoundFrontend/Public/MetasoundSendNode.h b/Engine/Plugins/Runtime/Metasound/Source/MetasoundFrontend/Public/MetasoundSendNode.h
--- a/Engine/Plugins/Runtime/Metasound/Source/MetasoundFrontend/Public/MetasoundSendNode.h
+++ b/Engine/Plugins/Runtime/Metasound/Source/MetasoundFrontend/Public/MetasoundSendNode.h
@@ -3,6 +3,7 @@
 
 #include "CoreMinimal.h"
 #include "MetasoundBuilderInterface.h"
+#include "MetasoundNode.h"
 #include "MetasoundNodeInterface.h"
 #include "MetasoundOperatorInterface.h"
 #include "MetasoundDataReference.h"
@@ -11,14 +12,13 @@
 
 #include <type_traits>
 
-#define LOCTEXT_NAMESPACE "MetasoundGraphCore"
+#define LOCTEXT_NAMESPACE "MetasoundFrontend"
 
 
 namespace Metasound
 {
-	// This convenience node can be registered and will invoke static_cast<ToDataType>(FromDataType) every time it is executed.
 	template<typename TDataType>
-	class TSendNode : public INode
+	class TSendNode : public FNode
 	{
 	public:
 		static const FString& GetAddressInputName()
@@ -33,6 +33,43 @@ namespace Metasound
 			return SendInput;
 		}
 
+		static FVertexInterface DeclareVertexInterface()
+		{
+			return FVertexInterface(
+				FInputVertexInterface(
+					TInputDataVertexModel<FSendAddress>(GetAddressInputName(), FText::GetEmpty()),
+					TInputDataVertexModel<TDataType>(GetSendInputName(), FText::GetEmpty())
+				),
+				FOutputVertexInterface(
+				)
+			);
+		}
+
+		static const FNodeInfo& GetNodeInfo()
+		{
+			auto InitNodeInfo = []() -> FNodeInfo
+			{
+				FNodeInfo Info;
+
+				Info.ClassName = FName(*FString::Printf(TEXT("Send %s"), *GetSendInputName()));
+				Info.MajorVersion = 1;
+				Info.MinorVersion = 0;
+				Info.Description = LOCTEXT("Metasound_SendNodeDescription", "Sends data from a send node with the same name.");
+				Info.Author = PluginAuthor;
+				Info.PromptIfMissing = PluginNodeMissingPrompt;
+				Info.DefaultInterface = DeclareVertexInterface();
+				Info.CategoryHierarchy = {LOCTEXT("Metasound_SendNodeCategory", "Send"), FText::AsCultureInvariant(GetMetasoundDataTypeString<TDataType>())};
+				Info.Keywords = {TEXT("send"), GetMetasoundDataTypeName<TDataType>()};
+
+				return Info;
+			};
+
+			static const FNodeInfo Info = InitNodeInfo();
+
+			return Info;
+		}
+
+
 	private:
 		class TSendOperator : public TExecutableOperator<TSendOperator>
 		{
@@ -100,20 +137,9 @@ namespace Metasound
 		};
 
 		public:
-			static FVertexInterface DeclareVertexInterface()
-			{
-				return FVertexInterface(
-					FInputVertexInterface(
-						TInputDataVertexModel<FSendAddress>(GetAddressInputName(), FText::GetEmpty()),
-						TInputDataVertexModel<TDataType>(GetSendInputName(), FText::GetEmpty())
-					),
-					FOutputVertexInterface(
-					)
-				);
-			}
 
 			TSendNode(const FNodeInitData& InInitData)
-				: NodeDescription(InInitData.InstanceName)
+				: FNode(InInitData.InstanceName, GetNodeInfo())
 				, Interface(DeclareVertexInterface())
 				, Factory(MakeOperatorFactoryRef<FSendOperatorFactory>())
 			{
@@ -121,43 +147,11 @@ namespace Metasound
 
 			virtual ~TSendNode() = default;
 
-			virtual const FName& GetClassName() const override
-			{
-				static const FName ClassName(*FString::Printf(TEXT("Send %s"), *GetSendInputName()));
-				return ClassName;
-			}
-
-			virtual const FString& GetInstanceName() const override
-			{
-				return NodeDescription;
-			}
-
-			virtual const FText& GetDescription() const override
-			{
-				static const FText Description = LOCTEXT("Metasound_SendNodeDescription", "Sends data to any other metasound anywhere else.");
-				return Description;
-			}
-
-			virtual const FText& GetAuthorName() const override
-			{
-				return PluginAuthor;
-			}
-
-			virtual const FText& GetPromptIfMissing() const override
-			{
-				return PluginNodeMissingPrompt;
-			}
-
 			virtual const FVertexInterface& GetVertexInterface() const override
 			{
 				return Interface;
 			}
 
-			virtual const FVertexInterface& GetDefaultVertexInterface() const override
-			{
-				return Interface;
-			}
-
 			virtual bool SetVertexInterface(const FVertexInterface& InInterface) override
 			{
 				return Interface == InInterface;
@@ -174,7 +168,6 @@ namespace Metasound
 			}
 
 		private:
-			FString NodeDescription;
 			FVertexInterface Interface;
 			FOperatorFactorySharedRef Factory;
 	};
diff --git a/Engine/Plugins/Runtime/Metasound/Source/MetasoundGraphCore/Private/MetasoundFacade.cpp b/Engine/Plugins/Runtime/Metasound/Source/MetasoundGraphCore/Private/MetasoundFacade.cpp
--- a/Engine/Plugins/Runtime/Metasound/Source/MetasoundGraphCore/Private/MetasoundFacade.cpp
+++ b/Engine/Plugins/Runtime/Metasound/Source/MetasoundGraphCore/Private/MetasoundFacade.cpp
@@ -19,11 +19,6 @@ namespace Metasound
 		return VertexInterface;
 	}
 
-	const FVertexInterface& FNodeFacade::GetDefaultVertexInterface() const
-	{
-		return VertexInterface;
-	}
-
 	bool FNodeFacade::SetVertexInterface(const FVertexInterface& InInterface)
 	{
 		return InInterface == VertexInterface;
diff --git a/Engine/Plugins/Runtime/Metasound/Source/MetasoundGraphCore/Private/MetasoundGraph.cpp b/Engine/Plugins/Runtime/Metasound/Source/MetasoundGraphCore/Private/MetasoundGraph.cpp
--- a/Engine/Plugins/Runtime/Metasound/Source/MetasoundGraphCore/Private/MetasoundGraph.cpp
+++ b/Engine/Plugins/Runtime/Metasound/Source/MetasoundGraphCore/Private/MetasoundGraph.cpp
@@ -6,35 +6,20 @@
 
 namespace Metasound
 {
-	FGraph::FGraph(const FString& InDescription)
-		: Description(InDescription)
+	FGraph::FGraph(const FString& InInstanceName)
+	: InstanceName(InInstanceName)
+	, Metadata(FNodeInfo::GetEmpty())
 	{
 	}
 
-	const FText& FGraph::GetDescription() const
-	{
-		return FText::GetEmpty();
-	}
-
 	const FString& FGraph::GetInstanceName() const
 	{
-		return Description;
-	}
-
-	const FName& FGraph::GetClassName() const
-	{
-		static const FName ClassName("Graph");
-		return ClassName;
+		return InstanceName;
 	}
 
-	const FText& FGraph::GetAuthorName() const
+	const FNodeInfo& FGraph::GetMetadata() const
 	{
-		return PluginAuthor;
-	}
-
-	const FText& FGraph::GetPromptIfMissing() const
-	{
-		return PluginNodeMissingPrompt;
+		return Metadata;
 	}
 
 	bool FGraph::AddInputDataDestination(const INode& InNode, const FVertexKey& InVertexKey)
@@ -53,7 +38,7 @@ namespace Metasound
 
 	void FGraph::AddInputDataDestination(const FInputDataDestination& InDestination)
 	{
-		VertexInterface.GetInputInterface().Add(InDestination.Vertex);
+		Metadata.DefaultInterface.GetInputInterface().Add(InDestination.Vertex);
 		InputDestinations.Add(MakeDestinationDataVertexKey(InDestination), InDestination);
 	}
 
@@ -78,7 +63,7 @@ namespace Metasound
 
 	void FGraph::AddOutputDataSource(const FOutputDataSource& InSource)
 	{
-		VertexInterface.GetOutputInterface().Add(InSource.Vertex);
+		Metadata.DefaultInterface.GetOutputInterface().Add(InSource.Vertex);
 		OutputSources.Add(MakeSourceDataVertexKey(InSource), InSource);
 	}
 
@@ -130,23 +115,17 @@ namespace Metasound
 
 	const FVertexInterface& FGraph::GetVertexInterface() const 
 	{
-		return VertexInterface;
+		return Metadata.DefaultInterface;
 	}
 
-	const FVertexInterface& FGraph::GetDefaultVertexInterface() const 
-	{
-		static const FVertexInterface EmptyInterface;
-
-		return EmptyInterface;
-	}
 
 	bool FGraph::SetVertexInterface(const FVertexInterface& InInterface)
 	{
-		return InInterface == VertexInterface;
+		return InInterface == Metadata.DefaultInterface;
 	}
 
 	bool FGraph::IsVertexInterfaceSupported(const FVertexInterface& InInterface) const
 	{
-		return InInterface == VertexInterface;
+		return InInterface == Metadata.DefaultInterface;
 	}
 }
diff --git a/Engine/Plugins/Runtime/Metasound/Source/MetasoundGraphCore/Private/MetasoundNode.cpp b/Engine/Plugins/Runtime/Metasound/Source/MetasoundGraphCore/Private/MetasoundNode.cpp
--- a/Engine/Plugins/Runtime/Metasound/Source/MetasoundGraphCore/Private/MetasoundNode.cpp
+++ b/Engine/Plugins/Runtime/Metasound/Source/MetasoundGraphCore/Private/MetasoundNode.cpp
@@ -19,29 +19,8 @@ namespace Metasound
 	}
 
 	/** Return the type name of this node. */
-	const FName& FNode::GetClassName() const
+	const FNodeInfo& FNode::GetMetadata() const
 	{
-		return Info.ClassName;
-	}
-
-	/** Return a longer text description describing how this node is used. */
-	const FText& FNode::GetDescription() const
-	{
-		return Info.Description;
-	}
-
-	/** Return the original author of this node class. */
-	const FText& FNode::GetAuthorName() const
-	{
-		return Info.AuthorName;
-	}
-
-	/** 
-	 *  Return an optional prompt on how users can get the plugin this node is in,
-	 *  if they have found a metasound that uses this node but don't have this plugin downloaded or enabled.
-	 */
-	const FText& FNode::GetPromptIfMissing() const
-	{
-		return Info.PromptIfMissing;
+		return Info;
 	}
 }
diff --git a/Engine/Plugins/Runtime/Metasound/Source/MetasoundGraphCore/Public/MetasoundDataReference.h b/Engine/Plugins/Runtime/Metasound/Source/MetasoundGraphCore/Public/MetasoundDataReference.h
--- a/Engine/Plugins/Runtime/Metasound/Source/MetasoundGraphCore/Public/MetasoundDataReference.h
+++ b/Engine/Plugins/Runtime/Metasound/Source/MetasoundGraphCore/Public/MetasoundDataReference.h
@@ -64,7 +64,7 @@ namespace Metasound
 
 	/** Return the data type FName for a registered data type. */
 	template<typename DataType>
-	const FName GetMetasoundDataTypeName() 
+	const FName& GetMetasoundDataTypeName() 
 	{
 		static const FName TypeName = FName(TDataReferenceTypeInfo<std::decay_t<DataType>>::TypeName);
 
@@ -73,7 +73,7 @@ namespace Metasound
 
 	/** Return the data type string for a registered data type. */
 	template<typename DataType>
-	const FString GetMetasoundDataTypeString() 
+	const FString& GetMetasoundDataTypeString() 
 	{
 		static const FString TypeString = FString(TDataReferenceTypeInfo<std::decay_t<DataType>>::TypeName);
 
diff --git a/Engine/Plugins/Runtime/Metasound/Source/MetasoundGraphCore/Public/MetasoundFacade.h b/Engine/Plugins/Runtime/Metasound/Source/MetasoundGraphCore/Public/MetasoundFacade.h
--- a/Engine/Plugins/Runtime/Metasound/Source/MetasoundGraphCore/Public/MetasoundFacade.h
+++ b/Engine/Plugins/Runtime/Metasound/Source/MetasoundGraphCore/Public/MetasoundFacade.h
@@ -10,25 +10,6 @@
 
 namespace Metasound
 {
-	// Helper template to determine whether a static class function is declared
-	// for a given template class.
-	template <typename U>
-	class TIsVertexInterfaceDeclared
-	{
-		private:
-			template<typename T, T> 
-			struct Helper;
-
-			// Check for "static FVertexInterface U::DeclareVertexInterface()"
-			template<typename T>
-			static uint8 Check(Helper<FVertexInterface(*)(), &T::DeclareVertexInterface>*);
-
-			template<typename T> static uint16 Check(...);
-		public:
-			// If the function exists, then "Value" is true. Otherwise "Value" is false.
-			static constexpr bool Value = sizeof(Check<U>(0)) == sizeof(uint8_t);
-	};
-
 	// Helper template to determine whether a static class function is declared
 	// for a given template class.
 	template <typename U>
@@ -81,13 +62,9 @@ namespace Metasound
 		// Require that OperatorType is subclass of IOperator
 		static_assert(std::is_base_of<IOperator, OperatorType>::value, "To use the FNodeFacade consturctor, the OperatorType must be derived from IOperator");
 
-		// Require static FVertexInterface OperatorType::DeclareVertexInterface() exists.
-		static_assert(TIsVertexInterfaceDeclared<OperatorType>::Value, "To use the FNodeFacade constructor, the OperatorType must have the static function \"static FVertexInterface OperatorType::DeclareVertexInterface()\"");
-
 		// Require static TUniquePtr<IOperator> OperatorType::CreateOperator(const FCreateOperatorParams&, TArray<TUniquePtr<IOperatorBuildError>>&) exists.
 		static_assert(TIsFactoryMethodDeclared<OperatorType>::Value, "To use the FNodeFacade constructor, the OperatorType must have the static function \"static TUniquePtr<IOperator> OperatorType::CreateOperator(const FCreateOperatorParams&, TArray<TUniquePtr<IOperatorBuildError>>&)\"");
 
-
 		// Require static const FNodeInfo& OperatorType::GetNodeInfo() exists.
 		static_assert(TIsNodeInfoDeclared<OperatorType>::Value, "To use the FNodeFacade constructor, the OperatorType must have the static function \"static const FNodeInfo& OperatorType::GetNodeInfo()\"");
 
@@ -96,11 +73,6 @@ namespace Metasound
 			return OperatorType::GetNodeInfo(); 
 		}
 
-		static FVertexInterface DeclareVertexInterface() 
-		{ 
-			return OperatorType::DeclareVertexInterface(); 
-		}
-
 		static TUniquePtr<TUniquePtr<IOperator>> CreateOperator(const FCreateOperatorParams& InParams, FBuildErrorArray& OutErrors) 
 		{
 			return OperatorType::CreateOperator(InParams, OutErrors);
@@ -148,16 +120,14 @@ namespace Metasound
 			FNodeFacade(const FString& InInstanceName, TFacadeOperatorClass<OperatorType> OperatorClass)
 			:	FNode(InInstanceName, OperatorType::GetNodeInfo())
 			, 	Factory(MakeShared<FFactory, ESPMode::ThreadSafe>(&OperatorType::CreateOperator))
-			,	VertexInterface(OperatorType::DeclareVertexInterface())
 			{
+				VertexInterface = GetMetadata().DefaultInterface;
 			}
 
 			virtual ~FNodeFacade() = default;
 
 			virtual const FVertexInterface& GetVertexInterface() const override;
 
-			virtual const FVertexInterface& GetDefaultVertexInterface() const override;
-
 			/** FNodeFacade has a static vertex interface. This will return true
 			 * only if the input interface is equal to the default interface.
 			 */
@@ -177,6 +147,4 @@ namespace Metasound
 
 			FVertexInterface VertexInterface;
 	};
-
-
 }
diff --git a/Engine/Plugins/Runtime/Metasound/Source/MetasoundGraphCore/Public/MetasoundGraph.h b/Engine/Plugins/Runtime/Metasound/Source/MetasoundGraphCore/Public/MetasoundGraph.h
--- a/Engine/Plugins/Runtime/Metasound/Source/MetasoundGraphCore/Public/MetasoundGraph.h
+++ b/Engine/Plugins/Runtime/Metasound/Source/MetasoundGraphCore/Public/MetasoundGraph.h
@@ -21,27 +21,16 @@ namespace Metasound
 			virtual ~FGraph() = default;
 
 			/** Return the name of this specific instance of the node class. */
-			virtual const FString& GetInstanceName() const override;
+			const FString& GetInstanceName() const override;
 
-			/** Return the type name of this node. */
-			virtual const FName& GetClassName() const override;
-
-			/** Return a longer text description describing how this node is used. */
-			virtual const FText& GetDescription() const override;
-
-			/** Return the original author of this node class. */
-			virtual const FText& GetAuthorName() const override;
-
-			virtual const FText& GetPromptIfMissing() const override;
+			/** Return metadata about this graph. */
+			const FNodeInfo& GetMetadata() const override;
 
 			/** Retrieve all the edges associated with a graph. */
-			virtual const TArray<FDataEdge>& GetDataEdges() const override;
+			const TArray<FDataEdge>& GetDataEdges() const override;
 
 			/** Return the current vertex interface. */
-			virtual const FVertexInterface& GetVertexInterface() const override;
-
-			/** Return the default vertex interface. */
-			virtual const FVertexInterface& GetDefaultVertexInterface() const override;
+			const FVertexInterface& GetVertexInterface() const override;
 
 			/** Set the vertex interface. If the vertex was successfully changed, returns true. 
 			 *
@@ -49,7 +38,7 @@ namespace Metasound
 			 *
 			 * @return True on success, false otherwise.
 			 */
-			virtual bool SetVertexInterface(const FVertexInterface& InInterface) override;
+			bool SetVertexInterface(const FVertexInterface& InInterface) override;
 
 			/** Expresses whether a specific vertex interface is supported.
 			 *
@@ -57,13 +46,13 @@ namespace Metasound
 			 *
 			 * @return True if the interface is supported, false otherwise. 
 			 */
-			virtual bool IsVertexInterfaceSupported(const FVertexInterface& InInterface) const override;
+			bool IsVertexInterfaceSupported(const FVertexInterface& InInterface) const override;
 
 			/** Get vertices which contain input parameters. */
-			virtual const FInputDataDestinationCollection& GetInputDataDestinations() const override;
+			const FInputDataDestinationCollection& GetInputDataDestinations() const override;
 
 			/** Get vertices which contain output parameters. */
-			virtual const FOutputDataSourceCollection& GetOutputDataSources() const override;
+			const FOutputDataSourceCollection& GetOutputDataSources() const override;
 
 			/** Add an edge to the graph. */
 			void AddDataEdge(const FDataEdge& InEdge);
@@ -116,12 +105,12 @@ namespace Metasound
 			// TODO: Add ability to remove things.
 
 		private:
-			FString Description;
+			FString InstanceName;
+			FNodeInfo Metadata;
 
 
 			TArray<FDataEdge> Edges;
 
-			FVertexInterface VertexInterface;
 			FInputDataDestinationCollection InputDestinations;
 			FOutputDataSourceCollection OutputSources;
 	};
diff --git a/Engine/Plugins/Runtime/Metasound/Source/MetasoundGraphCore/Public/MetasoundNode.h b/Engine/Plugins/Runtime/Metasound/Source/MetasoundGraphCore/Public/MetasoundNode.h
--- a/Engine/Plugins/Runtime/Metasound/Source/MetasoundGraphCore/Public/MetasoundNode.h
+++ b/Engine/Plugins/Runtime/Metasound/Source/MetasoundGraphCore/Public/MetasoundNode.h
@@ -7,14 +7,6 @@
 
 namespace Metasound
 {
-	struct FNodeInfo
-	{
-		FName ClassName;
-		FText Description;
-		FText AuthorName;
-		FText PromptIfMissing;
-	};
-
 	class METASOUNDGRAPHCORE_API FNode : public INode
 	{
 		public:
@@ -23,27 +15,14 @@ namespace Metasound
 			virtual ~FNode() = default;
 
 			/** Return the name of this specific instance of the node class. */
-			virtual const FString& GetInstanceName() const override;
-
-			/** Return the type name of this node. */
-			virtual const FName& GetClassName() const override;
+			const FString& GetInstanceName() const override;
 
-			/** Return a longer text description describing how this node is used. */
-			virtual const FText& GetDescription() const override;
-
-			/** Return the original author of this node class. */
-			virtual const FText& GetAuthorName() const override;
-
-			/** 
-			 *  Return an optional prompt on how users can get the plugin this node is in,
-			 *  if they have found a metasound that uses this node but don't have this plugin downloaded or enabled.
-			 */
-			virtual const FText& GetPromptIfMissing() const override;
+			/** Return metadata associated with this node. */
+			const FNodeInfo& GetMetadata() const override;
 
 		private:
 
 			FString InstanceName;
 			FNodeInfo Info;
-
 	};
 }
diff --git a/Engine/Plugins/Runtime/Metasound/Source/MetasoundGraphCore/Public/MetasoundNodeInterface.h b/Engine/Plugins/Runtime/Metasound/Source/MetasoundGraphCore/Public/MetasoundNodeInterface.h
--- a/Engine/Plugins/Runtime/Metasound/Source/MetasoundGraphCore/Public/MetasoundNodeInterface.h
+++ b/Engine/Plugins/Runtime/Metasound/Source/MetasoundGraphCore/Public/MetasoundNodeInterface.h
@@ -15,7 +15,6 @@ namespace Metasound
 	static const FText PluginAuthor = NSLOCTEXT("MetasoundGraphCore", "Metasound_DefaultAuthor", "Epic Games, Inc.");
 	static const FText PluginNodeMissingPrompt = NSLOCTEXT("MetasoundGraphCore", "Metasound_DefaultMissingPrompt", "Make sure that the Metasound plugin is loaded.");
 
-
 	/**
 	 * This struct is used to pass in any arguments required for constructing a single node instance.
 	 * because of this, all FNode implementations have to implement a constructor that takes an FNodeInitData instance.
@@ -30,12 +29,50 @@ namespace Metasound
 		{
 			checkf(ParamMap.Contains(ParamName), TEXT("Tried to use node initialization parameter that didn't exist!"));
 
-			// If this hits, check that you have a constructor for DataType that can interpret the parameter.
-			//checkf(ParamMap[ParamName].IsCompatibleWithType<ParamType>(), TEXT("Tried to parse initialization parameter to an invalid type!"));
 			return TDataTypeLiteralFactory<ParamType>::CreateAny(ParamMap[ParamName]);
 		}
 	};
 
+	/** Provides metadata for a given node. */
+	struct FNodeInfo
+	{
+		// TODO: rename this class to FNodeMetadata
+
+		/** Name of class. Used for registration and lookup. */
+		FName ClassName;
+
+		/** Major version of node. Used for registration and lookup. */
+		int32 MajorVersion = -1;
+
+		/** Minor version of node. */
+		int32 MinorVersion = -1;
+
+		/** Human readable description of node. */
+		FText Description;
+
+		/** Author information. */
+		FText Author;
+
+		/** Human readable prompt for acquiring plugin in case node is not loaded. */
+		FText PromptIfMissing;
+
+		/** Default vertex interface for the node */
+		FVertexInterface DefaultInterface;
+
+		/** Hierarchy of categories for displaying node. */
+		TArray<FText> CategoryHierarchy;
+
+		/** List of keywords for contextual node searching. */
+		TArray<FName> Keywords;
+
+		/** Returns an empty FNodeInfo object. */
+		static const FNodeInfo& GetEmpty()
+		{
+			static const FNodeInfo EmptyInfo;
+			return EmptyInfo;
+		}
+	};
+
 	/** INodeBase
 	 * 
 	 * Interface for all nodes that can describe their name, type, inputs and outputs.
@@ -49,26 +86,11 @@ namespace Metasound
 			virtual const FString& GetInstanceName() const = 0;
 
 			/** Return the type name of this node. */
-			virtual const FName& GetClassName() const = 0;
-
-			/** Return a longer text description describing how this node is used. */
-			virtual const FText& GetDescription() const = 0;
-
-			/** Return the original author of this node class. */
-			virtual const FText& GetAuthorName() const = 0;
-
-			/** 
-			 *  Return an optional prompt on how users can get the plugin this node is in,
-			 *  if they have found a metasound that uses this node but don't have this plugin downloaded or enabled.
-			 */
-			virtual const FText& GetPromptIfMissing() const = 0;
+			virtual const FNodeInfo& GetMetadata() const = 0;
 
 			/** Return the current vertex interface. */
 			virtual const FVertexInterface& GetVertexInterface() const = 0;
 
-			/** Return the default vertex interface. */
-			virtual const FVertexInterface& GetDefaultVertexInterface() const = 0;
-
 			/** Set the vertex interface. If the vertex was successfully changed, returns true. 
 			 *
 			 * @param InInterface - New interface for node. 
diff --git a/Engine/Plugins/Runtime/Metasound/Source/MetasoundStandardNodes/Private/MetasoundADSRNode.cpp b/Engine/Plugins/Runtime/Metasound/Source/MetasoundStandardNodes/Private/MetasoundADSRNode.cpp
--- a/Engine/Plugins/Runtime/Metasound/Source/MetasoundStandardNodes/Private/MetasoundADSRNode.cpp
+++ b/Engine/Plugins/Runtime/Metasound/Source/MetasoundStandardNodes/Private/MetasoundADSRNode.cpp
@@ -147,18 +147,6 @@ namespace Metasound
 		);
 	}
 
-	const FNodeInfo& FADSROperator::GetNodeInfo()
-	{
-		static const FNodeInfo Info = {
-			FName(TEXT("ADSR")),
-			LOCTEXT("Metasound_ADSRNodeDescription", "Emits an ADSR (Attack, decay, sustain, & release) envelope when bopped."),
-			PluginAuthor,
-			PluginNodeMissingPrompt
-		};
-
-		return Info;
-	}
-
 	FVertexInterface FADSROperator::DeclareVertexInterface()
 	{
 		static const FVertexInterface Interface(
@@ -177,6 +165,28 @@ namespace Metasound
 		return Interface;
 	}
 
+	const FNodeInfo& FADSROperator::GetNodeInfo()
+	{
+		auto InitNodeInfo = []() -> FNodeInfo
+		{
+			FNodeInfo Info;
+			Info.ClassName = FName(TEXT("ADSR"));
+			Info.MajorVersion = 1;
+			Info.MinorVersion = 0;
+			Info.Description = LOCTEXT("Metasound_ADSRNodeDescription", "Emits an ADSR (Attack, decay, sustain, & release) envelope when bopped.");
+			Info.Author = PluginAuthor;
+			Info.PromptIfMissing = PluginNodeMissingPrompt;
+			Info.DefaultInterface = DeclareVertexInterface();
+
+			return Info;
+		};
+
+		static const FNodeInfo Info = InitNodeInfo();
+
+		return Info;
+	}
+
+
 	TUniquePtr<IOperator> FADSROperator::CreateOperator(const FCreateOperatorParams& InParams, FBuildErrorArray& OutErrors)
 	{
 		const FADSRNode& ADSRNode = static_cast<const FADSRNode&>(InParams.Node);
diff --git a/Engine/Plugins/Runtime/Metasound/Source/MetasoundStandardNodes/Private/MetasoundAudioMultiplyNode.cpp b/Engine/Plugins/Runtime/Metasound/Source/MetasoundStandardNodes/Private/MetasoundAudioMultiplyNode.cpp
--- a/Engine/Plugins/Runtime/Metasound/Source/MetasoundStandardNodes/Private/MetasoundAudioMultiplyNode.cpp
+++ b/Engine/Plugins/Runtime/Metasound/Source/MetasoundStandardNodes/Private/MetasoundAudioMultiplyNode.cpp
@@ -115,13 +115,22 @@ namespace Metasound
 
 	const FNodeInfo& FAudioMultiplyOperator::GetNodeInfo()
 	{
-		static const FNodeInfo Info = {
-			FName(TEXT("AudioMultiply")),
-			LOCTEXT("NodeDescription", "Multiply two audio streams, sample by sample."),
-			PluginAuthor,
-			PluginNodeMissingPrompt
+		auto InitNodeInfo = []() -> FNodeInfo
+		{
+			FNodeInfo Info;
+			Info.ClassName = FName(TEXT("AudioMultiply"));
+			Info.MajorVersion = 1;
+			Info.MinorVersion = 0;
+			Info.Description = LOCTEXT("AudioMultiply_NodeDescription", "Multiply two audio streams, sample by sample."),
+			Info.Author = PluginAuthor;
+			Info.PromptIfMissing = PluginNodeMissingPrompt;
+			Info.DefaultInterface = DeclareVertexInterface();
+
+			return Info;
 		};
 
+		static const FNodeInfo Info = InitNodeInfo();
+
 		return Info;
 	};
 
diff --git a/Engine/Plugins/Runtime/Metasound/Source/MetasoundStandardNodes/Private/MetasoundOscNode.cpp b/Engine/Plugins/Runtime/Metasound/Source/MetasoundStandardNodes/Private/MetasoundOscNode.cpp
--- a/Engine/Plugins/Runtime/Metasound/Source/MetasoundStandardNodes/Private/MetasoundOscNode.cpp
+++ b/Engine/Plugins/Runtime/Metasound/Source/MetasoundStandardNodes/Private/MetasoundOscNode.cpp
@@ -78,18 +78,6 @@ namespace Metasound
 		Phase -= FMath::FloorToFloat(Phase / TwoPi) * TwoPi;
 	}
 
-	const FNodeInfo& FOscOperator::GetNodeInfo()
-	{
-		static const FNodeInfo Info = {
-			FName(TEXT("Osc")),
-			LOCTEXT("Metasound_OscNodeDescription", "Emits an audio signal of a sinusoid."),
-			PluginAuthor,
-			PluginNodeMissingPrompt
-		};
-
-		return Info;
-	}
-
 	FVertexInterface FOscOperator::DeclareVertexInterface()
 	{
 		static const FVertexInterface Interface(
@@ -105,6 +93,28 @@ namespace Metasound
 	}
 
 
+	const FNodeInfo& FOscOperator::GetNodeInfo()
+	{
+		auto InitNodeInfo = []() -> FNodeInfo
+		{
+			FNodeInfo Info;
+			Info.ClassName = FName(TEXT("Osc"));
+			Info.MajorVersion = 1;
+			Info.MinorVersion = 0;
+			Info.Description = LOCTEXT("Metasound_OscNodeDescription", "Emits an audio signal of a sinusoid.");
+			Info.Author = PluginAuthor;
+			Info.PromptIfMissing = PluginNodeMissingPrompt;
+			Info.DefaultInterface = DeclareVertexInterface();
+
+			return Info;
+		};
+
+		static const FNodeInfo Info = InitNodeInfo();
+
+		return Info;
+	}
+
+
 	FOscNode::FOscNode(const FString& InName, float InDefaultFrequency)
 	:	FNodeFacade(InName, TFacadeOperatorClass<FOscOperator>())
 	,	DefaultFrequency(InDefaultFrequency)
diff --git a/Engine/Plugins/Runtime/Metasound/Source/MetasoundStandardNodes/Private/MetasoundPeriodicBopNode.cpp b/Engine/Plugins/Runtime/Metasound/Source/MetasoundStandardNodes/Private/MetasoundPeriodicBopNode.cpp
--- a/Engine/Plugins/Runtime/Metasound/Source/MetasoundStandardNodes/Private/MetasoundPeriodicBopNode.cpp
+++ b/Engine/Plugins/Runtime/Metasound/Source/MetasoundStandardNodes/Private/MetasoundPeriodicBopNode.cpp
@@ -105,18 +105,6 @@ namespace Metasound
 		return MakeUnique<FPeriodicBopOperator>(InParams.OperatorSettings, Period);
 	}
 
-	const FNodeInfo& FPeriodicBopOperator::GetNodeInfo()
-	{
-		static const FNodeInfo Info = {
-			FName(TEXT("PeriodicBop")),
-			LOCTEXT("Metasound_PeriodicBopNodeDescription", "Emits a bop periodically based on the period duration given."),
-			PluginAuthor,
-			PluginNodeMissingPrompt
-		};
-
-		return Info;
-	}
-
 	FVertexInterface FPeriodicBopOperator::DeclareVertexInterface()
 	{
 		static const FVertexInterface Interface(
@@ -131,6 +119,28 @@ namespace Metasound
 		return Interface;
 	}
 
+	const FNodeInfo& FPeriodicBopOperator::GetNodeInfo()
+	{
+		auto InitNodeInfo = []() -> FNodeInfo
+		{
+			FNodeInfo Info;
+			Info.ClassName = FName(TEXT("PeriodicBop"));
+			Info.MajorVersion = 1;
+			Info.MinorVersion = 0;
+			Info.Description = LOCTEXT("Metasound_PeriodicBopNodeDescription", "Emits a bop periodically based on the period duration given.");
+			Info.Author = PluginAuthor;
+			Info.PromptIfMissing = PluginNodeMissingPrompt;
+			Info.DefaultInterface = DeclareVertexInterface();
+
+			return Info;
+		};
+
+		static const FNodeInfo Info = InitNodeInfo();
+
+		return Info;
+	}
+
+
 	FPeriodicBopNode::FPeriodicBopNode(const FString& InName, float InDefaultPeriodInSeconds)
 	:	FNodeFacade(InName, TFacadeOperatorClass<FPeriodicBopOperator>())
 	,	DefaultPeriod(InDefaultPeriodInSeconds)
