Metasound Frontend Graph

diff --git a/Engine/Plugins/Runtime/Metasound/Source/MetasoundFrontend/Private/MetasoundAssetBase.cpp b/Engine/Plugins/Runtime/Metasound/Source/MetasoundFrontend/Private/MetasoundAssetBase.cpp
--- a/Engine/Plugins/Runtime/Metasound/Source/MetasoundFrontend/Private/MetasoundAssetBase.cpp
+++ b/Engine/Plugins/Runtime/Metasound/Source/MetasoundFrontend/Private/MetasoundAssetBase.cpp
@@ -60,7 +60,7 @@ Metasound::Frontend::FGraphHandle FMetasoundAssetBase::GetRootGraphHandle() cons
 	using EFromClass = Metasound::Frontend::Path::EFromClass;
 
 	FDescPath PathToGraph = FDescPath()[EFromDocument::ToRootClass][EFromClass::ToGraph];
-	FHandleInitParams InitParams = { GetGraphAccessPoint(), PathToGraph, GetDocument().RootClass.Metadata.NodeName, MakeWeakObjectPtr(GetOwningAsset()) };
+	FHandleInitParams InitParams = { GetGraphAccessPoint(), PathToGraph, INDEX_NONE, GetDocument().RootClass.UniqueID, MakeWeakObjectPtr(GetOwningAsset()) };
 	return FGraphHandle(GetPrivateToken(), InitParams);
 }
 
@@ -83,7 +83,7 @@ TArray<Metasound::Frontend::FGraphHandle> FMetasoundAssetBase::GetAllSubgraphHan
 		if (bIsSubgraph)
 		{
 			FDescPath SubgraphPath = RootPathForDependencyGraphs[Dependency.UniqueID][EFromClass::ToGraph];
-			FHandleInitParams InitParams = { GetGraphAccessPoint(), SubgraphPath, Dependency.Metadata.NodeName, MakeWeakObjectPtr(GetOwningAsset()) };
+			FHandleInitParams InitParams = { GetGraphAccessPoint(), SubgraphPath, INDEX_NONE, Dependency.UniqueID, MakeWeakObjectPtr(GetOwningAsset()) };
 			OutArray.Emplace(GetPrivateToken(), InitParams);
 		}
 	}
diff --git a/Engine/Plugins/Runtime/Metasound/Source/MetasoundFrontend/Private/MetasoundFrontend.cpp b/Engine/Plugins/Runtime/Metasound/Source/MetasoundFrontend/Private/MetasoundFrontend.cpp
--- a/Engine/Plugins/Runtime/Metasound/Source/MetasoundFrontend/Private/MetasoundFrontend.cpp
+++ b/Engine/Plugins/Runtime/Metasound/Source/MetasoundFrontend/Private/MetasoundFrontend.cpp
@@ -4,20 +4,18 @@
 
 #include "HAL/FileManager.h"
 #include "HAL/IConsoleManager.h"
+#include "MetasoundAudioBuffer.h"
+#include "MetasoundDataTypeRegistrationMacro.h"
+#include "MetasoundFrontendGraph.h"
+#include "MetasoundFrontendRegistries.h"
 #include "MetasoundJsonBackend.h"
 #include "MetasoundOperatorBuilder.h"
+#include "MetasoundPrimitives.h"
+#include "MetasoundRouter.h"
 #include "Modules/ModuleManager.h"
 #include "StructDeserializer.h"
 #include "StructSerializer.h"
 #include "Serialization/MemoryReader.h"
-#include "MetasoundFrontendRegistries.h"
-
-#include "MetasoundDataTypeRegistrationMacro.h"
-
-#include "MetasoundPrimitives.h"
-
-#include "MetasoundAudioBuffer.h"
-#include "MetasoundRouter.h"
 
 static int32 MetasoundUndoRollLimitCvar = 128;
 FAutoConsoleVariableRef CVarMetasoundUndoRollLimit(
@@ -368,10 +366,10 @@ namespace Metasound
 			}
 		}
 
-		FInputHandle::FInputHandle(FHandleInitParams::EPrivateToken PrivateToken, const FHandleInitParams& InParams, const FString& InputName)
+		FInputHandle::FInputHandle(FHandleInitParams::EPrivateToken PrivateToken, const FHandleInitParams& InParams, const FString& InputName, EDefaultTag InTag)
 			: ITransactable(MetasoundUndoRollLimitCvar, InParams.InOwningAsset)
 			, NodePtr(InParams.InAccessPoint, InParams.InPath)
-			, NodeClass(InParams.InAccessPoint, Path::GetDependencyPath(InParams.InClassName))
+			, NodeClass(InParams.InAccessPoint, Path::GetDependencyPath(InParams.DependencyID))
 			, InputPtr(InParams.InAccessPoint, NodeClass.GetPath()[Path::EFromClass::ToInputs][*InputName])
 			, OutputNodePtr(nullptr, FDescPath())
 		{
@@ -382,12 +380,12 @@ namespace Metasound
 			}
 		}
 
-		FInputHandle::FInputHandle(FHandleInitParams::EPrivateToken PrivateToken, const FHandleInitParams& InParams)
+		FInputHandle::FInputHandle(FHandleInitParams::EPrivateToken PrivateToken, const FHandleInitParams& InParams, const FString& InputName, EOutputNodeTag InTag)
 			: ITransactable(MetasoundUndoRollLimitCvar, InParams.InOwningAsset)
 			, NodePtr(InParams.InAccessPoint, InParams.InPath)
 			, NodeClass(nullptr, FDescPath())
 			, InputPtr(nullptr, FDescPath())
-			, OutputNodePtr(InParams.InAccessPoint, Path::GetOutputDescriptionPath(InParams.InPath, InParams.InClassName))
+			, OutputNodePtr(InParams.InAccessPoint, Path::GetOutputDescriptionPath(InParams.InPath, InputName))
 		{
 			if (InParams.InAccessPoint.IsValid())
 			{
@@ -400,8 +398,8 @@ namespace Metasound
 		{
 			FDescPath NullPath = FDescPath();
 			FString NullString = FString();
-			FHandleInitParams InitParams = { nullptr, NullPath , NullString, nullptr };
-			return FInputHandle(FHandleInitParams::PrivateToken, InitParams, NullString);
+			FHandleInitParams InitParams = { nullptr, NullPath , FMetasoundNodeDescription::InvalidID, FMetasoundClassDescription::InvalidID, nullptr };
+			return FInputHandle(FHandleInitParams::PrivateToken, InitParams, NullString, FInputHandle::Default);
 		}
 
 		bool FInputHandle::IsValid() const
@@ -510,7 +508,7 @@ namespace Metasound
 			check(Connection);
 
 			FString OutputName = Connection->OutputName;
-			uint32 OutputNodeID = Connection->NodeID;
+			int32 OutputNodeID = Connection->NodeID;
 
 			// All node connections are in the same graph, so we just need to go up one level to the Nodes array
 			// and look up the node by it's unique ID.
@@ -518,13 +516,13 @@ namespace Metasound
 
 			if (NodeClass.IsValid())
 			{
-				FHandleInitParams InitParams = { NodePtr.GetAccessPoint(), OutputNodePath, NodeClass->Metadata.NodeName, OwningAsset };
-				return FOutputHandle(FHandleInitParams::PrivateToken, InitParams, OutputName);
+				FHandleInitParams InitParams = { NodePtr.GetAccessPoint(), OutputNodePath, FMetasoundNodeDescription::InvalidID, NodeClass->UniqueID, OwningAsset };
+				return FOutputHandle(FHandleInitParams::PrivateToken, InitParams, OutputName, FOutputHandle::Default);
 			}
 			else if (OutputNodePtr.IsValid())
 			{
-				FHandleInitParams InitParams = { NodePtr.GetAccessPoint(), OutputNodePath, OutputNodePtr->Name, OwningAsset };
-				return FOutputHandle(FHandleInitParams::PrivateToken, InitParams, OutputName);
+				FHandleInitParams InitParams = { NodePtr.GetAccessPoint(), OutputNodePath, FMetasoundNodeDescription::InvalidID, FMetasoundClassDescription::InvalidID, OwningAsset };
+				return FOutputHandle(FHandleInitParams::PrivateToken, InitParams, OutputName, FOutputHandle::Default);
 			}
 			else
 			{
@@ -566,7 +564,7 @@ namespace Metasound
 				return false;
 			}
 
-			uint32 OutputNodeID = InHandle.GetOwningNodeID();
+			int32 OutputNodeID = InHandle.GetOwningNodeID();
 			FString OutputName = InHandle.GetOutputName();
 
 			FMetasoundNodeConnectionDescription* Connection = GetConnectionDescription();
@@ -624,7 +622,7 @@ namespace Metasound
 				return false;
 			}
 
-			const uint32 OutputNodeID = InHandle.GetOwningNodeID();
+			const int32 OutputNodeID = InHandle.GetOwningNodeID();
 			for (int32 i = 0; i < NodePtr->InputConnections.Num(); ++i)
 			{
 				if (NodePtr->InputConnections[i].NodeID == OutputNodeID)
@@ -644,7 +642,8 @@ namespace Metasound
 			{
 				NodePtr.GetAccessPoint(),
 				OwningGraphPath,
-				FString(), // Caching the graph name is not necessary here.
+				FMetasoundNodeDescription::InvalidID,
+				FMetasoundClassDescription::InvalidID,
 				OwningAsset
 			};
 
@@ -706,10 +705,10 @@ namespace Metasound
 			return nullptr;
 		}
 
-		FOutputHandle::FOutputHandle(FHandleInitParams::EPrivateToken InToken /* unused */, const FHandleInitParams& InParams, const FString& InOutputName)
+		FOutputHandle::FOutputHandle(FHandleInitParams::EPrivateToken InToken /* unused */, const FHandleInitParams& InParams, const FString& InOutputName, FOutputHandle::EDefaultTag InTag)
 			: ITransactable(MetasoundUndoRollLimitCvar, InParams.InOwningAsset)
 			, NodePtr(InParams.InAccessPoint, InParams.InPath)
-			, NodeClass(InParams.InAccessPoint, Path::GetDependencyPath(InParams.InClassName))
+			, NodeClass(InParams.InAccessPoint, Path::GetDependencyPath(InParams.DependencyID))
 			, OutputPtr(InParams.InAccessPoint, NodeClass.GetPath()[Path::EFromClass::ToOutputs][*InOutputName])
 			, InputNodePtr(nullptr, FDescPath())
 		{
@@ -720,12 +719,12 @@ namespace Metasound
 			}
 		}
 
-		FOutputHandle::FOutputHandle(FHandleInitParams::EPrivateToken InToken, const FHandleInitParams& InParams)
+		FOutputHandle::FOutputHandle(FHandleInitParams::EPrivateToken InToken, const FHandleInitParams& InParams, const FString& InOutputName, FOutputHandle::EInputNodeTag InTag)
 			: ITransactable(MetasoundUndoRollLimitCvar, InParams.InOwningAsset)
 			, NodePtr(InParams.InAccessPoint, InParams.InPath)
 			, NodeClass(nullptr, FDescPath())
 			, OutputPtr(nullptr, FDescPath())
-			, InputNodePtr(InParams.InAccessPoint, Path::GetInputDescriptionPath(InParams.InPath, InParams.InClassName))
+			, InputNodePtr(InParams.InAccessPoint, Path::GetInputDescriptionPath(InParams.InPath, InOutputName))
 		{
 			if (InParams.InAccessPoint.IsValid())
 			{
@@ -738,8 +737,8 @@ namespace Metasound
 		{
 			FString NullString;
 			FDescPath NullPath;
-			FHandleInitParams InitParams = { nullptr, NullPath, NullString, nullptr };
-			return FOutputHandle(FHandleInitParams::PrivateToken, InitParams, NullString);
+			FHandleInitParams InitParams = { nullptr, NullPath, FMetasoundNodeDescription::InvalidID, FMetasoundClassDescription::InvalidID, nullptr };
+			return FOutputHandle(FHandleInitParams::PrivateToken, InitParams, NullString, FOutputHandle::Default);
 		}
 
 		bool FOutputHandle::IsValid() const
@@ -814,11 +813,11 @@ namespace Metasound
 			}
 		}
 
-		uint32 FOutputHandle::GetOwningNodeID() const
+		int32 FOutputHandle::GetOwningNodeID() const
 		{
 			if (!IsValid())
 			{
-				return INDEX_NONE;
+				return FMetasoundNodeDescription::InvalidID;
 			}
 
 			return NodePtr->UniqueID;
@@ -858,7 +857,7 @@ namespace Metasound
 		{
 			if (InNodeClassType != EMetasoundClassType::Input && InNodeClassType != EMetasoundClassType::Output)
 			{
-				return TDescriptionPtr<FMetasoundClassDescription>(InitParams.InAccessPoint, Path::GetDependencyPath(InitParams.InClassName));
+				return TDescriptionPtr<FMetasoundClassDescription>(InitParams.InAccessPoint, Path::GetDependencyPath(InitParams.DependencyID));
 			}
 			else
 			{
@@ -884,7 +883,7 @@ namespace Metasound
 		{
 			FDescPath InvalidPath;
 			FString InvalidClassName;
-			FHandleInitParams InitParams = { nullptr, InvalidPath, InvalidClassName, nullptr };
+			FHandleInitParams InitParams = { nullptr, InvalidPath, FMetasoundNodeDescription::InvalidID, FMetasoundClassDescription::InvalidID, nullptr };
 			return FNodeHandle(FHandleInitParams::PrivateToken, InitParams, EMetasoundClassType::Invalid);
 		}
 
@@ -910,8 +909,8 @@ namespace Metasound
 
 				const FString& NodeClassName = NodePtr->Name;
 
-				FHandleInitParams InitParams = { NodePtr.GetAccessPoint(), NodePath, NodeClassName, OwningAsset };
-				OutArray.Emplace(FHandleInitParams::PrivateToken, InitParams);
+				FHandleInitParams InitParams = { NodePtr.GetAccessPoint(), NodePath, NodePtr->UniqueID, FMetasoundClassDescription::InvalidID, OwningAsset };
+				OutArray.Emplace(FHandleInitParams::PrivateToken, InitParams, NodeClassName, FInputHandle::OutputNode);
 			}
 			else
 			{
@@ -922,8 +921,8 @@ namespace Metasound
 				{
 					FDescPath NodePath = NodePtr.GetPath();
 					FString ClassName = GetNodeClassName();
-					FHandleInitParams InitParams = { NodePtr.GetAccessPoint(), NodePath, ClassName, OwningAsset };
-					OutArray.Emplace(FHandleInitParams::PrivateToken, InitParams, InputDescription.Name);
+					FHandleInitParams InitParams = { NodePtr.GetAccessPoint(), NodePath, NodePtr->UniqueID, NodeClass->UniqueID, OwningAsset };
+					OutArray.Emplace(FHandleInitParams::PrivateToken, InitParams, InputDescription.Name, FInputHandle::Default);
 				}
 			}
 
@@ -946,8 +945,8 @@ namespace Metasound
 
 				const FString& NodeClassName = NodePtr->Name;
 
-				FHandleInitParams InitParams = { NodePtr.GetAccessPoint(), NodePath, NodeClassName, OwningAsset };
-				OutArray.Emplace(FHandleInitParams::PrivateToken, InitParams);
+				FHandleInitParams InitParams = { NodePtr.GetAccessPoint(), NodePath, NodePtr->UniqueID, FMetasoundClassDescription::InvalidID, OwningAsset };
+				OutArray.Emplace(FHandleInitParams::PrivateToken, InitParams, NodeClassName, FOutputHandle::InputNode);
 			}
 			else
 			{
@@ -957,8 +956,8 @@ namespace Metasound
 				{
 					FDescPath NodePath = NodePtr.GetPath();
 					FString NodeClassName = GetNodeClassName();
-					FHandleInitParams InitParams = { NodePtr.GetAccessPoint(), NodePath, NodeClassName, OwningAsset };
-					OutArray.Emplace(FHandleInitParams::PrivateToken, InitParams, OutputDescription.Name);
+					FHandleInitParams InitParams = { NodePtr.GetAccessPoint(), NodePath, NodePtr->UniqueID, NodeClass->UniqueID, OwningAsset };
+					OutArray.Emplace(FHandleInitParams::PrivateToken, InitParams, OutputDescription.Name, FOutputHandle::Default);
 				}
 			}
 
@@ -978,8 +977,8 @@ namespace Metasound
 				ensureAlwaysMsgf(InName == NodeClassName, TEXT("An output node's input connection should always be the same as it's class name!"));
 
 				FDescPath NodePath = NodePtr.GetPath();
-				FHandleInitParams InitParams = { NodePtr.GetAccessPoint(), NodePath, NodeClassName , OwningAsset };
-				return FInputHandle(FHandleInitParams::PrivateToken, InitParams);
+				FHandleInitParams InitParams = { NodePtr.GetAccessPoint(), NodePath, NodePtr->UniqueID, FMetasoundClassDescription::InvalidID, OwningAsset };
+				return FInputHandle(FHandleInitParams::PrivateToken, InitParams, NodeClassName, FInputHandle::OutputNode);
 			}
 
 			TArray<FMetasoundInputDescription>& InputDescriptions = NodeClass->Inputs;
@@ -990,8 +989,8 @@ namespace Metasound
 				{
 					FString ClassName = GetNodeClassName();
 					FDescPath NodePath = NodePtr.GetPath();
-					FHandleInitParams InitParams = { NodePtr.GetAccessPoint(), NodePath, ClassName, OwningAsset };
-					return FInputHandle(FHandleInitParams::PrivateToken, InitParams, InputDescription.Name);
+					FHandleInitParams InitParams = { NodePtr.GetAccessPoint(), NodePath, NodePtr->UniqueID, NodeClass->UniqueID, OwningAsset };
+					return FInputHandle(FHandleInitParams::PrivateToken, InitParams, InputDescription.Name, FInputHandle::Default);
 				}
 			}
 
@@ -1012,8 +1011,8 @@ namespace Metasound
 				const FString& NodeClassName = NodePtr->Name;
 				ensureAlwaysMsgf(InName == NodeClassName, TEXT("An input node's output connection should always be the same as it's class name!"));
 				FDescPath NodePath = NodePtr.GetPath();
-				FHandleInitParams InitParams = { NodePtr.GetAccessPoint(), NodePath, NodeClassName , OwningAsset };
-				return FOutputHandle(FHandleInitParams::PrivateToken, InitParams);
+				FHandleInitParams InitParams = { NodePtr.GetAccessPoint(), NodePath, FMetasoundNodeDescription::InvalidID, FMetasoundClassDescription::InvalidID, OwningAsset };
+				return FOutputHandle(FHandleInitParams::PrivateToken, InitParams, NodeClassName, FOutputHandle::InputNode);
 			}
 
 			TArray<FMetasoundOutputDescription>& OutputDescriptions = NodeClass->Outputs;
@@ -1024,8 +1023,8 @@ namespace Metasound
 				{
 					FString NodeClassName = GetNodeClassName();
 					FDescPath NodePath = NodePtr.GetPath();
-					FHandleInitParams InitParams = { NodePtr.GetAccessPoint(), NodePath, NodeClassName , OwningAsset };
-					return FOutputHandle(FHandleInitParams::PrivateToken, InitParams, OutputDescription.Name);
+					FHandleInitParams InitParams = { NodePtr.GetAccessPoint(), NodePath, FMetasoundNodeDescription::InvalidID, NodeClass->UniqueID, OwningAsset };
+					return FOutputHandle(FHandleInitParams::PrivateToken, InitParams, OutputDescription.Name, FOutputHandle::Default);
 				}
 			}
 
@@ -1116,16 +1115,16 @@ namespace Metasound
 			}
 
 			FDescPath ContainedGraphPath = NodeClass.GetPath()[Path::EFromClass::ToGraph];
-			FHandleInitParams InitParams = { NodeClass.GetAccessPoint(), ContainedGraphPath, NodeClass->Metadata.NodeName, OwningAsset };
+			FHandleInitParams InitParams = { NodeClass.GetAccessPoint(), ContainedGraphPath, FMetasoundNodeDescription::InvalidID, NodeClass->UniqueID, OwningAsset };
 			// Todo: link this up to look for externally implemented graphs as well.
 			OutGraph = FGraphHandle(FHandleInitParams::PrivateToken, InitParams);
 		}
 
-		uint32 FNodeHandle::GetNodeID() const
+		int32 FNodeHandle::GetNodeID() const
 		{
 			if (!IsValid())
 			{
-				return INDEX_NONE;
+				return FMetasoundNodeDescription::InvalidID;
 			}
 
 			return NodePtr->UniqueID;
@@ -1142,18 +1141,18 @@ namespace Metasound
 			return NodePtr->Name;
 		}
 
-		uint32 FNodeHandle::GetNodeID(const FDescPath& InNodePath)
+		int32 FNodeHandle::GetNodeID(const FDescPath& InNodePath)
 		{
 			if (!ensureAlwaysMsgf(InNodePath.Path.Num() != 0, TEXT("Tried to get a node ID from an empty path.")))
 			{
-				return INDEX_NONE;
+				return FMetasoundNodeDescription::InvalidID;
 			}
 
 			const Path::FElement& LastElementInPath = InNodePath.Path.Last();
 			
 			if (!ensureAlwaysMsgf(LastElementInPath.CurrentDescType == Path::EDescType::Node, TEXT("Tried to get the node ID for a path that was not set up for a node.")))
 			{
-				return INDEX_NONE;
+				return FMetasoundNodeDescription::InvalidID;
 			}
 
 			return LastElementInPath.LookupID;
@@ -1172,25 +1171,25 @@ namespace Metasound
 			}
 		}
 
-		uint32 FGraphHandle::FindNewUniqueNodeId()
+		int32 FGraphHandle::FindNewUniqueNodeId()
 		{
 			// Assumption here is that we will never need more than ten thousand nodes,
 			// and four digits are easy enough to read/remember when looking at metasound graph documents.
-			static const uint32 NodeIDMax = 9999;
+			static const int32 NodeIDMax = 9999;
 
 			if (!IsValid())
 			{
-				return INDEX_NONE;
+				return FMetasoundNodeDescription::InvalidID;
 			}
 
 			TArray<FMetasoundNodeDescription>& Nodes = GraphPtr->Nodes;
 
-			if (!ensureAlwaysMsgf(((uint32)Nodes.Num()) < NodeIDMax, TEXT("Too many nodes to guarantee a unique node ID. Increase the value of NodeIDMax.")))
+			if (!ensureAlwaysMsgf(((int32)Nodes.Num()) < NodeIDMax, TEXT("Too many nodes to guarantee a unique node ID. Increase the value of NodeIDMax.")))
 			{
-				return INDEX_NONE;
+				return FMetasoundNodeDescription::InvalidID;
 			}
 
-			while (uint32 RandomID = FMath::RandRange(1, NodeIDMax))
+			while (int32 RandomID = FMath::RandRange(1, NodeIDMax))
 			{
 				// Scan through the nodes in this graph to see if they match this ID.
 				// If it does, generate a new random ID.
@@ -1210,28 +1209,28 @@ namespace Metasound
 				}
 			}
 
-			return INDEX_NONE;
+			return FMetasoundNodeDescription::InvalidID;
 		}
 
-		uint32 FGraphHandle::FindNewUniqueDependencyId()
+		int32 FGraphHandle::FindNewUniqueDependencyId()
 		{
 			// Assumption here is that we will never need more than ten thousand dependencies,
 			// and four digits are easy enough to read/remember when looking at metasound graph documents.
-			static const uint32 DependencyIDMax = 9999;
+			static const int32 DependencyIDMax = 9999;
 
 			if (!IsValid())
 			{
-				return INDEX_NONE;
+				return FMetasoundClassDescription::InvalidID;
 			}
 
 			TArray<FMetasoundClassDescription>& Dependencies = OwningDocument->Dependencies;
 
 			if (!ensureAlwaysMsgf(Dependencies.Num() < DependencyIDMax, TEXT("Too many nodes to guarantee a unique node ID. Increase the value of NodeIDMax.")))
 			{
-				return INDEX_NONE;
+				return FMetasoundClassDescription::InvalidID;
 			}
 
-			while (uint32 RandomID = FMath::RandRange(1, DependencyIDMax))
+			while (int32 RandomID = FMath::RandRange(1, DependencyIDMax))
 			{
 				// Scan through the nodes in this graph to see if they match this ID.
 				// If it does, generate a new random ID.
@@ -1251,7 +1250,8 @@ namespace Metasound
 				}
 			}
 
-			return INDEX_NONE;
+			ensureAlwaysMsgf(false, TEXT("Failed to generate a new unique ID"));
+			return FMetasoundClassDescription::InvalidID;
 		}
 
 		FMetasoundLiteralDescription* FGraphHandle::GetLiteralDescriptionForInput(const FString& InInputName, FName& OutDataType) const
@@ -1302,13 +1302,13 @@ namespace Metasound
 		{
 			FDescPath InvalidPath;
 			FString InvalidClassName;
-			FHandleInitParams InitParams = { nullptr, InvalidPath, InvalidClassName, nullptr };
+			FHandleInitParams InitParams = { nullptr, InvalidPath, FMetasoundNodeDescription::InvalidID, FMetasoundClassDescription::InvalidID, nullptr };
 			return FGraphHandle(FHandleInitParams::Token, InitParams);
 		}
 
 		bool FGraphHandle::IsValid() const
 		{
-			return GraphPtr.IsValid() && GraphsClassDeclaration.IsValid();
+			return OwningDocument.IsValid() && GraphPtr.IsValid() && GraphsClassDeclaration.IsValid();
 		}
 
 		void FGraphHandle::CopyGraph(FGraphHandle& InOther)
@@ -1382,14 +1382,14 @@ namespace Metasound
 			{
 				FDescPath NodePath = GraphPtr.GetPath()[Path::EFromGraph::ToNodes][NodeDescription.UniqueID];
 
-				FHandleInitParams InitParams = { GraphPtr.GetAccessPoint(), NodePath, NodeDescription.Name, OwningAsset };
+				FHandleInitParams InitParams = { GraphPtr.GetAccessPoint(), NodePath, NodeDescription.UniqueID, NodeDescription.DependencyID, OwningAsset };
 				OutArray.Emplace(FHandleInitParams::PrivateToken, InitParams, NodeDescription.ObjectTypeOfNode);
 			}
 
 			return OutArray;
 		}
 
-		FNodeHandle FGraphHandle::GetNodeWithId(uint32 InNodeId) const
+		FNodeHandle FGraphHandle::GetNodeWithId(int32 InNodeId) const
 		{
 			if (!IsValid())
 			{
@@ -1402,7 +1402,7 @@ namespace Metasound
 				if (NodeDescription.UniqueID == InNodeId)
 				{
 					FDescPath NodePath = GraphPtr.GetPath()[Path::EFromGraph::ToNodes][InNodeId];
-					FHandleInitParams InitParams = { GraphPtr.GetAccessPoint(), NodePath, NodeDescription.Name, OwningAsset };
+					FHandleInitParams InitParams = { GraphPtr.GetAccessPoint(), NodePath, NodeDescription.UniqueID, NodeDescription.DependencyID, OwningAsset };
 					return FNodeHandle(FHandleInitParams::PrivateToken, InitParams, NodeDescription.ObjectTypeOfNode);
 				}
 			}
@@ -1426,7 +1426,7 @@ namespace Metasound
 				if (NodeDescription.ObjectTypeOfNode == EMetasoundClassType::Output)
 				{
 					FDescPath NodePath = GraphPtr.GetPath()[Path::EFromGraph::ToNodes][NodeDescription.UniqueID];
-					FHandleInitParams InitParams = { GraphPtr.GetAccessPoint(), NodePath, NodeDescription.Name, OwningAsset };
+					FHandleInitParams InitParams = { GraphPtr.GetAccessPoint(), NodePath, NodeDescription.UniqueID, NodeDescription.DependencyID, OwningAsset };
 					OutArray.Emplace(FHandleInitParams::PrivateToken, InitParams, NodeDescription.ObjectTypeOfNode);
 				}
 			}
@@ -1449,7 +1449,7 @@ namespace Metasound
 				if (NodeDescription.ObjectTypeOfNode == EMetasoundClassType::Input)
 				{
 					FDescPath NodePath = GraphPtr.GetPath()[Path::EFromGraph::ToNodes][NodeDescription.UniqueID];
-					FHandleInitParams InitParams = { GraphPtr.GetAccessPoint(), NodePath, NodeDescription.Name, OwningAsset };
+					FHandleInitParams InitParams = { GraphPtr.GetAccessPoint(), NodePath, NodeDescription.UniqueID, NodeDescription.DependencyID, OwningAsset };
 					OutArray.Emplace(FHandleInitParams::PrivateToken, InitParams, NodeDescription.ObjectTypeOfNode);
 				}
 			}
@@ -1509,7 +1509,7 @@ namespace Metasound
 				{
 					FDescPath NodePath = GraphPtr.GetPath()[Path::EFromGraph::ToNodes][NodeDescription.UniqueID];
 
-					FHandleInitParams InitParams = { GraphPtr.GetAccessPoint(), NodePath, InName, OwningAsset };
+					FHandleInitParams InitParams = { GraphPtr.GetAccessPoint(), NodePath, NodeDescription.UniqueID, NodeDescription.DependencyID, OwningAsset };
 					return FNodeHandle(FHandleInitParams::PrivateToken, InitParams, NodeDescription.ObjectTypeOfNode);
 				}
 			}
@@ -1533,7 +1533,7 @@ namespace Metasound
 					FDescPath NodePath = GraphPtr.GetPath()[Path::EFromGraph::ToNodes][NodeDescription.UniqueID];
 
 					// todo: add special enum for input/output nodes
-					FHandleInitParams InitParams = { GraphPtr.GetAccessPoint(), NodePath, InName, OwningAsset };
+					FHandleInitParams InitParams = { GraphPtr.GetAccessPoint(), NodePath, NodeDescription.UniqueID, FMetasoundClassDescription::InvalidID, OwningAsset };
 					return FNodeHandle(FHandleInitParams::PrivateToken, InitParams, NodeDescription.ObjectTypeOfNode);
 				}
 			}
@@ -1551,9 +1551,9 @@ namespace Metasound
 			
 			// @todo: verify that InDescription.TypeName is a valid Metasound type.
 
-			const uint32 NewUniqueId = FindNewUniqueNodeId();
+			const int32 NewUniqueId = FindNewUniqueNodeId();
 			
-			if (!ensureAlwaysMsgf(NewUniqueId != INDEX_NONE, TEXT("FindNewUniqueNodeId failed!")))
+			if (!ensureAlwaysMsgf(NewUniqueId != FMetasoundNodeDescription::InvalidID, TEXT("FindNewUniqueNodeId failed!")))
 			{
 				return FNodeHandle::InvalidHandle();
 			}
@@ -1581,7 +1581,7 @@ namespace Metasound
 			GraphPtr->Nodes.Add(NewNodeDescription);
 
 			FDescPath NodePath = GraphPtr.GetPath()[Path::EFromGraph::ToNodes][NewNodeDescription.UniqueID];
-			FHandleInitParams InitParams = { GraphPtr.GetAccessPoint(), NodePath, NewNodeDescription.Name, OwningAsset };
+			FHandleInitParams InitParams = { GraphPtr.GetAccessPoint(), NodePath, NewUniqueId, FMetasoundClassDescription::InvalidID, OwningAsset };
 
 			// todo: add special enum for input and output nodes
 			return FNodeHandle(FHandleInitParams::PrivateToken, InitParams, NewNodeDescription.ObjectTypeOfNode);
@@ -1653,9 +1653,9 @@ namespace Metasound
 
 			// @todo: verify that InDescription.TypeName is a valid Metasound type.
 
-			const uint32 NewUniqueId = FindNewUniqueNodeId();
+			const int32 NewUniqueId = FindNewUniqueNodeId();
 
-			if (!ensureAlwaysMsgf(NewUniqueId != INDEX_NONE, TEXT("FindNewUniqueNodeId failed")))
+			if (!ensureAlwaysMsgf(NewUniqueId != FMetasoundNodeDescription::InvalidID, TEXT("FindNewUniqueNodeId failed")))
 			{
 				return FNodeHandle::InvalidHandle();
 			}
@@ -1686,7 +1686,7 @@ namespace Metasound
 			FDescPath NodePath = GraphPtr.GetPath()[Path::EFromGraph::ToNodes][NewNodeDescription.UniqueID];
 
 			// todo: add special enum or class for input/output nodes
-			FHandleInitParams InitParams = { GraphPtr.GetAccessPoint(), NodePath, NewNodeDescription.Name, OwningAsset };
+			FHandleInitParams InitParams = { GraphPtr.GetAccessPoint(), NodePath, NewNodeDescription.UniqueID, NewNodeDescription.DependencyID, OwningAsset };
 			return FNodeHandle(FHandleInitParams::PrivateToken, InitParams, NewNodeDescription.ObjectTypeOfNode);
 		}
 
@@ -1891,70 +1891,59 @@ namespace Metasound
 				return FNodeHandle::InvalidHandle();
 			}
 
-			// First, scan our dependency list to see if this node already exists there, and if not, get it.
-			TArray<FMetasoundClassDescription>& Dependencies = OwningDocument->Dependencies;
-
-			bool bFoundMatchingDependencyInDocument = false;
+			const int32 NodeID = FindNewUniqueNodeId();
+			if (!ensureAlwaysMsgf(NodeID != FMetasoundNodeDescription::InvalidID, TEXT("Call to FindNewUniqueNodeId failed!")))
+			{
+				return FNodeHandle::InvalidHandle();
+			}
 
-			for (FMetasoundClassDescription& Dependency : Dependencies)
+			auto IsMatchingDependency = [&](const FMetasoundClassDescription& InDependency) -> bool
 			{
-				if (Dependency.Metadata.NodeName == InNodeClass.NodeName && Dependency.Metadata.NodeType == InNodeClass.NodeType)
+				// TODO: may need to add Input/Ouptut node types to external lookup.
+				if ((InNodeClass.NodeType == EMetasoundClassType::External) && (InDependency.Metadata.NodeType == EMetasoundClassType::External))
 				{
-					bFoundMatchingDependencyInDocument = true;
-
-					// If this dependency was in the document's dependency list, check to see if we need to add it to this class' dependencies.
-					bool bFoundDependencyInLocalClass = false;
-					for (uint32 DependencyID : GraphsClassDeclaration->DependencyIDs)
-					{
-						if (DependencyID == Dependency.UniqueID)
-						{
-							bFoundDependencyInLocalClass = true;
-							break;
-						}
-					}
-
-					if (!bFoundDependencyInLocalClass)
-					{
-						// This dependency is already referenced somewhere in the document, but not for this graph's class. Add it.
-						GraphsClassDeclaration->DependencyIDs.Add(Dependency.UniqueID);
-						UE_LOG(LogTemp, Verbose, TEXT("Adding %s as a dependency for Metasound graph %s in Document %s"), *InNodeClass.NodeName, *GetGraphMetadata().NodeName, *OwningDocument->RootClass.Metadata.NodeName);
-					}
-
-					break;
+					// For external nodes, we rely on the external node lookup info.
+					const bool bIsLookupInfoEqual = (InDependency.ExternalNodeClassLookupInfo.ExternalNodeClassName == InNodeClass.LookupKey.NodeName) && (InDependency.ExternalNodeClassLookupInfo.ExternalNodeClassHash == InNodeClass.LookupKey.NodeHash);
+					return bIsLookupInfoEqual;
+				}
+				else if ((InNodeClass.NodeType == EMetasoundClassType::MetasoundGraph) && (InDependency.Metadata.NodeType == EMetasoundClassType::MetasoundGraph))
+				{
+					// TODO: This dependency lookup should also take into consideration version info, and possibly even more.
+					// For graph nodes, we only rely on the name.
+					return InNodeClass.LookupKey.NodeName == InDependency.ExternalNodeClassLookupInfo.ExternalNodeClassName;
 				}
-			}
 
-			
+				return false;
+			};
+
+			// Scan dependency list to see if the class already exists in dependency list.
+			const FMetasoundClassDescription* Dependency = OwningDocument->Dependencies.FindByPredicate(IsMatchingDependency);
 
-			// If we haven't added a node of this class to the graph yet, add it to the dependencies for this class.
-			if (!bFoundMatchingDependencyInDocument)
+			// Add dependency to owning document if it does not exist.
+			if (nullptr == Dependency)
 			{
 				FMetasoundClassDescription NewDependencyClassDescription = GenerateClassDescription(InNodeClass);
 				NewDependencyClassDescription.UniqueID = FindNewUniqueDependencyId();
-				Dependencies.Add(NewDependencyClassDescription);
-				GraphsClassDeclaration->DependencyIDs.Add(NewDependencyClassDescription.UniqueID);
 
-				UE_LOG(LogTemp, Verbose, TEXT("Adding %s is used in graph %s, adding as a new dependency for Metasound Document %s"), *InNodeClass.NodeName, *GetGraphMetadata().NodeName, *OwningDocument->RootClass.Metadata.NodeName);
+				Dependency = &(OwningDocument->Dependencies.Add_GetRef(NewDependencyClassDescription));
 			}
 
-			// Add a new node instance for this class.
-
-			const uint32 NewUniqueId = FindNewUniqueNodeId();
-			if (!ensureAlwaysMsgf(NewUniqueId != INDEX_NONE, TEXT("Call to FindNewUniqueNodeId failed!")))
-			{
-				return FNodeHandle::InvalidHandle();
-			}
+			// Add dependency ID to graph
+			GraphsClassDeclaration->DependencyIDs.AddUnique(Dependency->UniqueID);
 
 			// Add a node for this output to the graph description.
 			FMetasoundNodeDescription NewNodeDescription;
 			NewNodeDescription.Name = InNodeClass.NodeName;
-			NewNodeDescription.UniqueID = NewUniqueId;
+			NewNodeDescription.UniqueID = NodeID;
+			NewNodeDescription.DependencyID = Dependency->UniqueID;
 			NewNodeDescription.ObjectTypeOfNode = InNodeClass.NodeType;
 
 			GraphPtr->Nodes.Add(NewNodeDescription);
 
+			// Create node handle.
 			FDescPath NodePath = GraphPtr.GetPath()[Path::EFromGraph::ToNodes][NewNodeDescription.UniqueID];
-			FHandleInitParams InitParams = { GraphPtr.GetAccessPoint(), NodePath, NewNodeDescription.Name, OwningAsset };
+			FHandleInitParams InitParams = { GraphPtr.GetAccessPoint(), NodePath, NodeID, Dependency->UniqueID, OwningAsset };
+
 			return FNodeHandle(FHandleInitParams::PrivateToken, InitParams, NewNodeDescription.ObjectTypeOfNode);
 		}
 
@@ -2042,10 +2031,10 @@ namespace Metasound
 			// Remove its ID as a dependency for the graph.
 			if (NodesOfClass < 2)
 			{
-				TArray<uint32>& DependencyIDs = GraphsClassDeclaration->DependencyIDs;
+				TArray<int32>& DependencyIDs = GraphsClassDeclaration->DependencyIDs;
 				int32 IndexOfDependencyToRemove = -1;
 
-				uint32 UniqueIDForThisDependency = INDEX_NONE;
+				int32 UniqueIDForThisDependency = FMetasoundClassDescription::InvalidID;
 				TArray<FMetasoundClassDescription>& DependencyClasses = OwningDocument->Dependencies;
 
 				// scan the owning document's depenency classes for a dependency with this name.
@@ -2203,8 +2192,8 @@ namespace Metasound
 			}
 
 			// Create a new class in this graph's dependencies list:
-			uint32 NewUniqueIDForGraph = FindNewUniqueDependencyId();
-			if (!ensureAlwaysMsgf(NewUniqueIDForGraph != INDEX_NONE, TEXT("Call to FindNewUniqueNodeId failed!")))
+			int32 NewUniqueIDForGraph = FindNewUniqueDependencyId();
+			if (!ensureAlwaysMsgf(NewUniqueIDForGraph != FMetasoundClassDescription::InvalidID, TEXT("Call to FindNewUniqueNodeId failed!")))
 			{
 				return BuildInvalidTupleHandle();
 			}
@@ -2218,13 +2207,13 @@ namespace Metasound
 			GraphsClassDeclaration->DependencyIDs.Add(NewUniqueIDForGraph);
 
 			// Generate a new FGraphHandle for this subgraph:
-			FDescPath PathForNewGraph = FDescPath()[Path::EFromDocument::ToDependencies][*InInfo.NodeName][Path::EFromClass::ToGraph];
-			FHandleInitParams InitParams = { GraphsClassDeclaration.GetAccessPoint(), PathForNewGraph, InInfo.NodeName, OwningAsset };
+			FDescPath PathForNewGraph = FDescPath()[Path::EFromDocument::ToDependencies][NewUniqueIDForGraph][Path::EFromClass::ToGraph];
+			FHandleInitParams InitParams = { GraphsClassDeclaration.GetAccessPoint(), PathForNewGraph, FMetasoundNodeDescription::InvalidID, NewUniqueIDForGraph, OwningAsset };
 			FGraphHandle SubgraphHandle = FGraphHandle(FHandleInitParams::PrivateToken, InitParams);
 
 			// Create the node for this subgraph in the current graph:
-			const uint32 NewUniqueId = FindNewUniqueNodeId();
-			if (!ensureAlwaysMsgf(NewUniqueId != INDEX_NONE, TEXT("Call to FindNewUniqueNodeId failed!")))
+			const int32 NewUniqueId = FindNewUniqueNodeId();
+			if (!ensureAlwaysMsgf(NewUniqueId != FMetasoundNodeDescription::InvalidID, TEXT("Call to FindNewUniqueNodeId failed!")))
 			{
 				return BuildInvalidTupleHandle();
 			}
@@ -2233,12 +2222,13 @@ namespace Metasound
 			FMetasoundNodeDescription NewNodeDescription;
 			NewNodeDescription.Name = InInfo.NodeName;
 			NewNodeDescription.UniqueID = NewUniqueId;
+			NewNodeDescription.DependencyID = NewUniqueIDForGraph;
 			NewNodeDescription.ObjectTypeOfNode = InInfo.NodeType;
 
 			GraphPtr->Nodes.Add(NewNodeDescription);
 
 			FDescPath NodePath = GraphPtr.GetPath()[Path::EFromGraph::ToNodes][NewNodeDescription.UniqueID];
-			FHandleInitParams NodeInitParams = { GraphPtr.GetAccessPoint(), NodePath, InInfo.NodeName, OwningAsset };
+			FHandleInitParams NodeInitParams = { GraphPtr.GetAccessPoint(), NodePath, NewUniqueId, NewUniqueIDForGraph, OwningAsset };
 			FNodeHandle SubgraphNode = FNodeHandle(FHandleInitParams::PrivateToken, NodeInitParams, NewNodeDescription.ObjectTypeOfNode);
 
 			return TTuple<FGraphHandle, FNodeHandle>(SubgraphHandle, SubgraphNode);
@@ -2248,299 +2238,31 @@ namespace Metasound
 		{
 			if (!IsValid())
 			{
-				return nullptr;
+				return TUniquePtr<IOperator>(nullptr);
 			}
 
 			// TODO: Implement inflation step here.
 
-			// At this point, we should be left with a flat graph of externally implemented nodes.
-			using FGraph = Metasound::FGraph;
-			using FDataEdge = Metasound::FDataEdge;
-			using INodePtr = TUniquePtr<INode>;
-
-			TArray<FName> DataTypes = GetAllAvailableDataTypes();
-
-			auto IsInputNode = [](const FMetasoundNodeDescription& InDescription) -> bool
-			{
-				return InDescription.ObjectTypeOfNode == EMetasoundClassType::Input;
-			};
-
-			auto IsOutputNode = [](const FMetasoundNodeDescription& InDescription) -> bool
-			{
-				return InDescription.ObjectTypeOfNode == EMetasoundClassType::Output;
-			};
-
-			int32 TotalNodesGenerated = 0;
-
-			// These maps are used to fix up input and output destinations after we have fully connected the graph.
-			TMap<int32, uint32> InputIndexToInputNodeIDs;
-			TMap<int32, uint32> OutputIndexToInputNodeIDs;
-
-			// Helper for auto-generating nodes.
-			auto GetNodePtr = [&](const FMetasoundNodeDescription& InDescription) -> INodePtr
-			{
-				if (IsInputNode(InDescription))
-				{
-					// Find this in the graph's class description's list of inputs, until we find a match.
-					const TArray<FMetasoundInputDescription>& InputDescriptions = GraphsClassDeclaration->Inputs;
-
-					int32 InputDescriptionIndex = 0;
-
-					for (const FMetasoundInputDescription& InputDescription : InputDescriptions)
-					{
-						if (InputDescription.Name == InDescription.Name)
-						{
-							// We found a match. now we just need to create the input node.
-							FDataTypeLiteralParam LiteralParam = GetLiteralParamForDataType(InputDescription.TypeName, InputDescription.LiteralValue);
-
-							if (!ensureAlwaysMsgf(DoesDataTypeSupportLiteralType(InputDescription.TypeName, InputDescription.LiteralValue.LiteralType), TEXT("Tried to use an unsupported literal type!")))
-							{
-								// print out some info about the type.
-								UE_LOG(LogTemp, Display, TEXT("Data Type %s supports the following literal types:"), *InputDescription.TypeName.ToString());
-
-								FDataTypeRegistryInfo DataTypeInfo;
-								
-								// this shouldn't hit at all, because this should be a registered data type.
-								ensure(GetTraitsForDataType(InputDescription.TypeName, DataTypeInfo));
-
-								if (DataTypeInfo.bIsBoolParsable)
-								{
-									UE_LOG(LogTemp, Display, TEXT("   Boolean"));
-								}
-
-								if (DataTypeInfo.bIsIntParsable)
-								{
-									UE_LOG(LogTemp, Display, TEXT("    Integer"));
-								}
-
-								if (DataTypeInfo.bIsFloatParsable)
-								{
-									UE_LOG(LogTemp, Display, TEXT("    Float"));
-								}
-
-								if (DataTypeInfo.bIsStringParsable)
-								{
-									UE_LOG(LogTemp, Display, TEXT("    String"));
-								}
-
-								if (DataTypeInfo.bIsProxyParsable)
-								{
-									check(DataTypeInfo.ProxyGeneratorClass);
-									UE_LOG(LogTemp, Display, TEXT("    U%s*"), *DataTypeInfo.ProxyGeneratorClass->GetName());
-								}
-
-								if (DataTypeInfo.bIsProxyArrayParsable)
-								{
-									check(DataTypeInfo.ProxyGeneratorClass);
-									UE_LOG(LogTemp, Display, TEXT("    TArray<U%s*>"), *DataTypeInfo.ProxyGeneratorClass->GetName());
-								}
-
-								return nullptr;
-							}
-
-							FInputNodeConstructorParams InitParams =
-							{
-								InputDescription.Name,
-								InputDescription.Name,
-								InSettings,
-								MoveTemp(LiteralParam)
-							};
-
-							INodePtr InputNode = ConstructInputNode(InputDescription.TypeName, MoveTemp(InitParams));
-							if (!ensureAlwaysMsgf(InputNode.IsValid(), TEXT("Failed to construct a valid input node for Data Type %s!"), *InputDescription.TypeName.ToString()))
-							{
-								return nullptr;
-							}
-
-							TotalNodesGenerated++;
-
-							InputIndexToInputNodeIDs.Add(InputDescriptionIndex, InDescription.UniqueID);
-
-							return InputNode;
-						}
-						else
-						{
-							InputDescriptionIndex++;
-						}
-					}
-
-					// if we hit this, the document has been corrupted in some way, because we didn't have a matching Input Description for this node.
-					ensureAlwaysMsgf(false, TEXT("Document corrupted! found input node %s but couldn't find a matching Input Description for it in the Class Description."), *InDescription.Name);
-					return nullptr;
-				}
-				else if (IsOutputNode(InDescription))
-				{
-					int32 OutputDescriptionIndex = 0;
-
-					// Find this in the graph's class description's list of inputs, until we find a match.
-					const TArray<FMetasoundOutputDescription>& OutputDescriptions = GraphsClassDeclaration->Outputs;
-					for (const FMetasoundOutputDescription& OutputDescription : OutputDescriptions)
-					{
-						if (OutputDescription.Name == InDescription.Name)
-						{
-							// We found a match. now we just need to create the input node.
-							FOutputNodeConstrutorParams InitParams =
-							{
-								OutputDescription.Name,
-								OutputDescription.Name,
-							};
-
-							INodePtr OutputNode = ConstructOutputNode(OutputDescription.TypeName, InitParams);
-							if (!ensureAlwaysMsgf(OutputNode.IsValid(), TEXT("Failed to construct a valid input node for Data Type %s!"), *OutputDescription.TypeName.ToString()))
-							{
-								return nullptr;
-							}
-
-							TotalNodesGenerated++;
-
-							OutputIndexToInputNodeIDs.Add(OutputDescriptionIndex, InDescription.UniqueID);
-
-							return OutputNode;
-						}
-						else
-						{
-							OutputDescriptionIndex++;
-						}
-					}
-
-					// if we hit this, the document has been corrupted in some way, because we didn't have a matching Input Description for this node.
-					ensureAlwaysMsgf(false, TEXT("Document corrupted! found input node %s but couldn't find a matching Input Description for it in the Class Description."), *InDescription.Name);
-					return nullptr;
-				}
-				else
-				{
-					if (!ensureAlwaysMsgf(InDescription.ObjectTypeOfNode == EMetasoundClassType::External, TEXT("At this point in construction, we should only need to look up external nodes.")))
-					{
-						return nullptr;
-					}
-
-					// Find the node class in the dependencies.
-					const TArray<FMetasoundClassDescription>& DependencyDescriptions = OwningDocument->Dependencies;
-					for (const FMetasoundClassDescription& DependencyDescription : DependencyDescriptions)
-					{
-						// TODO: Add the dependency ID to the node so that we can look it up directly.
-						if (DependencyDescription.Metadata.NodeName == InDescription.Name)
-						{
-							// We found a match. now we just need to create the input node.
-							const FMetasoundExternalClassLookupInfo& LookupInfo = DependencyDescription.ExternalNodeClassLookupInfo;
-							FNodeInitData InitData;
-							InitData.InstanceName.Append(InDescription.Name);
-							InitData.InstanceName.AppendChar('_');
-							InitData.InstanceName.AppendInt(InDescription.UniqueID);
-							
-							// Copy over our initialization params.
-							for (auto& StaticParamTuple : InDescription.StaticParameters)
-							{
-								FDataTypeLiteralParam LiteralParam = GetLiteralParam(StaticParamTuple.Value);
-
-								if (LiteralParam.IsValid())
-								{
-									InitData.ParamMap.Add(StaticParamTuple.Key, MoveTemp(LiteralParam));
-								}
-							}
-
-							INodePtr ExternalNode = ConstructExternalNode(LookupInfo.ExternalNodeClassName, LookupInfo.ExternalNodeClassHash, InitData);
-							if (!ensureAlwaysMsgf(ExternalNode.IsValid(), TEXT("Failed to construct a valid external node for Node Class %s!"), *DependencyDescription.Metadata.NodeName))
-							{
-								return nullptr;
-							}
-
-							TotalNodesGenerated++;
-							return ExternalNode;
-						}
-					}
-
-					ensureAlwaysMsgf(false, TEXT("Document corrupted! found node %s but couldn't find a matching Class Description for it in the Dependencies."), *InDescription.Name);
-					return nullptr;
-				}
-			};
-
-			FGraph GraphToBuild = FGraph(GraphsClassDeclaration->Metadata.NodeName);
-
-			TMap<uint32, INodePtr> GraphNodes;
-
-			// Step 1: Initialize Nodes
-			const TArray<FMetasoundNodeDescription>& NodeDescriptions = GraphPtr->Nodes;
-
-			for (const FMetasoundNodeDescription& NodeDescription : NodeDescriptions)
-			{
-				INodePtr NodePtr = GetNodePtr(NodeDescription);
-				
-				if (!NodePtr.IsValid())
-				{
-					return nullptr;
-				}
-				else
-				{
-					GraphNodes.Add(NodeDescription.UniqueID, MoveTemp(NodePtr));
-				}
-			}
-
-			// Sanity check that we created enough input nodes and output nodes.
-			const int32 NumInputsInMetasoundClass = GraphsClassDeclaration->Inputs.Num();
-			if (InputIndexToInputNodeIDs.Num() != NumInputsInMetasoundClass)
-			{
-				ensureAlwaysMsgf(false, TEXT("Mismatch between number of inputs in the metasound graph (%d) and number of inputs declared in it's class description (%d)!"), InputIndexToInputNodeIDs.Num(), NumInputsInMetasoundClass);
-				return nullptr;
-			}
-
-			const int32 NumOutputsInMetasoundClass = GraphsClassDeclaration->Outputs.Num();
-			if (OutputIndexToInputNodeIDs.Num() != NumOutputsInMetasoundClass)
-			{
-				ensureAlwaysMsgf(false, TEXT("Mismatch between number of outputs in the metasound graph (%d) and number of outputs declared in it's class description (%d)!"), OutputIndexToInputNodeIDs.Num(), NumOutputsInMetasoundClass);
-				return nullptr;
-			}
-
-			if (TotalNodesGenerated != GraphPtr->Nodes.Num())
+			if (!(GraphsClassDeclaration.IsValid() && OwningDocument.IsValid()))
 			{
-				ensureAlwaysMsgf(false, TEXT("Created %d of %d needed nodes!"), TotalNodesGenerated, GraphPtr->Nodes.Num());
-				return nullptr;
-			}
-
-			// Step 2: Connect Nodes Inside The Graph
-			for (const FMetasoundNodeDescription& NodeDescription : NodeDescriptions)
-			{
-				// TODO: create a INode type that houses literals.
-
-				INodePtr& NodeToConnectTo = GraphNodes[NodeDescription.UniqueID];
-
-				for (const FMetasoundNodeConnectionDescription& InputConnection : NodeDescription.InputConnections)
-				{
-					const bool bCanMakeConnection = InputConnection.NodeID != FMetasoundNodeConnectionDescription::DisconnectedNodeID
-						&& ensureAlwaysMsgf(GraphNodes.Contains(InputConnection.NodeID), TEXT("Connection in document describes a node ID that doesn't exist!"));
-					
-					if (bCanMakeConnection)
-					{
-						INodePtr& NodeToConnectFrom = GraphNodes[InputConnection.NodeID];
-						GraphToBuild.AddDataEdge(*NodeToConnectFrom, *InputConnection.OutputName, *NodeToConnectTo, *InputConnection.InputName);
-					}
-				}
+				// Need valid description of this class and list of dependencies
+				// in order to build operator.
+				return TUniquePtr<IOperator>(nullptr);
 			}
 
-			// Step 3: Declare our Input Destinations
-			const TArray<FMetasoundInputDescription>& InputDescriptions = GraphsClassDeclaration->Inputs;
-			for (auto& InputTuple : InputIndexToInputNodeIDs)
-			{
-				INodePtr& InputNodePtr = GraphNodes[InputTuple.Value];
-				const FString& InVertexName = InputDescriptions[InputTuple.Key].Name;
-				GraphToBuild.AddInputDataDestination(*InputNodePtr, *InVertexName);
-			}
+			// TODO: bubble up errors. 
+			FFrontendGraphBuilder GraphBuilder;
+			TUniquePtr<FFrontendGraph> Graph = GraphBuilder.CreateGraph(*GraphsClassDeclaration, OwningDocument->Dependencies);
 
-			// Step 4: Declare our Output Destinations
-			const TArray<FMetasoundOutputDescription>& OutputDescriptions = GraphsClassDeclaration->Outputs;
-			for (auto& OutputTuple : OutputIndexToInputNodeIDs)
+			if (!Graph.IsValid())
 			{
-				INodePtr& OutputNodePtr = GraphNodes[OutputTuple.Value];
-				const FString& OutVertexName = OutputDescriptions[OutputTuple.Key].Name;
-				GraphToBuild.AddOutputDataSource(*OutputNodePtr, *OutVertexName);
+				return TUniquePtr<IOperator>(nullptr);
 			}
 
-			/** NOTE: In the future- we should split steps 1-4 above, and 5 below. GraphToBuild and the node map can be cached on the graph handle. */
-
 			// Step 5: Invoke Operator Builder
-			FOperatorBuilder Builder(FOperatorBuilderSettings::GetDefaultSettings());
+			FOperatorBuilder OperatorBuilder(FOperatorBuilderSettings::GetDefaultSettings());
 
-			return Builder.BuildGraphOperator(GraphToBuild, InSettings, InEnvironment, OutBuildErrors);
+			return OperatorBuilder.BuildGraphOperator(*Graph, InSettings, InEnvironment, OutBuildErrors);
 		}
 
 		const FText& FGraphHandle::GetInputDisplayName(FString InputName) const
diff --git a/Engine/Plugins/Runtime/Metasound/Source/MetasoundFrontend/Private/MetasoundFrontendBaseClasses.cpp b/Engine/Plugins/Runtime/Metasound/Source/MetasoundFrontend/Private/MetasoundFrontendBaseClasses.cpp
--- a/Engine/Plugins/Runtime/Metasound/Source/MetasoundFrontend/Private/MetasoundFrontendBaseClasses.cpp
+++ b/Engine/Plugins/Runtime/Metasound/Source/MetasoundFrontend/Private/MetasoundFrontendBaseClasses.cpp
@@ -679,219 +679,226 @@ namespace Metasound
 			}
 		}
 
-		FDescriptionAccessPoint::FDescriptionUnwindStep FDescriptionAccessPoint::GoToNext(FDescPath& InPath, FDescriptionUnwindStep InElement)
+		FDescriptionAccessPoint::FDescriptionUnwindStep FDescriptionAccessPoint::GoToNextFromDocument(FMetasoundDocument& InDocument, FDescPath& InPath, const Path::FElement& InNext)
 		{
-			auto InvalidUnwindStep = []() -> FDescriptionUnwindStep
+			switch (InNext.CurrentDescType)
 			{
-				FDescriptionUnwindStep InvalidStep;
-				InvalidStep.DescriptionStructPtr.Set<void*>(nullptr);
-				InvalidStep.Type = Path::EDescType::Invalid;
-				return InvalidStep;
-			};
-
-			if (!ensureMsgf(InPath.Path.Num() != 0, TEXT("Attempted to unwind an empty path.")))
-			{
-				return InvalidUnwindStep();
-			}
+				case Path::EDescType::Document:
+				{
+					FDescriptionUnwindStep UnwindStep;
+					UnwindStep.DescriptionStructPtr.Set<FMetasoundDocument*>(&InDocument);
+					UnwindStep.Type = Path::EDescType::Document;
+					return UnwindStep;
+				}
+				case Path::EDescType::Class:
+				{
+					FDescriptionUnwindStep UnwindStep;
+					UnwindStep.DescriptionStructPtr.Set<FMetasoundClassDescription*>(&(InDocument.RootClass));
+					UnwindStep.Type = Path::EDescType::Class;
+					return UnwindStep;
+				}
+				case Path::EDescType::DocDependencies:
+				{
+					// The next element in a path after Dependencies will always be the name of the dependency.
+					Path::FElement DependencyElement = InPath.Path[0];
+					InPath.Path.RemoveAt(0);
+					if (!ensureAlwaysMsgf(DependencyElement.CurrentDescType == Path::EDescType::Class, TEXT("Invalid path set up.")))
+					{
+						return FDescriptionUnwindStep::CreateInvalid();
+					}
 
-			Path::FElement NextStep = InPath.Path[0];
+					FString& DependencyName = DependencyElement.LookupName;
 
-			InPath.Path.RemoveAt(0);
+					int32 DependencyID = DependencyElement.LookupID;
+					
+					if (!ensureAlwaysMsgf(DependencyID != INDEX_NONE || DependencyName.Len() > 0, TEXT("Path to a dependency did not include a valid ID or dependency name.")))
+					{
+						return FDescriptionUnwindStep::CreateInvalid();
+					}
 
-			switch (InElement.Type)
-			{
-				case Path::EDescType::Document:
-				{
-					FMetasoundDocument* Document = InElement.DescriptionStructPtr.Get<FMetasoundDocument*>();
-					switch (NextStep.CurrentDescType)
+					if (DependencyID == FMetasoundClassDescription::RootClassID)
 					{
-						case Path::EDescType::Document:
-						{
-							FDescriptionUnwindStep UnwindStep;
-							UnwindStep.DescriptionStructPtr.Set<FMetasoundDocument*>(Document);
-							UnwindStep.Type = Path::EDescType::Document;
-							return UnwindStep;
-						}
-						case Path::EDescType::Class:
-						{
-							FDescriptionUnwindStep UnwindStep;
-							UnwindStep.DescriptionStructPtr.Set<FMetasoundClassDescription*>(&(Document->RootClass));
-							UnwindStep.Type = Path::EDescType::Class;
-							return UnwindStep;
-						}
-						case Path::EDescType::DocDependencies:
-						{
-							// The next element in a path after Dependencies will always be the name of the dependency.
-							Path::FElement DependencyElement = InPath.Path[0];
-							InPath.Path.RemoveAt(0);
-							if (!ensureAlwaysMsgf(DependencyElement.CurrentDescType == Path::EDescType::Class, TEXT("Invalid path set up.")))
-							{
-								return InvalidUnwindStep();
-							}
+						FDescriptionUnwindStep UnwindStep;
+						UnwindStep.DescriptionStructPtr.Set<FMetasoundClassDescription*>(&(InDocument.RootClass));
+						UnwindStep.Type = Path::EDescType::Class;
+						return UnwindStep;
+					}
 
-							FString& DependencyName = DependencyElement.LookupName;
+					TArray<FMetasoundClassDescription>& DependenciesList = InDocument.Dependencies;
 
-							uint32 DependencyID = DependencyElement.LookupID;
-							
-							if (!ensureAlwaysMsgf(DependencyID != INDEX_NONE || DependencyName.Len() > 0, TEXT("Path to a dependency did not include a valid ID or dependency name.")))
+					// Dependencies can be looked up by ID or by name.
+					if (DependencyID != INDEX_NONE)
+					{
+						// Scan the dependencies list for the matching lookup ID.
+						for (FMetasoundClassDescription& Dependency : DependenciesList)
+						{
+							if (Dependency.UniqueID == DependencyID)
 							{
-								return InvalidUnwindStep();
+								FDescriptionUnwindStep UnwindStep;
+								UnwindStep.DescriptionStructPtr.Set<FMetasoundClassDescription*>(&Dependency);
+								UnwindStep.Type = Path::EDescType::Class;
+								return UnwindStep;
 							}
-
-							if (DependencyID == FMetasoundClassDescription::RootClassID)
+						}
+					}
+					else
+					{
+						// TODO: remove this chunk of code in the "else{}" block. Not sure if it has to be here. 
+						checkNoEntry();
+						// fall back to scanning the dependencies list for the matching lookup name.
+						for (FMetasoundClassDescription& Dependency : DependenciesList)
+						{
+							if (Dependency.Metadata.NodeName == DependencyName)
 							{
 								FDescriptionUnwindStep UnwindStep;
-								UnwindStep.DescriptionStructPtr.Set<FMetasoundClassDescription*>(&(Document->RootClass));
+								UnwindStep.DescriptionStructPtr.Set<FMetasoundClassDescription*>(&Dependency);
 								UnwindStep.Type = Path::EDescType::Class;
 								return UnwindStep;
 							}
+						}
+					}
+					
 
-							TArray<FMetasoundClassDescription>& DependenciesList = Document->Dependencies;
+					// If we reached the end of the Dependencies list and didn't find a match, ensure.
+					ensureAlwaysMsgf(false, TEXT("Couldn't find dependency %s in path."), *DependencyName);
+					return FDescriptionUnwindStep::CreateInvalid();
+				}
+			}
 
-							// Dependencies can be looked up by ID or by name.
-							if (DependencyID != INDEX_NONE)
-							{
-								// Scan the dependencies list for the matching lookup ID.
-								for (FMetasoundClassDescription& Dependency : DependenciesList)
-								{
-									if (Dependency.UniqueID == DependencyID)
-									{
-										FDescriptionUnwindStep UnwindStep;
-										UnwindStep.DescriptionStructPtr.Set<FMetasoundClassDescription*>(&Dependency);
-										UnwindStep.Type = Path::EDescType::Class;
-										return UnwindStep;
-									}
-								}
-							}
-							else
-							{
-								// fall back to scanning the dependencies list for the matching lookup name.
-								for (FMetasoundClassDescription& Dependency : DependenciesList)
-								{
-									if (Dependency.Metadata.NodeName == DependencyName)
-									{
-										FDescriptionUnwindStep UnwindStep;
-										UnwindStep.DescriptionStructPtr.Set<FMetasoundClassDescription*>(&Dependency);
-										UnwindStep.Type = Path::EDescType::Class;
-										return UnwindStep;
-									}
-								}
-							}
-							
+			checkNoEntry();
+			return FDescriptionUnwindStep::CreateInvalid();
+		}
 
-							// If we reached the end of the Dependencies list and didn't find a match, ensure.
-							ensureAlwaysMsgf(false, TEXT("Couldn't find dependency %s in path."), *DependencyName);
-							return InvalidUnwindStep();
-						}
-					}
-					break;
+		FDescriptionAccessPoint::FDescriptionUnwindStep FDescriptionAccessPoint::GoToNextFromClass(FMetasoundClassDescription& InClassDescription, FDescPath& InPath, const Path::FElement& InNext)
+		{
+			switch (InNext.CurrentDescType)
+			{
+				case Path::EDescType::Graph:
+				{
+					FDescriptionUnwindStep UnwindStep;
+					UnwindStep.DescriptionStructPtr.Set<FMetasoundGraphDescription*>(&(InClassDescription.Graph));
+					UnwindStep.Type = Path::EDescType::Graph;
+					return UnwindStep;
 				}
-				case Path::EDescType::Class:
+				case Path::EDescType::ClassDependencies:
 				{
-					FMetasoundClassDescription* ClassDescription = InElement.DescriptionStructPtr.Get<FMetasoundClassDescription*>();
-					switch (NextStep.CurrentDescType)
+					FDescriptionUnwindStep UnwindStep;
+					UnwindStep.DescriptionStructPtr.Set<FClassDependencyIDs*>(&(InClassDescription.DependencyIDs));
+					UnwindStep.Type = Path::EDescType::ClassDependencies;
+					return UnwindStep;
+				}
+				case Path::EDescType::Inputs:
+				{
+					// The next element after an Inputs element should always be the name of an input.
+					Path::FElement InputElement = InPath.Path[0];
+					InPath.Path.RemoveAt(0);
+					if (!ensureAlwaysMsgf(InputElement.CurrentDescType == Path::EDescType::Input, TEXT("Invalid path set up.")))
 					{
-						case Path::EDescType::Graph:
+						return FDescriptionUnwindStep::CreateInvalid();
+					}
+
+					FString& InputName = InputElement.LookupName;
+
+					// Scan the inputs list for the lookup name.
+					TArray<FMetasoundInputDescription>& InputsList = InClassDescription.Inputs;
+					for (FMetasoundInputDescription& Input : InputsList)
+					{
+						if (Input.Name == InputName)
 						{
 							FDescriptionUnwindStep UnwindStep;
-							UnwindStep.DescriptionStructPtr.Set<FMetasoundGraphDescription*>(&(ClassDescription->Graph));
-							UnwindStep.Type = Path::EDescType::Graph;
+							UnwindStep.DescriptionStructPtr.Set<FMetasoundInputDescription*>(&Input);
+							UnwindStep.Type = Path::EDescType::Input;
 							return UnwindStep;
 						}
-						case Path::EDescType::ClassDependencies:
+					}
+
+					// If we reached the end of the Inputs list and didn't find a match, ensure.
+					ensureAlwaysMsgf(false, TEXT("Couldn't find input %s in path."), *InputName);
+					return FDescriptionUnwindStep::CreateInvalid();
+				}
+				case Path::EDescType::Outputs:
+				{
+					// The next element after an Outputs element should always be the name of an output.
+					Path::FElement OutputElement = InPath.Path[0];
+					InPath.Path.RemoveAt(0);
+					if (!ensureAlwaysMsgf(OutputElement.CurrentDescType == Path::EDescType::Output, TEXT("Invalid path set up.")))
+					{
+						return FDescriptionUnwindStep::CreateInvalid();
+					}
+				
+					FString& OutputName = OutputElement.LookupName;
+
+					// Scan the outputs list for the lookup name.
+					TArray<FMetasoundOutputDescription>& OutputsList = InClassDescription.Outputs;
+					for (FMetasoundOutputDescription& Output : OutputsList)
+					{
+						if (Output.Name == OutputElement.LookupName)
 						{
 							FDescriptionUnwindStep UnwindStep;
-							UnwindStep.DescriptionStructPtr.Set<FClassDependencyIDs*>(&(ClassDescription->DependencyIDs));
-							UnwindStep.Type = Path::EDescType::ClassDependencies;
+							UnwindStep.DescriptionStructPtr.Set<FMetasoundOutputDescription*>(&Output);
+							UnwindStep.Type = Path::EDescType::Output;
 							return UnwindStep;
 						}
-						case Path::EDescType::Inputs:
-						{
-							// The next element after an Inputs element should always be the name of an input.
-							Path::FElement InputElement = InPath.Path[0];
-							InPath.Path.RemoveAt(0);
-							if (!ensureAlwaysMsgf(InputElement.CurrentDescType == Path::EDescType::Input, TEXT("Invalid path set up.")))
-							{
-								return InvalidUnwindStep();
-							}
+					}
 
-							FString& InputName = InputElement.LookupName;
+					// If we reached the end of the Inputs list and didn't find a match, ensure.
+					ensureAlwaysMsgf(false, TEXT("Couldn't find output %s in path."), *InNext.LookupName);
+					return FDescriptionUnwindStep::CreateInvalid();
+				}
+				case Path::EDescType::Metadata:
+				{
+					FDescriptionUnwindStep UnwindStep;
+					UnwindStep.DescriptionStructPtr.Set<FMetasoundClassMetadata*>(&(InClassDescription.Metadata));
+					UnwindStep.Type = Path::EDescType::Metadata;
 
-							// Scan the inputs list for the lookup name.
-							TArray<FMetasoundInputDescription>& InputsList = ClassDescription->Inputs;
-							for (FMetasoundInputDescription& Input : InputsList)
-							{
-								if (Input.Name == InputName)
-								{
-									FDescriptionUnwindStep UnwindStep;
-									UnwindStep.DescriptionStructPtr.Set<FMetasoundInputDescription*>(&Input);
-									UnwindStep.Type = Path::EDescType::Input;
-									return UnwindStep;
-								}
-							}
+					return UnwindStep;
+				}
+				default:
+				{
+					ensureAlwaysMsgf(false, TEXT("Invalid path- Tried to path directly from a Class Description to a type that wasn't a direct memember of the Class"));
+					return FDescriptionUnwindStep::CreateInvalid();
+				}
+			}
+		}
 
-							// If we reached the end of the Inputs list and didn't find a match, ensure.
-							ensureAlwaysMsgf(false, TEXT("Couldn't find input %s in path."), *InputName);
-							return InvalidUnwindStep();
-						}
-						case Path::EDescType::Outputs:
-						{
-							// The next element after an Outputs element should always be the name of an output.
-							Path::FElement OutputElement = InPath.Path[0];
-							InPath.Path.RemoveAt(0);
-							if (!ensureAlwaysMsgf(OutputElement.CurrentDescType == Path::EDescType::Output, TEXT("Invalid path set up.")))
-							{
-								return InvalidUnwindStep();
-							}
-						
-							FString& OutputName = OutputElement.LookupName;
 
-							// Scan the outputs list for the lookup name.
-							TArray<FMetasoundOutputDescription>& OutputsList = ClassDescription->Outputs;
-							for (FMetasoundOutputDescription& Output : OutputsList)
-							{
-								if (Output.Name == OutputElement.LookupName)
-								{
-									FDescriptionUnwindStep UnwindStep;
-									UnwindStep.DescriptionStructPtr.Set<FMetasoundOutputDescription*>(&Output);
-									UnwindStep.Type = Path::EDescType::Output;
-									return UnwindStep;
-								}
-							}
+		FDescriptionAccessPoint::FDescriptionUnwindStep FDescriptionAccessPoint::GoToNext(FDescPath& InPath, FDescriptionUnwindStep InElement)
+		{
+			if (!ensureMsgf(InPath.Path.Num() != 0, TEXT("Attempted to unwind an empty path.")))
+			{
+				return FDescriptionUnwindStep::CreateInvalid();
+			}
 
-							// If we reached the end of the Inputs list and didn't find a match, ensure.
-							ensureAlwaysMsgf(false, TEXT("Couldn't find output %s in path."), *NextStep.LookupName);
-							return InvalidUnwindStep();
-						}
-						case Path::EDescType::Metadata:
-						{
-							FDescriptionUnwindStep UnwindStep;
-							UnwindStep.DescriptionStructPtr.Set<FMetasoundClassMetadata*>(&(ClassDescription->Metadata));
-							UnwindStep.Type = Path::EDescType::Metadata;
+			Path::FElement NextStep = InPath.Path[0];
 
-							return UnwindStep;
-						}
-						default:
-						{
-							ensureAlwaysMsgf(false, TEXT("Invalid path- Tried to path directly from a Class Description to a type that wasn't a direct memember of the Class"));
-							return InvalidUnwindStep();
-						}
-					}
+			InPath.Path.RemoveAt(0);
 
+			switch (InElement.Type)
+			{
+				case Path::EDescType::Document:
+				{
+					FMetasoundDocument* Document = InElement.DescriptionStructPtr.Get<FMetasoundDocument*>();
+					return GoToNextFromDocument(*Document, InPath, NextStep);
+					break;
+				}
+				case Path::EDescType::Class:
+				{
+					FMetasoundClassDescription* ClassDescription = InElement.DescriptionStructPtr.Get<FMetasoundClassDescription*>();
+					return GoToNextFromClass(*ClassDescription, InPath, NextStep);
 					break;
 				}
 				case Path::EDescType::Graph:
 				{
 					if (!ensureAlwaysMsgf(NextStep.CurrentDescType == Path::EDescType::Nodes, TEXT("Invalid path. the Graph description only contains the Nodes list.")))
 					{
-						return InvalidUnwindStep();
+						return FDescriptionUnwindStep::CreateInvalid();
 					}
 
 					FMetasoundGraphDescription* GraphDescription = InElement.DescriptionStructPtr.Get<FMetasoundGraphDescription*>();
 				
 					if (!ensureAlwaysMsgf(InPath.Path.Num() != 0, TEXT("Incomplete path! path stopped at Nodes list without specifying a node ID.")))
 					{
-						return InvalidUnwindStep();
+						return FDescriptionUnwindStep::CreateInvalid();
 					}
 
 					Path::FElement NodeElement = InPath.Path[0];
@@ -899,10 +906,10 @@ namespace Metasound
 
 					if (!ensureAlwaysMsgf(NodeElement.CurrentDescType == Path::EDescType::Node, TEXT("Invalid path! a Nodes element must always be followed by a Node ID.")))
 					{
-						return InvalidUnwindStep();
+						return FDescriptionUnwindStep::CreateInvalid();
 					}
 				
-					uint32 NodeID = NodeElement.LookupID;
+					int32 NodeID = NodeElement.LookupID;
 
 					TArray<FMetasoundNodeDescription>& NodeList = GraphDescription->Nodes;
 					for (FMetasoundNodeDescription& Node : NodeList)
@@ -920,12 +927,12 @@ namespace Metasound
 				case Path::EDescType::Inputs:
 				{
 					ensureAlwaysMsgf(false, TEXT("Invalid path. Inputs will always be pathed directly after a Class."));
-					return InvalidUnwindStep();
+					return FDescriptionUnwindStep::CreateInvalid();
 				}
 				case Path::EDescType::DocDependencies:
 				{
 					ensureAlwaysMsgf(false, TEXT("Invalid path. Document dependencies should always follow after a Document and should always list a specific dependency by ID or name."));
-					return InvalidUnwindStep();
+					return FDescriptionUnwindStep::CreateInvalid();
 				}
 				case Path::EDescType::Input:
 				case Path::EDescType::Output:
@@ -934,13 +941,13 @@ namespace Metasound
 				default:
 				{
 					ensureAlwaysMsgf(false, TEXT("Invalid path. Input, Output, Node, and Metadata don't have any child elements."));
-					return InvalidUnwindStep();
+					return FDescriptionUnwindStep::CreateInvalid();
 				}
 			}
 
 			// if we ever hit this, we likely missed a return on one of these branches.
 			checkNoEntry();
-			return InvalidUnwindStep();
+			return FDescriptionUnwindStep::CreateInvalid();
 		}
 
 		FDescPath Path::GetPathToClassForNode(FDescPath InPathForNode, FString& InNodeName)
@@ -967,17 +974,7 @@ namespace Metasound
 			return InPathForGraph;
 		}
 
-		Metasound::Frontend::FDescPath Path::GetClassDescription(const FString& InClassName)
-		{
-			return FDescPath()[Path::EFromDocument::ToDependencies][*InClassName];
-		}
-
-		FDescPath Path::GetDependencyPath(const FString& InDependencyName)
-		{
-			return FDescPath()[Path::EFromDocument::ToDependencies][*InDependencyName];
-		}
-
-		FDescPath Path::GetDependencyPath(uint32 InDependencyID)
+		FDescPath Path::GetDependencyPath(int32 InDependencyID)
 		{
 			return FDescPath()[Path::EFromDocument::ToDependencies][InDependencyID];
 		}
diff --git a/Engine/Plugins/Runtime/Metasound/Source/MetasoundFrontend/Private/MetasoundFrontendGraph.cpp b/Engine/Plugins/Runtime/Metasound/Source/MetasoundFrontend/Private/MetasoundFrontendGraph.cpp
new file mode 100644
--- /dev/null
+++ b/Engine/Plugins/Runtime/Metasound/Source/MetasoundFrontend/Private/MetasoundFrontendGraph.cpp
@@ -0,0 +1,413 @@
+// Copyright Epic Games, Inc. All Rights Reserved.
+
+#include "MetasoundFrontendGraph.h"
+
+#include "Algo/AllOf.h"
+#include "Algo/AnyOf.h"
+#include "Algo/Transform.h"
+#include "CoreMinimal.h"
+#include "MetasoundFrontend.h"
+#include "MetasoundFrontendDataLayout.h"
+#include "MetasoundGraph.h"
+#include "MetasoundNodeInterface.h"
+
+namespace Metasound
+{
+	FFrontendGraph::FFrontendGraph(const FString& InInstanceName)
+	:	FGraph(InInstanceName)
+	{
+	}
+
+	void FFrontendGraph::AddInputNode(int32 InDependencyId, int32 InIndex, const FVertexKey& InVertexKey, TUniquePtr<INode> InNode)
+	{
+		if (InNode.IsValid())
+		{
+			// There shouldn't be duplicate IDs. 
+			check(!InputNodes.Contains(InIndex));
+
+			// Input nodes need an extra Index value to keep track of their position in the graph's inputs.
+			InputNodes.Add(InIndex, InNode.Get());
+			AddInputDataDestination(*InNode, InVertexKey);
+			AddNode(InDependencyId, MoveTemp(InNode));
+		}
+	}
+
+	void FFrontendGraph::AddOutputNode(int32 InNodeID, int32 InIndex, const FVertexKey& InVertexKey, TUniquePtr<INode> InNode)
+	{
+		if (InNode.IsValid())
+		{
+			// There shouldn't be duplicate IDs. 
+			check(!OutputNodes.Contains(InIndex));
+
+			// Output nodes need an extra Index value to keep track of their position in the graph's inputs.
+			OutputNodes.Add(InIndex, InNode.Get());
+			AddOutputDataSource(*InNode, InVertexKey);
+			AddNode(InNodeID, MoveTemp(InNode));
+		}
+	}
+
+	/** Store a node on this graph. */
+	void FFrontendGraph::AddNode(int32 InNodeID, TUniquePtr<INode> InNode)
+	{
+		if (InNode.IsValid())
+		{
+			// There shouldn't be duplicate IDs. 
+			check(!NodeMap.Contains(InNodeID));
+
+			NodeMap.Add(InNodeID, InNode.Get());
+			StoreNode(MoveTemp(InNode));
+		}
+	}
+
+	const INode* FFrontendGraph::FindNode(int32 InNodeID) const
+	{
+		INode* const* NodePtr = NodeMap.Find(InNodeID);
+
+		if (nullptr != NodePtr)
+		{
+			return *NodePtr;
+		}
+
+		return nullptr;
+	}
+
+	const INode* FFrontendGraph::FindInputNode(int32 InIndex) const
+	{
+		INode* const* NodePtr = InputNodes.Find(InIndex);
+
+		if (nullptr != NodePtr)
+		{
+			return *NodePtr;
+		}
+
+		return nullptr;
+	}
+
+	const INode* FFrontendGraph::FindOutputNode(int32 InIndex) const
+	{
+		INode* const* NodePtr = OutputNodes.Find(InIndex);
+
+		if (nullptr != NodePtr)
+		{
+			return *NodePtr;
+		}
+
+		return nullptr;
+	}
+
+	/** Returns true if all edges, destinations and sources refer to 
+	 * nodes stored in this graph. */
+	bool FFrontendGraph::OwnsAllReferencedNodes() const
+	{
+		const TArray<FDataEdge>& AllEdges = GetDataEdges();
+		for (const FDataEdge& Edge : AllEdges)
+		{
+			if (!StoredNodes.Contains(Edge.From.Node))
+			{
+				return false;
+			}
+
+			if (!StoredNodes.Contains(Edge.To.Node))
+			{
+				return false;
+			}
+		}
+
+		const FInputDataDestinationCollection& AllInputDestinations = GetInputDataDestinations();
+		for (auto& DestTuple : AllInputDestinations)
+		{
+			if (!StoredNodes.Contains(DestTuple.Value.Node))
+			{
+				return false;
+			}
+		}
+
+		const FOutputDataSourceCollection& AllOutputSources = GetOutputDataSources();
+		for (auto& SourceTuple : AllOutputSources)
+		{
+			if (!StoredNodes.Contains(SourceTuple.Value.Node))
+			{
+				return false;
+			}
+		}
+
+		return true;
+	}
+
+
+	void FFrontendGraph::StoreNode(TUniquePtr<INode> InNode)
+	{
+		check(InNode.IsValid());
+		StoredNodes.Add(InNode.Get());
+		Storage.Add(MoveTemp(InNode));
+	}
+
+
+	TUniquePtr<INode> FFrontendGraphBuilder::CreateInputNode(const FMetasoundNodeDescription& InNode, const FMetasoundInputDescription& InDescription) const
+	{
+		if (!ensureAlwaysMsgf(Frontend::DoesDataTypeSupportLiteralType(InDescription.TypeName, InDescription.LiteralValue.LiteralType), TEXT("[%s] cannot be constructed with the provided literal type."), *InDescription.TypeName.ToString()))
+		{
+			return TUniquePtr<INode>(nullptr);
+		}
+
+		FDataTypeLiteralParam LiteralParam = Frontend::GetLiteralParamForDataType(InDescription.TypeName, InDescription.LiteralValue);
+
+		FInputNodeConstructorParams InitParams =
+		{
+			InDescription.Name,
+			InDescription.Name,
+			MoveTemp(LiteralParam)
+		};
+
+		return FMetasoundFrontendRegistryContainer::Get()->ConstructInputNode(InDescription.TypeName, MoveTemp(InitParams));
+	}
+
+	TUniquePtr<INode> FFrontendGraphBuilder::CreateOutputNode(const FMetasoundNodeDescription& InNode, const FMetasoundOutputDescription& InDescription) const
+	{
+		FOutputNodeConstrutorParams InitParams =
+		{
+			InDescription.Name,
+			InDescription.Name,
+		};
+
+		return FMetasoundFrontendRegistryContainer::Get()->ConstructOutputNode(InDescription.TypeName, InitParams);
+	}
+
+	TUniquePtr<INode> FFrontendGraphBuilder::CreateExternalNode(const FMetasoundNodeDescription& InNode, const FMetasoundClassDescription& InClass) const
+	{
+		check(InClass.Metadata.NodeType == EMetasoundClassType::External);
+		check(InNode.ObjectTypeOfNode == EMetasoundClassType::External);
+
+		// We found a match. now we just need to create the input node.
+		const FMetasoundExternalClassLookupInfo& LookupInfo = InClass.ExternalNodeClassLookupInfo;
+
+		FNodeInitData InitData;
+
+		InitData.InstanceName.Append(InNode.Name);
+		InitData.InstanceName.AppendChar('_');
+		InitData.InstanceName.AppendInt(InNode.UniqueID);
+		
+		// Copy over our initialization params.
+		for (auto& StaticParamTuple : InNode.StaticParameters)
+		{
+			FDataTypeLiteralParam LiteralParam = Frontend::GetLiteralParam(StaticParamTuple.Value);
+
+			if (LiteralParam.IsValid())
+			{
+				InitData.ParamMap.Add(StaticParamTuple.Key, MoveTemp(LiteralParam));
+			}
+		}
+
+		return FMetasoundFrontendRegistryContainer::Get()->ConstructExternalNode(LookupInfo.ExternalNodeClassName, LookupInfo.ExternalNodeClassHash, InitData);
+	}
+
+	// TODO: add errors here. Most will be a "PromptIfMissing"...
+	void FFrontendGraphBuilder::AddExternalNodesToGraph(const TArray<FMetasoundNodeDescription>& InNodes, const TMap<int32, const FMetasoundClassDescription*>& InClasses, FFrontendGraph& OutGraph) const
+	{
+		for (const FMetasoundNodeDescription& Node : InNodes)
+		{
+			if (Node.ObjectTypeOfNode == EMetasoundClassType::External)
+			{
+				const FMetasoundClassDescription* ClassDesc = InClasses.FindRef(Node.DependencyID);
+
+				if (ensure(nullptr != ClassDesc))
+				{
+					OutGraph.AddNode(Node.UniqueID, CreateExternalNode(Node, *ClassDesc));
+				}
+			}
+		}
+	}
+
+	void FFrontendGraphBuilder::AddEdgesToGraph(const FMetasoundGraphDescription& InGraphDescription, FFrontendGraph& OutGraph) const
+	{
+		for (const FMetasoundNodeDescription& NodeDescription : InGraphDescription.Nodes)
+		{
+			const INode* ToNode = OutGraph.FindNode(NodeDescription.UniqueID);
+
+			if (nullptr == ToNode)
+			{
+				// TODO: log errror. 
+				continue;
+			}
+
+			for (const FMetasoundNodeConnectionDescription& InputConnection : NodeDescription.InputConnections)
+			{
+				const bool bIsConnected = InputConnection.NodeID != FMetasoundNodeConnectionDescription::DisconnectedNodeID;
+				if (!bIsConnected)
+				{
+					continue;
+				}
+
+				const INode* FromNode = OutGraph.FindNode(InputConnection.NodeID);
+				if (nullptr == FromNode)
+				{
+					// TODO: log error.
+					continue;
+				}
+
+				OutGraph.AddDataEdge(*FromNode, *InputConnection.OutputName, *ToNode, *InputConnection.InputName);
+			}
+		}
+	}
+
+	void FFrontendGraphBuilder::AddInputDestinationsToGraph(const TArray<FMetasoundInputDescription>& InInputDescriptions, const TArray<FMetasoundNodeDescription>& InNodes, FFrontendGraph& OutGraph) const
+	{
+		for (int32 InputIndex = 0; InputIndex < InInputDescriptions.Num(); InputIndex++)
+		{
+			const FMetasoundInputDescription& InputDescription = InInputDescriptions[InputIndex];
+
+			auto IsInputWithSameName = [&](const FMetasoundNodeDescription& InNode) 
+			{ 
+				return (InNode.Name == InputDescription.Name) && (InNode.ObjectTypeOfNode == EMetasoundClassType::Input);
+			};
+
+			const FMetasoundNodeDescription* NodeDescription = InNodes.FindByPredicate(IsInputWithSameName);
+
+			if (nullptr == NodeDescription)
+			{
+				// TODO: add error.
+				continue;
+			}
+
+			OutGraph.AddInputNode(NodeDescription->UniqueID, InputIndex, InputDescription.Name, CreateInputNode(*NodeDescription, InputDescription));
+		}
+	}
+
+	void FFrontendGraphBuilder::AddOutputSourcesToGraph(const TArray<FMetasoundOutputDescription>& InOutputDescriptions, const TArray<FMetasoundNodeDescription>& InNodes, FFrontendGraph& OutGraph) const
+	{
+		for (int32 OutputIndex = 0; OutputIndex < InOutputDescriptions.Num(); OutputIndex++)
+		{
+			const FMetasoundOutputDescription& OutputDescription = InOutputDescriptions[OutputIndex];
+
+			auto IsOutputWithSameName = [&](const FMetasoundNodeDescription& InNode) 
+			{ 
+				return (InNode.Name == OutputDescription.Name) && (InNode.ObjectTypeOfNode == EMetasoundClassType::Output);
+			};
+
+			const FMetasoundNodeDescription* NodeDescription = InNodes.FindByPredicate(IsOutputWithSameName);
+
+			if (nullptr == NodeDescription)
+			{
+				// TODO: add error.
+				continue;
+			}
+
+			OutGraph.AddOutputNode(NodeDescription->UniqueID, OutputIndex, OutputDescription.Name, CreateOutputNode(*NodeDescription, OutputDescription));
+		}
+	}
+
+	void FFrontendGraphBuilder::SplitNodesByType(const TArray<FMetasoundNodeDescription>& InNodes, TArray<FMetasoundNodeDescription>& OutExternalNodes, TArray<FMetasoundNodeDescription>& OutInputNodes, TArray<FMetasoundNodeDescription>& OutOutputNodes) const
+	{
+		for (const FMetasoundNodeDescription& Node : InNodes)
+		{
+			switch (Node.ObjectTypeOfNode)
+			{
+				case EMetasoundClassType::External:
+
+					OutExternalNodes.Add(Node);
+					break;
+
+				case EMetasoundClassType::Input:
+
+					OutInputNodes.Add(Node);
+					break;
+
+				case EMetasoundClassType::Output:
+
+					OutOutputNodes.Add(Node);
+					break;
+
+				default:
+					checkNoEntry();
+			}
+		}
+	}
+
+	/** Check that all dependencies are C++ class dependencies. */
+	bool FFrontendGraphBuilder::IsFlat(const FMetasoundDocument& InDocument) const
+	{
+		return IsFlat(InDocument.RootClass, InDocument.Dependencies);
+	}
+
+	bool FFrontendGraphBuilder::IsFlat(const FMetasoundClassDescription& InRoot, const TArray<FMetasoundClassDescription>& InDependencies) const
+	{
+		// All dependencies are external dependencies in a flat graph
+		auto IsClassExternal = [&](const FMetasoundClassDescription& InDesc) { return InDesc.Metadata.NodeType == EMetasoundClassType::External; };
+		const bool bIsEveryDependencyExternal = Algo::AllOf(InDependencies, IsClassExternal);
+
+		if (!bIsEveryDependencyExternal)
+		{
+			return false;
+		}
+
+		// None of the nodes are subgraphs in a flat graph
+		auto IsGraphNode = [&](const FMetasoundNodeDescription& InDesc) { return InDesc.ObjectTypeOfNode == EMetasoundClassType::MetasoundGraph; };
+		const bool bIsAnyNodeAGraph = Algo::AnyOf(InRoot.Graph.Nodes, IsGraphNode);
+
+		if (bIsAnyNodeAGraph)
+		{
+			return false;
+		}
+
+		// All the dependencies are available in a flat graph.
+		TSet<int32> AvailableDependencies;
+		Algo::Transform(InDependencies, AvailableDependencies, [](const FMetasoundClassDescription& InDesc) { return InDesc.UniqueID; });
+
+		auto IsDependencyMet = [&](const FMetasoundNodeDescription& InNode) 
+		{ 
+			// Currently, input/output nodes do not have dependencies. This may change.
+			const bool bIsInputOrOutput = (InNode.ObjectTypeOfNode == EMetasoundClassType::Input) || (InNode.ObjectTypeOfNode == EMetasoundClassType::Output);
+			const bool bIsExternal = InNode.ObjectTypeOfNode == EMetasoundClassType::External;
+
+			return bIsInputOrOutput || (bIsExternal && AvailableDependencies.Contains(InNode.DependencyID));
+		};
+
+		const bool bIsEveryDependencyMet = Algo::AllOf(InRoot.Graph.Nodes, IsDependencyMet);
+		return bIsEveryDependencyMet;
+	}
+
+	
+	/* Metasound document should be inflated by now. */
+	TUniquePtr<FFrontendGraph> FFrontendGraphBuilder::CreateGraph(const FMetasoundDocument& InDocument) const
+	{
+		return CreateGraph(InDocument.RootClass, InDocument.Dependencies);
+	}
+
+	TUniquePtr<FFrontendGraph> FFrontendGraphBuilder::CreateGraph(const FMetasoundClassDescription& InRoot, const TArray<FMetasoundClassDescription>& InDependencies) const
+	{
+		if (!IsFlat(InRoot, InDependencies))
+		{
+			checkNoEntry();
+			return TUniquePtr<FFrontendGraph>(nullptr);
+		}
+
+		TUniquePtr<FFrontendGraph> Graph = MakeUnique<FFrontendGraph>(InRoot.Metadata.NodeName);
+
+		TMap<int32, const FMetasoundClassDescription*> NodeMap;
+
+		for (const FMetasoundClassDescription& InDesc : InDependencies)
+		{
+			NodeMap.Add(InDesc.UniqueID, &InDesc);
+		}
+
+		TArray<FMetasoundNodeDescription> ExternalNodes;
+		TArray<FMetasoundNodeDescription> InputNodes;
+		TArray<FMetasoundNodeDescription> OutputNodes;
+
+		SplitNodesByType(InRoot.Graph.Nodes, ExternalNodes, InputNodes, OutputNodes);
+
+		// TODO: will likely want to bubble up errors here for case where
+		// a datatype or node is not registered. 
+		AddExternalNodesToGraph(ExternalNodes, NodeMap, *Graph);
+
+		AddInputDestinationsToGraph(InRoot.Inputs, InputNodes, *Graph);
+
+		AddOutputSourcesToGraph(InRoot.Outputs, OutputNodes, *Graph);
+
+		AddEdgesToGraph(InRoot.Graph, *Graph);
+
+		check(Graph->OwnsAllReferencedNodes());
+
+		return MoveTemp(Graph);
+	}
+}
diff --git a/Engine/Plugins/Runtime/Metasound/Source/MetasoundFrontend/Private/MetasoundFrontendGraph.h b/Engine/Plugins/Runtime/Metasound/Source/MetasoundFrontend/Private/MetasoundFrontendGraph.h
new file mode 100644
--- /dev/null
+++ b/Engine/Plugins/Runtime/Metasound/Source/MetasoundFrontend/Private/MetasoundFrontendGraph.h
@@ -0,0 +1,144 @@
+// Copyright Epic Games, Inc. All Rights Reserved.
+
+#pragma once
+
+#include "CoreMinimal.h"
+
+#include "MetasoundFrontend.h"
+#include "MetasoundFrontendDataLayout.h"
+#include "MetasoundGraph.h"
+#include "MetasoundNodeInterface.h"
+
+namespace Metasound
+{
+	/** FFrontendGraph is a utility graph for use in the frontend. It can own nodes
+	 * that live within the graph and provides query interfaces for finding nodes
+	 * by dependency ID or input/output index. 
+	 */
+	class FFrontendGraph : public FGraph
+	{
+		// - TODO: consider adding ability to add in another sub graph
+		public:
+			/** FFrontendGraph constructor.
+			 *
+			 * @parma InInstanceName - Name of this graph.
+			 */
+			FFrontendGraph(const FString& InInstanceName);
+
+			virtual ~FFrontendGraph() = default;
+
+			/** Add an input node to this graph.
+			 *
+			 * @param InNodeID - The NodeID related to the parent FMetasoundClassDescription.
+			 * @param InIndex - The positional index for the input.
+			 * @param InVertexKey - The key for the graph input vertex.
+			 * @param InNode - A unique pointer to an input node. 
+			 */
+			void AddInputNode(int32 InNodeID, int32 InIndex, const FVertexKey& InVertexKey, TUniquePtr<INode> InNode);
+
+			/** Add an output node to this graph.
+			 *
+			 * @param InNodeID - The NodeID related to the parent FMetasoundClassDescription.
+			 * @param InIndex - The positional index for the output.
+			 * @param InVertexKey - The key for the graph output vertex.
+			 * @param InNode - A unique pointer to an output node. 
+			 */
+			void AddOutputNode(int32 InNodeID, int32 InIndex, const FVertexKey& InVertexKey, TUniquePtr<INode> InNode);
+
+			/** Store a node on this graph. 
+			 *
+			 * @param InNodeID - The NodeID related to the parent FMetasoundClassDescription.
+			 * @param InNode - A unique pointer to a node. 
+			 */
+			void AddNode(int32 InNodeID, TUniquePtr<INode> InNode);
+
+			/** Retrieve node by dependency ID.
+			 *
+			 * @param InNodeID - The NodeID of the requested Node.
+			 *
+			 * @return Pointer to the Node if it is stored on this graph. nullptr otherwise. 
+			 */
+			const INode* FindNode(int32 InNodeID) const;
+
+			/** Retrieve node by input index.
+			 *
+			 * @param InIndex - The index of the requested input.
+			 *
+			 * @return Pointer to the Node if it is stored on this graph. nullptr otherwise. 
+			 */
+			const INode* FindInputNode(int32 InIndex) const;
+
+			/** Retrieve node by output index.
+			 *
+			 * @param InIndex - The index of the requested output.
+			 *
+			 * @return Pointer to the Node if it is stored on this graph. nullptr otherwise. 
+			 */
+			const INode* FindOutputNode(int32 InIndex) const;
+
+			/** Returns true if all edges, destinations and sources refer to 
+			 * nodes stored in this graph. */
+			bool OwnsAllReferencedNodes() const;
+
+		private:
+
+			void StoreNode(TUniquePtr<INode> InNode);
+
+			TMap<int32, INode*> InputNodes;
+			TMap<int32, INode*> OutputNodes;
+
+			TMap<int32, INode*> NodeMap;
+			TSet<const INode*> StoredNodes;
+			TArray<TUniquePtr<INode>> Storage;
+	};
+
+	/** FFrontendGraphBuilder builds a FFrontendGraph from a FMetasoundDoucment
+	 * or FMetasoundClassDescription.
+	 */
+	class FFrontendGraphBuilder
+	{
+	public:
+
+		/** Check that all dependencies are C++ class dependencies. 
+		 * 
+		 * @param InDocument - Document containing dependencies.
+		 *
+		 * @return True if all dependencies are C++ classes. False otherwise.
+		 */
+		bool IsFlat(const FMetasoundDocument& InDocument) const;
+
+		/** Check that all dependencies are C++ class dependencies. 
+		 * 
+		 * @param InDependencies - Array of dependencies to check.
+		 *
+		 * @return True if all dependencies are C++ classes. False otherwise.
+		 */
+		bool IsFlat(const FMetasoundClassDescription& InRoot, const TArray<FMetasoundClassDescription>& InDependencies) const;
+		
+		/* Metasound document should be in order to create this graph. */
+		TUniquePtr<FFrontendGraph> CreateGraph(const FMetasoundDocument& InDocument) const;
+
+		/** Creates a graph given the root class and all of it's dependencies. */
+		TUniquePtr<FFrontendGraph> CreateGraph(const FMetasoundClassDescription& InRootClass, const TArray<FMetasoundClassDescription>& InDependencies) const;
+
+	private:
+
+		TUniquePtr<INode> CreateInputNode(const FMetasoundNodeDescription& InNode, const FMetasoundInputDescription& InDescription) const;
+
+		TUniquePtr<INode> CreateOutputNode(const FMetasoundNodeDescription& InNode, const FMetasoundOutputDescription& InDescription) const;
+
+		TUniquePtr<INode> CreateExternalNode(const FMetasoundNodeDescription& InNode, const FMetasoundClassDescription& InClass) const;
+
+		void SplitNodesByType(const TArray<FMetasoundNodeDescription>& InNodes, TArray<FMetasoundNodeDescription>& OutExternalNodes, TArray<FMetasoundNodeDescription>& OutInputNodes, TArray<FMetasoundNodeDescription>& OutOutputNodes) const;
+
+
+		// TODO: add errors here. Most will be a "PromptIfMissing"...
+		void AddExternalNodesToGraph(const TArray<FMetasoundNodeDescription>& InNodes, const TMap<int32, const FMetasoundClassDescription*>& InClasses, FFrontendGraph& OutGraph) const;
+
+		void AddEdgesToGraph(const FMetasoundGraphDescription& InGraphDescription, FFrontendGraph& OutGraph) const;
+
+		void AddInputDestinationsToGraph(const TArray<FMetasoundInputDescription>& InInputDescriptions, const TArray<FMetasoundNodeDescription>& InNodes, FFrontendGraph& OutGraph) const;
+		void AddOutputSourcesToGraph(const TArray<FMetasoundOutputDescription>& InOutputDescriptions, const TArray<FMetasoundNodeDescription>& InNodes, FFrontendGraph& OutGraph) const;
+
+	};
+}
diff --git a/Engine/Plugins/Runtime/Metasound/Source/MetasoundFrontend/Public/MetasoundFrontend.h b/Engine/Plugins/Runtime/Metasound/Source/MetasoundFrontend/Public/MetasoundFrontend.h
--- a/Engine/Plugins/Runtime/Metasound/Source/MetasoundFrontend/Public/MetasoundFrontend.h
+++ b/Engine/Plugins/Runtime/Metasound/Source/MetasoundFrontend/Public/MetasoundFrontend.h
@@ -209,8 +209,12 @@ namespace Metasound
 			// for the FNodeHandle, FInputHandle and FOutputHandle, this is a path to the individual node within a Graph description.
 			const FDescPath& InPath;
 
-			// Class name for the graph or node to get a handle for.
-			const FString& InClassName;
+			// Name for the vertex to get a handle for.
+			//const FString& InName;
+
+			// ID for node or class.
+			int32 NodeID;
+			int32 DependencyID;
 
 			// The asset that owns the FMetasoundDocument this handle is associated with.
 			TWeakObjectPtr<UObject> InOwningAsset;
@@ -232,10 +236,13 @@ namespace Metasound
 		public:
 			FOutputHandle() = delete;
 
-			explicit FOutputHandle(FHandleInitParams::EPrivateToken InToken, const FHandleInitParams& InParams, const FString& InOutputName);
+			enum EInputNodeTag { InputNode };
+			enum EDefaultTag { Default };
+
+			explicit FOutputHandle(FHandleInitParams::EPrivateToken InToken, const FHandleInitParams& InParams, const FString& InOutputName, EDefaultTag InTag);
 
 			// Constructor used for the outgoing connection from an input node.
-			explicit FOutputHandle(FHandleInitParams::EPrivateToken InToken, const FHandleInitParams& InParams);
+			explicit FOutputHandle(FHandleInitParams::EPrivateToken InToken, const FHandleInitParams& InParams, const FString& InOutputName, EInputNodeTag InTag);
 
 			~FOutputHandle() = default;
 
@@ -250,7 +257,7 @@ namespace Metasound
 			FName GetOutputType() const;
 			FString GetOutputName() const;
 			FText GetOutputTooltip() const;
-			uint32 GetOwningNodeID() const;
+			int32 GetOwningNodeID() const;
 
 			// @todo: with the current spec, we'd have to scan all nodes in the graph for this output to solve this.
 			// Is this worth it, or should we just require folks to use FInputHandle::GetCurrentlyConnectedOutput()?
@@ -279,11 +286,14 @@ namespace Metasound
 			FInputHandle() = delete;
 			~FInputHandle() = default;
 
+			enum EOutputNodeTag { OutputNode };
+			enum EDefaultTag { Default };
+
 			// Sole constructor for FInputHandle. Can only be used by friend classes for FHandleInitParams.
-			explicit FInputHandle(FHandleInitParams::EPrivateToken PrivateToken, const FHandleInitParams& InParams, const FString& InputName);
+			explicit FInputHandle(FHandleInitParams::EPrivateToken PrivateToken, const FHandleInitParams& InParams, const FString& InputName, EDefaultTag InTag);
 
 			// constructor used exclusively by output nodes.
-			explicit FInputHandle(FHandleInitParams::EPrivateToken PrivateToken, const FHandleInitParams& InParams);
+			explicit FInputHandle(FHandleInitParams::EPrivateToken PrivateToken, const FHandleInitParams& InParams, const FString& InputName, EOutputNodeTag InTag);
 
 			static FInputHandle InvalidHandle();
 
@@ -356,8 +366,8 @@ namespace Metasound
 			// Otherwise it will return an invalid FGraphHandle.
 			void GetContainedGraph(FGraphHandle& OutGraph);
 
-			uint32 GetNodeID() const;
-			static uint32 GetNodeID(const FDescPath& InNodePath);
+			int32 GetNodeID() const;
+			static int32 GetNodeID(const FDescPath& InNodePath);
 
 			const FString& GetNodeName() const;
 
@@ -372,7 +382,7 @@ namespace Metasound
 			// an externally implemented node, or a metasound graph itself.
 			EMetasoundClassType NodeClassType;
 
-			uint32 NodeID;
+			int32 NodeID;
 		};
 
 		class METASOUNDFRONTEND_API FGraphHandle : protected ITransactable
@@ -389,7 +399,7 @@ namespace Metasound
 				using namespace Path;
 
 				FDescPath PathToGraph = FDescPath()[EFromDocument::ToRootClass][EFromClass::ToGraph];
-				FHandleInitParams InitParams = { InAccessPoint, PathToGraph, InRootMetasoundDocument.RootClass.Metadata.NodeName, MakeWeakObjectPtr(InOwner) };
+				FHandleInitParams InitParams = { InAccessPoint, PathToGraph, INDEX_NONE, InRootMetasoundDocument.RootClass.UniqueID, MakeWeakObjectPtr(InOwner) };
 				return FGraphHandle(FHandleInitParams::PrivateToken, InitParams);
 			}
 
@@ -409,7 +419,7 @@ namespace Metasound
 			void SetEditorData(const FMetasoundEditorData& InEditorData);
 
 			TArray<FNodeHandle> GetAllNodes();
-			FNodeHandle GetNodeWithId(uint32 InNodeId) const;
+			FNodeHandle GetNodeWithId(int32 InNodeId) const;
 			TArray<FNodeHandle> GetOutputNodes();
 			TArray<FNodeHandle> GetInputNodes();
 
@@ -525,10 +535,10 @@ namespace Metasound
 			bool RemoveNodeInternal(const FNodeHandle& InNode);
 
 			// Scans all existing node ids to guarantee a new unique ID.
-			uint32 FindNewUniqueNodeId();
+			int32 FindNewUniqueNodeId();
 
 			// Scans all existing dependency IDs in the root dependency list to guarantee a new unique ID when adding a dependency.
-			uint32 FindNewUniqueDependencyId();
+			int32 FindNewUniqueDependencyId();
 
 			FMetasoundLiteralDescription* GetLiteralDescriptionForInput(const FString& InInputName, FName& OutDataType) const;
 
diff --git a/Engine/Plugins/Runtime/Metasound/Source/MetasoundFrontend/Public/MetasoundFrontendBaseClasses.h b/Engine/Plugins/Runtime/Metasound/Source/MetasoundFrontend/Public/MetasoundFrontendBaseClasses.h
--- a/Engine/Plugins/Runtime/Metasound/Source/MetasoundFrontend/Public/MetasoundFrontendBaseClasses.h
+++ b/Engine/Plugins/Runtime/Metasound/Source/MetasoundFrontend/Public/MetasoundFrontendBaseClasses.h
@@ -193,15 +193,15 @@ namespace Metasound
 			{
 				EDescType CurrentDescType;
 				
-				// This is only used if EDescType is Inputs, Outputs, or Dependencies.
+				// This is only used if EDescType is Inputs or Outputs
 				FString LookupName;
 
-				// This is only used if EDescType == Node.
-				uint32 LookupID;
+				// This is only used if EDescType == Node or Dependencies
+				int32 LookupID;
 
 				FElement()
 					: CurrentDescType(EDescType::Invalid)
-					, LookupID(INDEX_NONE)
+					, LookupID(FMetasoundNodeDescription::InvalidID)
 				{}
 			};
 		}
@@ -365,7 +365,7 @@ namespace Metasound
 				return NewPath;
 			}
 
-			FDescPath operator[](uint32 InElement) const
+			FDescPath operator[](int32 InElement) const
 			{
 				if (!ensureAlwaysMsgf(Path.Num() > 0, TEXT("Tried to append to a path that had no root.")))
 				{
@@ -383,16 +383,19 @@ namespace Metasound
 				}
 
 				Path::FElement NewElement = Path::FElement();
+				NewElement.LookupID = InElement;
 
 				if (LastElementInPath.CurrentDescType == Path::EDescType::Nodes)
 				{
 					NewElement.CurrentDescType = Path::EDescType::Node;
-					NewElement.LookupID = InElement;
 				}
 				else if (LastElementInPath.CurrentDescType == Path::EDescType::ClassDependencies)
 				{
 					NewElement.CurrentDescType = Path::EDescType::DocDependencies;
-					NewElement.LookupID = InElement;
+				}
+				else if (LastElementInPath.CurrentDescType == Path::EDescType::DocDependencies)
+				{
+					NewElement.CurrentDescType = Path::EDescType::Class;
 				}
 
 				NewPath.Path.Add(NewElement);
@@ -416,8 +419,8 @@ namespace Metasound
 				{
 					case Path::EDescType::DocDependencies:
 					{
-						NewElement.CurrentDescType = Path::EDescType::Class;
-						NewElement.LookupName = FString(InElementName);
+						ensureAlwaysMsgf(false, TEXT("Invalid path set up. DocDependencies for are accessed by using a int32."));
+						return FDescPath();
 						break;
 					}
 					case Path::EDescType::Inputs:
@@ -440,17 +443,17 @@ namespace Metasound
 						return FDescPath();
 						break;
 					}
-					/** Dependencies for individual classes are navigated to via a uint32. */
+					/** Dependencies for individual classes are navigated to via a int32. */
 					case Path::EDescType::ClassDependencies:
 					{
-						ensureAlwaysMsgf(false, TEXT("Invalid path set up. Dependencies for individual classes are accessed by using a uint32."));
+						ensureAlwaysMsgf(false, TEXT("Invalid path set up. Dependencies for individual classes are accessed by using a int32."));
 						return FDescPath();
 						break;;
 					}
-					/** Node list descriptions are navigated via a uint32. */
+					/** Node list descriptions are navigated via a int32. */
 					case Path::EDescType::Nodes:
 					{
-						ensureAlwaysMsgf(false, TEXT("Invalid path set up. Node lists are accessed by using a uint32."));
+						ensureAlwaysMsgf(false, TEXT("Invalid path set up. Node lists are accessed by using a int32."));
 						return FDescPath();
 						break;
 					}
@@ -497,11 +500,7 @@ namespace Metasound
 			// Returns a path for whatever class description encapsulates the given graph.
 			FDescPath GetOwningClassDescription(FDescPath InPathForGraph);
 			
-			FDescPath GetClassDescription(const FString& InClassName);
-
-			FDescPath GetDependencyPath(const FString& InDependencyName);
-
-			FDescPath GetDependencyPath(uint32 InDependencyID);
+			FDescPath GetDependencyPath(int32 InDependencyID);
 
 			// Given the path to an input node, returns the path to the corresponding  input description.
 			FDescPath GetInputDescriptionPath(FDescPath InPathForInputNode, const FString& InputName);
@@ -514,7 +513,7 @@ namespace Metasound
 			FString GetPrintableString(FDescPath InPath);
 		}
 
-		using FClassDependencyIDs = TArray<uint32>;
+		using FClassDependencyIDs = TArray<int32>;
 
 		/**
 		 * Utility class that lives alongside a FMetasoundClassDescription to allow it's individual elements to be safely read and edited.
@@ -602,8 +601,22 @@ namespace Metasound
 			{
 				FMetasoundDescriptionPtr DescriptionStructPtr;
 				Path::EDescType Type;
+
+				// TODO: make some simple constructors for this so we don't have
+				// to manually fill out info each time. 
+
+				static FDescriptionUnwindStep CreateInvalid()
+				{
+					FDescriptionUnwindStep InvalidStep;
+					InvalidStep.DescriptionStructPtr.Set<void*>(nullptr);
+					InvalidStep.Type = Path::EDescType::Invalid;
+					return InvalidStep;
+				}
 			};
 
+			FDescriptionUnwindStep GoToNextFromDocument(FMetasoundDocument& InDocument, FDescPath& InPath, const Path::FElement& InNext);
+			FDescriptionUnwindStep GoToNextFromClass(FMetasoundClassDescription& InClass, FDescPath& InPath, const Path::FElement& InNext);
+
 			// This function pops off elements of InPath to go to the next level down in a Metasound description.
 			FDescriptionUnwindStep GoToNext(FDescPath& InPath, FDescriptionUnwindStep InElement);
 		};
diff --git a/Engine/Plugins/Runtime/Metasound/Source/MetasoundFrontend/Public/MetasoundFrontendDataLayout.h b/Engine/Plugins/Runtime/Metasound/Source/MetasoundFrontend/Public/MetasoundFrontendDataLayout.h
--- a/Engine/Plugins/Runtime/Metasound/Source/MetasoundFrontend/Public/MetasoundFrontendDataLayout.h
+++ b/Engine/Plugins/Runtime/Metasound/Source/MetasoundFrontend/Public/MetasoundFrontendDataLayout.h
@@ -6,7 +6,6 @@
 #include "MetasoundGraph.h"
 #include "MetasoundFrontendDataLayout.generated.h"
 
-
 UENUM()
 enum class EMetasoundClassType : uint8
 {
@@ -173,7 +172,7 @@ struct FMetasoundNodeConnectionDescription
 {
 	GENERATED_BODY()
 
-	static const uint32 DisconnectedNodeID = INDEX_NONE;
+	static constexpr int32 DisconnectedNodeID = INDEX_NONE;
 
 	// Name of the input this connection is for. Should match an input name in the Inputs array of this node's FMetasoundClassDescription.
 	UPROPERTY()
@@ -181,7 +180,7 @@ struct FMetasoundNodeConnectionDescription
 
 	// Unique ID for the node this connection is from.
 	UPROPERTY()
-	uint32 NodeID;
+	int32 NodeID;
 
 	// FMetasoundOutputDescription::Name of the output this connection is from if NodeID is valid,
 	// if FMetasoundNodeConnectionDescription::DisconnectedNodeID.
@@ -200,10 +199,14 @@ struct FMetasoundNodeDescription
 {
 	GENERATED_BODY()
 
+	static constexpr int32 InvalidID = INDEX_NONE;
 
 	// Unique integer given to this node. Only has to be unique within the Nodes array.
 	UPROPERTY()
-	uint32 UniqueID;
+	int32 UniqueID;
+
+	UPROPERTY()
+	int32 DependencyID;
 
 	// This will either match a FMetasoundObjectDescription::Metadata::Name in the Dependencies list of the owning FMetasoundObjectDescription,
 	// Or it will match a FMetasoundInputDescription::Name or FMetasoundOutputDescription::Name, depending on the value of ObjectTypeOfNode.
@@ -271,12 +274,13 @@ struct FMetasoundClassDescription
 {
 	GENERATED_BODY()
 
-	static const uint32 RootClassID = 0;
+	static constexpr int32 RootClassID = 0;
+	static constexpr int32 InvalidID = INDEX_NONE;
 
 	// Unique ID for this FMetasoundClassDescription in FMetasoundDocument::Dependencies.
 	// this will be set to FMetasoundClassDescription::RootClassID for FMetasoundDocument::RootClass.
 	UPROPERTY()
-	uint32 UniqueID;
+	int32 UniqueID;
 
 	UPROPERTY(EditAnywhere, Category = CustomView)
 	FMetasoundClassMetadata Metadata;
@@ -293,7 +297,7 @@ struct FMetasoundClassDescription
 	// If this object is itself of the Metasound type, here we list the unique IDs of other objects it depends on to be built.
 	// These unique IDs are then found in FMetasoundDocument::Dependencies.
 	UPROPERTY()
-	TArray<uint32> DependencyIDs;
+	TArray<int32> DependencyIDs;
 
 	// If this object is itself of the Metasound type, We can fully describe that metasound's graph here.
 	// If this object is a Metasound and this Graph is empty,
diff --git a/Engine/Plugins/Runtime/Metasound/Source/MetasoundFrontend/Public/MetasoundFrontendRegistries.h b/Engine/Plugins/Runtime/Metasound/Source/MetasoundFrontend/Public/MetasoundFrontendRegistries.h
--- a/Engine/Plugins/Runtime/Metasound/Source/MetasoundFrontend/Public/MetasoundFrontendRegistries.h
+++ b/Engine/Plugins/Runtime/Metasound/Source/MetasoundFrontend/Public/MetasoundFrontendRegistries.h
@@ -56,7 +56,6 @@ namespace Metasound
 			FName NodeName;
 
 			// A hash generated from the input types and output types for this node.
-			// TODO: Write up some tests to ensure this is deterministic.
 			uint32 NodeHash = 0;
 
 			FORCEINLINE bool operator==(const FNodeRegistryKey& Other) const
diff --git a/Engine/Plugins/Runtime/Metasound/Source/MetasoundGraphCore/Public/MetasoundLiteral.h b/Engine/Plugins/Runtime/Metasound/Source/MetasoundGraphCore/Public/MetasoundLiteral.h
--- a/Engine/Plugins/Runtime/Metasound/Source/MetasoundGraphCore/Public/MetasoundLiteral.h
+++ b/Engine/Plugins/Runtime/Metasound/Source/MetasoundGraphCore/Public/MetasoundLiteral.h
@@ -143,9 +143,6 @@ namespace Metasound
 		const FString& InNodeName;
 		const FString& InVertexName;
 
-		// the operator settings that we will compile with
-		const FOperatorSettings& InSettings;
-
 		FDataTypeLiteralParam InitParam;
 	};
 
