Metasounds: Having IOperator::GetInputs() and GetOutputs() return by value instead of reference to avoid storing extra data on the operator objects.
#rnx

diff --git a/Engine/Plugins/Runtime/Metasound/Source/MetasoundEngine/Private/MetasoundWavePlayerNode.cpp b/Engine/Plugins/Runtime/Metasound/Source/MetasoundEngine/Private/MetasoundWavePlayerNode.cpp
--- a/Engine/Plugins/Runtime/Metasound/Source/MetasoundEngine/Private/MetasoundWavePlayerNode.cpp
+++ b/Engine/Plugins/Runtime/Metasound/Source/MetasoundEngine/Private/MetasoundWavePlayerNode.cpp
@@ -45,7 +45,6 @@ namespace Metasound
 			, AudioBuffer(FAudioBufferWriteRef::CreateNew(InSettings))
 		{
 			check(AudioBuffer->Num() == InSettings.GetNumFramesPerBlock());
-			OutputDataReferences.AddDataReadReference(TEXT("Audio"), FAudioBufferReadRef(AudioBuffer));
 		}
 		
 		FWavePlayerOperator(
@@ -62,16 +61,19 @@ namespace Metasound
 			, DecoderOutput(MoveTemp(InDecoderOutput))
 		{
 			check(AudioBuffer->Num() == InSettings.GetNumFramesPerBlock());
-			OutputDataReferences.AddDataReadReference(TEXT("Audio"), FAudioBufferReadRef(AudioBuffer));
 		}
 
-		virtual const FDataReferenceCollection& GetInputs() const override
+		virtual FDataReferenceCollection GetInputs() const override
 		{
+			FDataReferenceCollection InputDataReferences;
+			InputDataReferences.AddDataReadReference(TEXT("Audio"), FWaveAssetReadRef(Wave));
 			return InputDataReferences;
 		}
 
-		virtual const FDataReferenceCollection& GetOutputs() const override
+		virtual FDataReferenceCollection GetOutputs() const override
 		{
+			FDataReferenceCollection OutputDataReferences;
+			OutputDataReferences.AddDataReadReference(TEXT("Audio"), FAudioBufferReadRef(AudioBuffer));
 			return OutputDataReferences;
 		}
 
@@ -113,8 +115,6 @@ namespace Metasound
 
 		FWaveAssetReadRef Wave;
 		FAudioBufferWriteRef AudioBuffer;
-		FDataReferenceCollection InputDataReferences;
-		FDataReferenceCollection OutputDataReferences;
 
 		// Decoder/IO. 
 		Audio::ICodec::FDecoderPtr Decoder;
diff --git a/Engine/Plugins/Runtime/Metasound/Source/MetasoundFrontend/Public/MetasoundAutoConverterNode.h b/Engine/Plugins/Runtime/Metasound/Source/MetasoundFrontend/Public/MetasoundAutoConverterNode.h
--- a/Engine/Plugins/Runtime/Metasound/Source/MetasoundFrontend/Public/MetasoundAutoConverterNode.h
+++ b/Engine/Plugins/Runtime/Metasound/Source/MetasoundFrontend/Public/MetasoundAutoConverterNode.h
@@ -84,18 +84,20 @@ namespace Metasound
 					: FromData(InFromDataReference)
 					, ToData(InToDataReference) 
 				{
-					Outputs.AddDataReadReference<ToDataType>(GetOutputName(), ToData);
 				}
 
 				virtual ~FConverterOperator() {}
 
-				virtual const FDataReferenceCollection& GetInputs() const override
+				virtual FDataReferenceCollection GetInputs() const override
 				{
+					FDataReferenceCollection Inputs;
 					return Inputs;
 				}
 
-				virtual const FDataReferenceCollection& GetOutputs() const override
+				virtual FDataReferenceCollection GetOutputs() const override
 				{
+					FDataReferenceCollection Outputs;
+					Outputs.AddDataReadReference<ToDataType>(GetOutputName(), ToData);
 					return Outputs;
 				}
 
@@ -107,9 +109,6 @@ namespace Metasound
 			private:
 				TDataReadReference<FromDataType> FromData;
 				TDataWriteReference<ToDataType> ToData;
-
-				FDataReferenceCollection Inputs;
-				FDataReferenceCollection Outputs;
 		};
 
 		/** FConverterOperatorFactory creates an operator which converts from 
diff --git a/Engine/Plugins/Runtime/Metasound/Source/MetasoundFrontend/Public/MetasoundInputNode.h b/Engine/Plugins/Runtime/Metasound/Source/MetasoundFrontend/Public/MetasoundInputNode.h
--- a/Engine/Plugins/Runtime/Metasound/Source/MetasoundFrontend/Public/MetasoundInputNode.h
+++ b/Engine/Plugins/Runtime/Metasound/Source/MetasoundFrontend/Public/MetasoundInputNode.h
@@ -11,7 +11,7 @@
 
 namespace Metasound
 {
-	/** TInputOperator supplies a writeable input and a readable output. */
+	/** TInputOperator supplies a writable input and a readable output. */
 	template<typename DataType>
 	class TInputOperator : public IOperator
 	{
@@ -26,12 +26,12 @@ namespace Metasound
 
 			virtual ~TInputOperator() {}
 
-			virtual const FDataReferenceCollection& GetInputs() const override
+			virtual FDataReferenceCollection GetInputs() const override
 			{
 				return Inputs;
 			}
 
-			virtual const FDataReferenceCollection& GetOutputs() const override
+			virtual FDataReferenceCollection GetOutputs() const override
 			{
 				return Outputs;
 			}
@@ -53,7 +53,7 @@ namespace Metasound
 		public:
 			// If the data type is parsable from a literal type, then the data type 
 			// can be registered as an input type with the frontend.  To make a 
-			// DataType registerable, either create a constructor for the data type
+			// DataType registrable, either create a constructor for the data type
 			// which accepts the one of the supported literal types with an optional 
 			// FOperatorSettings argument, or create a default constructor, or specialize
 			// this factory with an implementation for that specific data type.
diff --git a/Engine/Plugins/Runtime/Metasound/Source/MetasoundFrontend/Public/MetasoundOutputNode.h b/Engine/Plugins/Runtime/Metasound/Source/MetasoundFrontend/Public/MetasoundOutputNode.h
--- a/Engine/Plugins/Runtime/Metasound/Source/MetasoundFrontend/Public/MetasoundOutputNode.h
+++ b/Engine/Plugins/Runtime/Metasound/Source/MetasoundFrontend/Public/MetasoundOutputNode.h
@@ -26,12 +26,12 @@ namespace Metasound
 
 				virtual ~FOutputOperator() {}
 
-				virtual const FDataReferenceCollection& GetInputs() const override
+				virtual FDataReferenceCollection GetInputs() const override
 				{
-					return Inputs;
+					return {};
 				}
 
-				virtual const FDataReferenceCollection& GetOutputs() const override
+				virtual FDataReferenceCollection GetOutputs() const override
 				{
 					return Outputs;
 				}
@@ -42,7 +42,6 @@ namespace Metasound
 				}
 
 			private:
-				FDataReferenceCollection Inputs;
 				FDataReferenceCollection Outputs;
 		};
 
diff --git a/Engine/Plugins/Runtime/Metasound/Source/MetasoundFrontend/Public/MetasoundReceiveNode.h b/Engine/Plugins/Runtime/Metasound/Source/MetasoundFrontend/Public/MetasoundReceiveNode.h
--- a/Engine/Plugins/Runtime/Metasound/Source/MetasoundFrontend/Public/MetasoundReceiveNode.h
+++ b/Engine/Plugins/Runtime/Metasound/Source/MetasoundFrontend/Public/MetasoundReceiveNode.h
@@ -83,18 +83,19 @@ namespace Metasound
 					, CachedReceiverParams({InOperatorSettings})
 					, Receiver(FDataTransmissionCenter::Get().RegisterNewReceiver<TDataType>(CachedSendAddress, CachedReceiverParams))
 				{
-					Outputs.AddDataReadReference<TDataType>(GetSendOutputName(), TDataReadReference<TDataType>(OutputData));
 				}
 
 				virtual ~TReceiverOperator() {}
 
-				virtual const FDataReferenceCollection& GetInputs() const override
+				virtual FDataReferenceCollection GetInputs() const override
 				{
-					return Inputs;
+					return {};
 				}
 
-				virtual const FDataReferenceCollection& GetOutputs() const override
+				virtual FDataReferenceCollection GetOutputs() const override
 				{
+					FDataReferenceCollection Outputs;
+					Outputs.AddDataReadReference<TDataType>(GetSendOutputName(), TDataReadReference<TDataType>(OutputData));
 					return Outputs;
 				}
 
@@ -120,9 +121,6 @@ namespace Metasound
 				FReceiverInitParams CachedReceiverParams;
 
 				TReceiverPtr<TDataType> Receiver;
-
-				FDataReferenceCollection Inputs;
-				FDataReferenceCollection Outputs;
 		};
 
 		class FReceiverOperatorFactory : public IOperatorFactory
diff --git a/Engine/Plugins/Runtime/Metasound/Source/MetasoundFrontend/Public/MetasoundSendNode.h b/Engine/Plugins/Runtime/Metasound/Source/MetasoundFrontend/Public/MetasoundSendNode.h
--- a/Engine/Plugins/Runtime/Metasound/Source/MetasoundFrontend/Public/MetasoundSendNode.h
+++ b/Engine/Plugins/Runtime/Metasound/Source/MetasoundFrontend/Public/MetasoundSendNode.h
@@ -82,20 +82,21 @@ namespace Metasound
 					, CachedSenderParams({InOperatorSettings, 0.0f})
 					, Sender(FDataTransmissionCenter::Get().RegisterNewSend<TDataType>(CachedSendAddress, CachedSenderParams))
 				{
-					Inputs.AddDataReadReference<FSendAddress>(GetAddressInputName(), SendAddress);
-					Inputs.AddDataReadReference<TDataType>(GetSendInputName(), TDataReadReference<TDataType>(InputData));
 				}
 
 				virtual ~TSendOperator() {}
 
-				virtual const FDataReferenceCollection& GetInputs() const override
+				virtual FDataReferenceCollection GetInputs() const override
 				{
+					FDataReferenceCollection Inputs;
+					Inputs.AddDataReadReference<FSendAddress>(GetAddressInputName(), SendAddress);
+					Inputs.AddDataReadReference<TDataType>(GetSendInputName(), TDataReadReference<TDataType>(InputData));
 					return Inputs;
 				}
 
-				virtual const FDataReferenceCollection& GetOutputs() const override
+				virtual FDataReferenceCollection GetOutputs() const override
 				{
-					return Outputs;
+					return {};
 				}
 
 				void Execute()
@@ -117,9 +118,6 @@ namespace Metasound
 				FSenderInitParams CachedSenderParams;
 
 				TSenderPtr<TDataType> Sender;
-
-				FDataReferenceCollection Inputs;
-				FDataReferenceCollection Outputs;
 		};
 
 		class FSendOperatorFactory : public IOperatorFactory
diff --git a/Engine/Plugins/Runtime/Metasound/Source/MetasoundGraphCore/Private/MetasoundGraphOperator.cpp b/Engine/Plugins/Runtime/Metasound/Source/MetasoundGraphCore/Private/MetasoundGraphOperator.cpp
--- a/Engine/Plugins/Runtime/Metasound/Source/MetasoundGraphCore/Private/MetasoundGraphOperator.cpp
+++ b/Engine/Plugins/Runtime/Metasound/Source/MetasoundGraphCore/Private/MetasoundGraphOperator.cpp
@@ -36,12 +36,12 @@ namespace Metasound
 		Outputs = InCollection;
 	}
 
-	const FDataReferenceCollection& FGraphOperator::GetInputs() const
+	FDataReferenceCollection FGraphOperator::GetInputs() const
 	{
 		return Inputs;
 	}
 
-	const FDataReferenceCollection& FGraphOperator::GetOutputs() const
+	FDataReferenceCollection FGraphOperator::GetOutputs() const
 	{
 		return Outputs;
 	}
diff --git a/Engine/Plugins/Runtime/Metasound/Source/MetasoundGraphCore/Public/MetasoundExecutableOperator.h b/Engine/Plugins/Runtime/Metasound/Source/MetasoundGraphCore/Public/MetasoundExecutableOperator.h
--- a/Engine/Plugins/Runtime/Metasound/Source/MetasoundGraphCore/Public/MetasoundExecutableOperator.h
+++ b/Engine/Plugins/Runtime/Metasound/Source/MetasoundGraphCore/Public/MetasoundExecutableOperator.h
@@ -114,7 +114,7 @@ namespace Metasound
 				return Operator.IsValid();
 			}
 
-			virtual const FDataReferenceCollection& GetInputs() const override
+			virtual FDataReferenceCollection GetInputs() const override
 			{
 				static FDataReferenceCollection EmptyCollection;
 
@@ -126,7 +126,7 @@ namespace Metasound
 				return EmptyCollection;
 			}
 
-			virtual const FDataReferenceCollection& GetOutputs() const override
+			virtual FDataReferenceCollection GetOutputs() const override
 			{
 				static FDataReferenceCollection EmptyCollection;
 
diff --git a/Engine/Plugins/Runtime/Metasound/Source/MetasoundGraphCore/Public/MetasoundGraphOperator.h b/Engine/Plugins/Runtime/Metasound/Source/MetasoundGraphCore/Public/MetasoundGraphOperator.h
--- a/Engine/Plugins/Runtime/Metasound/Source/MetasoundGraphCore/Public/MetasoundGraphOperator.h
+++ b/Engine/Plugins/Runtime/Metasound/Source/MetasoundGraphCore/Public/MetasoundGraphOperator.h
@@ -26,9 +26,9 @@ namespace Metasound
 
 			void SetOutputs(const FDataReferenceCollection& InCollection);
 
-			virtual const FDataReferenceCollection& GetInputs() const override;
+			virtual FDataReferenceCollection GetInputs() const override;
 
-			virtual const FDataReferenceCollection& GetOutputs() const override;
+			virtual FDataReferenceCollection GetOutputs() const override;
 
 			virtual FExecuteFunction GetExecuteFunction() override;
 
diff --git a/Engine/Plugins/Runtime/Metasound/Source/MetasoundGraphCore/Public/MetasoundOperatorInterface.h b/Engine/Plugins/Runtime/Metasound/Source/MetasoundGraphCore/Public/MetasoundOperatorInterface.h
--- a/Engine/Plugins/Runtime/Metasound/Source/MetasoundGraphCore/Public/MetasoundOperatorInterface.h
+++ b/Engine/Plugins/Runtime/Metasound/Source/MetasoundGraphCore/Public/MetasoundOperatorInterface.h
@@ -33,14 +33,14 @@ namespace Metasound
 		 * exceptions are input operators and graph operators which need to interface with external
 		 * systems even after the operator is created.
 		 */
-		virtual const FDataReferenceCollection& GetInputs() const = 0;
+		virtual FDataReferenceCollection GetInputs() const = 0;
 
 		/** Return the output parameters associated with this IOperator.
 		 *
 		 * Implementations of IOperator should return a collection of their output parameters
 		 * which other nodes can read.
 		 */
-		virtual const FDataReferenceCollection& GetOutputs() const = 0;
+		virtual FDataReferenceCollection GetOutputs() const = 0;
 
 		/** Return the execution function to call during graph execution.
 		 *
diff --git a/Engine/Plugins/Runtime/Metasound/Source/MetasoundStandardNodes/Private/MetasoundADSRNode.cpp b/Engine/Plugins/Runtime/Metasound/Source/MetasoundStandardNodes/Private/MetasoundADSRNode.cpp
--- a/Engine/Plugins/Runtime/Metasound/Source/MetasoundStandardNodes/Private/MetasoundADSRNode.cpp
+++ b/Engine/Plugins/Runtime/Metasound/Source/MetasoundStandardNodes/Private/MetasoundADSRNode.cpp
@@ -33,34 +33,30 @@ namespace Metasound
 
 			FADSROperator(const FOperatorSettings& InSettings, const FBopReadRef& InBop, const FADSRReferences& InADSRData);
 
-			virtual const FDataReferenceCollection& GetInputs() const override;
-			virtual const FDataReferenceCollection& GetOutputs() const override;
+			virtual FDataReferenceCollection GetInputs() const override;
+			virtual FDataReferenceCollection GetOutputs() const override;
 			void Execute();
 
 		private:
 
 			void GenerateEnvelope(int32 InStartFrame, int32 InEndFrame);
 
-			const FOperatorSettings OperatorSettings;
-
 			// TODO: write envelope gen for metasound more suited to this processing structure. 
 			Audio::FEnvelope Envelope;
 
 			FBopReadRef Bop;
 			bool bNotReleased;
+			float SampleRate;
 			FSampleTime TimeTillRelease;
 			FSampleTime TimePerBlock;
 			FADSRReferences ADSRReferences;
 			FAudioBufferWriteRef EnvelopeBuffer;
-
-			FDataReferenceCollection OutputDataReferences;
-			FDataReferenceCollection InputDataReferences;
 	};
 
 	FADSROperator::FADSROperator(const FOperatorSettings& InSettings, const FBopReadRef& InBop, const FADSRReferences& InADSRData)
-	:	OperatorSettings(InSettings)
-	,	Bop(InBop)
+	:	Bop(InBop)
 	,	bNotReleased(false)
+	,	SampleRate(InSettings.GetSampleRate())
 	,	TimeTillRelease(0, InSettings.GetSampleRate())
 	,	TimePerBlock(InSettings.GetNumFramesPerBlock(), InSettings.GetSampleRate())
 	,	ADSRReferences(InADSRData)
@@ -68,18 +64,25 @@ namespace Metasound
 	{
 		check(EnvelopeBuffer->Num() == InSettings.GetNumFramesPerBlock());
 
-		OutputDataReferences.AddDataReadReference(TEXT("Envelope"), FAudioBufferReadRef(EnvelopeBuffer));
 
-		Envelope.Init(OperatorSettings.GetSampleRate());
+		Envelope.Init(SampleRate);
 	}
 
-	const FDataReferenceCollection& FADSROperator::GetInputs() const
+	FDataReferenceCollection FADSROperator::GetInputs() const
 	{
+		FDataReferenceCollection InputDataReferences;
+		InputDataReferences.AddDataReadReference(TEXT("Bop"), FBopReadRef(Bop));
+		InputDataReferences.AddDataReadReference(TEXT("Attack"), FFloatTimeReadRef(ADSRReferences.Attack));
+		InputDataReferences.AddDataReadReference(TEXT("Decay"), FFloatTimeReadRef(ADSRReferences.Decay));
+		InputDataReferences.AddDataReadReference(TEXT("Sustain"), FFloatTimeReadRef(ADSRReferences.Sustain));
+		InputDataReferences.AddDataReadReference(TEXT("Release"), FFloatTimeReadRef(ADSRReferences.Release));
 		return InputDataReferences;
 	}
 
-	const FDataReferenceCollection& FADSROperator::GetOutputs() const
+	FDataReferenceCollection FADSROperator::GetOutputs() const
 	{
+		FDataReferenceCollection OutputDataReferences;
+		OutputDataReferences.AddDataReadReference(TEXT("Envelope"), FAudioBufferReadRef(EnvelopeBuffer));
 		return OutputDataReferences;
 	}
 
diff --git a/Engine/Plugins/Runtime/Metasound/Source/MetasoundStandardNodes/Private/MetasoundAudioMultiplyNode.cpp b/Engine/Plugins/Runtime/Metasound/Source/MetasoundStandardNodes/Private/MetasoundAudioMultiplyNode.cpp
--- a/Engine/Plugins/Runtime/Metasound/Source/MetasoundStandardNodes/Private/MetasoundAudioMultiplyNode.cpp
+++ b/Engine/Plugins/Runtime/Metasound/Source/MetasoundStandardNodes/Private/MetasoundAudioMultiplyNode.cpp
@@ -22,25 +22,27 @@ namespace Metasound
 			static TUniquePtr<IOperator> CreateOperator(const FCreateOperatorParams& InParams, FBuildErrorArray& OutErrors);
 
 			FAudioMultiplyOperator(const FOperatorSettings& InSettings, const FAudioBufferReadRef& InBuffer1, const FAudioBufferReadRef& InBuffer2)
-			:	OperatorSettings(InSettings)
-			,	InputBuffer1(InBuffer1)
+			:	InputBuffer1(InBuffer1)
 			,	InputBuffer2(InBuffer2)
 			,	OutputBuffer(FAudioBufferWriteRef::CreateNew(InSettings))
+			,	BlockSize(InSettings.GetNumFramesPerBlock())
 			{
-				check(OutputBuffer->Num() == InSettings.GetNumFramesPerBlock());
-				check(InputBuffer1->Num() == InSettings.GetNumFramesPerBlock());
-				check(InputBuffer2->Num() == InSettings.GetNumFramesPerBlock());
+				check(OutputBuffer->Num() == BlockSize);
+				check(InputBuffer1->Num() == BlockSize);
+				check(InputBuffer2->Num() == BlockSize);
 
-				OutputDataReferences.AddDataReadReference(TEXT("Audio"), FAudioBufferReadRef(OutputBuffer));
 			}
 
-			virtual const FDataReferenceCollection& GetInputs() const override
+			virtual FDataReferenceCollection GetInputs() const override
 			{
+				FDataReferenceCollection InputDataReferences;
 				return InputDataReferences;
 			}
 
-			virtual const FDataReferenceCollection& GetOutputs() const override
+			virtual FDataReferenceCollection GetOutputs() const override
 			{
+				FDataReferenceCollection OutputDataReferences;
+				OutputDataReferences.AddDataReadReference(TEXT("Audio"), FAudioBufferReadRef(OutputBuffer));
 				return OutputDataReferences;
 			}
 
@@ -48,20 +50,17 @@ namespace Metasound
 			{
 				using namespace Audio;
 
-				// TODO: add buffer vector op to multipy two buffers (not in place). 
+				// TODO: add buffer vector op to multiply two buffers (not in place). 
 				// TODO: what to do about RESTRICT? This could be something for the builder to keep in mind.
-				FMemory::Memcpy(OutputBuffer->GetData(), InputBuffer1->GetData(), sizeof(float) * OperatorSettings.GetNumFramesPerBlock());
+				FMemory::Memcpy(OutputBuffer->GetData(), InputBuffer1->GetData(), sizeof(float) * BlockSize);
 				MultiplyBuffersInPlace(*InputBuffer2, *OutputBuffer);
 			}
 
 		private:
-			const FOperatorSettings OperatorSettings;
-
 			FAudioBufferReadRef InputBuffer1;
 			FAudioBufferReadRef InputBuffer2;
 			FAudioBufferWriteRef OutputBuffer;
-			FDataReferenceCollection OutputDataReferences;
-			FDataReferenceCollection InputDataReferences;
+			int32 BlockSize;
 	};
 
 	FAudioMultiplyNode::FAudioMultiplyNode(const FString& InInstanceName)
diff --git a/Engine/Plugins/Runtime/Metasound/Source/MetasoundStandardNodes/Private/MetasoundOscNode.cpp b/Engine/Plugins/Runtime/Metasound/Source/MetasoundStandardNodes/Private/MetasoundOscNode.cpp
--- a/Engine/Plugins/Runtime/Metasound/Source/MetasoundStandardNodes/Private/MetasoundOscNode.cpp
+++ b/Engine/Plugins/Runtime/Metasound/Source/MetasoundStandardNodes/Private/MetasoundOscNode.cpp
@@ -24,52 +24,54 @@ namespace Metasound
 
 			FOscOperator(const FOperatorSettings& InSettings, const FFrequencyReadRef& InFrequency);
 
-			virtual const FDataReferenceCollection& GetInputs() const override;
+			virtual FDataReferenceCollection GetInputs() const override;
 
-			virtual const FDataReferenceCollection& GetOutputs() const override;
+			virtual FDataReferenceCollection GetOutputs() const override;
 
 			void Execute();
 
 		private:
-			const FOperatorSettings OperatorSettings;
 			const float TwoPi;
 			float Phase;
+			float SampleRate;
+			int32 BlockSize;
 
 			FFrequencyReadRef Frequency;
 			FAudioBufferWriteRef AudioBuffer;
 
-			FDataReferenceCollection InputDataReferences;
-			FDataReferenceCollection OutputDataReferences;
 	};
 
 	FOscOperator::FOscOperator(const FOperatorSettings& InSettings, const FFrequencyReadRef& InFrequency)
-	:	OperatorSettings(InSettings)
-	,	TwoPi(2.f * PI)
+	:	TwoPi(2.f * PI)
 	,	Phase(0.f)
+	,	SampleRate(InSettings.GetSampleRate())
+	,	BlockSize(InSettings.GetNumFramesPerBlock())
 	,	Frequency(InFrequency)
 	,	AudioBuffer(FAudioBufferWriteRef::CreateNew(InSettings))
 	{
 		check(AudioBuffer->Num() == InSettings.GetNumFramesPerBlock());
-
-		OutputDataReferences.AddDataReadReference(TEXT("Audio"), FAudioBufferReadRef(AudioBuffer));
 	}
 
-	const FDataReferenceCollection& FOscOperator::GetInputs() const
+	FDataReferenceCollection FOscOperator::GetInputs() const
 	{
+		FDataReferenceCollection InputDataReferences;
+		InputDataReferences.AddDataReadReference(TEXT("Audio"), FFrequencyReadRef(Frequency));
 		return InputDataReferences;
 	}
 
-	const FDataReferenceCollection& FOscOperator::GetOutputs() const
+	FDataReferenceCollection FOscOperator::GetOutputs() const
 	{
+		FDataReferenceCollection OutputDataReferences;
+		OutputDataReferences.AddDataReadReference(TEXT("Audio"), FAudioBufferReadRef(AudioBuffer));
 		return OutputDataReferences;
 	}
 
 	void FOscOperator::Execute()
 	{
-		const float PhaseDelta = Frequency->GetRadiansPerSample(OperatorSettings.GetSampleRate());
+		const float PhaseDelta = Frequency->GetRadiansPerSample(SampleRate);
 		float* Data = AudioBuffer->GetData();
 
-		for (int32 i = 0; i < OperatorSettings.GetNumFramesPerBlock(); i++)
+		for (int32 i = 0; i < BlockSize; i++)
 		{
 			Data[i] = FMath::Sin(Phase);
 			Phase += PhaseDelta;
diff --git a/Engine/Plugins/Runtime/Metasound/Source/MetasoundStandardNodes/Private/MetasoundPeriodicBopNode.cpp b/Engine/Plugins/Runtime/Metasound/Source/MetasoundStandardNodes/Private/MetasoundPeriodicBopNode.cpp
--- a/Engine/Plugins/Runtime/Metasound/Source/MetasoundStandardNodes/Private/MetasoundPeriodicBopNode.cpp
+++ b/Engine/Plugins/Runtime/Metasound/Source/MetasoundStandardNodes/Private/MetasoundPeriodicBopNode.cpp
@@ -32,42 +32,45 @@ namespace Metasound
 
 			FPeriodicBopOperator(const FOperatorSettings& InSettings, const FFloatTimeReadRef& InPeriod);
 
-			virtual const FDataReferenceCollection& GetInputs() const override;
+			virtual FDataReferenceCollection GetInputs() const override;
 
-			virtual const FDataReferenceCollection& GetOutputs() const override;
+			virtual FDataReferenceCollection GetOutputs() const override;
 
 			void Execute();
 
 		private:
 
-			FOperatorSettings OperatorSettings;
 			FBopWriteRef Bop;
 			FFloatTimeReadRef Period;
 
 			float ExecuteDurationInSamples;
+			float SampleRate;
 			float SampleCountdown;
 
-			FDataReferenceCollection InputDataReferences;
-			FDataReferenceCollection OutputDataReferences;
 	};
 
 	FPeriodicBopOperator::FPeriodicBopOperator(const FOperatorSettings& InSettings, const FFloatTimeReadRef& InPeriod)
-	:	OperatorSettings(InSettings)
-	,	Bop(FBopWriteRef::CreateNew(InSettings))
+	:	Bop(FBopWriteRef::CreateNew(InSettings))
 	,	Period(InPeriod)
 	,	ExecuteDurationInSamples(InSettings.GetNumFramesPerBlock())
+	,	SampleRate(InSettings.GetSampleRate())
 	,	SampleCountdown(0.f)
 	{
-		OutputDataReferences.AddDataReadReference(TEXT("Bop"), FBopReadRef(Bop));
 	}
 
-	const FDataReferenceCollection& FPeriodicBopOperator::GetInputs() const
+	FDataReferenceCollection FPeriodicBopOperator::GetInputs() const
 	{
+		FDataReferenceCollection InputDataReferences;
+		InputDataReferences.AddDataReadReference(TEXT("Period"), FFloatTimeReadRef(Period));
+
 		return InputDataReferences;
 	}
 
-	const FDataReferenceCollection& FPeriodicBopOperator::GetOutputs() const
+	FDataReferenceCollection FPeriodicBopOperator::GetOutputs() const
 	{
+		FDataReferenceCollection OutputDataReferences;
+		OutputDataReferences.AddDataReadReference(TEXT("Bop"), FBopReadRef(Bop));
+
 		return OutputDataReferences;
 	}
 
@@ -76,7 +79,7 @@ namespace Metasound
 		// Advance internal counter to get rid of old bops.
 		Bop->AdvanceBlock();
 
-		float PeriodInSamples = FMath::Max(Period->GetSeconds(), MinimumPeriodSeconds) * OperatorSettings.GetSampleRate();
+		float PeriodInSamples = FMath::Max(Period->GetSeconds(), MinimumPeriodSeconds) * SampleRate;
 
 		PeriodInSamples = FMath::Max(PeriodInSamples, MinimumPeriodSamples);
 
