Multichannel Metasound Source
- Moving Frontend UObject registration logic to FMetasoundUObjectRegistry
- Allow FMetasoundAssetBase to support multiple FMetasoundArchetypes
- Adding graph synchronization between UEdGraph and FMetasoundDocument
- Added TAssetPtr<> for tracking validity of objects that cannot be wrapped in a smart pointer.
- Added support for mono and stereo to UMetasoundSource
- Added tools for comparing FMetasoundARchetypes.
- Added LogMetasound

diff --git a/Engine/Plugins/Runtime/Metasound/Source/MetasoundEditor/Private/MetasoundDetailCustomization.cpp b/Engine/Plugins/Runtime/Metasound/Source/MetasoundEditor/Private/MetasoundDetailCustomization.cpp
--- a/Engine/Plugins/Runtime/Metasound/Source/MetasoundEditor/Private/MetasoundDetailCustomization.cpp
+++ b/Engine/Plugins/Runtime/Metasound/Source/MetasoundEditor/Private/MetasoundDetailCustomization.cpp
@@ -8,9 +8,9 @@
 #include "DetailLayoutBuilder.h"
 #include "DetailWidgetRow.h"
 #include "IDetailGroup.h"
-#include "Metasound.h"
 #include "MetasoundAssetBase.h"
 #include "MetasoundFrontend.h"
+#include "MetasoundUObjectRegistry.h"
 #include "PropertyEditorDelegates.h"
 #include "PropertyHandle.h"
 #include "SlateCore/Public/Styling/SlateColor.h"
@@ -180,7 +180,7 @@ namespace Metasound
 			DetailLayout.GetObjectsBeingCustomized(Objects);
 			if (Objects.Num() > 0)
 			{
-				FMetasoundAssetBase* MetasoundAsset = Frontend::GetObjectAsAssetBase(Objects[0].Get());
+				FMetasoundAssetBase* MetasoundAsset = IMetasoundUObjectRegistry::Get().GetObjectAsAssetBase(Objects[0].Get());
 				check(MetasoundAsset);
 
 				Frontend::FGraphHandle GraphHandle = MetasoundAsset->GetRootGraphHandle();
diff --git a/Engine/Plugins/Runtime/Metasound/Source/MetasoundEditor/Private/MetasoundEditor.cpp b/Engine/Plugins/Runtime/Metasound/Source/MetasoundEditor/Private/MetasoundEditor.cpp
--- a/Engine/Plugins/Runtime/Metasound/Source/MetasoundEditor/Private/MetasoundEditor.cpp
+++ b/Engine/Plugins/Runtime/Metasound/Source/MetasoundEditor/Private/MetasoundEditor.cpp
@@ -16,12 +16,14 @@
 #include "HAL/PlatformApplicationMisc.h"
 #include "IDetailsView.h"
 #include "Kismet2/BlueprintEditorUtils.h"
+#include "Metasound.h"
 #include "MetasoundEditorCommands.h"
 #include "MetasoundEditorGraph.h"
 #include "MetasoundEditorGraphBuilder.h"
 #include "MetasoundEditorGraphNode.h"
 #include "MetasoundEditorGraphSchema.h"
 #include "MetasoundEditorTabFactory.h"
+#include "MetasoundUObjectRegistry.h"
 #include "Misc/Attribute.h"
 #include "Modules/ModuleManager.h"
 #include "PropertyEditorModule.h"
@@ -34,6 +36,7 @@
 
 #define LOCTEXT_NAMESPACE "MetasoundEditor"
 
+
 namespace Metasound
 {
 	namespace Editor
@@ -96,7 +99,7 @@ namespace Metasound
 
 		Frontend::FGraphHandle FEditor::InitMetasound(UObject& InMetasound)
 		{
-			FMetasoundAssetBase* MetasoundAsset = Frontend::GetObjectAsAssetBase(&InMetasound);
+			FMetasoundAssetBase* MetasoundAsset = IMetasoundUObjectRegistry::Get().GetObjectAsAssetBase(&InMetasound);
 			check(MetasoundAsset);
 
 			if (!MetasoundAsset->GetGraph())
@@ -111,10 +114,40 @@ namespace Metasound
 			return MetasoundAsset->GetRootGraphHandle();
 		}
 
+		bool FEditor::RebuildGraph() const
+		{
+			UObject* MetasoundObj = GetMetasoundObject();
+			if (nullptr != MetasoundObj)
+			{
+				if (IMetasoundUObjectRegistry::Get().IsRegisteredClass(MetasoundObj))
+				{
+					FGraphBuilder::RebuildGraph(*MetasoundObj);
+					return true;
+				}
+			}
+
+			return false;
+		}
+
+		bool FEditor::SynchronizeGraph() const
+		{
+			UObject* MetasoundObj = GetMetasoundObject();
+			if (nullptr != MetasoundObj)
+			{
+				if (IMetasoundUObjectRegistry::Get().IsRegisteredClass(MetasoundObj))
+				{
+					FGraphBuilder::SynchronizeGraph(*MetasoundObj);
+					return true;
+				}
+			}
+
+			return false;
+		}
+
 		void FEditor::InitMetasoundEditor(const EToolkitMode::Type Mode, const TSharedPtr<IToolkitHost>& InitToolkitHost, UObject* ObjectToEdit)
 		{
 			check(ObjectToEdit);
-			checkf(Frontend::IsObjectAMetasoundArchetype(ObjectToEdit), TEXT("Object passed in was not registered as a valid metasound archetype!"));
+			checkf(IMetasoundUObjectRegistry::Get().IsRegisteredClass(ObjectToEdit), TEXT("Object passed in was not registered as a valid metasound archetype!"));
 			
 			// Support undo/redo
 			ObjectToEdit->SetFlags(RF_Transactional);
@@ -233,9 +266,14 @@ namespace Metasound
 		}
 
 		void FEditor::NotifyPostChange(const FPropertyChangedEvent& PropertyChangedEvent, FProperty* PropertyThatChanged)
-		{
+		{	
 			if (MetasoundGraphEditor.IsValid() && PropertyChangedEvent.ChangeType != EPropertyChangeType::Interactive)
 			{
+				// If a property change event occurs outside of the metasound UEdGraph and results in the metasound document changing, 
+				// then the document and the UEdGraph need to be synchronized. There may be a better trigger for this call to reduce
+				// the number of times the graph is synchronized. 
+				SynchronizeGraph();
+
 				MetasoundGraphEditor->NotifyGraphChanged();
 			}
 		}
@@ -353,7 +391,7 @@ namespace Metasound
 
 		void FEditor::Import()
 		{
-			FMetasoundAssetBase* MetasoundAsset = Frontend::GetObjectAsAssetBase(Metasound);
+			FMetasoundAssetBase* MetasoundAsset = IMetasoundUObjectRegistry::Get().GetObjectAsAssetBase(Metasound);
 			if (MetasoundAsset)
 			{
 				// TODO: Prompt OFD and provide path from user
@@ -366,18 +404,36 @@ namespace Metasound
 
 				if (Frontend::ImportJSONAssetToMetasound(InputPath, MetasoundDoc))
 				{
-					Frontend::GetObjectForDocument(MetasoundDoc, OutputPath);
+					TArray<UClass*> ImportClasses = IMetasoundUObjectRegistry::Get().GetUClassesForArchetype(MetasoundDoc.Archetype.ArchetypeName);
+
+					if (ImportClasses.Num() < 1)
+					{
+						UE_LOG(LogTemp, Warning, TEXT("Cannot create UObject from Metasound document. No UClass supports archetype \"%s\""), *MetasoundDoc.Archetype.ArchetypeName.ToString());
+					}
+					else
+					{
+						if (ImportClasses.Num() > 1)
+						{
+							for (UClass* Cls : ImportClasses)
+							{
+								// TODO: could do a modal dialog to give user choice of import type.
+								UE_LOG(LogTemp, Warning, TEXT("Duplicate UClass support archetype \"%s\" with UClass \"%s\""), *MetasoundDoc.Archetype.ArchetypeName.ToString(), *Cls->GetName());
+							}
+						}
+
+						IMetasoundUObjectRegistry::Get().NewObject(ImportClasses[0], MetasoundDoc, OutputPath);
+					}
 				}
 				else
 				{
-					UE_LOG(LogTemp, Warning, TEXT("Couldn't import Metasound at path: %s"), *InputPath);
+					UE_LOG(LogTemp, Warning, TEXT("Could not import Metasound at path: %s"), *InputPath);
 				}
 			}
 		}
 
 		void FEditor::Export()
 		{
-			FMetasoundAssetBase* InMetasoundAsset = Frontend::GetObjectAsAssetBase(Metasound);
+			FMetasoundAssetBase* InMetasoundAsset = IMetasoundUObjectRegistry::Get().GetObjectAsAssetBase(Metasound);
 			check(InMetasoundAsset);
 
 			if (!InMetasoundAsset)
@@ -584,7 +640,7 @@ namespace Metasound
 			InEvents.OnTextCommitted = FOnNodeTextCommitted::CreateSP(this, &FEditor::OnNodeTitleCommitted);
 			InEvents.OnNodeDoubleClicked = FSingleNodeEvent::CreateSP(this, &FEditor::PlaySingleNode);
 
-			FMetasoundAssetBase* MetasoundAsset = Frontend::GetObjectAsAssetBase(Metasound);
+			FMetasoundAssetBase* MetasoundAsset = IMetasoundUObjectRegistry::Get().GetObjectAsAssetBase(Metasound);
 			check(MetasoundAsset);
 
 			return SNew(SGraphEditor)
@@ -645,7 +701,7 @@ namespace Metasound
 
 			UObject* ParentMetasoundObject = Graph->ParentMetasound;
 			check(ParentMetasoundObject);
-			if (FMetasoundAssetBase* MetasoundAsset = Frontend::GetObjectAsAssetBase(ParentMetasoundObject))
+			if (FMetasoundAssetBase* MetasoundAsset = IMetasoundUObjectRegistry::Get().GetObjectAsAssetBase(ParentMetasoundObject))
 			{
 				Frontend::FGraphHandle GraphHandle = MetasoundAsset->GetRootGraphHandle();
 				for (FGraphPanelSelectionSet::TConstIterator NodeIt(SelectedNodes); NodeIt; ++NodeIt)
@@ -742,7 +798,7 @@ namespace Metasound
 				}
 			}
 
-			FMetasoundAssetBase* MetasoundAsset = Metasound::Frontend::GetObjectAsAssetBase(Metasound);
+			FMetasoundAssetBase* MetasoundAsset = Metasound::IMetasoundUObjectRegistry::Get().GetObjectAsAssetBase(Metasound);
 			check(MetasoundAsset);
 			FGraphHandle GraphHandle = MetasoundAsset->GetRootGraphHandle();
 
@@ -829,7 +885,7 @@ namespace Metasound
 // 				Location = MetasoundGraphEditor->GetPasteLocation();
 // 			}
 // 
-// 			FMetasoundAssetBase* MetasoundAsset = Frontend::GetObjectAsAssetBase(Metasound);
+// 			FMetasoundAssetBase* MetasoundAsset = IMetasoundUObjectRegistry::Get().GetObjectAsAssetBase(Metasound);
 // 			check(MetasoundAsset);
 // 
 // 			UEdGraph* Graph = MetasoundAsset->GetGraph();
@@ -895,7 +951,7 @@ namespace Metasound
 // 			FString ClipboardContent;
 // 			FPlatformApplicationMisc::ClipboardPaste(ClipboardContent);
 // 
-// 			FMetasoundAssetBase* AssetBase = Frontend::GetObjectAsAssetBase(Metasound);
+// 			FMetasoundAssetBase* AssetBase = IMetasoundUObjectRegistry::Get().GetObjectAsAssetBase(Metasound);
 // 			check(AssetBase);
 // 
 // 			PastedDocument = FMetasoundDocument();
diff --git a/Engine/Plugins/Runtime/Metasound/Source/MetasoundEditor/Private/MetasoundEditor.h b/Engine/Plugins/Runtime/Metasound/Source/MetasoundEditor/Private/MetasoundEditor.h
--- a/Engine/Plugins/Runtime/Metasound/Source/MetasoundEditor/Private/MetasoundEditor.h
+++ b/Engine/Plugins/Runtime/Metasound/Source/MetasoundEditor/Private/MetasoundEditor.h
@@ -47,6 +47,8 @@ namespace Metasound
 			virtual UObject* GetMetasoundObject() const override;
 			virtual void SetSelection(const TArray<UObject*>& SelectedObjects) override;
 			virtual bool GetBoundsForSelectedNodes(FSlateRect& Rect, float Padding) override;
+			bool RebuildGraph() const override;
+			bool SynchronizeGraph() const override;
 
 			/** IToolkit interface */
 			virtual FName GetToolkitFName() const override;
diff --git a/Engine/Plugins/Runtime/Metasound/Source/MetasoundEditor/Private/MetasoundEditorGraphBuilder.cpp b/Engine/Plugins/Runtime/Metasound/Source/MetasoundEditor/Private/MetasoundEditorGraphBuilder.cpp
--- a/Engine/Plugins/Runtime/Metasound/Source/MetasoundEditor/Private/MetasoundEditorGraphBuilder.cpp
+++ b/Engine/Plugins/Runtime/Metasound/Source/MetasoundEditor/Private/MetasoundEditorGraphBuilder.cpp
@@ -11,6 +11,7 @@
 #include "MetasoundEditorModule.h"
 #include "MetasoundFrontend.h"
 #include "MetasoundFrontendRegistries.h"
+#include "MetasoundUObjectRegistry.h"
 #include "Modules/ModuleManager.h"
 #include "ScopedTransaction.h"
 #include "Templates/Tuple.h"
@@ -34,7 +35,7 @@ namespace Metasound
 		{
 			const FScopedTransaction Transaction(LOCTEXT("AddMetasoundGraphNode", "Add Metasound Node"));
 
-			FMetasoundAssetBase* MetasoundAsset = Frontend::GetObjectAsAssetBase(&InMetasound);
+			FMetasoundAssetBase* MetasoundAsset = IMetasoundUObjectRegistry::Get().GetObjectAsAssetBase(&InMetasound);
 			check(MetasoundAsset);
 
 			UEdGraph& Graph = MetasoundAsset->GetGraphChecked();
@@ -63,7 +64,7 @@ namespace Metasound
 
 		Frontend::FNodeHandle FGraphBuilder::AddNodeHandle(UObject& InMetasound, const Frontend::FNodeClassInfo& InClassInfo)
 		{
-			FMetasoundAssetBase* MetasoundAsset = Frontend::GetObjectAsAssetBase(&InMetasound);
+			FMetasoundAssetBase* MetasoundAsset = IMetasoundUObjectRegistry::Get().GetObjectAsAssetBase(&InMetasound);
 			check(MetasoundAsset);
 
 			return MetasoundAsset->GetRootGraphHandle().AddNewNode(InClassInfo);
@@ -94,12 +95,13 @@ namespace Metasound
 			return Categories;
 		}
 
-		FString FGraphBuilder::GenerateUniqueInputName(const UObject& InMetasound, const FName InBaseName)
+		FString FGraphBuilder::GenerateUniqueInputName(UObject& InMetasound, const FName InBaseName)
 		{
 			FString NameBase = GetDataTypeDisplayName(InBaseName);
-			const FMetasoundAssetBase* MetasoundAsset = Frontend::GetObjectAsAssetBase(&InMetasound);
+			FMetasoundAssetBase* MetasoundAsset = IMetasoundUObjectRegistry::Get().GetObjectAsAssetBase(&InMetasound);
 			check(MetasoundAsset);
 
+			// TODO: To achieve const correctness, this needs a FConstGraphHandle.
 			const Frontend::FGraphHandle GraphHandle = MetasoundAsset->GetRootGraphHandle();
 
 			int32 i = 1;
@@ -112,12 +114,13 @@ namespace Metasound
 			return NewNodeName;
 		}
 
-		FString FGraphBuilder::GenerateUniqueOutputName(const UObject& InMetasound, const FName InBaseName)
+		FString FGraphBuilder::GenerateUniqueOutputName(UObject& InMetasound, const FName InBaseName)
 		{
 			FString NameBase = GetDataTypeDisplayName(InBaseName);
-			const FMetasoundAssetBase* MetasoundAsset = Frontend::GetObjectAsAssetBase(&InMetasound);
+			FMetasoundAssetBase* MetasoundAsset = IMetasoundUObjectRegistry::Get().GetObjectAsAssetBase(&InMetasound);
 			check(MetasoundAsset);
 
+			// TODO: To achieve const correctness, this needs a FConstGraphHandle.
 			const Frontend::FGraphHandle GraphHandle = MetasoundAsset->GetRootGraphHandle();
 
 			int32 i = 1;
@@ -143,7 +146,7 @@ namespace Metasound
 			Description.TypeName = InTypeName;
 			Description.ToolTip = InToolTip;
 
-			FMetasoundAssetBase* MetasoundAsset = Frontend::GetObjectAsAssetBase(&InMetasound);
+			FMetasoundAssetBase* MetasoundAsset = IMetasoundUObjectRegistry::Get().GetObjectAsAssetBase(&InMetasound);
 			check(MetasoundAsset);
 
 			Frontend::FGraphHandle GraphHandle = MetasoundAsset->GetRootGraphHandle();
@@ -237,7 +240,7 @@ namespace Metasound
 			Description.TypeName = InTypeName;
 			Description.ToolTip = InToolTip;
 
-			FMetasoundAssetBase* MetasoundAsset = Frontend::GetObjectAsAssetBase(&InMetasound);
+			FMetasoundAssetBase* MetasoundAsset = IMetasoundUObjectRegistry::Get().GetObjectAsAssetBase(&InMetasound);
 			check(MetasoundAsset);
 
 			Frontend::FGraphHandle GraphHandle = MetasoundAsset->GetRootGraphHandle();
@@ -261,7 +264,7 @@ namespace Metasound
 
 			const Frontend::FNodeHandle& NodeHandle = InNode.GetNodeHandle();
 
-			FMetasoundAssetBase* MetasoundAsset = Frontend::GetObjectAsAssetBase(Graph->GetMetasound());
+			FMetasoundAssetBase* MetasoundAsset = IMetasoundUObjectRegistry::Get().GetObjectAsAssetBase(Graph->GetMetasound());
 			check(MetasoundAsset);
 
 			Frontend::FGraphHandle GraphHandle = MetasoundAsset->GetRootGraphHandle();
@@ -297,7 +300,7 @@ namespace Metasound
 		{
 			using namespace Frontend;
 
-			FMetasoundAssetBase* MetasoundAsset = GetObjectAsAssetBase(&InMetasound);
+			FMetasoundAssetBase* MetasoundAsset = IMetasoundUObjectRegistry::Get().GetObjectAsAssetBase(&InMetasound);
 			check(MetasoundAsset);
 
 			FGraphHandle GraphHandle = MetasoundAsset->GetRootGraphHandle();
@@ -392,7 +395,7 @@ namespace Metasound
 
 			IMetasoundEditorModule& EditorModule = FModuleManager::GetModuleChecked<IMetasoundEditorModule>("MetasoundEditor");
 
-			const FMetasoundAssetBase* MetasoundAsset = Frontend::GetObjectAsAssetBase(&InGraphNode.GetMetasoundChecked());
+			FMetasoundAssetBase* MetasoundAsset = IMetasoundUObjectRegistry::Get().GetObjectAsAssetBase(&InGraphNode.GetMetasoundChecked());
 			check(MetasoundAsset);
 
 			Frontend::FGraphHandle GraphHandle = MetasoundAsset->GetRootGraphHandle();
@@ -400,58 +403,379 @@ namespace Metasound
 			TArray<Frontend::FInputHandle> InputHandles = InNodeHandle.GetAllInputs();
 			for (int32 i = 0; i < InputHandles.Num(); ++i)
 			{
-				const Frontend::FInputHandle& InputHandle = InputHandles[i];
+				AddPinToNode(InGraphNode, InputHandles[i]);
+			}
+
+			TArray<Frontend::FOutputHandle> OutputHandles = InNodeHandle.GetAllOutputs();
+			for (int32 i = 0; i < OutputHandles.Num(); ++i)
+			{
+				AddPinToNode(InGraphNode, OutputHandles[i]);
+			}
 
-				FText DisplayName;
-				if (InGraphNode.GetNodeHandle().GetNodeType() == EMetasoundClassType::Input)
+			InGraphNode.MarkPackageDirty();
+		}
+
+		bool FGraphBuilder::IsMatchingInputHandleAndPin(const Frontend::FInputHandle& InInputHandle, const UEdGraphPin& InEditorPin)
+		{
+			if (EEdGraphPinDirection::EGPD_Input == InEditorPin.Direction)
+			{
+				if (InEditorPin.GetName() == InInputHandle.GetInputName())
 				{
-					DisplayName = GraphHandle.GetInputDisplayName(InputHandle.GetInputName());
+					IMetasoundEditorModule& EditorModule = FModuleManager::GetModuleChecked<IMetasoundEditorModule>("MetasoundEditor");
+
+					if (InEditorPin.PinType == EditorModule.FindDataType(InInputHandle.GetInputType()).PinType)
+					{
+						UEdGraphNode* EditorNode = InEditorPin.GetOwningNode();
+						UMetasoundEditorGraphNode* MetasoundEditorNode = Cast<UMetasoundEditorGraphNode>(EditorNode);
+
+						if (nullptr != MetasoundEditorNode)
+						{
+							if  (MetasoundEditorNode->GetNodeHandle().GetNodeID() == InInputHandle.GetOwningNodeID())
+							{
+								return true;
+							}
+						}
+					}
 				}
-				else
+			}
+			return false;
+		}
+
+		bool FGraphBuilder::IsMatchingOutputHandleAndPin(const Frontend::FOutputHandle& InOutputHandle, const UEdGraphPin& InEditorPin)
+		{
+			if (EEdGraphPinDirection::EGPD_Output == InEditorPin.Direction)
+			{
+				if (InEditorPin.GetName() == InOutputHandle.GetOutputName())
 				{
-					DisplayName = FText::FromString(*InputHandle.GetInputName());
+					IMetasoundEditorModule& EditorModule = FModuleManager::GetModuleChecked<IMetasoundEditorModule>("MetasoundEditor");
+
+					if (InEditorPin.PinType == EditorModule.FindDataType(InOutputHandle.GetOutputType()).PinType)
+					{
+						UEdGraphNode* EditorNode = InEditorPin.GetOwningNode();
+						UMetasoundEditorGraphNode* MetasoundEditorNode = Cast<UMetasoundEditorGraphNode>(EditorNode);
+
+						if (nullptr != MetasoundEditorNode)
+						{
+							if  (MetasoundEditorNode->GetNodeHandle().GetNodeID() == InOutputHandle.GetOwningNodeID())
+							{
+								return true;
+							}
+						}
+					}
 				}
+			}
+			return false;
+		}
+
+		UEdGraphPin* FGraphBuilder::AddPinToNode(UMetasoundEditorGraphNode& InEditorNode, Frontend::FInputHandle& InInputHandle)
+		{
+			IMetasoundEditorModule& EditorModule = FModuleManager::GetModuleChecked<IMetasoundEditorModule>("MetasoundEditor");
 
-				FEdGraphPinType PinType("MetasoundNode", NAME_None, nullptr, EPinContainerType::None, false, FEdGraphTerminalType());
-				UEdGraphPin* NewPin = InGraphNode.CreatePin(EGPD_Input, PinType, *DisplayName.ToString(), i);
-				if (ensureAlways(NewPin))
+			FEdGraphPinType PinType("MetasoundNode", NAME_None, nullptr, EPinContainerType::None, false, FEdGraphTerminalType());
+			UEdGraphPin* NewPin = InEditorNode.CreatePin(EGPD_Input, PinType, FName(*InInputHandle.GetInputName()));
+
+			if (ensureAlways(NewPin))
+			{
+				NewPin->PinToolTip = InInputHandle.GetInputTooltip().ToString();
+				NewPin->PinType = EditorModule.FindDataType(InInputHandle.GetInputType()).PinType;
+			}
+
+			return NewPin;
+		}
+
+		UEdGraphPin* FGraphBuilder::AddPinToNode(UMetasoundEditorGraphNode& InEditorNode, Frontend::FOutputHandle& InOutputHandle)
+		{
+			IMetasoundEditorModule& EditorModule = FModuleManager::GetModuleChecked<IMetasoundEditorModule>("MetasoundEditor");
+			FText DisplayName;
+
+			FEdGraphPinType PinType("MetasoundNode", NAME_None, nullptr, EPinContainerType::None, false, FEdGraphTerminalType());
+			UEdGraphPin* NewPin = InEditorNode.CreatePin(EGPD_Output, PinType, FName(*InOutputHandle.GetOutputName()));
+
+			if (ensureAlways(NewPin))
+			{
+				NewPin->PinToolTip = InOutputHandle.GetOutputTooltip().ToString();
+				NewPin->PinType = EditorModule.FindDataType(InOutputHandle.GetOutputType()).PinType;
+			}
+
+			return NewPin;
+		}
+
+		bool FGraphBuilder::SynchronizeGraph(UObject& InMetasound)
+		{
+			using namespace Frontend;
+
+			bool bIsEditorGraphDirty = false;
+
+			FMetasoundAssetBase* MetasoundAsset = IMetasoundUObjectRegistry::Get().GetObjectAsAssetBase(&InMetasound);
+			check(MetasoundAsset);
+
+			// Get all nodes from metasound graph
+			FGraphHandle GraphHandle = MetasoundAsset->GetRootGraphHandle();
+			TArray<FNodeHandle> Nodes = GraphHandle.GetAllNodes();
+
+			// Get all nodes from metasound editor graph
+			UMetasoundEditorGraph* EditorGraph = CastChecked<UMetasoundEditorGraph>(MetasoundAsset->GetGraph());
+			TArray<UMetasoundEditorGraphNode*> EditorNodes;
+			for (UEdGraphNode* EdGraphNode : EditorGraph->Nodes)
+			{
+				UMetasoundEditorGraphNode* EditorNode = Cast<UMetasoundEditorGraphNode>(EdGraphNode);
+				if (ensure(nullptr != EditorNode))
+				{
+					EditorNodes.Add(EditorNode);
+				}
+			}
+
+			// Find existing pairs of nodes and editor nodes
+			struct FNodePair
+			{
+				UMetasoundEditorGraphNode* EditorNode = nullptr;
+				FNodeHandle Node = FNodeHandle::InvalidHandle();;
+			};
+			TMap<int32, FNodePair> PairedNodes;
+
+			// Reverse iterate so paired nodes can safely be removed from the array.
+			for (int32 i = Nodes.Num() - 1; i >= 0; i--)
+			{
+				FNodeHandle Node = Nodes[i];
+				auto IsEditorNodeWithSameNodeID = [&](const UMetasoundEditorGraphNode* InEditorNode)
 				{
-					NewPin->PinToolTip = InputHandle.GetInputTooltip().ToString();
+					return InEditorNode->GetNodeID() == Node.GetNodeID();
+				};
 
-					FEditorDataType DataType = EditorModule.FindDataType(InputHandle.GetInputType());
-					NewPin->PinType = DataType.PinType;
+				UMetasoundEditorGraphNode* EditorNode = nullptr;
+				if (UMetasoundEditorGraphNode** PointerEditorNode = EditorNodes.FindByPredicate(IsEditorNodeWithSameNodeID))
+				{
+					EditorNode = *PointerEditorNode;
+				}
+
+				if (nullptr != EditorNode)
+				{
+					PairedNodes.Add(Node.GetNodeID(), { EditorNode, Node });
+					EditorNodes.RemoveSwap(EditorNode);
+					Nodes.RemoveAtSwap(i);
 				}
 			}
 
-			TArray<Frontend::FOutputHandle> OutputHandles = InNodeHandle.GetAllOutputs();
-			for (int32 i = 0; i < OutputHandles.Num(); ++i)
+			// Nodes contains nodes which need to be added.
+			// EditorNodes contains nodes that need to be removed. 
+			// PairedNodes contains pairs which we have to check have synchronized pins
+
+			// Add and remove nodes first in order to make sure correct editor nodes
+			// exist before attempting to synchronize connections.
+			for (UMetasoundEditorGraphNode* EditorNode : EditorNodes)
+			{
+				bIsEditorGraphDirty |= EditorGraph->RemoveNode(EditorNode);
+			}
+
+			// Add missing nodes. 
+			bIsEditorGraphDirty |= (Nodes.Num() > 0);
+			for (FNodeHandle Node : Nodes)
 			{
-				const Frontend::FOutputHandle& OutputHandle = OutputHandles[i];
-				FEdGraphPinType PinType("MetasoundNode", NAME_None, nullptr, EPinContainerType::None, false, FEdGraphTerminalType());
+				UEdGraphNode* NewNode = AddNode(InMetasound, EditorGraph->GetGoodPlaceForNewNode(), Node, false /* bInSelectNewNode */);
+				PairedNodes.Add(Node.GetNodeID(), {Cast<UMetasoundEditorGraphNode>(NewNode), Node});
+			}
 
-				FText DisplayName;
-				if (InGraphNode.GetNodeHandle().GetNodeType() == EMetasoundClassType::Output)
+			// Synchronize pins on node pairs.
+			for (const TPair<int32, FNodePair>& IdNodePair : PairedNodes)
+			{
+				UMetasoundEditorGraphNode* EditorNode = IdNodePair.Value.EditorNode;
+				FNodeHandle Node = IdNodePair.Value.Node;
+
+				bIsEditorGraphDirty |= SynchronizeNodePins(*IdNodePair.Value.EditorNode, IdNodePair.Value.Node);
+			}
+
+			// Synchronize connections.
+			bIsEditorGraphDirty |= SynchronizeConnections(InMetasound);
+
+			if (bIsEditorGraphDirty)
+			{
+				InMetasound.PostEditChange();
+				InMetasound.MarkPackageDirty();
+			}
+
+			return bIsEditorGraphDirty;
+		}
+
+		bool FGraphBuilder::SynchronizeNodePins(UMetasoundEditorGraphNode& InEditorNode, Frontend::FNodeHandle InNode)
+		{
+			bool bIsNodeDirty = false;
+
+			TArray<Frontend::FInputHandle> InputHandles = InNode.GetAllInputs();
+			TArray<Frontend::FOutputHandle> OutputHandles = InNode.GetAllOutputs();
+			TArray<UEdGraphPin*> EditorPins = InEditorNode.Pins;
+
+			// Filter out pins which are not paired.
+			for (int32 i = EditorPins.Num() - 1; i >= 0; i--)
+			{
+				UEdGraphPin* Pin = EditorPins[i];
+
+				auto IsMatchingInputHandle = [&](const Frontend::FInputHandle& InputHandle)
+				{
+					return IsMatchingInputHandleAndPin(InputHandle, *Pin);
+				};
+
+				auto IsMatchingOutputHandle = [&](const Frontend::FOutputHandle& OutputHandle)
+				{
+					return IsMatchingOutputHandleAndPin(OutputHandle, *Pin);
+				};
+
+				switch (Pin->Direction)
 				{
-					DisplayName = GraphHandle.GetInputDisplayName(OutputHandle.GetOutputName());
+					case EEdGraphPinDirection::EGPD_Input:
+					{
+						int32 MatchingInputIndex = InputHandles.FindLastByPredicate(IsMatchingInputHandle);
+						if (INDEX_NONE != MatchingInputIndex)
+						{
+							InputHandles.RemoveAtSwap(MatchingInputIndex);
+							EditorPins.RemoveAtSwap(i);
+						}
+					}
+					break;
+
+					case EEdGraphPinDirection::EGPD_Output:
+					{
+						int32 MatchingOutputIndex = OutputHandles.FindLastByPredicate(IsMatchingOutputHandle);
+						if (INDEX_NONE != MatchingOutputIndex)
+						{
+							OutputHandles.RemoveAtSwap(MatchingOutputIndex);
+							EditorPins.RemoveAtSwap(i);
+						}
+					}
+					break;
 				}
-				else
+			}
+
+			bIsNodeDirty |= (InputHandles.Num() > 0);
+			bIsNodeDirty |= (OutputHandles.Num() > 0);
+			bIsNodeDirty |= (EditorPins.Num() > 0);
+
+			// Remove any unused editor pins.
+			for (UEdGraphPin* Pin : EditorPins)
+			{
+				InEditorNode.RemovePin(Pin);
+			}
+
+			// Add unmatched input pins
+			for (Frontend::FInputHandle& InputHandle : InputHandles)
+			{
+				AddPinToNode(InEditorNode, InputHandle);
+			}
+
+			for (Frontend::FOutputHandle& OutputHandle : OutputHandles)
+			{
+				AddPinToNode(InEditorNode, OutputHandle);
+			}
+			
+			if (bIsNodeDirty)
+			{
+				InEditorNode.MarkPackageDirty();
+			}
+
+			return bIsNodeDirty;
+		}
+
+		bool FGraphBuilder::SynchronizeConnections(UObject& InMetasound)
+		{
+			bool bIsGraphDirty = false;
+
+			FMetasoundAssetBase* MetasoundAsset = IMetasoundUObjectRegistry::Get().GetObjectAsAssetBase(&InMetasound);
+			check(MetasoundAsset);
+
+			Frontend::FGraphHandle GraphHandle = MetasoundAsset->GetRootGraphHandle();
+
+			UMetasoundEditorGraph* EditorGraph = CastChecked<UMetasoundEditorGraph>(MetasoundAsset->GetGraph());
+			TArray<UMetasoundEditorGraphNode*> EditorNodes;
+
+			// Cache map of editor nodes indexed by node id.
+			TMap<int32, UMetasoundEditorGraphNode*> EditorNodesByID;
+			for (UEdGraphNode* EdGraphNode : EditorGraph->Nodes)
+			{
+				UMetasoundEditorGraphNode* MetasoundEditorNode = Cast<UMetasoundEditorGraphNode>(EdGraphNode);
+				if (nullptr != MetasoundEditorNode)
 				{
-					DisplayName = FText::FromString(*OutputHandle.GetOutputName());
+					EditorNodes.Add(MetasoundEditorNode);
+					EditorNodesByID.Add(MetasoundEditorNode->GetNodeID(), MetasoundEditorNode);
 				}
+			}
 
-				UEdGraphPin* NewPin = InGraphNode.CreatePin(EGPD_Output, PinType, *DisplayName.ToString(), i);
-				if (ensureAlways(NewPin))
+			// Iterate through all nodes in metasound editor graph and synchronize connections.
+			for (UEdGraphNode* EdGraphNode : EditorGraph->Nodes)
+			{
+				bool bIsNodeDirty = false;
+
+				UMetasoundEditorGraphNode* MetasoundEditorNode = Cast<UMetasoundEditorGraphNode>(EdGraphNode);
+				Frontend::FNodeHandle Node = MetasoundEditorNode->GetNodeHandle();
+
+				TArray<UEdGraphPin*> Pins = MetasoundEditorNode->GetAllPins();
+				TArray<Frontend::FInputHandle> NodeInputs = Node.GetAllInputs();
+
+				for (Frontend::FInputHandle& NodeInput : NodeInputs)
 				{
-					NewPin->PinToolTip = OutputHandle.GetOutputTooltip().ToString();
+					auto IsMatchingInputPin = [&](const UEdGraphPin* Pin) -> bool
+					{
+						return IsMatchingInputHandleAndPin(NodeInput, *Pin);
+					};
+
+					UEdGraphPin* MatchingPin = nullptr;
+					if (UEdGraphPin** DoublePointer = Pins.FindByPredicate(IsMatchingInputPin))
+					{
+						MatchingPin = *DoublePointer;
+					}
+
+					if (ensure(nullptr != MatchingPin))
+					{
+						// Remove pin so it isn't used twice. 
+						Pins.Remove(MatchingPin);
 
-					const FName OutputType = OutputHandle.GetOutputType();
-					FEditorDataType DataType = EditorModule.FindDataType(OutputType);
-					NewPin->PinType = DataType.PinType;
+						Frontend::FOutputHandle OutputHandle = NodeInput.GetCurrentlyConnectedOutput();
+						if (OutputHandle.IsValid())
+						{
+							bool bAddLink = false;
+
+							if (MatchingPin->LinkedTo.Num() < 1)
+							{
+								// No link currently exists. Add the appropriate link.
+								bAddLink = true;
+							}
+							else if (!IsMatchingOutputHandleAndPin(OutputHandle, *MatchingPin->LinkedTo[0]))
+							{
+								// The wrong link exists.
+								MatchingPin->BreakAllPinLinks();
+								bAddLink = true;
+							}
+
+							if (bAddLink)
+							{
+								UMetasoundEditorGraphNode* OutputEditorNode = EditorNodesByID[OutputHandle.GetOwningNodeID()];
+								UEdGraphPin* OutputPin = OutputEditorNode->FindPinChecked(OutputHandle.GetOutputName(), EEdGraphPinDirection::EGPD_Output);
+								MatchingPin->MakeLinkTo(OutputPin);
+								bIsNodeDirty = true;
+							}
+						}
+						else if (MatchingPin->LinkedTo.Num() > 0)
+						{
+							// No link should exist.
+							MatchingPin->BreakAllPinLinks();
+							bIsNodeDirty = true;
+						}
+					}
+				}
+
+				if (bIsNodeDirty)
+				{
+					EdGraphNode->MarkPackageDirty();
 				}
+
+				bIsGraphDirty |= bIsNodeDirty;
 			}
 
-			InGraphNode.MarkPackageDirty();
+			if (bIsGraphDirty)
+			{
+				EditorGraph->MarkPackageDirty();
+			}
+
+			return bIsGraphDirty;
 		}
 	} // namespace Editor
 } // namespace Metasound
-#undef LOCTEXT_NAMESPACE
\ No newline at end of file
+#undef LOCTEXT_NAMESPACE
diff --git a/Engine/Plugins/Runtime/Metasound/Source/MetasoundEditor/Private/MetasoundEditorGraphBuilder.h b/Engine/Plugins/Runtime/Metasound/Source/MetasoundEditor/Private/MetasoundEditorGraphBuilder.h
--- a/Engine/Plugins/Runtime/Metasound/Source/MetasoundEditor/Private/MetasoundEditorGraphBuilder.h
+++ b/Engine/Plugins/Runtime/Metasound/Source/MetasoundEditor/Private/MetasoundEditorGraphBuilder.h
@@ -39,9 +39,9 @@ namespace Metasound
 			// Adds a node handle with the given class info
 			static Frontend::FNodeHandle AddNodeHandle(UObject& InMetasound, const Frontend::FNodeClassInfo& InClassInfo);
 
-			static FString GenerateUniqueInputName(const UObject& InMetasound, const FName InBaseName);
+			static FString GenerateUniqueInputName(UObject& InMetasound, const FName InBaseName);
 
-			static FString GenerateUniqueOutputName(const UObject& InMetasound, const FName InBaseName);
+			static FString GenerateUniqueOutputName(UObject& InMetasound, const FName InBaseName);
 
 			static FString GetDataTypeDisplayName(const FName& InDataTypeName);
 
@@ -62,6 +62,34 @@ namespace Metasound
 			static void RebuildGraph(UObject& InMetasound);
 
 			static void RebuildNodePins(UMetasoundEditorGraphNode& InGraphNode, Frontend::FNodeHandle& InNodeHandle, bool bInRecordTransaction = true);
+
+			// Adds an Input UEdGraphPin to a UMetasoundEditorGraphNode
+			static UEdGraphPin* AddPinToNode(UMetasoundEditorGraphNode& InEditorNode, Frontend::FInputHandle& InInputHandle);
+
+			// Adds an Output UEdGraphPin to a UMetasoundEditorGraphNode
+			static UEdGraphPin* AddPinToNode(UMetasoundEditorGraphNode& InEditorNode, Frontend::FOutputHandle& InOutputHandle);
+
+			// Adds and removes nodes, pins and connections so that the UEdGraph of the Metasound matches the FMetasoundDocumentModel
+			//
+			// @return True if the UEdGraph was altered, false otherwise. 
+			static bool SynchronizeGraph(UObject& InMetasound);
+
+			// Adds and removes pins so that the UMetasoundEditorGraphNode matches the InNode.
+			//
+			// @return True if the UMetasoundEditorGraphNode was altered. False otherwise.
+			static bool SynchronizeNodePins(UMetasoundEditorGraphNode& InEditorNode, Frontend::FNodeHandle InNode);
+
+			// Adds and removes connections so that the UEdGraph of the metasound has the same
+			// connections as the FMetasoundDocument graph.
+			//
+			// @return True if the UEdGraph was altered. False otherwise. 
+			static bool SynchronizeConnections(UObject& InMetasound);
+
+			// Returns true if the FInputHandle and UEdGraphPin match each other. 
+			static bool IsMatchingInputHandleAndPin(const Frontend::FInputHandle& InInputHandle, const UEdGraphPin& InEditorPin);
+
+			// Returns true if the FOutputHandle and UEdGraphPin match each other. 
+			static bool IsMatchingOutputHandleAndPin(const Frontend::FOutputHandle& InOutputHandle, const UEdGraphPin& InEditorPin);
 		};
 	} // namespace Editor
 } // namespace Metasound
diff --git a/Engine/Plugins/Runtime/Metasound/Source/MetasoundEditor/Private/MetasoundEditorGraphNode.cpp b/Engine/Plugins/Runtime/Metasound/Source/MetasoundEditor/Private/MetasoundEditorGraphNode.cpp
--- a/Engine/Plugins/Runtime/Metasound/Source/MetasoundEditor/Private/MetasoundEditorGraphNode.cpp
+++ b/Engine/Plugins/Runtime/Metasound/Source/MetasoundEditor/Private/MetasoundEditorGraphNode.cpp
@@ -13,6 +13,7 @@
 #include "MetasoundEditorGraphBuilder.h"
 #include "MetasoundEditorGraphSchema.h"
 #include "MetasoundEditorCommands.h"
+#include "MetasoundUObjectRegistry.h"
 #include "ScopedTransaction.h"
 #include "ToolMenus.h"
 
@@ -85,15 +86,15 @@ const UObject& UMetasoundEditorGraphNode::GetMetasoundChecked() const
 	return EdGraph->GetMetasoundChecked();
 }
 
-Metasound::Frontend::FGraphHandle UMetasoundEditorGraphNode::GetRootGraphHandle() const
+Metasound::Frontend::FGraphHandle UMetasoundEditorGraphNode::GetRootGraphHandle()
 {
-	const FMetasoundAssetBase* MetasoundAsset = Metasound::Frontend::GetObjectAsAssetBase(&GetMetasoundChecked());
+	FMetasoundAssetBase* MetasoundAsset = Metasound::IMetasoundUObjectRegistry::Get().GetObjectAsAssetBase(&GetMetasoundChecked());
 	check(MetasoundAsset);
 
 	return MetasoundAsset->GetRootGraphHandle();
 }
 
-Metasound::Frontend::FNodeHandle UMetasoundEditorGraphNode::GetNodeHandle() const
+Metasound::Frontend::FNodeHandle UMetasoundEditorGraphNode::GetNodeHandle()
 {
 	return GetRootGraphHandle().GetNodeWithId(NodeID);
 }
@@ -181,8 +182,17 @@ bool UMetasoundEditorGraphNode::CanCreateUnderSpecifiedSchema(const UEdGraphSche
 
 bool UMetasoundEditorGraphNode::CanUserDeleteNode() const
 {
-	const FString& NodeName = GetNodeHandle().GetNodeName();
-	Metasound::Frontend::FGraphHandle GraphHandle = GetRootGraphHandle();
+	const FMetasoundAssetBase* ConstMetasoundAsset = Metasound::IMetasoundUObjectRegistry::Get().GetObjectAsAssetBase(&GetMetasoundChecked());
+
+	// TODO: create FConstGraphHandle so that this const cast is unnecessary.
+	FMetasoundAssetBase* MetasoundAsset = const_cast<FMetasoundAssetBase*>(ConstMetasoundAsset);
+	check(MetasoundAsset);
+
+	Metasound::Frontend::FGraphHandle GraphHandle = MetasoundAsset->GetRootGraphHandle();
+	Metasound::Frontend::FNodeHandle NodeHandle = GraphHandle.GetNodeWithId(NodeID);
+
+	const FString& NodeName = NodeHandle.GetNodeName();
+
 	if (GraphHandle.IsRequiredInput(NodeName))
 	{
 		return false;
@@ -215,19 +225,25 @@ FText UMetasoundEditorGraphNode::GetNodeTitle(ENodeTitleType::Type TitleType) co
 {
 	using namespace Metasound::Frontend;
 
-	FNodeHandle NodeHandle = GetNodeHandle();
+	const FMetasoundAssetBase* ConstMetasoundAsset = Metasound::IMetasoundUObjectRegistry::Get().GetObjectAsAssetBase(&GetMetasoundChecked());
+
+	// TODO: create FConstGraphHandle so that this const cast is unnecessary.
+	FMetasoundAssetBase* MetasoundAsset = const_cast<FMetasoundAssetBase*>(ConstMetasoundAsset);
+	check(MetasoundAsset);
+
+	FGraphHandle GraphHandle = MetasoundAsset->GetRootGraphHandle();
+	FNodeHandle NodeHandle = GraphHandle.GetNodeWithId(NodeID);
+
 	switch (NodeHandle.GetNodeType())
 	{
 		case EMetasoundClassType::Input:
 		{
-			FGraphHandle GraphHandle = GetRootGraphHandle();
 			return FText::Format(LOCTEXT("MetasoundGraphNode_TitleFormat", "Input {0}"), GraphHandle.GetInputDisplayName(NodeHandle.GetNodeName()));
 		}
 		break;
 
 		case EMetasoundClassType::Output:
 		{
-			FGraphHandle GraphHandle = GetRootGraphHandle();
 			return FText::Format(LOCTEXT("MetasoundGraphNode_TitleFormat", "Output {0}"), GraphHandle.GetOutputDisplayName(NodeHandle.GetNodeName()));
 		}
 		break;
@@ -312,19 +328,27 @@ FText UMetasoundEditorGraphNode::GetTooltipText() const
 {
 	using namespace Metasound::Frontend;
 
-	FNodeHandle NodeHandle = GetNodeHandle();
+	using namespace Metasound::Frontend;
+
+	const FMetasoundAssetBase* ConstMetasoundAsset = Metasound::IMetasoundUObjectRegistry::Get().GetObjectAsAssetBase(&GetMetasoundChecked());
+
+	// TODO: create FConstGraphHandle so that this const cast is unnecessary.
+	FMetasoundAssetBase* MetasoundAsset = const_cast<FMetasoundAssetBase*>(ConstMetasoundAsset);
+	check(MetasoundAsset);
+
+	FGraphHandle GraphHandle = MetasoundAsset->GetRootGraphHandle();
+	FNodeHandle NodeHandle = GraphHandle.GetNodeWithId(NodeID);
+
 	switch (NodeHandle.GetNodeType())
 	{
 		case EMetasoundClassType::Input:
 		{
-			FGraphHandle GraphHandle = GetRootGraphHandle();
 			return GraphHandle.GetInputToolTip(NodeHandle.GetNodeName());
 		}
 		break;
 
 		case EMetasoundClassType::Output:
 		{
-			FGraphHandle GraphHandle = GetRootGraphHandle();
 			return GraphHandle.GetOutputToolTip(NodeHandle.GetNodeName());
 		}
 		break;
diff --git a/Engine/Plugins/Runtime/Metasound/Source/MetasoundEditor/Public/IMetasoundEditor.h b/Engine/Plugins/Runtime/Metasound/Source/MetasoundEditor/Public/IMetasoundEditor.h
--- a/Engine/Plugins/Runtime/Metasound/Source/MetasoundEditor/Public/IMetasoundEditor.h
+++ b/Engine/Plugins/Runtime/Metasound/Source/MetasoundEditor/Public/IMetasoundEditor.h
@@ -5,7 +5,6 @@
 #include "CoreMinimal.h"
 
 #include "Layout/SlateRect.h"
-#include "Metasound.h"
 #include "Toolkits/AssetEditorToolkit.h"
 #include "UObject/Object.h"
 
@@ -15,4 +14,10 @@ class IMetasoundEditor : public FAssetEditorToolkit
 	virtual UObject* GetMetasoundObject() const = 0;
 	virtual void SetSelection(const TArray<UObject*>& SelectedObjects) = 0;
 	virtual bool GetBoundsForSelectedNodes(FSlateRect& Rect, float Padding) = 0;
+
+	// Rebuilds all nodes and connections from root FMetasoundDocument. Results in loss of node layout.
+	virtual bool RebuildGraph() const = 0;
+
+	// Synchronizes the nodes and connections of the graph to the root FMetasoundDocument.
+	virtual bool SynchronizeGraph() const = 0;
 };
diff --git a/Engine/Plugins/Runtime/Metasound/Source/MetasoundEditor/Public/MetasoundEditorGraph.h b/Engine/Plugins/Runtime/Metasound/Source/MetasoundEditor/Public/MetasoundEditorGraph.h
--- a/Engine/Plugins/Runtime/Metasound/Source/MetasoundEditor/Public/MetasoundEditorGraph.h
+++ b/Engine/Plugins/Runtime/Metasound/Source/MetasoundEditor/Public/MetasoundEditorGraph.h
@@ -3,7 +3,6 @@
 
 #include "CoreMinimal.h"
 #include "EdGraph/EdGraph.h"
-#include "Metasound.h"
 #include "MetasoundEditor.h"
 #include "UObject/ObjectMacros.h"
 
diff --git a/Engine/Plugins/Runtime/Metasound/Source/MetasoundEditor/Public/MetasoundEditorGraphNode.h b/Engine/Plugins/Runtime/Metasound/Source/MetasoundEditor/Public/MetasoundEditorGraphNode.h
--- a/Engine/Plugins/Runtime/Metasound/Source/MetasoundEditor/Public/MetasoundEditorGraphNode.h
+++ b/Engine/Plugins/Runtime/Metasound/Source/MetasoundEditor/Public/MetasoundEditorGraphNode.h
@@ -55,8 +55,8 @@ public:
 	UObject& GetMetasoundChecked();
 	const UObject& GetMetasoundChecked() const;
 
-	Metasound::Frontend::FGraphHandle GetRootGraphHandle() const;
-	Metasound::Frontend::FNodeHandle GetNodeHandle() const;
+	Metasound::Frontend::FGraphHandle GetRootGraphHandle();
+	Metasound::Frontend::FNodeHandle GetNodeHandle();
 
 	void SetNodeID(uint32 InNodeID);
 	uint32 GetNodeID() const;
diff --git a/Engine/Plugins/Runtime/Metasound/Source/MetasoundEngine/Private/Metasound.cpp b/Engine/Plugins/Runtime/Metasound/Source/MetasoundEngine/Private/Metasound.cpp
--- a/Engine/Plugins/Runtime/Metasound/Source/MetasoundEngine/Private/Metasound.cpp
+++ b/Engine/Plugins/Runtime/Metasound/Source/MetasoundEngine/Private/Metasound.cpp
@@ -6,11 +6,9 @@
 #include "StructSerializer.h"
 #include "UObject/UnrealType.h"
 
-#include "MetasoundArchetypeRegistration.h"
-
 UMetasound::UMetasound(const FObjectInitializer& ObjectInitializer)
 	: Super(ObjectInitializer)
-	, FMetasoundAssetBase(RootMetasoundDocument)
+	, FMetasoundAssetBase(FMetasoundArchetype())
 {
 }
 
@@ -43,6 +41,37 @@ void UMetasound::SetGraph(UEdGraph* InGraph)
 }
 #endif // WITH_EDITORONLY_DATA
 
+// Returns document object responsible for serializing asset
+Metasound::Frontend::TAccessPtr<FMetasoundDocument> UMetasound::GetDocument()
+{
+	return Metasound::Frontend::MakeAccessPtr(MetasoundDocumentAccessPoint, MetasoundDocument);
+}
+
+// Returns document object responsible for serializing asset
+Metasound::Frontend::TAccessPtr<const FMetasoundDocument> UMetasound::GetDocument() const
+{
+	return Metasound::Frontend::MakeAccessPtr<const FMetasoundDocument>(MetasoundDocumentAccessPoint, MetasoundDocument);
+}
+
+const TArray<FMetasoundArchetype>& UMetasound::GetPreferredArchetypes() const
+{
+	// Not preferred archetypes for a basic UMetasound.
+	static const TArray<FMetasoundArchetype> Preferred;
+	return Preferred;
+}
+
+bool UMetasound::IsArchetypeSupported(const FMetasoundArchetype& InArchetype) const
+{
+	// All archetypes are supported.
+	return true;
+}
+
+const FMetasoundArchetype& UMetasound::GetPreferredArchetype(const FMetasoundDocument& InDocument) const
+{
+	// Prefer to keep original archetype.
+	return InDocument.Archetype;
+}
+
 // This can be used to update the metadata (name, author, etc) for this metasound.
 // @param InMetadata may be updated with any corrections we do to the input metadata.
 void UMetasound::SetMetadata(FMetasoundClassMetadata& InMetadata)
@@ -57,12 +86,3 @@ void UMetasound::PostLoad()
 	Super::PostLoad();
 }
 
-FMetasoundArchetype UMetasound::GetArchetype() const
-{
-	FMetasoundArchetype Archetype;
-	static FName MetasoundArchetypeName = TEXT("Generic Metasound");
-
-	Archetype.ArchetypeName = MetasoundArchetypeName;
-
-	return Archetype;
-}
diff --git a/Engine/Plugins/Runtime/Metasound/Source/MetasoundEngine/Private/MetasoundEngineModule.cpp b/Engine/Plugins/Runtime/Metasound/Source/MetasoundEngine/Private/MetasoundEngineModule.cpp
--- a/Engine/Plugins/Runtime/Metasound/Source/MetasoundEngine/Private/MetasoundEngineModule.cpp
+++ b/Engine/Plugins/Runtime/Metasound/Source/MetasoundEngine/Private/MetasoundEngineModule.cpp
@@ -2,12 +2,12 @@
 #include "MetasoundEngineModule.h"
 
 #include "Metasound.h"
-#include "MetasoundArchetypeRegistration.h"
-#include "MetasoundSource.h"
-#include "Modules/ModuleManager.h"
-#include "MetasoundWave.h"
 #include "MetasoundDataReference.h"
 #include "MetasoundDataTypeRegistrationMacro.h"
+#include "MetasoundSource.h"
+#include "MetasoundUObjectRegistry.h"
+#include "MetasoundWave.h"
+#include "Modules/ModuleManager.h"
 
 DEFINE_LOG_CATEGORY(LogMetasoundEngine);
 
@@ -17,8 +17,11 @@ class FMetasoundEngineModule : public IMetasoundEngineModule
 {
 	virtual void StartupModule() override
 	{
-		Metasound::Frontend::RegisterArchetype<UMetasoundSource>();
-		Metasound::Frontend::RegisterArchetype<UMetasound>();
+		// If there is no archetype name, use UMetasound
+		Metasound::IMetasoundUObjectRegistry::RegisterUClassArchetype<UMetasound>(TEXT(""));
+
+		// Register preferred archetypes
+		Metasound::IMetasoundUObjectRegistry::RegisterUClassPreferredArchetypes<UMetasoundSource>();
 
 		UE_LOG(LogMetasoundEngine, Log, TEXT("Metasound Engine Initialized"));
 	}
diff --git a/Engine/Plugins/Runtime/Metasound/Source/MetasoundEngine/Private/MetasoundSource.cpp b/Engine/Plugins/Runtime/Metasound/Source/MetasoundEngine/Private/MetasoundSource.cpp
--- a/Engine/Plugins/Runtime/Metasound/Source/MetasoundEngine/Private/MetasoundSource.cpp
+++ b/Engine/Plugins/Runtime/Metasound/Source/MetasoundEngine/Private/MetasoundSource.cpp
@@ -2,18 +2,17 @@
 
 #include "MetasoundSource.h"
 
+#include "AssetRegistryModule.h"
 #include "CoreMinimal.h"
 #include "Internationalization/Text.h"
-#include "AssetRegistryModule.h"
-
 #include "MetasoundAssetBase.h"
+#include "MetasoundAudioFormats.h"
 #include "MetasoundBop.h"
 #include "MetasoundEngineEnvironment.h"
 #include "MetasoundGenerator.h"
-#include "MetasoundAudioFormats.h"
+#include "MetasoundLog.h"
 #include "MetasoundPrimitives.h"
 
-#include "MetasoundArchetypeRegistration.h"
 
 #if WITH_EDITORONLY_DATA
 #include "EdGraph/EdGraph.h"
@@ -25,9 +24,8 @@ static FName MetasoundSourceArchetypeName = FName(TEXT("Metasound Source"));
 
 UMetasoundSource::UMetasoundSource(const FObjectInitializer& ObjectInitializer)
 	: Super(ObjectInitializer)
-	, FMetasoundAssetBase(RootMetasoundDocument)
+	, FMetasoundAssetBase(GetMonoSourceArchetype())
 {
-	// TODO: 
 	NumChannels = 1;
 	Duration = INDEFINITELY_LOOPING_DURATION;
 	bLooping = true;
@@ -42,6 +40,36 @@ void UMetasoundSource::SetMetadata(FMetasoundClassMetadata& InMetadata)
 	RootMetasoundDocument.RootClass.Metadata = InMetadata;
 }
 
+#if WITH_EDITOR
+
+void UMetasoundSource::PostEditChangeProperty(FPropertyChangedEvent& InEvent)
+{
+	Super::PostEditChangeProperty(InEvent);
+
+	if (InEvent.GetPropertyName() == GET_MEMBER_NAME_CHECKED(UMetasoundSource, OutputFormat)) 
+	{
+		switch (OutputFormat)
+		{
+			case EMetasoundSourceAudioFormat::Stereo:
+
+				NumChannels = 2;
+				ensure(SetArchetype(GetStereoSourceArchetype()));
+
+				break;
+
+			case EMetasoundSourceAudioFormat::Mono:
+			default:
+
+				NumChannels = 1;
+				ensure(SetArchetype(GetMonoSourceArchetype()));
+
+				break;
+		}
+	}
+}
+
+#endif
+
 bool UMetasoundSource::IsPlayable() const
 {
 	// todo: cache off whether this metasound is buildable to an operator.
@@ -63,7 +91,6 @@ ISoundGeneratorPtr UMetasoundSource::CreateSoundGenerator(const FSoundGeneratorI
 {
 	using namespace Metasound;
 
-	NumChannels = 1;
 	Duration = INDEFINITELY_LOOPING_DURATION;
 	bLooping = true;
 	SampleRate = InParams.SampleRate;
@@ -87,29 +114,61 @@ ISoundGeneratorPtr UMetasoundSource::CreateSoundGenerator(const FSoundGeneratorI
 	}
 	Environment.SetValue<FAudioDeviceHandle>(GetAudioDeviceHandleVariableName(), DeviceHandle);
 
+	// Create the operator from the graph handle.
 	TArray<IOperatorBuilder::FBuildErrorPtr> BuildErrors;
 
 	Frontend::FGraphHandle RootGraph = GetRootGraphHandle();
 	ensureAlways(RootGraph.IsValid());
 
 	TUniquePtr<IOperator> Operator = RootGraph.BuildOperator(InSettings, Environment, BuildErrors);
-	if (ensureAlways(Operator.IsValid()))
+
+	if (!Operator.IsValid())
+	{
+		// Log build errors that resulted in a null operator.
+		UE_LOG(LogMetasound, Error, TEXT("Failed to build Metasound operator from graph in MetasoundSource [%s]"), *GetName());
+		for (const IOperatorBuilder::FBuildErrorPtr& Error : BuildErrors)
+		{
+			if (Error.IsValid())
+			{
+				UE_LOG(LogMetasound, Warning, TEXT("MetasoundSource [%s] build error [%s] \"%s\""), *GetName(), *(Error->GetErrorType().ToString()), *(Error->GetErrorDescription().ToString()));
+			}
+		}
+	}
+	else
 	{
 		FDataReferenceCollection Outputs = Operator->GetOutputs();
+		TArrayView<FAudioBufferReadRef> OutputBuffers;
 
+		// Get output audio buffers.
+		if (EMetasoundSourceAudioFormat::Stereo == OutputFormat)
+		{
+			if (!Outputs.ContainsDataReadReference<FStereoAudioFormat>(GetAudioOutputName()))
+			{
+				UE_LOG(LogMetasound, Warning, TEXT("MetasoundSource [%s] does not contain stereo output [%s] in output"), *GetName(), *GetAudioOutputName());
+			}
+			OutputBuffers = Outputs.GetDataReadReferenceOrConstruct<FStereoAudioFormat>(GetAudioOutputName(), InSettings)->GetBuffers();
+		}
+		else if (EMetasoundSourceAudioFormat::Mono == OutputFormat)
+		{
+			if (!Outputs.ContainsDataReadReference<FStereoAudioFormat>(GetAudioOutputName()))
+			{
+				UE_LOG(LogMetasound, Warning, TEXT("MetasoundSource [%s] does not contain mono output [%s] in output"), *GetName(), *GetAudioOutputName());
+			}
+			OutputBuffers = Outputs.GetDataReadReferenceOrConstruct<FMonoAudioFormat>(GetAudioOutputName(), InSettings)->GetBuffers();
+		}
+
+		// Create the FMetasoundGenerator.
 		FMetasoundGeneratorInitParams InitParams =
 		{
 			MoveTemp(Operator),
-			Outputs.GetDataReadReferenceOrConstruct<FAudioBuffer>(GetAudioOutputName(), InSettings.GetNumFramesPerBlock()),
+			OutputBuffers,
 			Outputs.GetDataReadReferenceOrConstruct<FBop>(GetIsFinishedOutputName(), InSettings, false)
 		};
 
 		return ISoundGeneratorPtr(new FMetasoundGenerator(MoveTemp(InitParams)));
 	}
-	else
-	{
-		return ISoundGeneratorPtr(nullptr);
-	}
+
+	return ISoundGeneratorPtr(nullptr);
 }
 
 void UMetasoundSource::PostLoad()
@@ -119,6 +178,13 @@ void UMetasoundSource::PostLoad()
 	ConformDocumentToArchetype();
 }
 
+const TArray<FMetasoundArchetype>& UMetasoundSource::GetPreferredArchetypes() const 
+{
+	static TArray<FMetasoundArchetype> Preferred({GetMonoSourceArchetype(), GetStereoSourceArchetype()});
+
+	return Preferred;
+}
+
 const FString& UMetasoundSource::GetOnPlayInputName()
 {
 	static FString BopInputName = FString(TEXT("On Play"));
@@ -143,40 +209,87 @@ const FString& UMetasoundSource::GetAudioDeviceHandleVariableName()
 	return AudioDeviceHandleVarName;
 }
 
-FMetasoundArchetype UMetasoundSource::GetArchetype() const
+const FMetasoundArchetype& UMetasoundSource::GetBaseArchetype()
+{
+	auto CreateBaseArchetype = []() -> FMetasoundArchetype
+	{
+		FMetasoundArchetype Archetype;
+		
+		FMetasoundInputDescription OnPlayBop;
+		OnPlayBop.Name = UMetasoundSource::GetOnPlayInputName();
+		OnPlayBop.DisplayName = FText::FromString(OnPlayBop.Name);
+		OnPlayBop.TypeName = Metasound::Frontend::GetDataTypeName<Metasound::FBop>();
+		OnPlayBop.ToolTip = LOCTEXT("OnPlayBopToolTip", "Bop executed when this source is first played.");
+
+		Archetype.RequiredInputs.Add(OnPlayBop);
+
+		FMetasoundOutputDescription OnFinished;
+		OnFinished.Name = UMetasoundSource::GetIsFinishedOutputName();
+		OnFinished.DisplayName = FText::FromString(OnFinished.Name);
+		OnFinished.TypeName = Metasound::Frontend::GetDataTypeName<Metasound::FBop>();
+		OnFinished.ToolTip = LOCTEXT("OnFinishedToolTip", "Bop executed to initiate stopping the source.");
+
+		Archetype.RequiredOutputs.Add(OnFinished);
+
+		FMetasoundEnvironmentVariableDescription AudioDeviceHandle;
+		AudioDeviceHandle.Name = UMetasoundSource::GetAudioDeviceHandleVariableName();
+		AudioDeviceHandle.DisplayName = FText::FromString(AudioDeviceHandle.Name);
+		AudioDeviceHandle.ToolTip = LOCTEXT("AudioDeviceHandleToolTip", "Audio device handle");
+
+		Archetype.EnvironmentVariables.Add(AudioDeviceHandle);
+
+		return Archetype;
+	};
+
+	static const FMetasoundArchetype BaseArchetype = CreateBaseArchetype();
+
+	return BaseArchetype;
+}
+
+const FMetasoundArchetype& UMetasoundSource::GetMonoSourceArchetype()
 {
-	FMetasoundArchetype Archetype;
-	Archetype.ArchetypeName = MetasoundSourceArchetypeName;
+	auto CreateMonoArchetype = []() -> FMetasoundArchetype
+	{
+		FMetasoundArchetype Archetype = GetBaseArchetype();
+		Archetype.ArchetypeName = FName(TEXT("MonoSource"));
+
+		FMetasoundOutputDescription GeneratedAudio;
+		GeneratedAudio.Name = UMetasoundSource::GetAudioOutputName();
+		GeneratedAudio.DisplayName = LOCTEXT("GeneratedMono", "Mono");
+		GeneratedAudio.TypeName = Metasound::Frontend::GetDataTypeName<Metasound::FMonoAudioFormat>();
+		GeneratedAudio.ToolTip = LOCTEXT("GeneratedAudioToolTip", "The resulting output audio from this source.");
 
-	FMetasoundInputDescription OnPlayBop;
-	OnPlayBop.Name = GetOnPlayInputName();
-	OnPlayBop.DisplayName = FText::FromString(OnPlayBop.Name);
-	OnPlayBop.TypeName = Metasound::Frontend::GetDataTypeName<Metasound::FBop>();
-	OnPlayBop.ToolTip = LOCTEXT("OnPlayBopToolTip", "Bop executed when this source is first played.");
+		Archetype.RequiredOutputs.Add(GeneratedAudio);
+
+		return Archetype;
+	};
+
+	static const FMetasoundArchetype MonoArchetype = CreateMonoArchetype();
+
+	return MonoArchetype;
+}
+
+const FMetasoundArchetype& UMetasoundSource::GetStereoSourceArchetype()
+{
+	auto CreateStereoArchetype = []() -> FMetasoundArchetype
+	{
+		FMetasoundArchetype Archetype = GetBaseArchetype();
+		Archetype.ArchetypeName = FName(TEXT("StereoSource"));
 
-	Archetype.RequiredInputs.Add(OnPlayBop);
+		FMetasoundOutputDescription GeneratedAudio;
+		GeneratedAudio.Name = UMetasoundSource::GetAudioOutputName();
+		GeneratedAudio.DisplayName = LOCTEXT("GeneratedStereo", "Stereo");
+		GeneratedAudio.TypeName = Metasound::Frontend::GetDataTypeName<Metasound::FStereoAudioFormat>();
+		GeneratedAudio.ToolTip = LOCTEXT("GeneratedAudioToolTip", "The resulting output audio from this source.");
 
-	FMetasoundOutputDescription GeneratedAudio;
-	GeneratedAudio.Name = GetAudioOutputName();
-	GeneratedAudio.DisplayName = FText::FromString(GeneratedAudio.Name);
-	GeneratedAudio.TypeName = Metasound::Frontend::GetDataTypeName<Metasound::FAudioBuffer>();
-	GeneratedAudio.ToolTip = LOCTEXT("GeneratedAudioToolTip", "The resulting output audio from this source.");
-	Archetype.RequiredOutputs.Add(GeneratedAudio);
+		Archetype.RequiredOutputs.Add(GeneratedAudio);
 
-	FMetasoundOutputDescription OnFinished;
-	OnFinished.Name = GetIsFinishedOutputName();
-	OnFinished.DisplayName = FText::FromString(OnFinished.Name);
-	OnFinished.TypeName = Metasound::Frontend::GetDataTypeName<Metasound::FBop>();
-	OnFinished.ToolTip = LOCTEXT("OnFinishedToolTip", "Bop executed to initiate stopping the source.");
-	Archetype.RequiredOutputs.Add(OnFinished);
+		return Archetype;
+	};
 
-	FMetasoundEnvironmentVariableDescription AudioDeviceHandle;
-	AudioDeviceHandle.Name = GetAudioDeviceHandleVariableName();
-	AudioDeviceHandle.DisplayName = FText::FromString(AudioDeviceHandle.Name);
-	AudioDeviceHandle.ToolTip = LOCTEXT("AudioDeviceHandleToolTip", "Audio device handle");
-	Archetype.EnvironmentVariables.Add(AudioDeviceHandle);
+	static const FMetasoundArchetype StereoArchetype = CreateStereoArchetype();
 
-	return Archetype;
+	return StereoArchetype;
 }
 
 #undef LOCTEXT_NAMESPACE
diff --git a/Engine/Plugins/Runtime/Metasound/Source/MetasoundEngine/Private/MetasoundUObjectRegistry.cpp b/Engine/Plugins/Runtime/Metasound/Source/MetasoundEngine/Private/MetasoundUObjectRegistry.cpp
new file mode 100644
--- /dev/null
+++ b/Engine/Plugins/Runtime/Metasound/Source/MetasoundEngine/Private/MetasoundUObjectRegistry.cpp
@@ -0,0 +1,181 @@
+// Copyright Epic Games, Inc. All Rights Reserved.
+
+#include "MetasoundUObjectRegistry.h"
+
+#include "Algo/Copy.h"
+#include "AssetRegistry/AssetRegistryModule.h"
+#include "CoreMinimal.h"
+#include "UObject/Object.h"
+
+namespace Metasound
+{
+	class FMetasoundUObjectRegistry : public IMetasoundUObjectRegistry
+	{
+		public:
+			void RegisterUClassArchetype(TUniquePtr<IMetasoundUObjectRegistryEntry>&& InEntry) override
+			{
+				if (InEntry.IsValid())
+				{
+					FName ArchetypeName = InEntry->GetArchetypeName();
+
+					EntriesByArchetype.Add(ArchetypeName, InEntry.Get());
+					Entries.Add(InEntry.Get());
+					Storage.Add(MoveTemp(InEntry));
+				}
+			}
+
+			TArray<UClass*> GetUClassesForArchetype(const FName& InArchetypeName) const override
+			{
+				TArray<UClass*> Classes;
+
+				TArray<const IMetasoundUObjectRegistryEntry*> EntriesForArchetype;
+				EntriesByArchetype.MultiFind(InArchetypeName, EntriesForArchetype);
+
+				for (const IMetasoundUObjectRegistryEntry* Entry : EntriesForArchetype)
+				{
+					if (nullptr != Entry)
+					{
+						if (UClass* Class = Entry->GetUClass())
+						{
+							Classes.Add(Class);
+						}
+					}
+				}
+
+				return Classes;
+			}
+
+			UObject* NewObject(UClass* InClass, const FMetasoundDocument& InDocument, const FString& InPath) const override
+			{
+				auto IsChildClassOfRegisteredClass = [&](const IMetasoundUObjectRegistryEntry* Entry)
+				{
+					return Entry->IsChildClass(InClass);
+				};
+
+				TArray<const IMetasoundUObjectRegistryEntry*> EntriesForClass = FindEntriesByPredicate(IsChildClassOfRegisteredClass);
+
+				for (const IMetasoundUObjectRegistryEntry* Entry : EntriesForClass)
+				{
+					if (Entry->GetArchetypeName() == InDocument.Archetype.ArchetypeName)
+					{
+						return NewObject(*Entry, InDocument, InPath);
+					}
+				}
+
+				return nullptr;
+			}
+
+
+			bool IsRegisteredClass(UObject* InObject) const override
+			{
+				return (nullptr != GetEntryByUObject(InObject));
+			}
+
+			FMetasoundAssetBase* GetObjectAsAssetBase(UObject* InObject) const override
+			{
+				if (const IMetasoundUObjectRegistryEntry* Entry = GetEntryByUObject(InObject))
+				{
+					return Entry->Cast(InObject);
+				}
+				return nullptr;
+			}
+
+			const FMetasoundAssetBase* GetObjectAsAssetBase(const UObject* InObject) const override
+			{
+				if (const IMetasoundUObjectRegistryEntry* Entry = GetEntryByUObject(InObject))
+				{
+					return Entry->Cast(InObject);
+				}
+				return nullptr;
+			}
+
+		private:
+			UObject* NewObject(const IMetasoundUObjectRegistryEntry& InEntry, const FMetasoundDocument& InDocument, const FString& InPath) const
+			{
+				UPackage* PackageToSaveTo = nullptr;
+
+				if (GIsEditor)
+				{
+					FText InvalidPathReason;
+					bool const bValidPackageName = FPackageName::IsValidLongPackageName(InPath, false, &InvalidPathReason);
+
+					if (!ensureAlwaysMsgf(bValidPackageName, TEXT("Tried to generate a Metasound UObject with an invalid package path/name Falling back to transient package, which means we won't be able to save this asset.")))
+					{
+						PackageToSaveTo = GetTransientPackage();
+					}
+					else
+					{
+						PackageToSaveTo = CreatePackage(*InPath);
+					}
+				}
+				else
+				{
+					PackageToSaveTo = GetTransientPackage();
+				}
+
+				UObject* NewMetasoundObject = InEntry.NewObject(PackageToSaveTo, *InDocument.RootClass.Metadata.NodeName);
+				FMetasoundAssetBase* NewAssetBase = InEntry.Cast(NewMetasoundObject);
+				if (ensure(nullptr != NewAssetBase))
+				{
+					NewAssetBase->SetDocument(InDocument);
+				}
+
+#if WITH_EDITOR
+				AsyncTask(ENamedThreads::GameThread, [NewMetasoundObject]()
+				{
+					FAssetRegistryModule::AssetCreated(NewMetasoundObject);
+					NewMetasoundObject->MarkPackageDirty();
+					// todo: how do you get the package for a uobject and save it? I forget
+				});
+#endif
+
+				return NewMetasoundObject;
+			}
+
+			const IMetasoundUObjectRegistryEntry* FindEntryByPredicate(TFunction<bool (const IMetasoundUObjectRegistryEntry*)> InPredicate) const
+			{
+				const IMetasoundUObjectRegistryEntry* const* Entry = Entries.FindByPredicate(InPredicate);
+
+				if (nullptr == Entry)
+				{
+					return nullptr;
+				}
+
+				return *Entry;
+			}
+
+			TArray<const IMetasoundUObjectRegistryEntry*> FindEntriesByPredicate(TFunction<bool (const IMetasoundUObjectRegistryEntry*)> InPredicate) const
+			{
+				TArray<const IMetasoundUObjectRegistryEntry*> FoundEntries;
+
+				Algo::CopyIf(Entries, FoundEntries, InPredicate);
+
+				return FoundEntries;
+			}
+
+			const IMetasoundUObjectRegistryEntry* GetEntryByUObject(const UObject* InObject) const
+			{
+				auto IsChildClassOfRegisteredClass = [&](const IMetasoundUObjectRegistryEntry* Entry)
+				{
+					if (nullptr == Entry)
+					{
+						return false;
+					}
+
+					return Entry->IsChildClass(InObject);
+				};
+
+				return FindEntryByPredicate(IsChildClassOfRegisteredClass);
+			}
+
+			TArray<TUniquePtr<IMetasoundUObjectRegistryEntry>> Storage;
+			TMultiMap<FName, const IMetasoundUObjectRegistryEntry*> EntriesByArchetype;
+			TArray<const IMetasoundUObjectRegistryEntry*> Entries;
+	};
+
+	IMetasoundUObjectRegistry& IMetasoundUObjectRegistry::Get()
+	{
+		static FMetasoundUObjectRegistry Registry;
+		return Registry;
+	}
+}
diff --git a/Engine/Plugins/Runtime/Metasound/Source/MetasoundEngine/Public/Metasound.h b/Engine/Plugins/Runtime/Metasound/Source/MetasoundEngine/Public/Metasound.h
--- a/Engine/Plugins/Runtime/Metasound/Source/MetasoundEngine/Public/Metasound.h
+++ b/Engine/Plugins/Runtime/Metasound/Source/MetasoundEngine/Public/Metasound.h
@@ -30,7 +30,15 @@ class METASOUNDENGINE_API UMetasound : public UObject, public FMetasoundAssetBas
 
 protected:
 	UPROPERTY(EditAnywhere, Category = CustomView)
-	FMetasoundDocument RootMetasoundDocument;
+	FMetasoundDocument MetasoundDocument;
+
+	Metasound::Frontend::FAccessPoint MetasoundDocumentAccessPoint;
+
+	// Returns document object responsible for serializing asset
+	Metasound::Frontend::TAccessPtr<FMetasoundDocument> GetDocument() override;
+
+	// Returns document object responsible for serializing asset
+	Metasound::Frontend::TAccessPtr<const FMetasoundDocument> GetDocument() const override;
 
 #if WITH_EDITORONLY_DATA
 	UPROPERTY()
@@ -45,7 +53,7 @@ public:
 	// while allowing editor to reference directly)
 	static FName GetDocumentPropertyName()
 	{
-		return GET_MEMBER_NAME_CHECKED(UMetasound, RootMetasoundDocument);
+		return GET_MEMBER_NAME_CHECKED(UMetasound, MetasoundDocument);
 	}
 
 	// Returns the graph associated with this Metasound. Graph is required to be referenced on
@@ -62,16 +70,6 @@ public:
 	virtual void SetGraph(UEdGraph* InGraph) override;
 #endif // WITH_EDITORONLY_DATA
 
-	FMetasoundDocument& GetDocument() override
-	{
-		return RootMetasoundDocument;
-	}
-
-	const FMetasoundDocument& GetDocument() const override
-	{
-		return RootMetasoundDocument;
-	}
-
 	UObject* GetOwningAsset() const override
 	{
 		// Hack to allow handles to manipulate while providing
@@ -79,13 +77,16 @@ public:
 		return const_cast<UObject*>(CastChecked<const UObject>(this));
 	}
 
+	const TArray<FMetasoundArchetype>& GetPreferredArchetypes() const override;
+
+	bool IsArchetypeSupported(const FMetasoundArchetype& InArchetype) const override;
+
+	const FMetasoundArchetype& GetPreferredArchetype(const FMetasoundDocument& InDocument) const override;
+
 	// Updates the Metasound's metadata (name, author, etc).
 	// @param InMetadata Metadata containing corrections to the class metadata.
 	void SetMetadata(FMetasoundClassMetadata& InMetadata) override;
 
-	// Get the most up to date archetype for metasound.
-	virtual FMetasoundArchetype GetArchetype() const override;
-
 	void PostLoad() override;
 
-};
\ No newline at end of file
+};
diff --git a/Engine/Plugins/Runtime/Metasound/Source/MetasoundEngine/Public/MetasoundArchetypeRegistration.h b/Engine/Plugins/Runtime/Metasound/Source/MetasoundEngine/Public/MetasoundArchetypeRegistration.h
deleted file mode 100644
diff --git a/Engine/Plugins/Runtime/Metasound/Source/MetasoundEngine/Public/MetasoundSource.h b/Engine/Plugins/Runtime/Metasound/Source/MetasoundEngine/Public/MetasoundSource.h
--- a/Engine/Plugins/Runtime/Metasound/Source/MetasoundEngine/Public/MetasoundSource.h
+++ b/Engine/Plugins/Runtime/Metasound/Source/MetasoundEngine/Public/MetasoundSource.h
@@ -16,6 +16,13 @@
 
 #include "MetasoundSource.generated.h"
 
+/** Declares the output audio format of the UMetasoundSource */
+UENUM()
+enum class EMetasoundSourceAudioFormat : uint8
+{
+	Mono,  //< Mono audio output.
+	Stereo //< Stereo audio output.
+};
 
 /**
  * This Metasound type can be played as an audio source.
@@ -29,6 +36,13 @@ protected:
 	UPROPERTY(EditAnywhere, Category = CustomView)
 	FMetasoundDocument RootMetasoundDocument;
 
+	// MetasoundSourceAccessPoint allows the RootMetasoundDocument to be safely
+	// accessed via the TAccessPtr. The lifetime of MetasoundSourceAccessPoint is
+	// tied to the lifetime of RootMetasoundDocument by their coexistence as members
+	// of the same class. When MetasoundSourceAccessPoint is destructed, the documents
+	// access ptr becomes invalid. 
+	Metasound::Frontend::FAccessPoint MetasoundSourceAccessPoint;
+
 #if WITH_EDITORONLY_DATA
 	UPROPERTY()
 	UEdGraph* Graph;
@@ -37,6 +51,10 @@ protected:
 public:
 	UMetasoundSource(const FObjectInitializer& ObjectInitializer);
 
+	// The output audio format of the metasound source.
+	UPROPERTY(EditAnywhere, BlueprintReadOnly, Category = Audio)
+	EMetasoundSourceAudioFormat OutputFormat;
+
 #if WITH_EDITORONLY_DATA
 	// Returns document name (for editor purposes, and avoids making document public for edit
 	// while allowing editor to reference directly)
@@ -79,15 +97,11 @@ public:
 	}
 #endif // WITH_EDITORONLY_DATA
 
-	FMetasoundDocument& GetDocument() override
-	{
-		return RootMetasoundDocument;
-	}
+#if WITH_EDITOR
 
-	const FMetasoundDocument& GetDocument() const override
-	{
-		return RootMetasoundDocument;
-	}
+	void PostEditChangeProperty(FPropertyChangedEvent& InEvent);
+
+#endif // WITH_EDITOR
 
 	UObject* GetOwningAsset() const override
 	{
@@ -107,7 +121,23 @@ public:
 	void PostLoad() override;
 
 	// Get the most up to date archetype for metasound sources.
-	virtual FMetasoundArchetype GetArchetype() const override;
+	const TArray<FMetasoundArchetype>& GetPreferredArchetypes() const override;
+
+protected:
+
+	Metasound::Frontend::TAccessPtr<FMetasoundDocument> GetDocument() override
+	{
+		// Return document using FAccessPoint to inform the TAccessPtr when the 
+		// object is no longer valid.
+		return Metasound::Frontend::MakeAccessPtr<FMetasoundDocument>(MetasoundSourceAccessPoint, RootMetasoundDocument);
+	}
+
+	Metasound::Frontend::TAccessPtr<const FMetasoundDocument> GetDocument() const override
+	{
+		// Return document using FAccessPoint to inform the TAccessPtr when the 
+		// object is no longer valid.
+		return Metasound::Frontend::MakeAccessPtr<const FMetasoundDocument>(MetasoundSourceAccessPoint, RootMetasoundDocument);
+	}
 
 private:
 
@@ -115,4 +145,7 @@ private:
 	static const FString& GetAudioOutputName();
 	static const FString& GetIsFinishedOutputName();
 	static const FString& GetAudioDeviceHandleVariableName();
+	static const FMetasoundArchetype& GetBaseArchetype();
+	static const FMetasoundArchetype& GetMonoSourceArchetype();
+	static const FMetasoundArchetype& GetStereoSourceArchetype();
 };
diff --git a/Engine/Plugins/Runtime/Metasound/Source/MetasoundEngine/Public/MetasoundUObjectRegistry.h b/Engine/Plugins/Runtime/Metasound/Source/MetasoundEngine/Public/MetasoundUObjectRegistry.h
new file mode 100644
--- /dev/null
+++ b/Engine/Plugins/Runtime/Metasound/Source/MetasoundEngine/Public/MetasoundUObjectRegistry.h
@@ -0,0 +1,197 @@
+// Copyright Epic Games, Inc. All Rights Reserved.
+
+#pragma once
+
+#include "CoreMinimal.h"
+#include "MetasoundAssetBase.h"
+#include "UObject/Object.h"
+
+namespace Metasound
+{
+	/** Interface for an entry into the Metasound-UObject Registry. 
+	 *
+	 * An entry provides information linking a FMetasoundArchetype to a UClass.
+	 * It also provides methods for accessing the FMetasoundAssetBase from a UObject.
+	 */
+	class IMetasoundUObjectRegistryEntry
+	{
+	public:
+		virtual ~IMetasoundUObjectRegistryEntry() = default;
+
+		/** Archetype name associated with this entry. */
+		virtual const FName& GetArchetypeName() const = 0;
+
+		/** UClass associated with this entry. */
+		virtual UClass* GetUClass() const = 0;
+
+		/** Returns true if the UObject is of a Class which is a child of this UClass associated with this entry. */
+		virtual bool IsChildClass(const UObject* InObject) const = 0;
+
+		/** Returns true if the UClass is a child of this UClass associated with this entry. */
+		virtual bool IsChildClass(const UClass* InClass) const = 0;
+
+		/** Attempts to cast the UObject to an FMetasoundAssetBase */
+		virtual FMetasoundAssetBase* Cast(UObject* InObject) const = 0;
+
+		/** Attempts to cast the UObject to an FMetasoundAssetBase */
+		virtual const FMetasoundAssetBase* Cast(const UObject* InObject) const = 0;
+
+		/** Creates a new object of the UClass type. */
+		virtual UObject* NewObject(UPackage* InPackage, const FName& InName) const = 0;
+
+	private:
+		IMetasoundUObjectRegistryEntry() = default;
+
+		/** Only the TMetasoundUObjectRegistryEntry can construct this class. */
+		template<typename UClassType>
+		friend class TMetasoundUObjectRegistryEntry;
+	};
+
+	/** An entry into the Metasound-UObject registry. 
+	 *
+	 * @Tparam UClassType A class which derives from UObject and FMetasoundAssetBase.
+	 */
+	template<typename UClassType>
+	class TMetasoundUObjectRegistryEntry : public IMetasoundUObjectRegistryEntry
+	{
+		// Ensure that this is a subclass of FMetasoundAssetBase and UObject.
+		static_assert(std::is_base_of<FMetasoundAssetBase, UClassType>::value, "UClass must be derived from FMetasoundAssetBase");
+		static_assert(std::is_base_of<UObject, UClassType>::value, "UClass must be derived from UObject");
+
+	public:
+		TMetasoundUObjectRegistryEntry(const FName& InArchetypeName)
+		:	ArchetypeName(InArchetypeName)
+		{
+		}
+
+		virtual ~TMetasoundUObjectRegistryEntry() = default;
+
+		const FName& GetArchetypeName() const
+		{
+			return ArchetypeName;
+		}
+
+		UClass* GetUClass() const override
+		{
+			return UClassType::StaticClass();
+		}
+
+		bool IsChildClass(const UObject* InObject) const override
+		{
+			if (nullptr != InObject)
+			{
+				return InObject->IsA(UClassType::StaticClass());
+			}
+			return false;
+		}
+
+		bool IsChildClass(const UClass* InClass) const override
+		{
+			if (nullptr != InClass)
+			{
+				return InClass->IsChildOf(UClassType::StaticClass());
+			}
+			return false;
+		}
+
+		FMetasoundAssetBase* Cast(UObject* InObject) const override
+		{
+			if (nullptr == InObject)
+			{
+				return nullptr;
+			}
+			return static_cast<FMetasoundAssetBase*>(CastChecked<UClassType>(InObject));
+		}
+
+		const FMetasoundAssetBase* Cast(const UObject* InObject) const override
+		{
+			if (nullptr == InObject)
+			{
+				return nullptr;
+			}
+			return static_cast<const FMetasoundAssetBase*>(CastChecked<const UClassType>(InObject));
+		}
+
+		UObject* NewObject(UPackage* InPackage, const FName& InName) const override
+		{
+			return ::NewObject<UClassType>(InPackage, InName);
+		}
+
+	private:
+
+		FName ArchetypeName;
+	};
+
+
+	/** IMetaoundUObjectRegistry contains IMetasoundUObjectRegistryEntrys. 
+	 *
+	 * Registered UObject classes can utilize the Metasound Editor. It also enables
+	 * the creation of a UObject directly from a FMetasoundDocument.
+	 */
+	class METASOUNDENGINE_API IMetasoundUObjectRegistry
+	{
+		public:
+			virtual ~IMetasoundUObjectRegistry() = default;
+
+			/** Return static singleton instance of the registry. */
+			static IMetasoundUObjectRegistry& Get();
+
+			/** Register all preferred archetypes of the UClass. 
+			 *
+			 * All Archtypes returned by the default objects GetPreferredArchetypes() will be registered. 
+			 */
+			template<typename UClassType>
+			static void RegisterUClassPreferredArchetypes()
+			{
+				static_assert(std::is_base_of<FMetasoundAssetBase, UClassType>::value, "UClass must be derived from FMetasoundAssetBase");
+
+				const TArray<FMetasoundArchetype>& PreferredArchetypes = GetDefault<UClassType>()->GetPreferredArchetypes();
+
+				for (const FMetasoundArchetype& Arch : PreferredArchetypes)
+				{
+					IMetasoundUObjectRegistry::RegisterUClassArchetype<UClassType>(Arch.ArchetypeName);
+				}
+			}
+
+			/** Register an archetype for the UClass. 
+			 *
+			 * @param InArchetypeName - An FMetasoundDocument to associate with the UClass.
+			 */
+			template<typename UClassType>
+			static void RegisterUClassArchetype(const FName& InArchetypeName)
+			{
+				using FRegistryEntryType = TMetasoundUObjectRegistryEntry<UClassType>;
+
+				IMetasoundUObjectRegistry::Get().RegisterUClassArchetype(MakeUnique<FRegistryEntryType>(InArchetypeName));
+			}
+
+			/** Adds an entry to the registry. */
+			virtual void RegisterUClassArchetype(TUniquePtr<IMetasoundUObjectRegistryEntry>&& InEntry) = 0;
+
+			/** Returns all UClasses registered to the archetype name. */
+			virtual TArray<UClass*> GetUClassesForArchetype(const FName& InArchetypeName) const = 0;
+
+			/** Creates a new object from a metasound doucment.
+			 *
+			 * @param InClass - A registered UClass to create.
+			 * @param InDocucument - The FMetasoundDocument to use when creating the class.
+			 * @param InPath - If in editor, the created asset will be stored at this content path.
+			 *
+			 * @return A new object. A nullptr on error.
+			 */
+			virtual UObject* NewObject(UClass* InClass, const FMetasoundDocument& InDocument, const FString& InPath) const = 0;
+
+			/** Returns true if the InObject is of a class or child class which is registered. */
+			virtual bool IsRegisteredClass(UObject* InObject) const = 0;
+
+			/** Returns casts the UObject to a FMetasoundAssetBase if the UObject is of a registered type.
+			 * If the UObject's UClass is not registered, then a nullptr is returned. 
+			 */
+			virtual FMetasoundAssetBase* GetObjectAsAssetBase(UObject* InObject) const = 0;
+
+			/** Returns casts the UObject to a FMetasoundAssetBase if the UObject is of a registered type.
+			 * If the UObject's UClass is not registered, then a nullptr is returned. 
+			 */
+			virtual const FMetasoundAssetBase* GetObjectAsAssetBase(const UObject* InObject) const = 0;
+	};
+}
diff --git a/Engine/Plugins/Runtime/Metasound/Source/MetasoundFrontend/Private/MetasoundArchetype.cpp b/Engine/Plugins/Runtime/Metasound/Source/MetasoundFrontend/Private/MetasoundArchetype.cpp
new file mode 100644
--- /dev/null
+++ b/Engine/Plugins/Runtime/Metasound/Source/MetasoundFrontend/Private/MetasoundArchetype.cpp
@@ -0,0 +1,247 @@
+// Copyright Epic Games, Inc. All Rights Reserved.
+
+#include "MetasoundArchetype.h"
+
+namespace Metasound
+{
+	namespace Frontend
+	{
+		namespace MetasoundArchetypeIntrinsics
+		{
+
+			bool IsLessThanInputDescription(const FMetasoundInputDescription& InputA, const FMetasoundInputDescription& InputB)
+			{
+				// For archetypes we can ignore display name, tooltip and literalvalue.
+				if (InputA.Name == InputB.Name)
+				{
+					return InputA.TypeName.FastLess(InputB.TypeName);
+				}
+				return InputA.Name < InputB.Name;
+			}
+
+			bool IsLessThanOutputDescription(const FMetasoundOutputDescription& OutputA, const FMetasoundOutputDescription& OutputB)
+			{
+				// For archetypes we can ignore display name and tooltip
+				if (OutputA.Name == OutputB.Name)
+				{
+					return OutputA.TypeName.FastLess(OutputB.TypeName);
+				}
+				return OutputA.Name < OutputB.Name;
+			}
+
+			bool IsLessThanEnvironmentDescription(const FMetasoundEnvironmentVariableDescription& EnvironmentA, const FMetasoundEnvironmentVariableDescription& EnvironmentB)
+			{
+				// For archetypes we can ignore display name and tooltip
+				return (EnvironmentA.Name < EnvironmentB.Name);
+			}
+
+			bool IsEquivalentEnvironmentDescription(const FMetasoundEnvironmentVariableDescription& EnvironmentA, const FMetasoundEnvironmentVariableDescription& EnvironmentB)
+			{
+				// For archetypes we can ignore display name and tooltip
+				return (EnvironmentA.Name == EnvironmentB.Name);
+			}
+
+
+
+			template<typename Type>
+			TArray<const Type*> MakePointerToArrayElements(const TArray<Type>& InArray)
+			{
+				TArray<const Type*> PointerArray;
+				for (const Type& Element : InArray)
+				{
+					PointerArray.Add(&Element);
+				}
+				return PointerArray;
+			}
+
+			template<typename DescriptionType, typename ComparisonType>
+			void DescriptionSetDifference(const TArray<DescriptionType>& InDescriptionsA, const TArray<DescriptionType>& InDescriptionsB, TArray<const DescriptionType*>& OutUniqueDescriptionsA, TArray<const DescriptionType*>& OutUniqueDescriptionsB, ComparisonType Compare)
+			{
+				TArray<const DescriptionType*> SortedA = MakePointerToArrayElements(InDescriptionsA);
+				SortedA.Sort(Compare);
+
+				TArray<const DescriptionType*> SortedB = MakePointerToArrayElements(InDescriptionsB);
+				SortedB.Sort(Compare);
+
+				int32 IndexA = 0;
+				int32 IndexB = 0;
+
+				while ((IndexA < SortedA.Num()) && (IndexB < SortedB.Num()))
+				{
+					if (Compare(*SortedA[IndexA], *SortedB[IndexB]))
+					{
+						OutUniqueDescriptionsA.Add(SortedA[IndexA]);
+						IndexA++;
+					}
+					else if (Compare(*SortedB[IndexB], *SortedA[IndexA]))
+					{
+						OutUniqueDescriptionsB.Add(SortedB[IndexB]);
+						IndexB++;
+					}
+					else
+					{
+						// Equal. Increment both indices.
+						IndexA++;
+						IndexB++;
+					}
+				}
+			}
+
+			template<typename DescriptionType, typename ComparisonType>
+			bool IsDescriptionSetEquivalent(const TArray<DescriptionType>& InDescriptionsA, const TArray<DescriptionType>& InDescriptionsB, ComparisonType Compare)
+			{
+				TArray<const DescriptionType*> UniqueA;
+				TArray<const DescriptionType*> UniqueB;
+
+				DescriptionSetDifference(InDescriptionsA, InDescriptionsB, UniqueA, UniqueB, Compare);
+
+				return (UniqueA.Num() == 0) && (UniqueB.Num() == 0);
+			}
+			
+			template<typename DescriptionType, typename ComparisonType>
+			bool IsDescriptionSetIncluded(const TArray<DescriptionType>& InSubset, const TArray<DescriptionType>& InSuperset, ComparisonType Compare)
+			{
+				TArray<const DescriptionType*> UniqueSubset;
+				TArray<const DescriptionType*> UniqueSuperset;
+
+				DescriptionSetDifference(InSubset, InSuperset, UniqueSubset, UniqueSuperset, Compare);
+
+				return UniqueSubset.Num() == 0;
+			}
+
+			template<typename DescriptionType, typename ComparisonType>
+			int32 DescriptionSetDifferenceCount(const TArray<DescriptionType>& InDescriptionsA, const TArray<DescriptionType>& InDescriptionsB, ComparisonType Compare)
+			{
+				TArray<const DescriptionType*> UniqueA;
+				TArray<const DescriptionType*> UniqueB;
+
+				DescriptionSetDifference(InDescriptionsA, InDescriptionsB, UniqueA, UniqueB, Compare);
+
+				return UniqueA.Num() + UniqueB.Num();
+			}
+		}
+
+		bool IsSubsetOfArchetype(const FMetasoundArchetype& InSubsetArchetype, const FMetasoundArchetype& InSupersetArchetype)
+		{
+			using namespace MetasoundArchetypeIntrinsics;
+
+			const bool bIsInputSetIncluded = IsDescriptionSetIncluded(InSubsetArchetype.RequiredInputs, InSupersetArchetype.RequiredInputs, &IsLessThanInputDescription);
+
+			const bool bIsOutputSetIncluded = IsDescriptionSetIncluded(InSubsetArchetype.RequiredOutputs, InSupersetArchetype.RequiredOutputs, &IsLessThanOutputDescription);
+
+			const bool bIsEnvironmentVariableSetIncluded = IsDescriptionSetIncluded(InSubsetArchetype.EnvironmentVariables, InSupersetArchetype.EnvironmentVariables, &IsLessThanEnvironmentDescription);
+
+			return bIsInputSetIncluded && bIsOutputSetIncluded && bIsEnvironmentVariableSetIncluded;
+		}
+
+		bool IsEquivalentArchetype(const FMetasoundArchetype& InInputArchetype, const FMetasoundArchetype& InTargetArchetype)
+		{
+			using namespace MetasoundArchetypeIntrinsics;
+
+			const bool bIsInputSetEquivalent = IsDescriptionSetEquivalent(InTargetArchetype.RequiredInputs, InInputArchetype.RequiredInputs, &IsLessThanInputDescription);
+
+			const bool bIsOutputSetEquivalent = IsDescriptionSetEquivalent(InTargetArchetype.RequiredOutputs, InInputArchetype.RequiredOutputs, &IsLessThanOutputDescription);
+
+			const bool bIsEnvironmentVariableSetEquivalent = IsDescriptionSetEquivalent(InTargetArchetype.EnvironmentVariables, InInputArchetype.EnvironmentVariables, &IsLessThanEnvironmentDescription);
+
+			return bIsInputSetEquivalent && bIsOutputSetEquivalent && bIsEnvironmentVariableSetEquivalent;
+		}
+
+		bool IsEqualArchetype(const FMetasoundArchetype& InInputArchetype, const FMetasoundArchetype& InTargetArchetype)
+		{
+			bool bIsMetadataEqual = InInputArchetype.ArchetypeName == InTargetArchetype.ArchetypeName;
+			bIsMetadataEqual &= InInputArchetype.MajorVersion == InTargetArchetype.MajorVersion;
+			bIsMetadataEqual &= InInputArchetype.MinorVersion == InTargetArchetype.MinorVersion;
+
+			if (bIsMetadataEqual)
+			{
+				// If metadata is equal, then inputs/outputs/environment variables should also be equal. 
+				// Consider updating archetype version info if this ensure is hit.
+				return ensure(IsEquivalentArchetype(InInputArchetype, InTargetArchetype));
+			}
+
+			return false;
+		}
+
+		int32 InputOutputDifferenceCount(const FMetasoundDocument& InDocument, const FMetasoundArchetype& InArchetype)
+		{
+			using namespace MetasoundArchetypeIntrinsics;
+
+			int32 DiffCount = DescriptionSetDifferenceCount(InDocument.RootClass.Inputs, InArchetype.RequiredInputs, &IsLessThanInputDescription);
+
+			DiffCount += DescriptionSetDifferenceCount(InDocument.RootClass.Outputs, InArchetype.RequiredOutputs, &IsLessThanOutputDescription);
+
+			return DiffCount;
+		}
+
+		int32 InputOutputDifferenceCount(const FMetasoundArchetype& InArchetypeA, const FMetasoundArchetype& InArchetypeB)
+		{
+			using namespace MetasoundArchetypeIntrinsics;
+
+			int32 DiffCount = DescriptionSetDifferenceCount(InArchetypeA.RequiredInputs, InArchetypeB.RequiredInputs, &IsLessThanInputDescription);
+
+			DiffCount += DescriptionSetDifferenceCount(InArchetypeA.RequiredOutputs, InArchetypeB.RequiredOutputs, &IsLessThanOutputDescription);
+
+			return DiffCount;
+		}
+
+		void GatherRequiredEnvironmentVariables(const FMetasoundDocument& InDocument, TArray<FMetasoundEnvironmentVariableDescription>& OutEnvironmentVariables)
+		{
+			using namespace MetasoundArchetypeIntrinsics;
+
+			OutEnvironmentVariables = InDocument.RootClass.EnvironmentVariables;
+
+			for (const FMetasoundClassDescription& MetasoundClass : InDocument.Dependencies)
+			{
+				for (const FMetasoundEnvironmentVariableDescription& EnvVar : MetasoundClass.EnvironmentVariables)
+				{
+					auto IsEquivalentEnvVar = [&](const FMetasoundEnvironmentVariableDescription& OtherEnvVar) 
+					{ 
+						return IsEquivalentEnvironmentDescription(EnvVar, OtherEnvVar); 
+					};
+
+					// Basically same as TArray::AddUnique except uses the `IsEquivalentEnvironmentDescription` instead of `operator==` to test for uniqueness.
+					if (nullptr == OutEnvironmentVariables.FindByPredicate(IsEquivalentEnvVar))
+					{
+						OutEnvironmentVariables.Add(EnvVar);
+					}
+				}
+			}
+		}
+
+		const FMetasoundArchetype* FindMostSimilarArchetypeSupportingEnvironment(const FMetasoundDocument& InDocument, const TArray<FMetasoundArchetype>& InCandidateArchetypes)
+		{
+			using namespace MetasoundArchetypeIntrinsics;
+
+			const FMetasoundArchetype* Result = nullptr;
+
+			TArray<const FMetasoundArchetype*> CandidateArchs = MakePointerToArrayElements(InCandidateArchetypes);
+
+			TArray<FMetasoundEnvironmentVariableDescription> AllDocumentEnvironmentVariables;
+			GatherRequiredEnvironmentVariables(InDocument, AllDocumentEnvironmentVariables);
+
+			// Remove all archetypes which do not provide all environment variables. 
+			auto DoesNotProvideAllEnvironmentVariables = [&](const FMetasoundArchetype* CandidateArch)
+			{
+				return !IsDescriptionSetIncluded(AllDocumentEnvironmentVariables, CandidateArch->EnvironmentVariables, &IsLessThanEnvironmentDescription);
+			};
+
+			CandidateArchs.RemoveAll(DoesNotProvideAllEnvironmentVariables);
+
+			// Return the archetype with the least amount of differences.
+			auto DifferencesFromDocument = [&](const FMetasoundArchetype* CandidateArch) -> int32
+			{ 
+				return InputOutputDifferenceCount(InDocument, *CandidateArch); 
+			};
+
+			Algo::SortBy(CandidateArchs, DifferencesFromDocument);
+
+			if (0 == CandidateArchs.Num())
+			{
+				return nullptr;
+			}
+
+			return CandidateArchs[0];
+		}
+	}
+}
diff --git a/Engine/Plugins/Runtime/Metasound/Source/MetasoundFrontend/Private/MetasoundAssetBase.cpp b/Engine/Plugins/Runtime/Metasound/Source/MetasoundFrontend/Private/MetasoundAssetBase.cpp
--- a/Engine/Plugins/Runtime/Metasound/Source/MetasoundFrontend/Private/MetasoundAssetBase.cpp
+++ b/Engine/Plugins/Runtime/Metasound/Source/MetasoundFrontend/Private/MetasoundAssetBase.cpp
@@ -2,56 +2,112 @@
 
 #include "MetasoundAssetBase.h"
 
-#include "MetasoundJsonBackend.h"
+#include "Algo/AnyOf.h"
 #include "HAL/FileManager.h"
 #include "IStructSerializerBackend.h"
+#include "MetasoundArchetype.h"
+#include "MetasoundJsonBackend.h"
 #include "StructSerializer.h"
 
-
 const FString FMetasoundAssetBase::FileExtension(TEXT(".metasound"));
 
 FMetasoundAssetBase::FMetasoundAssetBase()
 {
 }
 
-FMetasoundAssetBase::FMetasoundAssetBase(FMetasoundDocument& InDocument)
-	: AccessPoint(MakeShared<Metasound::Frontend::FDescriptionAccessPoint>(InDocument))
+FMetasoundAssetBase::FMetasoundAssetBase(const FMetasoundArchetype& InArchetype)
+:	Archetype(InArchetype)
 {
 }
 
 void FMetasoundAssetBase::SetMetadata(FMetasoundClassMetadata& InMetadata)
 {
 	InMetadata.NodeType = EMetasoundClassType::MetasoundGraph;
-	GetDocument().RootClass.Metadata = InMetadata;
+	GetDocumentChecked().RootClass.Metadata = InMetadata;
+}
+
+const FMetasoundArchetype& FMetasoundAssetBase::GetArchetype() const
+{
+	return Archetype;
+}
+
+bool FMetasoundAssetBase::SetArchetype(const FMetasoundArchetype& InArchetype)
+{
+	if (IsArchetypeSupported(InArchetype))
+	{
+		Archetype = InArchetype;
+		ConformDocumentToArchetype();
+
+		return true;
+	}
+
+	// Archetype was not set because it is not supported for this asset base. 
+	return false;
+}
+
+bool FMetasoundAssetBase::IsArchetypeSupported(const FMetasoundArchetype& InArchetype) const
+{
+	auto IsEqualArchetype = [&](const FMetasoundArchetype& SupportedArchetype)
+	{
+		return Metasound::Frontend::IsEqualArchetype(InArchetype, SupportedArchetype);
+	};
+
+	return Algo::AnyOf(GetPreferredArchetypes(), IsEqualArchetype);
+}
+
+const FMetasoundArchetype& FMetasoundAssetBase::GetPreferredArchetype(const FMetasoundDocument& InDocument) const
+{
+	// Default to archetype provided in case it is supported. 
+	if (IsArchetypeSupported(InDocument.Archetype))
+	{
+		return InDocument.Archetype;
+	}
+
+	// If existing archetype is not supported, get the most similar that still supports the documents environment.
+	const FMetasoundArchetype* SimilarArchetype = Metasound::Frontend::FindMostSimilarArchetypeSupportingEnvironment(InDocument, GetPreferredArchetypes());
+
+	if (nullptr != SimilarArchetype)
+	{
+		return *SimilarArchetype;
+	}
+
+	// Nothing found. Return the existing archetype for the FMetasoundAssetBase.
+	return GetArchetype();
 }
 
-void FMetasoundAssetBase::SetDocument(const FMetasoundDocument& InDocument)
+void FMetasoundAssetBase::SetDocument(const FMetasoundDocument& InDocument, bool bForceUpdateArchetype)
 {
-	FMetasoundDocument& Document = GetDocument();
+	FMetasoundDocument& Document = GetDocumentChecked();
 	Document = InDocument;
 
-	Document.Archetype = GetArchetype();
-	GetRootGraphHandle().FixDocumentToMatchArchetype();
+	FMetasoundArchetype NewArch = InDocument.Archetype;
+
+	if (bForceUpdateArchetype || (!IsArchetypeSupported(NewArch)))
+	{
+		NewArch = GetPreferredArchetype(Document);
+	}
+
+	ensure(SetArchetype(NewArch));
 }
 
 void FMetasoundAssetBase::ConformDocumentToArchetype()
 {
-	FMetasoundDocument& Document = GetDocument();
+	FMetasoundDocument& Document = GetDocumentChecked();
 	Document.Archetype = GetArchetype();
 	GetRootGraphHandle().FixDocumentToMatchArchetype();
 }
 
 FMetasoundClassMetadata FMetasoundAssetBase::GetMetadata()
 {
-	return GetDocument().RootClass.Metadata;
+	return GetDocumentChecked().RootClass.Metadata;
 }
 
-TWeakPtr<Metasound::Frontend::FDescriptionAccessPoint> FMetasoundAssetBase::GetGraphAccessPoint() const
+Metasound::Frontend::FDescriptionAccessPoint FMetasoundAssetBase::GetGraphAccessPoint() 
 {
-	return TWeakPtr<Metasound::Frontend::FDescriptionAccessPoint>(AccessPoint);
+	return Metasound::Frontend::FDescriptionAccessPoint(GetDocument());
 }
 
-Metasound::Frontend::FGraphHandle FMetasoundAssetBase::GetRootGraphHandle() const
+Metasound::Frontend::FGraphHandle FMetasoundAssetBase::GetRootGraphHandle()
 {
 	using FDescPath = Metasound::Frontend::FDescPath;
 	using FGraphHandle = Metasound::Frontend::FGraphHandle;
@@ -60,7 +116,7 @@ Metasound::Frontend::FGraphHandle FMetasoundAssetBase::GetRootGraphHandle() cons
 	using EFromClass = Metasound::Frontend::Path::EFromClass;
 
 	FDescPath PathToGraph = FDescPath()[EFromDocument::ToRootClass][EFromClass::ToGraph];
-	FHandleInitParams InitParams = { GetGraphAccessPoint(), PathToGraph, INDEX_NONE, GetDocument().RootClass.UniqueID, MakeWeakObjectPtr(GetOwningAsset()) };
+	FHandleInitParams InitParams = { GetGraphAccessPoint(), PathToGraph, INDEX_NONE, GetDocumentChecked().RootClass.UniqueID, MakeWeakObjectPtr(GetOwningAsset()) };
 	return FGraphHandle(GetPrivateToken(), InitParams);
 }
 
@@ -76,7 +132,7 @@ TArray<Metasound::Frontend::FGraphHandle> FMetasoundAssetBase::GetAllSubgraphHan
 
 	FDescPath RootPathForDependencyGraphs = FDescPath()[EFromDocument::ToDependencies];
 
-	TArray<FMetasoundClassDescription>& DependenciesList = GetDocument().Dependencies;
+	TArray<FMetasoundClassDescription>& DependenciesList = GetDocumentChecked().Dependencies;
 	for (FMetasoundClassDescription& Dependency : DependenciesList)
 	{
 		const bool bIsSubgraph = Dependency.Metadata.NodeType == EMetasoundClassType::MetasoundGraph && Dependency.Graph.Nodes.Num() > 0;
@@ -93,10 +149,36 @@ TArray<Metasound::Frontend::FGraphHandle> FMetasoundAssetBase::GetAllSubgraphHan
 
 bool FMetasoundAssetBase::ImportFromJSON(const FString& InJSON)
 {
-	return Metasound::Frontend::ImportJSONToMetasound(InJSON, GetDocument());
+	Metasound::Frontend::TAccessPtr<FMetasoundDocument> Document = GetDocument();
+	if (ensure(Document.IsValid()))
+	{
+		return Metasound::Frontend::ImportJSONToMetasound(InJSON, *Document);
+	}
+	return false;
 }
 
 bool FMetasoundAssetBase::ImportFromJSONAsset(const FString& InAbsolutePath)
 {
-	return Metasound::Frontend::ImportJSONAssetToMetasound(InAbsolutePath, GetDocument());
+	Metasound::Frontend::TAccessPtr<FMetasoundDocument> Document = GetDocument();
+	if (ensure(Document.IsValid()))
+	{
+		return Metasound::Frontend::ImportJSONAssetToMetasound(InAbsolutePath, *Document);
+	}
+	return false;
+}
+
+FMetasoundDocument& FMetasoundAssetBase::GetDocumentChecked()
+{
+	Metasound::Frontend::TAccessPtr<FMetasoundDocument> DocAccessPtr = GetDocument();
+
+	check(DocAccessPtr.IsValid());
+	return *DocAccessPtr;
+}
+
+const FMetasoundDocument& FMetasoundAssetBase::GetDocumentChecked() const
+{
+	Metasound::Frontend::TAccessPtr<const FMetasoundDocument> DocAccessPtr = GetDocument();
+
+	check(DocAccessPtr.IsValid());
+	return *DocAccessPtr;
 }
diff --git a/Engine/Plugins/Runtime/Metasound/Source/MetasoundFrontend/Private/MetasoundFrontend.cpp b/Engine/Plugins/Runtime/Metasound/Source/MetasoundFrontend/Private/MetasoundFrontend.cpp
--- a/Engine/Plugins/Runtime/Metasound/Source/MetasoundFrontend/Private/MetasoundFrontend.cpp
+++ b/Engine/Plugins/Runtime/Metasound/Source/MetasoundFrontend/Private/MetasoundFrontend.cpp
@@ -201,116 +201,6 @@ namespace Metasound
 			return false;
 		}
 
-		// Struct used for any pertinent archetype data.
-		struct FArchetypeRegistryElement
-		{
-			FMetasoundArchetype Archetype;
-			UClass* ArchetypeUClass;
-
-			// Constructor used to generate an instance of the UObject version of this archetype from scratch.
-			TUniqueFunction<UObject* (const FMetasoundDocument&, const FString&)> ObjectConstructor;
-
-			// template-generated lambdas used to safely sidecast to FMetasoundBase*.
-			TUniqueFunction<FMetasoundAssetBase* (UObject*)> SafeCast;
-			TUniqueFunction<const FMetasoundAssetBase* (const UObject*)> SafeConstCast;
-		};
-
-		static TMap<FName, FArchetypeRegistryElement> ArchetypeRegistry;
-
-		bool RegisterArchetype_Internal(FMetasoundArchetypeRegistryParams_Internal&& InParams)
-		{
-			FName ArchetypeName = InParams.ArchetypeDescription.ArchetypeName;
-
-			if (!ArchetypeRegistry.Contains(ArchetypeName))
-			{
-				FArchetypeRegistryElement RegistryElement = 
-				{ 
-					InParams.ArchetypeDescription,
-					InParams.ArchetypeUClass,
-					MoveTemp(InParams.ObjectGetter),
-					MoveTemp(InParams.SafeCast),
-					MoveTemp(InParams.SafeConstCast)
-				};
-
-				ArchetypeRegistry.Add(ArchetypeName, MoveTemp(RegistryElement));
-
-				return true;
-			}
-			else
-			{
-				return false;
-			}
-		}
-
-		TArray<FName> GetAllRegisteredArchetypes()
-		{
-			TArray<FName> OutArchetypes;
-			for (auto& RegistryTuple : ArchetypeRegistry)
-			{
-				OutArchetypes.Add(RegistryTuple.Key);
-			}
-
-			return OutArchetypes;
-		}
-
-		UObject* GetObjectForDocument(const FMetasoundDocument& InDocument, const FString& InPath)
-		{
-			FName ArchetypeName = InDocument.Archetype.ArchetypeName;
-			if (ArchetypeRegistry.Contains(ArchetypeName))
-			{
-				return ArchetypeRegistry[ArchetypeName].ObjectConstructor(InDocument, InPath);
-			}
-			else
-			{
-				return nullptr;
-			}
-		}
-
-		bool IsObjectAMetasoundArchetype(const UObject* InObject)
-		{
-			UClass* ObjectClass = InObject->GetClass();
-
-			for (auto& RegistryTuple : ArchetypeRegistry)
-			{
-				if (ObjectClass == RegistryTuple.Value.ArchetypeUClass)
-				{
-					return true;
-				}
-			}
-
-			return false;
-		}
-
-		FMetasoundAssetBase* GetObjectAsAssetBase(UObject* InObject)
-		{
-			UClass* ObjectClass = InObject->GetClass();
-
-			for (auto& RegistryTuple : ArchetypeRegistry)
-			{
-				if (ObjectClass == RegistryTuple.Value.ArchetypeUClass)
-				{
-					return RegistryTuple.Value.SafeCast(InObject);
-				}
-			}
-
-			return nullptr;
-		}
-
-		const FMetasoundAssetBase* GetObjectAsAssetBase(const UObject* InObject)
-		{
-			UClass* ObjectClass = InObject->GetClass();
-
-			for (auto& RegistryTuple : ArchetypeRegistry)
-			{
-				if (ObjectClass == RegistryTuple.Value.ArchetypeUClass)
-				{
-					return RegistryTuple.Value.SafeConstCast(InObject);
-				}
-			}
-
-			return nullptr;
-		}
-
 		void FGraphHandle::FixDocumentToMatchArchetype()
 		{
 			// TODO: Also check if this is the root class.
@@ -326,19 +216,22 @@ namespace Metasound
 			TArray<FMetasoundInputDescription>& CurrentInputs = DocumentToFixUp.RootClass.Inputs;
 			for (const FMetasoundInputDescription& RequiredInput : RequiredInputs)
 			{
-				bool bFoundRequiredInput = false;
-				for (const FMetasoundInputDescription& Input : CurrentInputs)
+				auto IsEqualInputName = [&](const FMetasoundInputDescription& InputDesc)
 				{
-					if (Input.TypeName == RequiredInput.TypeName && Input.Name == RequiredInput.Name)
+					return RequiredInput.Name == InputDesc.Name;
+				};
+
+				if (const FMetasoundInputDescription* InputWithSameName = CurrentInputs.FindByPredicate(IsEqualInputName))
+				{
+					if (InputWithSameName->TypeName != RequiredInput.TypeName)
 					{
-						bFoundRequiredInput = true;
-						break;
+						// Swap existing input because types do not match.
+						RemoveInput(InputWithSameName->Name);
+						AddNewInput(RequiredInput);
 					}
 				}
-
-				if (!bFoundRequiredInput)
+				else
 				{
-					//CurrentInputs.Add(RequiredInput);
 					AddNewInput(RequiredInput);
 				}
 			}
@@ -348,19 +241,22 @@ namespace Metasound
 			TArray<FMetasoundOutputDescription>& CurrentOutputs = DocumentToFixUp.RootClass.Outputs;
 			for (const FMetasoundOutputDescription& RequiredOutput : RequiredOutputs)
 			{
-				bool bFoundRequiredOutput = false;
-				for (const FMetasoundOutputDescription& Output : CurrentOutputs)
+				auto IsEqualOutputName = [&](const FMetasoundOutputDescription& OutputDesc)
+				{
+					return RequiredOutput.Name == OutputDesc.Name;
+				};
+
+				if (const FMetasoundOutputDescription* OutputWithSameName = CurrentOutputs.FindByPredicate(IsEqualOutputName))
 				{
-					if (Output.TypeName == RequiredOutput.TypeName && Output.Name == RequiredOutput.Name)
+					if (OutputWithSameName->TypeName != RequiredOutput.TypeName)
 					{
-						bFoundRequiredOutput = true;
-						break;
+						// Swap existing input because types do not match.
+						RemoveOutput(OutputWithSameName->Name);
+						AddNewOutput(RequiredOutput);
 					}
 				}
-
-				if (!bFoundRequiredOutput)
+				else
 				{
-					//CurrentOutputs.Add(RequiredOutput);
 					AddNewOutput(RequiredOutput);
 				}
 			}
@@ -371,7 +267,7 @@ namespace Metasound
 			, NodePtr(InParams.InAccessPoint, InParams.InPath)
 			, NodeClass(InParams.InAccessPoint, Path::GetDependencyPath(InParams.DependencyID))
 			, InputPtr(InParams.InAccessPoint, NodeClass.GetPath()[Path::EFromClass::ToInputs][*InputName])
-			, OutputNodePtr(nullptr, FDescPath())
+			, OutputNodePtr(TDescriptionPtr<FMetasoundOutputDescription>::CreateInvalid())
 		{
 			if (InParams.InAccessPoint.IsValid())
 			{
@@ -383,8 +279,8 @@ namespace Metasound
 		FInputHandle::FInputHandle(FHandleInitParams::EPrivateToken PrivateToken, const FHandleInitParams& InParams, const FString& InputName, EOutputNodeTag InTag)
 			: ITransactable(MetasoundUndoRollLimitCvar, InParams.InOwningAsset)
 			, NodePtr(InParams.InAccessPoint, InParams.InPath)
-			, NodeClass(nullptr, FDescPath())
-			, InputPtr(nullptr, FDescPath())
+			, NodeClass(TDescriptionPtr<FMetasoundClassDescription>::CreateInvalid())
+			, InputPtr(TDescriptionPtr<FMetasoundInputDescription>::CreateInvalid())
 			, OutputNodePtr(InParams.InAccessPoint, Path::GetOutputDescriptionPath(InParams.InPath, InputName))
 		{
 			if (InParams.InAccessPoint.IsValid())
@@ -398,7 +294,7 @@ namespace Metasound
 		{
 			FDescPath NullPath = FDescPath();
 			FString NullString = FString();
-			FHandleInitParams InitParams = { nullptr, NullPath , FMetasoundNodeDescription::InvalidID, FMetasoundClassDescription::InvalidID, nullptr };
+			FHandleInitParams InitParams = { FDescriptionAccessPoint::CreateInvalid(), NullPath , FMetasoundNodeDescription::InvalidID, FMetasoundClassDescription::InvalidID, nullptr };
 			return FInputHandle(FHandleInitParams::PrivateToken, InitParams, NullString, FInputHandle::Default);
 		}
 
@@ -492,6 +388,16 @@ namespace Metasound
 			}
 		}
 
+		int32 FInputHandle::GetOwningNodeID() const
+		{
+			if (!IsValid())
+			{
+				return FMetasoundNodeDescription::InvalidID;
+			}
+
+			return NodePtr->UniqueID;
+		}
+
 		FOutputHandle FInputHandle::GetCurrentlyConnectedOutput() const
 		{
 			if (!IsValid())
@@ -513,15 +419,24 @@ namespace Metasound
 			// All node connections are in the same graph, so we just need to go up one level to the Nodes array
 			// and look up the node by it's unique ID.
 			FDescPath OutputNodePath = (NodePtr.GetPath() << 1)[OutputNodeID];
+			TDescriptionPtr<FMetasoundNodeDescription> OtherNodePtr(NodePtr.GetAccessPoint(), OutputNodePath);
 
-			if (NodeClass.IsValid())
+			int32 OutputNodeDependencyId = FMetasoundClassDescription::InvalidID;
+			EMetasoundClassType OutputNodeType = EMetasoundClassType::Invalid;
+			if (OtherNodePtr.IsValid())
 			{
-				FHandleInitParams InitParams = { NodePtr.GetAccessPoint(), OutputNodePath, FMetasoundNodeDescription::InvalidID, NodeClass->UniqueID, OwningAsset };
-				return FOutputHandle(FHandleInitParams::PrivateToken, InitParams, OutputName, FOutputHandle::Default);
+				OutputNodeDependencyId = OtherNodePtr->DependencyID;
+				OutputNodeType = OtherNodePtr->ObjectTypeOfNode;
 			}
-			else if (OutputNodePtr.IsValid())
+
+			if (EMetasoundClassType::Input == OutputNodeType)
+			{
+				FHandleInitParams InitParams = { NodePtr.GetAccessPoint(), OutputNodePath, OutputNodeID, FMetasoundClassDescription::InvalidID, OwningAsset };
+				return FOutputHandle(FHandleInitParams::PrivateToken, InitParams, OutputName, FOutputHandle::InputNode);
+			}
+			else if (EMetasoundClassType::Invalid != OutputNodeType)
 			{
-				FHandleInitParams InitParams = { NodePtr.GetAccessPoint(), OutputNodePath, FMetasoundNodeDescription::InvalidID, FMetasoundClassDescription::InvalidID, OwningAsset };
+				FHandleInitParams InitParams = { NodePtr.GetAccessPoint(), OutputNodePath, OutputNodeID, OutputNodeDependencyId, OwningAsset };
 				return FOutputHandle(FHandleInitParams::PrivateToken, InitParams, OutputName, FOutputHandle::Default);
 			}
 			else
@@ -710,7 +625,7 @@ namespace Metasound
 			, NodePtr(InParams.InAccessPoint, InParams.InPath)
 			, NodeClass(InParams.InAccessPoint, Path::GetDependencyPath(InParams.DependencyID))
 			, OutputPtr(InParams.InAccessPoint, NodeClass.GetPath()[Path::EFromClass::ToOutputs][*InOutputName])
-			, InputNodePtr(nullptr, FDescPath())
+			, InputNodePtr(TDescriptionPtr<FMetasoundInputDescription>::CreateInvalid())
 		{
 			if (InParams.InAccessPoint.IsValid())
 			{
@@ -722,8 +637,8 @@ namespace Metasound
 		FOutputHandle::FOutputHandle(FHandleInitParams::EPrivateToken InToken, const FHandleInitParams& InParams, const FString& InOutputName, FOutputHandle::EInputNodeTag InTag)
 			: ITransactable(MetasoundUndoRollLimitCvar, InParams.InOwningAsset)
 			, NodePtr(InParams.InAccessPoint, InParams.InPath)
-			, NodeClass(nullptr, FDescPath())
-			, OutputPtr(nullptr, FDescPath())
+			, NodeClass(TDescriptionPtr<FMetasoundClassDescription>::CreateInvalid())
+			, OutputPtr(TDescriptionPtr<FMetasoundOutputDescription>::CreateInvalid())
 			, InputNodePtr(InParams.InAccessPoint, Path::GetInputDescriptionPath(InParams.InPath, InOutputName))
 		{
 			if (InParams.InAccessPoint.IsValid())
@@ -737,7 +652,7 @@ namespace Metasound
 		{
 			FString NullString;
 			FDescPath NullPath;
-			FHandleInitParams InitParams = { nullptr, NullPath, FMetasoundNodeDescription::InvalidID, FMetasoundClassDescription::InvalidID, nullptr };
+			FHandleInitParams InitParams = { FDescriptionAccessPoint::CreateInvalid(), NullPath, FMetasoundNodeDescription::InvalidID, FMetasoundClassDescription::InvalidID, nullptr };
 			return FOutputHandle(FHandleInitParams::PrivateToken, InitParams, NullString, FOutputHandle::Default);
 		}
 
@@ -862,7 +777,7 @@ namespace Metasound
 			else
 			{
 				// input nodes and output nodes don't have class descriptions.
-				return TDescriptionPtr<FMetasoundClassDescription>(nullptr, FDescPath());
+				return TDescriptionPtr<FMetasoundClassDescription>::CreateInvalid();
 			}
 		}
 
@@ -883,7 +798,7 @@ namespace Metasound
 		{
 			FDescPath InvalidPath;
 			FString InvalidClassName;
-			FHandleInitParams InitParams = { nullptr, InvalidPath, FMetasoundNodeDescription::InvalidID, FMetasoundClassDescription::InvalidID, nullptr };
+			FHandleInitParams InitParams = { FDescriptionAccessPoint::CreateInvalid(), InvalidPath, FMetasoundNodeDescription::InvalidID, FMetasoundClassDescription::InvalidID, nullptr };
 			return FNodeHandle(FHandleInitParams::PrivateToken, InitParams, EMetasoundClassType::Invalid);
 		}
 
@@ -1302,7 +1217,7 @@ namespace Metasound
 		{
 			FDescPath InvalidPath;
 			FString InvalidClassName;
-			FHandleInitParams InitParams = { nullptr, InvalidPath, FMetasoundNodeDescription::InvalidID, FMetasoundClassDescription::InvalidID, nullptr };
+			FHandleInitParams InitParams = { FDescriptionAccessPoint::CreateInvalid(), InvalidPath, FMetasoundNodeDescription::InvalidID, FMetasoundClassDescription::InvalidID, nullptr };
 			return FGraphHandle(FHandleInitParams::Token, InitParams);
 		}
 
@@ -2141,7 +2056,7 @@ namespace Metasound
 			{
 				
 				TJsonStructSerializerBackend<DefaultCharType> Backend(*FileWriter, EStructSerializerBackendFlags::Default);
-				FStructSerializer::Serialize<FMetasoundClassDescription>(GraphsClassDeclaration.GetChecked(), Backend);
+				FStructSerializer::Serialize<FMetasoundDocument>(*OwningDocument, Backend);
 		
 				FileWriter->Close();
 
diff --git a/Engine/Plugins/Runtime/Metasound/Source/MetasoundFrontend/Private/MetasoundFrontendBaseClasses.cpp b/Engine/Plugins/Runtime/Metasound/Source/MetasoundFrontend/Private/MetasoundFrontendBaseClasses.cpp
--- a/Engine/Plugins/Runtime/Metasound/Source/MetasoundFrontend/Private/MetasoundFrontendBaseClasses.cpp
+++ b/Engine/Plugins/Runtime/Metasound/Source/MetasoundFrontend/Private/MetasoundFrontendBaseClasses.cpp
@@ -11,8 +11,6 @@ namespace Metasound
 {
 	namespace Frontend
 	{
-		
-
 		TMap<FNodeRegistryKey, FNodeRegistryElement>& GetExternalNodeRegistry()
 		{
 			return FMetasoundFrontendRegistryContainer::Get()->GetExternalNodeRegistry();
@@ -508,179 +506,44 @@ namespace Metasound
 			return true;
 		}
 
-		FDescriptionAccessPoint::FDescriptionAccessPoint(FMetasoundDocument& InRootDocument)
-			: RootDocument(InRootDocument)
-		{
-		}
-
-		FMetasoundDocument& FDescriptionAccessPoint::GetRoot()
+		FDescriptionAccessPoint::FDescriptionAccessPoint(TAccessPtr<FMetasoundDocument> InRootDocumentPtr)
+			: RootDocumentPtr(InRootDocumentPtr)
 		{
-			return RootDocument;
 		}
 
-		FMetasoundClassDescription* FDescriptionAccessPoint::GetClassFromPath(const FDescPath& InPathFromRoot)
+		FMetasoundDocument& FDescriptionAccessPoint::GetRootChecked() const
 		{
-			FDescPath CurrentPath = InPathFromRoot;
-
-			FMetasoundDescriptionPtr Ptr;
-			Ptr.Set<FMetasoundDocument*>(&RootDocument);
-			FDescriptionUnwindStep CurrentStep = { Ptr, Path::EDescType::Document };
-
-			while (CurrentPath.Num() != 0 && CurrentStep.Type != Path::EDescType::Invalid)
-			{
-				CurrentStep = GoToNext(CurrentPath, CurrentStep);
-			}
-
-			if (ensureAlwaysMsgf(CurrentStep.Type == Path::EDescType::Class && CurrentPath.Num() == 0, TEXT("Couldn't resolve part of the path.")))
-			{
-				return CurrentStep.DescriptionStructPtr.Get<FMetasoundClassDescription*>();
-			}
-			else
-			{
-				return nullptr;
-			}
-		}
-
-		FMetasoundNodeDescription* FDescriptionAccessPoint::GetNodeFromPath(const FDescPath& InPathFromRoot)
-		{
-			FDescPath CurrentPath = InPathFromRoot;
-			
-			FMetasoundDescriptionPtr Ptr;
-			Ptr.Set<FMetasoundDocument*>(&RootDocument);
-			FDescriptionUnwindStep CurrentStep = { Ptr, Path::EDescType::Document };
-
-			while (CurrentPath.Num() != 0 && CurrentStep.Type != Path::EDescType::Invalid)
-			{
-				CurrentStep = GoToNext(CurrentPath, CurrentStep);
-			}
+			check(RootDocumentPtr.IsValid());
 
-			if (ensureAlwaysMsgf(CurrentPath.Num() == 0, TEXT("Couldn't resolve part of the path.")))
-			{
-				return CurrentStep.DescriptionStructPtr.Get<FMetasoundNodeDescription*>();
-			}
-			else
-			{
-				return nullptr;
-			}
+			return *RootDocumentPtr;
 		}
 
-		FMetasoundGraphDescription* FDescriptionAccessPoint::GetGraphFromPath(const FDescPath& InPathFromRoot)
+		bool FDescriptionAccessPoint::GetDescriptionPtrFromPath(const FDescPath& InPathFromRoot, FMetasoundDescriptionPtr& OutPtr) const
 		{
-			FDescPath CurrentPath = InPathFromRoot;
-			
-			FMetasoundDescriptionPtr Ptr;
-			Ptr.Set<FMetasoundDocument*>(&RootDocument);
-			FDescriptionUnwindStep CurrentStep = { Ptr, Path::EDescType::Document };
-
-			while (CurrentPath.Num() != 0 && CurrentStep.Type != Path::EDescType::Invalid)
+			if (!RootDocumentPtr.IsValid())
 			{
-				CurrentStep = GoToNext(CurrentPath, CurrentStep);
-			}
-
-			if (ensureAlwaysMsgf(CurrentPath.Num() == 0, TEXT("Couldn't resolve part of the path.")))
-			{
-				return CurrentStep.DescriptionStructPtr.Get<FMetasoundGraphDescription*>();
-			}
-			else
-			{
-				return nullptr;
-			}
-		}
-
-		FMetasoundInputDescription* FDescriptionAccessPoint::GetInputFromPath(const FDescPath& InPathFromRoot)
-		{
-			FDescPath CurrentPath = InPathFromRoot;
-			
-			FMetasoundDescriptionPtr Ptr;
-			Ptr.Set<FMetasoundDocument*>(&RootDocument);
-			FDescriptionUnwindStep CurrentStep = { Ptr, Path::EDescType::Document };
-
-			while (CurrentPath.Num() != 0 && CurrentStep.Type != Path::EDescType::Invalid)
-			{
-				CurrentStep = GoToNext(CurrentPath, CurrentStep);
-			}
-
-			if (ensureAlwaysMsgf(CurrentPath.Num() == 0, TEXT("Couldn't resolve part of the path.")))
-			{
-				return CurrentStep.DescriptionStructPtr.Get<FMetasoundInputDescription*>();
-			}
-			else
-			{
-				return nullptr;
+				return false;
 			}
-		}
 
-		FMetasoundOutputDescription* FDescriptionAccessPoint::GetOutputFromPath(const FDescPath& InPathFromRoot)
-		{
 			FDescPath CurrentPath = InPathFromRoot;
 			
-			FMetasoundDescriptionPtr Ptr;
-			Ptr.Set<FMetasoundDocument*>(&RootDocument);
-			FDescriptionUnwindStep CurrentStep = { Ptr, Path::EDescType::Document };
+			OutPtr.Set<FMetasoundDocument*>(RootDocumentPtr.Get());
+			FDescriptionUnwindStep CurrentStep = { OutPtr, Path::EDescType::Document };
 
 			while (CurrentPath.Num() != 0 && CurrentStep.Type != Path::EDescType::Invalid)
 			{
 				CurrentStep = GoToNext(CurrentPath, CurrentStep);
 			}
 
-			if (ensureAlwaysMsgf(CurrentPath.Num() == 0, TEXT("Couldn't resolve part of the path.")))
-			{
-				return CurrentStep.DescriptionStructPtr.Get<FMetasoundOutputDescription*>();
-			}
-			else
-			{
-				return nullptr;
-			}
-		}
+			OutPtr = CurrentStep.DescriptionStructPtr;
 
-		FMetasoundClassMetadata* FDescriptionAccessPoint::GetMetadataFromPath(const FDescPath& InPathFromRoot)
-		{
-			FDescPath CurrentPath = InPathFromRoot;
-			
-			FMetasoundDescriptionPtr Ptr;
-			Ptr.Set<FMetasoundDocument*>(&RootDocument);
-			FDescriptionUnwindStep CurrentStep = { Ptr, Path::EDescType::Document };
-
-			while (CurrentPath.Num() != 0 && CurrentStep.Type != Path::EDescType::Invalid)
-			{
-				CurrentStep = GoToNext(CurrentPath, CurrentStep);
-			}
-
-			if (ensureAlwaysMsgf(CurrentPath.Num() == 0, TEXT("Couldn't resolve part of the path.")))
-			{
-				return CurrentStep.DescriptionStructPtr.Get<FMetasoundClassMetadata*>();
-			}
-			else
-			{
-				return nullptr;
-			}
+			return CurrentStep.Type != Path::EDescType::Invalid;
 		}
 
-		Metasound::Frontend::FClassDependencyIDs* FDescriptionAccessPoint::GetClassDependencyIDsFromPath(const FDescPath& InPathFromRoot)
-		{
-			FDescPath CurrentPath = InPathFromRoot;
-			
-			FMetasoundDescriptionPtr Ptr;
-			Ptr.Set<FMetasoundDocument*>(&RootDocument);
-			FDescriptionUnwindStep CurrentStep = { Ptr, Path::EDescType::Document };
-
-			while (CurrentPath.Num() != 0 && CurrentStep.Type != Path::EDescType::Invalid)
-			{
-				CurrentStep = GoToNext(CurrentPath, CurrentStep);
-			}
-
-			if (ensureAlwaysMsgf(CurrentStep.Type != Path::EDescType::Invalid && CurrentPath.Num() == 0, TEXT("Couldn't resolve part of the path.")))
-			{
-				return CurrentStep.DescriptionStructPtr.Get<FClassDependencyIDs*>();
-			}
-			else
-			{
-				return nullptr;
-			}
-		}
 
-		FDescriptionAccessPoint::FDescriptionUnwindStep FDescriptionAccessPoint::GoToNextFromDocument(FMetasoundDocument& InDocument, FDescPath& InPath, const Path::FElement& InNext)
+		FDescriptionAccessPoint::FDescriptionUnwindStep FDescriptionAccessPoint::GoToNextFromDocument(FMetasoundDocument& InDocument, FDescPath& InPath, const Path::FElement& InNext) const
 		{
+			check(RootDocumentPtr.IsValid());
 			switch (InNext.CurrentDescType)
 			{
 				case Path::EDescType::Document:
@@ -769,7 +632,7 @@ namespace Metasound
 			return FDescriptionUnwindStep::CreateInvalid();
 		}
 
-		FDescriptionAccessPoint::FDescriptionUnwindStep FDescriptionAccessPoint::GoToNextFromClass(FMetasoundClassDescription& InClassDescription, FDescPath& InPath, const Path::FElement& InNext)
+		FDescriptionAccessPoint::FDescriptionUnwindStep FDescriptionAccessPoint::GoToNextFromClass(FMetasoundClassDescription& InClassDescription, FDescPath& InPath, const Path::FElement& InNext) const
 		{
 			switch (InNext.CurrentDescType)
 			{
@@ -862,7 +725,7 @@ namespace Metasound
 		}
 
 
-		FDescriptionAccessPoint::FDescriptionUnwindStep FDescriptionAccessPoint::GoToNext(FDescPath& InPath, FDescriptionUnwindStep InElement)
+		FDescriptionAccessPoint::FDescriptionUnwindStep FDescriptionAccessPoint::GoToNext(FDescPath& InPath, FDescriptionUnwindStep InElement) const
 		{
 			if (!ensureMsgf(InPath.Path.Num() != 0, TEXT("Attempted to unwind an empty path.")))
 			{
diff --git a/Engine/Plugins/Runtime/Metasound/Source/MetasoundFrontend/Public/MetasoundAccessPtr.h b/Engine/Plugins/Runtime/Metasound/Source/MetasoundFrontend/Public/MetasoundAccessPtr.h
new file mode 100644
--- /dev/null
+++ b/Engine/Plugins/Runtime/Metasound/Source/MetasoundFrontend/Public/MetasoundAccessPtr.h
@@ -0,0 +1,134 @@
+// Copyright Epic Games, Inc. All Rights Reserved.
+
+#pragma once
+
+#include "CoreMinimal.h"
+
+namespace Metasound
+{
+	namespace Frontend
+	{
+		class FAccessPoint;
+
+		struct FAccessToken {};
+
+		/** TAccessPtr
+		 *
+		 * TAccessPtr is used to determine whether an object has been destructed or not.
+		 * It is useful when an object cannot be wrapped in a TSharedPtr. A TAccessPtr
+		 * works similar to a TWeakPtr, but it cannot pin the object. Object pointers 
+		 * held within a TAccessPtr should only be accessed on a thread where the objects
+		 * cannot be destructed.
+		 *
+		 * If the TAccessPtr's underlying object is accessed when the pointer is invalid,
+		 * a fallback object will be returned. 
+		 *
+		 * @tparam Type - The type of object to track.
+		 */
+		template<typename Type>
+		class TAccessPtr
+		{
+		public:
+			static Type FallbackObject;
+
+			/** Create an invalid access pointer. */
+			static TAccessPtr<Type> CreateInvalid()
+			{
+				TSharedPtr<FAccessToken> InvalidToken;
+				return TAccessPtr<Type>(InvalidToken, FallbackObject);
+			}
+
+			/** Returns true if object is valid. Should only be called in same
+			 * thread which created the token. */
+			bool IsValid() const { return Token.IsValid(); }
+
+			FORCEINLINE explicit operator bool() const
+			{
+				return IsValid();
+			}
+
+			Type* Get() const
+			{
+				if (IsValid())
+				{
+					return Object;
+				}
+				return nullptr;
+			}
+
+			Type& operator*() const
+			{
+				if (ensure(IsValid()))
+				{
+					return *Object;
+				}
+
+				return FallbackObject;
+			}
+
+			Type& operator->() const
+			{
+				if (ensure(IsValid()))
+				{
+					return *Object;
+				}
+
+				return FallbackObject;
+			}
+
+			TAccessPtr() = delete;
+
+			TAccessPtr(const TAccessPtr<Type>& InOther) = default;
+			TAccessPtr<Type>& operator=(const TAccessPtr<Type>& InOther) = default;
+
+			TAccessPtr(TAccessPtr<Type>&& InOther) = default;
+			TAccessPtr& operator=(TAccessPtr<Type>&& InOther) = default;
+
+		private:
+			template<typename RelatedType>
+			friend TAccessPtr<RelatedType> MakeAccessPtr(const FAccessPoint& InAccessPoint, RelatedType& InRef);
+
+			TWeakPtr<FAccessToken> Token;
+			Type* Object;
+
+			TAccessPtr(TWeakPtr<FAccessToken> AccessToken, Type& InRef)
+			:	Token(AccessToken)
+			,	Object(&InRef)
+			{
+			}
+		};
+
+		template<typename Type>
+		Type TAccessPtr<Type>::FallbackObject = Type();
+
+		/** FAccessPoint acts as a lifecycle tracker for the TAccessPtrs it creates. 
+		 * When this object is destructed, all associated TAccessPtrs will become invalid.
+		 */
+		class FAccessPoint
+		{
+		public:
+
+			FAccessPoint() 
+			{
+				Token = MakeShared<FAccessToken>();
+			}
+
+		private:
+			template<typename Type>
+			friend TAccessPtr<Type> MakeAccessPtr(const FAccessPoint& InAccessPoint, Type& InRef);
+
+			FAccessPoint(const FAccessPoint&) = delete;
+			FAccessPoint& operator=(const FAccessPoint&) = delete;
+			FAccessPoint(FAccessPoint&&) = delete;
+			FAccessPoint& operator=(FAccessPoint&&) = delete;
+
+			TSharedPtr<FAccessToken> Token;
+		};
+
+		template<typename Type>
+		TAccessPtr<Type> MakeAccessPtr(const FAccessPoint& InAccessPoint, Type& InRef)
+		{
+			return TAccessPtr<Type>(InAccessPoint.Token, InRef);
+		}
+	}
+}
diff --git a/Engine/Plugins/Runtime/Metasound/Source/MetasoundFrontend/Public/MetasoundArchetype.h b/Engine/Plugins/Runtime/Metasound/Source/MetasoundFrontend/Public/MetasoundArchetype.h
new file mode 100644
--- /dev/null
+++ b/Engine/Plugins/Runtime/Metasound/Source/MetasoundFrontend/Public/MetasoundArchetype.h
@@ -0,0 +1,24 @@
+// Copyright Epic Games, Inc. All Rights Reserved.
+#pragma once
+
+#include "CoreMinimal.h"
+#include "MetasoundFrontendDataLayout.h"
+
+namespace Metasound
+{
+	namespace Frontend
+	{
+		METASOUNDFRONTEND_API bool IsSubsetOfArchetype(const FMetasoundArchetype& InSubsetArchetype, const FMetasoundArchetype& InSupersetArchetype);
+
+		METASOUNDFRONTEND_API bool IsEquivalentArchetype(const FMetasoundArchetype& InInputArchetype, const FMetasoundArchetype& InTargetArchetype);
+
+		METASOUNDFRONTEND_API bool IsEqualArchetype(const FMetasoundArchetype& InInputArchetype, const FMetasoundArchetype& InTargetArchetype);
+
+		METASOUNDFRONTEND_API int32 InputOutputDifferenceCount(const FMetasoundArchetype& InArchetypeA, const FMetasoundArchetype& InArchetypeB);
+
+		METASOUNDFRONTEND_API void GatherRequiredEnvironmentVariables(const FMetasoundDocument& InDocument, TArray<FMetasoundEnvironmentVariableDescription>& OutEnvironmentVariables);
+
+		METASOUNDFRONTEND_API const FMetasoundArchetype* FindMostSimilarArchetypeSupportingEnvironment(const FMetasoundDocument& InDocument, const TArray<FMetasoundArchetype>& InCandidateArchetypes);
+	}
+}
+
diff --git a/Engine/Plugins/Runtime/Metasound/Source/MetasoundFrontend/Public/MetasoundAssetBase.h b/Engine/Plugins/Runtime/Metasound/Source/MetasoundFrontend/Public/MetasoundAssetBase.h
--- a/Engine/Plugins/Runtime/Metasound/Source/MetasoundFrontend/Public/MetasoundAssetBase.h
+++ b/Engine/Plugins/Runtime/Metasound/Source/MetasoundFrontend/Public/MetasoundAssetBase.h
@@ -2,6 +2,7 @@
 #pragma once
 
 #include "CoreMinimal.h"
+#include "MetasoundAccessPtr.h"
 #include "MetasoundFrontend.h"
 #include "MetasoundFrontendDataLayout.h"
 #include "MetasoundFrontendBaseClasses.h"
@@ -10,13 +11,43 @@
 
 class UEdGraph;
 
+/** FMetasoundAssetBase is intended to be a mix-in subclass for UObjects which utilize
+ * Metasound assets.  It provides consistent access to FMetasoundDocuments, control
+ * over the FMetasoundArchetype of the FMetasoundDocument.  It also enables the UObject
+ * to be utilized by a host of other engine tools built to support Metasounds.
+ *
+ * Example Usage:
+ *
+ * class UMyMetasoundClass : public UObject, public FMetasoundAssetBase
+ * {
+ *
+ *     	UPROPERTY()
+ *     	FMetasoundDocument MetasoundDocument;
+ *
+ *	 	Metasound::Frontend::FAccessPoint MetasoundSourceAccessPoint;
+ *
+ * #if WITH_EDITORONLY_DATA
+ * 		UPROPERTY()
+ * 		UEdGraph* Graph;
+ * #endif // WITH_EDITORONLY_DATA
+ *
+ * 	public:
+ * 		// FMetasoundAssetBase Interface
+ * 		...
+ *      // End FMetasoundAssetBase Interface
+ * };
+ *
+ */
 class METASOUNDFRONTEND_API FMetasoundAssetBase
 {
 public:
 	static const FString FileExtension;
 
 	FMetasoundAssetBase();
-	FMetasoundAssetBase(FMetasoundDocument& InDocument);
+
+	// Construct an FMetasoundAssetBase with a default Metasound Archetype.
+	// @param InDefaultArchetype - Default archetype for a Metasound Document in this class.
+	FMetasoundAssetBase(const FMetasoundArchetype& InDefaultArchetype);
 
 	virtual ~FMetasoundAssetBase() = default;
 
@@ -40,7 +71,22 @@ public:
 	virtual void SetMetadata(FMetasoundClassMetadata& InMetadata);
 
 	// Returns  a description of the required inputs and outputs for this metasound UClass.
-	virtual FMetasoundArchetype GetArchetype() const = 0;
+	const FMetasoundArchetype& GetArchetype() const;
+
+	// Sets/overwrites the document archetype.
+	// @param InArchetype - The desired archetype .
+	//
+	// @return True on success. False if archetype is not supported by this object. 
+	bool SetArchetype(const FMetasoundArchetype& InArchetype);
+
+	// Returns an array of archetypes preferred for this class.
+	virtual const TArray<FMetasoundArchetype>& GetPreferredArchetypes() const = 0;
+
+	// Returns true if the archetype is supported by this object.
+	virtual bool IsArchetypeSupported(const FMetasoundArchetype& InArchetype) const;
+
+	// Returns the preferred archetype for the given document.
+	virtual const FMetasoundArchetype& GetPreferredArchetype(const FMetasoundDocument& InDocument) const;
 
 	// Returns the root class metadata
 	FMetasoundClassMetadata GetMetadata();
@@ -52,14 +98,20 @@ public:
 	bool ImportFromJSONAsset(const FString& InAbsolutePath);
 
 	// Returns handle for the root metasound graph of this asset.
-	Metasound::Frontend::FGraphHandle GetRootGraphHandle() const;
+	Metasound::Frontend::FGraphHandle GetRootGraphHandle();
 
 	// Returns all handles for subgraphs referenced
 	TArray<Metasound::Frontend::FGraphHandle> GetAllSubgraphHandles();
 
-	// Overwrites the entire document and fixes it up based on any 
-	// inputs or outputs in the archetype that are missing from the graph.
-	void SetDocument(const FMetasoundDocument& InDocument);
+	// Overwrites the existing document. If the document's archetype is not supported,
+	// the FMetasoundAssetBase be while queried for a new one using `GetPreferredArchetype`. If `bForceUpdateArchetype`
+	// is true, `GetPreferredArchetype` will be used whether or not the provided document's archetype
+	// is supported. 
+	void SetDocument(const FMetasoundDocument& InDocument, bool bForceUpdateArchetype=false);
+
+	
+	FMetasoundDocument& GetDocumentChecked();
+	const FMetasoundDocument& GetDocumentChecked() const;
 
 	// This must be called on UObject::PostLoad, as well as in this asset's UFactory, to fix up the root document based on the most recent version of the archetype.
 	void ConformDocumentToArchetype();
@@ -71,20 +123,22 @@ protected:
 		return Metasound::Frontend::FHandleInitParams::PrivateToken;
 	}
 
-	// Returns document object responsible for serializing asset
-	virtual FMetasoundDocument& GetDocument() = 0;
+	// Returns an access pointer to the document.
+	virtual Metasound::Frontend::TAccessPtr<FMetasoundDocument> GetDocument() = 0;
 
-	// Returns document object responsible for serializing asset
-	virtual const FMetasoundDocument& GetDocument() const = 0;
+	// Returns an access pointer to the document.
+	virtual Metasound::Frontend::TAccessPtr<const FMetasoundDocument> GetDocument() const = 0;
 
 	// Returns the owning asset responsible for transactions applied to metasound
 	virtual UObject* GetOwningAsset() const = 0;
 
-	// Returns a weak pointer that can be used to build a TDescriptionPtr
+	// Returns a access point that can be used to build a TDescriptionPtr
 	// for direct editing of the FMetasoundClassDescription tree.
 	// For advance use only, and requires knowledge of Metasound::Frontend::FDescPath syntax.
 	// For most use cases, use GetRootGraphHandle() instead.
-	TWeakPtr<Metasound::Frontend::FDescriptionAccessPoint> GetGraphAccessPoint() const;
+	Metasound::Frontend::FDescriptionAccessPoint GetGraphAccessPoint();
+
+private:
 
-	TSharedPtr<Metasound::Frontend::FDescriptionAccessPoint> AccessPoint;
+	FMetasoundArchetype Archetype;
 };
diff --git a/Engine/Plugins/Runtime/Metasound/Source/MetasoundFrontend/Public/MetasoundAudioBuffer.h b/Engine/Plugins/Runtime/Metasound/Source/MetasoundFrontend/Public/MetasoundAudioBuffer.h
--- a/Engine/Plugins/Runtime/Metasound/Source/MetasoundFrontend/Public/MetasoundAudioBuffer.h
+++ b/Engine/Plugins/Runtime/Metasound/Source/MetasoundFrontend/Public/MetasoundAudioBuffer.h
@@ -6,6 +6,7 @@
 #include "DSP/BufferVectorOperations.h"
 #include "MetasoundGraphCoreModule.h"
 #include "MetasoundDataReference.h"
+#include "MetasoundLog.h"
 
 /** Define which determines whether to check that the size of the audio buffer has not changed since initialization */
 #define METASOUNDGRAPHCORE_CHECKAUDIONUM !(UE_BUILD_SHIPPING || UE_BUILD_TEST)
@@ -61,7 +62,7 @@ namespace Metasound
 			FORCEINLINE const float* GetData() const
 			{
 #if METASOUNDGRAPHCORE_CHECKAUDIONUM
-				UE_CLOG(InitialNum != Buffer.Num(), LogMetasoundGraphCore, Error, TEXT("Metasound audio buffer size change detected.  Audio Buffers should not be resized."));
+				UE_CLOG(InitialNum != Buffer.Num(), LogMetasound, Error, TEXT("Metasound audio buffer size change detected.  Audio Buffers should not be resized."));
 #endif				
 				return Buffer.GetData();
 			}
@@ -70,7 +71,7 @@ namespace Metasound
 			FORCEINLINE float* GetData()
 			{
 #if METASOUNDGRAPHCORE_CHECKAUDIONUM
-				UE_CLOG(InitialNum != Buffer.Num(), LogMetasoundGraphCore, Error, TEXT("Metasound audio buffer size change detected.  Audio Buffers should not be resized."));
+				UE_CLOG(InitialNum != Buffer.Num(), LogMetasound, Error, TEXT("Metasound audio buffer size change detected.  Audio Buffers should not be resized."));
 #endif				
 				return Buffer.GetData();
 			}
@@ -79,7 +80,7 @@ namespace Metasound
 			FORCEINLINE int32 Num() const
 			{
 #if METASOUNDGRAPHCORE_CHECKAUDIONUM
-				UE_CLOG(InitialNum != Buffer.Num(), LogMetasoundGraphCore, Error, TEXT("Metasound audio buffer size change detected.  Audio Buffers should not be resized."));
+				UE_CLOG(InitialNum != Buffer.Num(), LogMetasound, Error, TEXT("Metasound audio buffer size change detected.  Audio Buffers should not be resized."));
 #endif				
 				return Buffer.Num();
 			}
@@ -88,7 +89,7 @@ namespace Metasound
 			FORCEINLINE operator const Audio::AlignedFloatBuffer& () const
 			{
 #if METASOUNDGRAPHCORE_CHECKAUDIONUM
-				UE_CLOG(InitialNum != Buffer.Num(), LogMetasoundGraphCore, Error, TEXT("Metasound audio buffer size change detected.  Audio Buffers should not be resized."));
+				UE_CLOG(InitialNum != Buffer.Num(), LogMetasound, Error, TEXT("Metasound audio buffer size change detected.  Audio Buffers should not be resized."));
 #endif				
 				return Buffer;
 			}
@@ -100,7 +101,7 @@ namespace Metasound
 			FORCEINLINE operator Audio::AlignedFloatBuffer& ()
 			{
 #if METASOUNDGRAPHCORE_CHECKAUDIONUM
-				UE_CLOG(InitialNum != Buffer.Num(), LogMetasoundGraphCore, Error, TEXT("Metasound audio buffer size change detected.  Audio Buffers should not be resized."));
+				UE_CLOG(InitialNum != Buffer.Num(), LogMetasound, Error, TEXT("Metasound audio buffer size change detected.  Audio Buffers should not be resized."));
 #endif				
 				return Buffer;
 			}
diff --git a/Engine/Plugins/Runtime/Metasound/Source/MetasoundFrontend/Public/MetasoundFrontend.h b/Engine/Plugins/Runtime/Metasound/Source/MetasoundFrontend/Public/MetasoundFrontend.h
--- a/Engine/Plugins/Runtime/Metasound/Source/MetasoundFrontend/Public/MetasoundFrontend.h
+++ b/Engine/Plugins/Runtime/Metasound/Source/MetasoundFrontend/Public/MetasoundFrontend.h
@@ -139,48 +139,6 @@ namespace Metasound
 		// @returns false if the file couldn't be found or parsed into a document.
 		METASOUNDFRONTEND_API bool ImportJSONAssetToMetasound(const FString& InPath, FMetasoundDocument& OutMetasoundDocument);
 
-		struct METASOUNDFRONTEND_API FMetasoundArchetypeRegistryParams_Internal
-		{
-			FMetasoundArchetype ArchetypeDescription;
-
-			// The UClass associated with this specific archetype.
-			UClass* ArchetypeUClass;
-
-			// template-generated lambdas used to safely sidecast to FMetasoundBase*.
-			TUniqueFunction<FMetasoundAssetBase* (UObject*)> SafeCast;
-			TUniqueFunction<const FMetasoundAssetBase* (const UObject*)> SafeConstCast;
-
-			// This function should construct a new UObject of the given archetype's type
-			// given a metasound document with a matching archetype.
-			// The first argument is the document to use.
-			// The second argument is the path relative to the content directory to save the soundwave to.
-			TUniqueFunction<UObject* (const FMetasoundDocument&, const FString&)> ObjectGetter;
-		};
-
-		// ಠ╭╮ಠ         ಠ╭╮ಠ           ಠ╭╮ಠ
-		// Please don't use this function.
-		// See RegisterArchetype<Class> instead, in MetasoundArchetypeRegistration.h
-		METASOUNDFRONTEND_API bool RegisterArchetype_Internal(FMetasoundArchetypeRegistryParams_Internal&& InParams);
-
-		METASOUNDFRONTEND_API TArray<FName> GetAllRegisteredArchetypes();
-
-		// Returns a new UObject, whose class corresponds to the archetype in the document.
-		// @param InDocument a fully generated metasound document, typically retrieved from ImportJSONAssetToMetasound().
-		// @param InPath, path in content directory to save the generated UAsset to (ex. "/game/MyDir/MyMetasoundAsset".
-		//                if InPath is invalid, we won't save to an asset.
-		// @returns nullptr if we couldn't find the archetype.
-		METASOUNDFRONTEND_API UObject* GetObjectForDocument(const FMetasoundDocument& InDocument, const FString& InPath);
-
-		// This returns true if the object is listed as one of our metasound archetypes.
-		// If it is, the object can be safely static cast to FMetasoundAssetBase*.
-		METASOUNDFRONTEND_API bool IsObjectAMetasoundArchetype(const UObject* InObject);
-
-		// These functions are used to safely sidecast between a UObject of a given metasound archetype
-		// and FMetasoundAssetBase*.
-		// @returns nullptr if the object wasn't a registered archetype.
-		METASOUNDFRONTEND_API FMetasoundAssetBase* GetObjectAsAssetBase(UObject* InObject);
-		METASOUNDFRONTEND_API const FMetasoundAssetBase* GetObjectAsAssetBase(const UObject* InObject);
-
 		// Struct that indicates whether an input and an output can be connected,
 		// and whether an intermediate node is necessary to connect the two.
 		struct METASOUNDFRONTEND_API FConnectability
@@ -202,7 +160,7 @@ namespace Metasound
 
 		struct METASOUNDFRONTEND_API FHandleInitParams
 		{
-			TWeakPtr<FDescriptionAccessPoint> InAccessPoint;
+			FDescriptionAccessPoint InAccessPoint;
 
 			// Path in the document to the element we're getting a handle to.
 			// for the FGraphHandle, this is a path to the graph within the graph's Class description.
@@ -303,6 +261,7 @@ namespace Metasound
 			FName GetInputType() const;
 			FString GetInputName() const;
 			FText GetInputTooltip() const;
+			int32 GetOwningNodeID() const;
 
 			FOutputHandle GetCurrentlyConnectedOutput() const;
 
@@ -381,8 +340,6 @@ namespace Metasound
 			// whether this node is an input or output to it's owning graph,
 			// an externally implemented node, or a metasound graph itself.
 			EMetasoundClassType NodeClassType;
-
-			int32 NodeID;
 		};
 
 		class METASOUNDFRONTEND_API FGraphHandle : protected ITransactable
@@ -394,7 +351,7 @@ namespace Metasound
 			// Sole constructor for FGraphHandle. Can only be used by friends of FHandleInitParams.
 			FGraphHandle(FHandleInitParams::EPrivateToken PrivateToken, const FHandleInitParams& InParams);
 
-			static FGraphHandle GetHandle(UObject* InOwner, FMetasoundDocument& InRootMetasoundDocument, TSharedPtr<Metasound::Frontend::FDescriptionAccessPoint>& InAccessPoint)
+			static FGraphHandle GetHandle(UObject* InOwner, FMetasoundDocument& InRootMetasoundDocument, const Metasound::Frontend::FDescriptionAccessPoint& InAccessPoint)
 			{
 				using namespace Path;
 
diff --git a/Engine/Plugins/Runtime/Metasound/Source/MetasoundFrontend/Public/MetasoundFrontendBaseClasses.h b/Engine/Plugins/Runtime/Metasound/Source/MetasoundFrontend/Public/MetasoundFrontendBaseClasses.h
--- a/Engine/Plugins/Runtime/Metasound/Source/MetasoundFrontend/Public/MetasoundFrontendBaseClasses.h
+++ b/Engine/Plugins/Runtime/Metasound/Source/MetasoundFrontend/Public/MetasoundFrontendBaseClasses.h
@@ -4,11 +4,10 @@
 
 #include "CoreMinimal.h"
 #include "Misc/TVariant.h"
+#include "MetasoundAccessPtr.h"
 #include "MetasoundFrontendDataLayout.h"
 #include "MetasoundFrontendRegistries.h"
 
-
-
 namespace Metasound
 {
 	namespace Frontend
@@ -522,80 +521,56 @@ namespace Metasound
 		 */
 		class METASOUNDFRONTEND_API FDescriptionAccessPoint
 		{
+			using FMetasoundDescriptionPtr = TVariant<
+				FMetasoundDocument*,
+				FMetasoundClassDescription*,
+				FMetasoundNodeDescription*,
+				FMetasoundGraphDescription*,
+				FMetasoundInputDescription*,
+				FMetasoundOutputDescription*,
+				FMetasoundClassMetadata*,
+				FClassDependencyIDs*,
+				void*>
+				;
 		public:
 			FDescriptionAccessPoint() = delete;
-			FDescriptionAccessPoint(FMetasoundDocument& InRootDocument);
+			FDescriptionAccessPoint(TAccessPtr<FMetasoundDocument> InRootDocument);
+
+			static FDescriptionAccessPoint CreateInvalid()
+			{
+				return FDescriptionAccessPoint(TAccessPtr<FMetasoundDocument>::CreateInvalid());
+			}
 
 			// Returns the root class description associated with this FDescriptionAccessPoint.
-			FMetasoundDocument& GetRoot();
+			FMetasoundDocument& GetRootChecked() const;
+
+			bool IsValid() const
+			{
+				return RootDocumentPtr.IsValid();
+			}
 
 			// Get a descriptor element at a specific path from the root of a metasound.
 			// @see FDescPath for more details.
 			// @returns a pointer to a descriptor struct, or an invalid ptr if the path is invalid.
 			template <typename MetasoundDescriptionType>
-			MetasoundDescriptionType* GetElementAtPath(const FDescPath& InPathFromRoot)
+			MetasoundDescriptionType* GetElementAtPath(const FDescPath& InPathFromRoot) const
 			{
-				if (TIsSame<MetasoundDescriptionType, FMetasoundClassDescription>::Value)
-				{
-					return (MetasoundDescriptionType*) GetClassFromPath(InPathFromRoot);
-				}
-				else if (TIsSame<MetasoundDescriptionType, FMetasoundNodeDescription>::Value)
-				{
-					return (MetasoundDescriptionType*) GetNodeFromPath(InPathFromRoot);
-				}
-				else if (TIsSame<MetasoundDescriptionType, FMetasoundGraphDescription>::Value)
-				{
-					return (MetasoundDescriptionType*) GetGraphFromPath(InPathFromRoot);
-				}
-				else if (TIsSame<MetasoundDescriptionType, FMetasoundInputDescription>::Value)
-				{
-					return (MetasoundDescriptionType*) GetInputFromPath(InPathFromRoot);
-				}
-				else if (TIsSame<MetasoundDescriptionType, FMetasoundOutputDescription>::Value)
+				FMetasoundDescriptionPtr Ptr;
+				if (ensure(GetDescriptionPtrFromPath(InPathFromRoot, Ptr)))
 				{
-					return (MetasoundDescriptionType*) GetOutputFromPath(InPathFromRoot);
-				}
-				else if (TIsSame<MetasoundDescriptionType, FMetasoundClassMetadata>::Value)
-				{
-					return (MetasoundDescriptionType*) GetMetadataFromPath(InPathFromRoot);
-				}
-				else if (TIsSame<MetasoundDescriptionType, FClassDependencyIDs>::Value)
-				{
-					return (MetasoundDescriptionType*) GetClassDependencyIDsFromPath(InPathFromRoot);
-				}
-				else
-				{
-					ensureAlwaysMsgf(false, TEXT("Tried to call GetElementAtPath with an invalid type."));
-					return nullptr;
+					return Ptr.Get<MetasoundDescriptionType*>();
 				}
+
+				return nullptr;
 			}
 
 		private:
 			// The root class description for all description structs
 			// accessed via this class.
-			FMetasoundDocument& RootDocument;
-
-			// These functions result in some duplicate code, but allow us to hide
-			// the implementation for GetElementAtPath.
-			FMetasoundClassDescription* GetClassFromPath(const FDescPath& InPathFromRoot);
-			FMetasoundNodeDescription* GetNodeFromPath(const FDescPath& InPathFromRoot);
-			FMetasoundGraphDescription* GetGraphFromPath(const FDescPath& InPathFromRoot);
-			FMetasoundInputDescription* GetInputFromPath(const FDescPath& InPathFromRoot);
-			FMetasoundOutputDescription* GetOutputFromPath(const FDescPath& InPathFromRoot);
-			FMetasoundClassMetadata* GetMetadataFromPath(const FDescPath& InPathFromRoot);
-			FClassDependencyIDs* GetClassDependencyIDsFromPath(const FDescPath& InPathFromRoot);
+			TAccessPtr<FMetasoundDocument> RootDocumentPtr;
 
-			using FMetasoundDescriptionPtr = TVariant<
-				FMetasoundDocument*,
-				FMetasoundClassDescription*,
-				FMetasoundNodeDescription*,
-				FMetasoundGraphDescription*,
-				FMetasoundInputDescription*,
-				FMetasoundOutputDescription*,
-				FMetasoundClassMetadata*,
-				FClassDependencyIDs*,
-				void*>
-				;
+
+			bool GetDescriptionPtrFromPath(const FDescPath& InPathFromRoot, FMetasoundDescriptionPtr& OutPtr) const;
 
 			struct FDescriptionUnwindStep
 			{
@@ -614,11 +589,11 @@ namespace Metasound
 				}
 			};
 
-			FDescriptionUnwindStep GoToNextFromDocument(FMetasoundDocument& InDocument, FDescPath& InPath, const Path::FElement& InNext);
-			FDescriptionUnwindStep GoToNextFromClass(FMetasoundClassDescription& InClass, FDescPath& InPath, const Path::FElement& InNext);
+			FDescriptionUnwindStep GoToNextFromDocument(FMetasoundDocument& InDocument, FDescPath& InPath, const Path::FElement& InNext) const;
+			FDescriptionUnwindStep GoToNextFromClass(FMetasoundClassDescription& InClass, FDescPath& InPath, const Path::FElement& InNext) const;
 
 			// This function pops off elements of InPath to go to the next level down in a Metasound description.
-			FDescriptionUnwindStep GoToNext(FDescPath& InPath, FDescriptionUnwindStep InElement);
+			FDescriptionUnwindStep GoToNext(FDescPath& InPath, FDescriptionUnwindStep InElement) const;
 		};
 
 		// Utility class used by the frontend handle classes to access the underlying Description data they are modifying.
@@ -646,14 +621,14 @@ namespace Metasound
 			FMetasoundClassMetadata,
 			FMetasoundClassDescription)");
 
-			TWeakPtr<FDescriptionAccessPoint> AccessPoint;
+			FDescriptionAccessPoint AccessPoint;
 			FDescPath PathFromRoot;
 			bool bAccessFailure = false;
 
 		public:
 			TDescriptionPtr() = delete;
 
-			TDescriptionPtr(TWeakPtr<FDescriptionAccessPoint> InAccessPoint, const FDescPath& InPathFromRoot)
+			TDescriptionPtr(FDescriptionAccessPoint InAccessPoint, const FDescPath& InPathFromRoot)
 				: AccessPoint(InAccessPoint)
 				, PathFromRoot(InPathFromRoot)
 			{
@@ -661,6 +636,11 @@ namespace Metasound
 				bAccessFailure = Get() == nullptr;
 			}
 
+			static TDescriptionPtr<MetasoundDescriptionType> CreateInvalid()
+			{
+				return TDescriptionPtr<MetasoundDescriptionType>(FDescriptionAccessPoint::CreateInvalid(), FDescPath());
+			}
+
 			TDescriptionPtr(TDescriptionPtr&& InDescriptionPtr) = default;
 			TDescriptionPtr(const TDescriptionPtr& InDescriptionPtr) = default;
 
@@ -679,7 +659,7 @@ namespace Metasound
 				return PathFromRoot;
 			}
 
-			TWeakPtr<FDescriptionAccessPoint> GetAccessPoint() const
+			const FDescriptionAccessPoint& GetAccessPoint() const
 			{
 				return AccessPoint;
 			}
@@ -688,16 +668,9 @@ namespace Metasound
 			// Note that this is NOT thread safe- this should be called on the same thread as whomever owns the lifecycle of the underlying MetasoundDescriptionType.
 			MetasoundDescriptionType* Get() const
 			{
-				if (TSharedPtr<FDescriptionAccessPoint> PinnedAccessPoint = AccessPoint.Pin())
+				if (AccessPoint.IsValid())
 				{
-					if (TIsSame<MetasoundDescriptionType, FMetasoundDocument>::Value)
-					{
-						return (MetasoundDescriptionType*) &(PinnedAccessPoint->GetRoot());
-					}
-					else
-					{
-						return PinnedAccessPoint->GetElementAtPath<MetasoundDescriptionType>(PathFromRoot);
-					}
+					return AccessPoint.GetElementAtPath<MetasoundDescriptionType>(PathFromRoot);
 				}
 				else
 				{
diff --git a/Engine/Plugins/Runtime/Metasound/Source/MetasoundFrontend/Public/MetasoundFrontendDataLayout.h b/Engine/Plugins/Runtime/Metasound/Source/MetasoundFrontend/Public/MetasoundFrontendDataLayout.h
--- a/Engine/Plugins/Runtime/Metasound/Source/MetasoundFrontend/Public/MetasoundFrontendDataLayout.h
+++ b/Engine/Plugins/Runtime/Metasound/Source/MetasoundFrontend/Public/MetasoundFrontendDataLayout.h
@@ -321,6 +321,12 @@ struct FMetasoundArchetype
 	UPROPERTY()
 	FName ArchetypeName;
 
+	UPROPERTY()
+	int32 MajorVersion;
+
+	UPROPERTY()
+	int32 MinorVersion;
+
 	// whatever inputs are required to support this archetype.
 	// call Special Inputs
 	UPROPERTY()
diff --git a/Engine/Plugins/Runtime/Metasound/Source/MetasoundFrontend/Public/MetasoundNodeRegistrationMacro.h b/Engine/Plugins/Runtime/Metasound/Source/MetasoundFrontend/Public/MetasoundNodeRegistrationMacro.h
--- a/Engine/Plugins/Runtime/Metasound/Source/MetasoundFrontend/Public/MetasoundNodeRegistrationMacro.h
+++ b/Engine/Plugins/Runtime/Metasound/Source/MetasoundFrontend/Public/MetasoundNodeRegistrationMacro.h
@@ -8,62 +8,66 @@
 #include "MetasoundFrontend.h"
 #include "MetasoundFrontendRegistries.h"
 
-// Utility class to ensure that a node class can use the constructor the frontend uses.
-template <typename NodeClass>
-struct ConstructorTakesNodeInitData {
+namespace Metasound
+{
+	// Utility class to ensure that a node class can use the constructor the frontend uses.
+	template <typename NodeClass>
+	struct ConstructorTakesNodeInitData {
 
-	// Use SFINAE trick to see if we have a valid constructor:
-	template<typename T>
-	static uint16 TestForConstructor(decltype(T(::Metasound::FNodeInitData()))*);
+		// Use SFINAE trick to see if we have a valid constructor:
+		template<typename T>
+		static uint16 TestForConstructor(decltype(T(::Metasound::FNodeInitData()))*);
 
-	template<typename T>
-	static uint8 TestForConstructor(...);
+		template<typename T>
+		static uint8 TestForConstructor(...);
 
-	static const bool Value = sizeof(TestForConstructor<NodeClass>(nullptr)) == sizeof(uint16);
-};
+		static const bool Value = sizeof(TestForConstructor<NodeClass>(nullptr)) == sizeof(uint16);
+	};
 
 
-template <typename FNodeType>
-bool RegisterNodeWithFrontend(const Metasound::FNodeInfo& InMetadata)
-{
-	// if we reenter this code (because DECLARE_METASOUND_DATA_REFERENCE_TYPES was called twice with the same type),
-	// we catch it here.
-	static bool bAlreadyRegisteredThisDataType = false;
-	if (bAlreadyRegisteredThisDataType)
+	template <typename FNodeType>
+	bool RegisterNodeWithFrontend(const Metasound::FNodeInfo& InMetadata)
 	{
-		UE_LOG(LogTemp, Display, TEXT("Tried to call METASOUND_REGISTER_NODE twice with the same class. ignoring the second call. Likely because METASOUND_REGISTER_NODE is in a header that's used in multiple modules. Consider moving it to a private header or cpp file."));
-		return false;
-	}
-
-	bAlreadyRegisteredThisDataType = true;
-
-	bool bSuccessfullyRegisteredNode = FMetasoundFrontendRegistryContainer::Get()->RegisterExternalNode(
-		[](const Metasound::FNodeInitData& InInitData) -> TUniquePtr<Metasound::INode>
-		{
-			return TUniquePtr<Metasound::INode>(new FNodeType(InInitData));
-		},
-		[=]() -> FMetasoundClassDescription
+		// if we reenter this code (because DECLARE_METASOUND_DATA_REFERENCE_TYPES was called twice with the same type),
+		// we catch it here.
+		static bool bAlreadyRegisteredThisDataType = false;
+		if (bAlreadyRegisteredThisDataType)
 		{
-			return Metasound::Frontend::GenerateClassDescription(InMetadata);
+			UE_LOG(LogTemp, Display, TEXT("Tried to call METASOUND_REGISTER_NODE twice with the same class. ignoring the second call. Likely because METASOUND_REGISTER_NODE is in a header that's used in multiple modules. Consider moving it to a private header or cpp file."));
+			return false;
 		}
-	);
 
-	ensureAlwaysMsgf(bSuccessfullyRegisteredNode, TEXT("Registering node class failed. Please check the logs."));
-	return bSuccessfullyRegisteredNode;
-}
+		bAlreadyRegisteredThisDataType = true;
 
-template <typename FNodeType>
-bool RegisterNodeWithFrontend()
-{
-	// Register a node using a prototype node.
-	// TODO: may want to add a warning here since we don't want to use this registration pathway.
-	Metasound::FNodeInitData InitData;
-	TUniquePtr<Metasound::INode> Node = MakeUnique<FNodeType>(InitData);
-	return RegisterNodeWithFrontend<FNodeType>(Node->GetMetadata());
+		bool bSuccessfullyRegisteredNode = FMetasoundFrontendRegistryContainer::Get()->RegisterExternalNode(
+			[](const Metasound::FNodeInitData& InInitData) -> TUniquePtr<Metasound::INode>
+			{
+				return TUniquePtr<Metasound::INode>(new FNodeType(InInitData));
+			},
+			[=]() -> FMetasoundClassDescription
+			{
+				return Metasound::Frontend::GenerateClassDescription(InMetadata);
+			}
+		);
+
+		ensureAlwaysMsgf(bSuccessfullyRegisteredNode, TEXT("Registering node class failed. Please check the logs."));
+		return bSuccessfullyRegisteredNode;
+	}
+
+	template <typename FNodeType>
+	bool RegisterNodeWithFrontend()
+	{
+		// Register a node using a prototype node.
+		// TODO: may want to add a warning here since we don't want to use this registration pathway.
+		Metasound::FNodeInitData InitData;
+		TUniquePtr<Metasound::INode> Node = MakeUnique<FNodeType>(InitData);
+		return RegisterNodeWithFrontend<FNodeType>(Node->GetMetadata());
+	}
 }
 
+// TODO: should remove this and force folks to use module startup/shutdown module system.
 #define METASOUND_REGISTER_NODE(NodeClass, ...) \
 	 static_assert(std::is_base_of<::Metasound::INodeBase, NodeClass>::value, "To be registered as a  Metasound Node," #NodeClass "need to be a derived class from Metasound::INodeBase, Metasound::INode, or Metasound::FNode."); \
-	 static_assert(::ConstructorTakesNodeInitData<NodeClass>::Value, "In order to be registered as a Metasound Node, " #NodeClass " needs to implement the following public constructor: " #NodeClass "(const Metasound::FNodeInitData& InInitData);"); \
-	 static bool bSuccessfullyRegistered##NodeClass  = FMetasoundFrontendRegistryContainer::Get()->EnqueueInitCommand([](){ ::RegisterNodeWithFrontend<NodeClass>(__VA_ARGS__); }); // This static bool is useful for debugging, but also is the only way the compiler will let us call this function outside of an expression.
+	 static_assert(::Metasound::ConstructorTakesNodeInitData<NodeClass>::Value, "In order to be registered as a Metasound Node, " #NodeClass " needs to implement the following public constructor: " #NodeClass "(const Metasound::FNodeInitData& InInitData);"); \
+	 static bool bSuccessfullyRegistered##NodeClass  = FMetasoundFrontendRegistryContainer::Get()->EnqueueInitCommand([](){ ::Metasound::RegisterNodeWithFrontend<NodeClass>(__VA_ARGS__); }); // This static bool is useful for debugging, but also is the only way the compiler will let us call this function outside of an expression.
 
diff --git a/Engine/Plugins/Runtime/Metasound/Source/MetasoundGenerator/Private/MetasoundGenerator.cpp b/Engine/Plugins/Runtime/Metasound/Source/MetasoundGenerator/Private/MetasoundGenerator.cpp
--- a/Engine/Plugins/Runtime/Metasound/Source/MetasoundGenerator/Private/MetasoundGenerator.cpp
+++ b/Engine/Plugins/Runtime/Metasound/Source/MetasoundGenerator/Private/MetasoundGenerator.cpp
@@ -16,13 +16,10 @@
 
 namespace Metasound
 {
-	
-
 	FMetasoundGenerator::FMetasoundGenerator(FMetasoundGeneratorInitParams&& InParams)
 		: NumChannels(0)
 		, NumFramesPerExecute(0)
 		, NumSamplesPerExecute(0)
-		, GraphOutputAudioRef(InParams.GraphOutputAudioRef) // Copy audio ref here to avoid unneeded construction of unused graph output audio object.
 		, OnFinishedBopRef(InParams.BopOnFinishRef)
 	{
 		SetGraph(MoveTemp(InParams));
@@ -35,10 +32,7 @@ namespace Metasound
 	bool FMetasoundGenerator::UpdateGraphOperator(FMetasoundGeneratorInitParams&& InParams)
 	{
 		// multichannel version:
-		//const int32 GraphNumChannels = InParams.GraphOutputAudioRef->GetNumChannels();
-
-		// single channel version:
-		const int32 GraphNumChannels = 1;
+		const int32 GraphNumChannels = InParams.OutputBuffers.Num();
 
 		if (GraphNumChannels == NumChannels)
 		{
@@ -56,26 +50,34 @@ namespace Metasound
 		NumFramesPerExecute = 0;
 		NumSamplesPerExecute = 0;
 
+		GraphOutputAudio.Reset();
+		if (InParams.OutputBuffers.Num() > 0)
+		{
+			GraphOutputAudio.Append(InParams.OutputBuffers.GetData(), InParams.OutputBuffers.Num());
+		}
+
+		OnFinishedBopRef = InParams.BopOnFinishRef;
+
 		// The graph operator and graph audio output contain all the values needed
 		// by the sound generator.
 		RootExecuter.SetOperator(MoveTemp(InParams.GraphOperator));
-		GraphOutputAudioRef = InParams.GraphOutputAudioRef;
-		OnFinishedBopRef = InParams.BopOnFinishRef;
+
 
 		// Query the graph output to get the number of output audio channels.
 		// Multichannel version:
-		//NumChannels = GraphOutputAudioRef->GetNumChannels();
+		NumChannels = GraphOutputAudio.Num();
 
-		// Single channel version:
-		NumChannels = 1;
 
 		// All buffers have same number of frames, so only need to query
 		// first buffer to know number of frames.
-
-		// Multichannel version:
-		//NumFramesPerExecute = GraphOutputAudioRef->GetBuffers()[0]->Num();
-		// single channel version:
-		NumFramesPerExecute = GraphOutputAudioRef->Num();
+		if (NumChannels > 0)
+		{
+			NumFramesPerExecute = GraphOutputAudio[0]->Num();
+		}
+		else
+		{
+			NumFramesPerExecute = 0;
+		}
 
 		NumSamplesPerExecute = NumFramesPerExecute * NumChannels;
 
@@ -88,11 +90,7 @@ namespace Metasound
 
 	int32 FMetasoundGenerator::GetNumChannels() const
 	{
-		// Multichannel version:
-		// return GraphOutputAudioRef->GetNumChannels();
-
-		// Single channel version:
-		return 1;
+		return GraphOutputAudio.Num();
 	}
 
 	int32 FMetasoundGenerator::OnGenerateAudio(float* OutAudio, int32 NumSamplesRemaining)
@@ -102,6 +100,12 @@ namespace Metasound
 			return 0;
 		}
 
+		if (NumSamplesPerExecute < 1)
+		{
+			FMemory::Memset(OutAudio, 0, sizeof(float) * NumSamplesRemaining);
+			return NumSamplesRemaining;
+		}
+
 		// If we have any audio left in the internal overflow buffer from 
 		// previous calls, write that to the output before generating more audio.
 		int32 NumSamplesWritten = FillWithBuffer(OverflowBuffer, OutAudio, NumSamplesRemaining);
@@ -136,7 +140,7 @@ namespace Metasound
 				OverflowBuffer.Reset();
 				OverflowBuffer.AddUninitialized(OverflowCount);
 
-				FMemory::Memcpy(OverflowBuffer.GetData(), &InterleavedAudioBuffer.GetData()[ThisLoopNumSamplesWritten], OverflowCount);
+				FMemory::Memcpy(OverflowBuffer.GetData(), &InterleavedAudioBuffer.GetData()[ThisLoopNumSamplesWritten], OverflowCount * sizeof(float));
 			}
 		}
 
@@ -146,6 +150,7 @@ namespace Metasound
 
 	int32 FMetasoundGenerator::GetDesiredNumSamplesToRenderPerCallback() const
 	{
+		// TODO: may improve performance if this number is increased. 
 		return NumFramesPerExecute;
 	}
 
@@ -177,11 +182,6 @@ namespace Metasound
 
 	void FMetasoundGenerator::InterleaveGeneratedAudio()
 	{
-		/**
-		// TODO: Convert MetasoundSource / MetasoundGenerator to take Multichannel Audio
-		// once we've got converter nodes up and running.
-		const FMultichannelAudioFormat& DeinterleavedAudio = *GraphOutputAudioRef;
-
 		// Prepare output buffer
 		InterleavedAudioBuffer.Reset();
 
@@ -190,32 +190,10 @@ namespace Metasound
 			InterleavedAudioBuffer.AddUninitialized(NumSamplesPerExecute);
 		}
 
-		const TArrayView<const FAudioBufferReadRef>& Buffers = DeinterleavedAudio.GetBuffers();
-
 		// Iterate over channels
 		for (int32 ChannelIndex = 0; ChannelIndex < NumChannels; ChannelIndex++)
 		{
-			const FAudioBuffer& InputBuffer = *Buffers[ChannelIndex];
-
-			const float* InputPtr = InputBuffer.GetData();
-			float* OutputPtr = &InterleavedAudioBuffer.GetData()[ChannelIndex];
-
-			// Assign values to output for single channel.
-			for (int32 FrameIndex = 0; FrameIndex < NumFramesPerExecute; FrameIndex++)
-			{
-				*OutputPtr = InputPtr[FrameIndex];
-				OutputPtr += NumChannels;
-			}
-		}
-
-		*/
-
-		// Single channel version (for now):
-		const FAudioBuffer& InputAudioBuffer = *GraphOutputAudioRef;
-
-		for (int32 ChannelIndex = 0; ChannelIndex < NumChannels; ChannelIndex++)
-		{
-			const FAudioBuffer& InputBuffer = InputAudioBuffer;
+			const FAudioBuffer& InputBuffer = *GraphOutputAudio[ChannelIndex];
 
 			const float* InputPtr = InputBuffer.GetData();
 			float* OutputPtr = &InterleavedAudioBuffer.GetData()[ChannelIndex];
@@ -227,6 +205,7 @@ namespace Metasound
 				OutputPtr += NumChannels;
 			}
 		}
+		// TODO: memcpy for single channel. 
 	}
 } 
 
diff --git a/Engine/Plugins/Runtime/Metasound/Source/MetasoundGenerator/Private/SynthComponentMetasoundExample.cpp b/Engine/Plugins/Runtime/Metasound/Source/MetasoundGenerator/Private/SynthComponentMetasoundExample.cpp
deleted file mode 100644
diff --git a/Engine/Plugins/Runtime/Metasound/Source/MetasoundGenerator/Public/MetasoundGenerator.h b/Engine/Plugins/Runtime/Metasound/Source/MetasoundGenerator/Public/MetasoundGenerator.h
--- a/Engine/Plugins/Runtime/Metasound/Source/MetasoundGenerator/Public/MetasoundGenerator.h
+++ b/Engine/Plugins/Runtime/Metasound/Source/MetasoundGenerator/Public/MetasoundGenerator.h
@@ -16,7 +16,8 @@ namespace Metasound
 	struct FMetasoundGeneratorInitParams
 	{
 		TUniquePtr<Metasound::IOperator> GraphOperator;
-		FAudioBufferReadRef GraphOutputAudioRef;
+		TArrayView<TDataReadReference<FAudioBuffer>> OutputBuffers;
+		// TODO: BopOnPlayRef;
 		FBopReadRef BopOnFinishRef;
 	};
 
@@ -106,7 +107,6 @@ namespace Metasound
 		//~ End FSoundGenerator
 
 	private:
-
 		// Internal set graph after checking compatibility.
 		void SetGraph(FMetasoundGeneratorInitParams&& InParams);
 
@@ -123,7 +123,7 @@ namespace Metasound
 		int32 NumFramesPerExecute;
 		int32 NumSamplesPerExecute;
 
-		FAudioBufferReadRef GraphOutputAudioRef;
+		TArray<FAudioBufferReadRef> GraphOutputAudio;
 
 		// This will be bopped when this metasound is finished playing.
 		FBopReadRef OnFinishedBopRef;
diff --git a/Engine/Plugins/Runtime/Metasound/Source/MetasoundGenerator/Public/SynthComponentMetasoundExample.h b/Engine/Plugins/Runtime/Metasound/Source/MetasoundGenerator/Public/SynthComponentMetasoundExample.h
deleted file mode 100644
diff --git a/Engine/Plugins/Runtime/Metasound/Source/MetasoundGraphCore/Private/MetasoundGraphCoreModule.cpp b/Engine/Plugins/Runtime/Metasound/Source/MetasoundGraphCore/Private/MetasoundGraphCoreModule.cpp
--- a/Engine/Plugins/Runtime/Metasound/Source/MetasoundGraphCore/Private/MetasoundGraphCoreModule.cpp
+++ b/Engine/Plugins/Runtime/Metasound/Source/MetasoundGraphCore/Private/MetasoundGraphCoreModule.cpp
@@ -3,10 +3,11 @@
 #include "MetasoundGraphCoreModule.h"
 
 #include "CoreMinimal.h"
+#include "MetasoundLog.h"
 #include "Modules/ModuleInterface.h"
 #include "Modules/ModuleManager.h"
 
-DEFINE_LOG_CATEGORY(LogMetasoundGraphCore);
+DEFINE_LOG_CATEGORY(LogMetasound);
 
 namespace Metasound 
 {
diff --git a/Engine/Plugins/Runtime/Metasound/Source/MetasoundGraphCore/Public/MetasoundDataReference.h b/Engine/Plugins/Runtime/Metasound/Source/MetasoundGraphCore/Public/MetasoundDataReference.h
--- a/Engine/Plugins/Runtime/Metasound/Source/MetasoundGraphCore/Public/MetasoundDataReference.h
+++ b/Engine/Plugins/Runtime/Metasound/Source/MetasoundGraphCore/Public/MetasoundDataReference.h
@@ -248,7 +248,7 @@ namespace Metasound
 			template <typename... ArgTypes>
 			static TDataWriteReference<DataType> CreateNew(ArgTypes&&... Args)
 			{
-				static_assert(std::is_constructible<DataType, ArgTypes...>::value, "Tried to call TDataWriteReference::CreateNew with args that don't match any constructor for an underlying type!");
+				static_assert(std::is_constructible<DataType, ArgTypes...>::value, "TDataWriteReference::CreateNew underlying type is not constructible with provided arguments.");
 				return TDataWriteReference<DataType>(EDataRefShouldConstruct::NewObject, Forward<ArgTypes>(Args)...);
 			}
 
diff --git a/Engine/Plugins/Runtime/Metasound/Source/MetasoundGraphCore/Public/MetasoundGraphCoreModule.h b/Engine/Plugins/Runtime/Metasound/Source/MetasoundGraphCore/Public/MetasoundGraphCoreModule.h
--- a/Engine/Plugins/Runtime/Metasound/Source/MetasoundGraphCore/Public/MetasoundGraphCoreModule.h
+++ b/Engine/Plugins/Runtime/Metasound/Source/MetasoundGraphCore/Public/MetasoundGraphCoreModule.h
@@ -5,8 +5,6 @@
 #include "CoreMinimal.h"
 #include "Modules/ModuleInterface.h"
 
-METASOUNDGRAPHCORE_API DECLARE_LOG_CATEGORY_EXTERN(LogMetasoundGraphCore, Log, All);
-
 namespace Metasound 
 {
 	class METASOUNDGRAPHCORE_API IMetasoundGraphCoreModule : public IModuleInterface
diff --git a/Engine/Plugins/Runtime/Metasound/Source/MetasoundGraphCore/Public/MetasoundLog.h b/Engine/Plugins/Runtime/Metasound/Source/MetasoundGraphCore/Public/MetasoundLog.h
new file mode 100644
--- /dev/null
+++ b/Engine/Plugins/Runtime/Metasound/Source/MetasoundGraphCore/Public/MetasoundLog.h
@@ -0,0 +1,8 @@
+// Copyright Epic Games, Inc. All Rights Reserved.
+
+#pragma once
+
+#include "CoreMinimal.h"
+
+METASOUNDGRAPHCORE_API DECLARE_LOG_CATEGORY_EXTERN(LogMetasound, Log, All);
+
diff --git a/Engine/Plugins/Runtime/Metasound/Source/MetasoundStandardNodes/Private/MetasoundAudioFormatNodes.cpp b/Engine/Plugins/Runtime/Metasound/Source/MetasoundStandardNodes/Private/MetasoundAudioFormatNodes.cpp
new file mode 100644
--- /dev/null
+++ b/Engine/Plugins/Runtime/Metasound/Source/MetasoundStandardNodes/Private/MetasoundAudioFormatNodes.cpp
@@ -0,0 +1,214 @@
+// Copyright Epic Games, Inc. All Rights Reserved.
+
+#include "MetasoundAudioFormatNodes.h"
+
+#include "MetasoundAudioFormats.h"
+#include "MetasoundFacade.h"
+#include "MetasoundOperatorInterface.h"
+
+#define LOCTEXT_NAMESPACE "MetasoundStandardNodes"
+
+namespace Metasound
+{
+	class FMonoAudioFormatOperator : public IOperator
+	{
+
+		static const FVertexKey& GetCenterName()
+		{
+			static const FVertexKey Name(TEXT("Center"));
+			return Name;
+		}
+
+		static const FVertexKey& GetMonoName()
+		{
+			static const FVertexKey Name(TEXT("Mono"));
+			return Name;
+		}
+
+		TDataReadReference<FMonoAudioFormat> MonoAudioRef;
+
+	public:
+
+		FMonoAudioFormatOperator(TDataReadReference<FMonoAudioFormat> InMonoAudioRef)
+		:	MonoAudioRef(InMonoAudioRef)
+		{
+		}
+
+
+		FDataReferenceCollection GetInputs() const override
+		{
+			FDataReferenceCollection Inputs;
+			Inputs.AddDataReadReference(GetCenterName(), MonoAudioRef->GetCenter());
+			return Inputs;	
+		}
+
+		FDataReferenceCollection GetOutputs() const override
+		{
+			FDataReferenceCollection Outputs;
+			Outputs.AddDataReadReference(GetMonoName(), MonoAudioRef);
+			return Outputs;
+		}
+
+		FExecuteFunction GetExecuteFunction() 
+		{
+			// No execution required. Mono format is created during operator instantiation.
+			return nullptr;
+		}
+		
+		
+		static TUniquePtr<IOperator> CreateOperator(const FCreateOperatorParams& InParams, TArray<TUniquePtr<IOperatorBuildError>>& OutErrors) 
+		{
+			TDataReadReference<FAudioBuffer> Center = InParams.InputDataReferences.GetDataReadReferenceOrConstruct<FAudioBuffer>(GetCenterName(), InParams.OperatorSettings);
+			TDataReadReference<FMonoAudioFormat> Mono = TDataReadReference<FMonoAudioFormat>::CreateNew(WriteCast(Center));
+			return MakeUnique<FMonoAudioFormatOperator>(Mono);
+		}
+
+		static const FNodeInfo& GetNodeInfo()
+		{
+			auto InitNodeInfo = []() -> FNodeInfo
+			{
+				FNodeInfo Info;
+
+				Info.ClassName = FName(TEXT("MonoAudioFormat"));
+				Info.MajorVersion = 1;
+				Info.MinorVersion = 0;
+				Info.Description = LOCTEXT("MonoAudioFormat_NodeDescription", "Mono Audio Format");
+				Info.Author = PluginAuthor;
+				Info.PromptIfMissing = PluginNodeMissingPrompt;
+
+				Info.DefaultInterface = FVertexInterface(
+					FInputVertexInterface(
+						TInputDataVertexModel<FAudioBuffer>(GetCenterName(), LOCTEXT("MonoAudioFormat_CenterChannel", "Audio"))
+					),
+					FOutputVertexInterface(
+						TOutputDataVertexModel<FMonoAudioFormat>(GetMonoName(), LOCTEXT("MonoAudioFormat_MonoOut", "Mono"))
+					)
+				);
+
+				return Info;
+			};
+
+			static const FNodeInfo Info = InitNodeInfo();
+			return Info;
+		}
+	};
+
+	FMonoAudioFormatNode::FMonoAudioFormatNode(const FString& InInstanceName)
+	: FNodeFacade(InInstanceName, TFacadeOperatorClass<FMonoAudioFormatOperator>())
+	{
+	}
+
+	FMonoAudioFormatNode::FMonoAudioFormatNode(const FNodeInitData& InInitData)
+	: FMonoAudioFormatNode(InInitData.InstanceName)
+	{
+	}
+
+	class FStereoAudioFormatOperator : public IOperator
+	{
+
+		static const FVertexKey& GetLeftName()
+		{
+			static const FVertexKey Name(TEXT("Left"));
+			return Name;
+		}
+
+		static const FVertexKey& GetRightName()
+		{
+			static const FVertexKey Name(TEXT("Right"));
+			return Name;
+		}
+
+		static const FVertexKey& GetStereoName()
+		{
+			static const FVertexKey Name(TEXT("Stereo"));
+			return Name;
+		}
+
+		TDataReadReference<FStereoAudioFormat> StereoAudio;
+
+	public:
+
+		FStereoAudioFormatOperator(TDataReadReference<FStereoAudioFormat> InStereoAudio)
+		:	StereoAudio(InStereoAudio)
+		{
+		}
+
+
+		FDataReferenceCollection GetInputs() const override
+		{
+			FDataReferenceCollection Inputs;
+			Inputs.AddDataReadReference(GetLeftName(), StereoAudio->GetLeft());
+			Inputs.AddDataReadReference(GetRightName(), StereoAudio->GetRight());
+			return Inputs;	
+		}
+
+		FDataReferenceCollection GetOutputs() const override
+		{
+			FDataReferenceCollection Outputs;
+			Outputs.AddDataReadReference(GetStereoName(), StereoAudio);
+			return Outputs;
+		}
+
+		FExecuteFunction GetExecuteFunction() 
+		{
+			// No execution required. Stereo format is created during operator instantiation.
+			return nullptr;
+		}
+		
+		
+		static TUniquePtr<IOperator> CreateOperator(const FCreateOperatorParams& InParams, TArray<TUniquePtr<IOperatorBuildError>>& OutErrors) 
+		{
+			TDataReadReference<FAudioBuffer> Left = InParams.InputDataReferences.GetDataReadReferenceOrConstruct<FAudioBuffer>(GetLeftName(), InParams.OperatorSettings);
+			TDataReadReference<FAudioBuffer> Right = InParams.InputDataReferences.GetDataReadReferenceOrConstruct<FAudioBuffer>(GetRightName(), InParams.OperatorSettings);
+			TDataReadReference<FStereoAudioFormat> Stereo = TDataReadReference<FStereoAudioFormat>::CreateNew(WriteCast(Left), WriteCast(Right));
+
+			return MakeUnique<FStereoAudioFormatOperator>(Stereo);
+		}
+
+		static const FNodeInfo& GetNodeInfo()
+		{
+			auto InitNodeInfo = []() -> FNodeInfo
+			{
+				FNodeInfo Info;
+
+				Info.ClassName = FName(TEXT("StereoAudioFormat"));
+				Info.MajorVersion = 1;
+				Info.MinorVersion = 0;
+				Info.Description = LOCTEXT("StereoAudioFormat_NodeDescription", "Stereo Audio Format");
+				Info.Author = PluginAuthor;
+				Info.PromptIfMissing = PluginNodeMissingPrompt;
+
+				Info.DefaultInterface = FVertexInterface(
+					FInputVertexInterface(
+						TInputDataVertexModel<FAudioBuffer>(GetLeftName(), LOCTEXT("StereoAudioFormat_LeftChannel", "Left Audio")),
+						TInputDataVertexModel<FAudioBuffer>(GetRightName(), LOCTEXT("StereoAudioFormat_RightChannel", "Right Audio"))
+					),
+					FOutputVertexInterface(
+						TOutputDataVertexModel<FStereoAudioFormat>(GetStereoName(), LOCTEXT("StereoAudioFormat_StereoOut", "Stereo"))
+					)
+				);
+
+				return Info;
+			};
+
+			static const FNodeInfo Info = InitNodeInfo();
+			return Info;
+		}
+	};
+
+	FStereoAudioFormatNode::FStereoAudioFormatNode(const FString& InInstanceName)
+	: FNodeFacade(InInstanceName, TFacadeOperatorClass<FStereoAudioFormatOperator>())
+	{
+	}
+
+	FStereoAudioFormatNode::FStereoAudioFormatNode(const FNodeInitData& InInitData)
+	: FStereoAudioFormatNode(InInitData.InstanceName)
+	{
+	}
+
+	METASOUND_REGISTER_NODE(FMonoAudioFormatNode);
+	METASOUND_REGISTER_NODE(FStereoAudioFormatNode);
+}
+
+
+#undef LOCTEXT_NAMESPACE
diff --git a/Engine/Plugins/Runtime/Metasound/Source/MetasoundStandardNodes/Public/MetasoundAudioFormatNodes.h b/Engine/Plugins/Runtime/Metasound/Source/MetasoundStandardNodes/Public/MetasoundAudioFormatNodes.h
new file mode 100644
--- /dev/null
+++ b/Engine/Plugins/Runtime/Metasound/Source/MetasoundStandardNodes/Public/MetasoundAudioFormatNodes.h
@@ -0,0 +1,24 @@
+// Copyright Epic Games, Inc. All Rights Reserved.
+
+#pragma once
+
+#include "CoreMinimal.h"
+#include "MetasoundAudioFormats.h"
+#include "MetasoundFacade.h"
+
+namespace Metasound
+{
+	class METASOUNDSTANDARDNODES_API FMonoAudioFormatNode : public FNodeFacade
+	{
+		public:
+			FMonoAudioFormatNode(const FString& InInstanceName);
+			FMonoAudioFormatNode(const FNodeInitData& InInitData);
+	};
+
+	class METASOUNDSTANDARDNODES_API FStereoAudioFormatNode : public FNodeFacade
+	{
+		public:
+			FStereoAudioFormatNode(const FString& InInstanceName);
+			FStereoAudioFormatNode(const FNodeInitData& InInitData);
+	};
+}
diff --git a/Engine/Plugins/Runtime/Metasound/Source/MetasoundStandardNodes/Public/MetasoundAudioFormats.h b/Engine/Plugins/Runtime/Metasound/Source/MetasoundStandardNodes/Public/MetasoundAudioFormats.h
--- a/Engine/Plugins/Runtime/Metasound/Source/MetasoundStandardNodes/Public/MetasoundAudioFormats.h
+++ b/Engine/Plugins/Runtime/Metasound/Source/MetasoundStandardNodes/Public/MetasoundAudioFormats.h
@@ -211,10 +211,12 @@ namespace Metasound
 	 * buffers within a FMultichannelAudioFormat object must contain the same 
 	 * number of audio frames.
 	 */
-	template<int32 NumChannels>
+	template<int32 TNumChannels>
 	class TStaticChannelAudioFormat : public IAudioDataType
 	{
 		public:
+			static constexpr int32 NumChannels = TNumChannels;
+
 			/** TStaticChannelAudioFormat Constructor
 			 *
 			 * @param InNumFrames - The number of frames per an audio buffer.
@@ -277,7 +279,7 @@ namespace Metasound
 			 *
 			 * This array will have GetNumChannels() elements.
 			 */
-			const TArrayView<const FAudioBufferReadRef> GetBuffers() const
+			const TArrayView<FAudioBufferReadRef> GetBuffers() const
 			{
 				return ReadableBuffers;
 			}
@@ -286,7 +288,7 @@ namespace Metasound
 			 *
 			 * This array will have GetNumChannels() elements.
 			 */
-			const TArrayView<const FAudioBufferWriteRef> GetBuffers() 
+			const TArrayView<FAudioBufferWriteRef> GetBuffers() 
 			{
 				return WritableBuffers;
 			}
@@ -333,8 +335,8 @@ namespace Metasound
 			}
 
 		private:
-			TArrayView<const FAudioBufferWriteRef> WritableBuffers;
-			TArrayView<const FAudioBufferReadRef> ReadableBuffers;
+			TArrayView<FAudioBufferWriteRef> WritableBuffers;
+			TArrayView<FAudioBufferReadRef> ReadableBuffers;
 
 			TArray<FAudioBufferWriteRef> WritableBufferStorage;
 			TArray<FAudioBufferReadRef> ReadableBufferStorage;
@@ -347,11 +349,6 @@ namespace Metasound
 	 */
 	class FMonoAudioFormat : public TStaticChannelAudioFormat<1>
 	{
-			// This is used to grant specific classes the ability to create FMonoAudioFormat
-			// using a FAudioBufferReadRef
-			//enum EPrivateToken { Token };
-			//static const EPrivateToken PrivateToken = EPrivateToken::Token;
-
 		public:
 			using Super = TStaticChannelAudioFormat<1>;
 
@@ -369,47 +366,12 @@ namespace Metasound
 			{
 			}
 
-			// TODO: clean up this mess.
-			// Construct a FMonoAudioFormat a readable buffer reference.
-			// 
-			// This constructor should only be used when it can be assured 
-			// that the constructed object will not provide writable access 
-			// to the passed in audio buffers. 
-			/*
-			FMonoAudioFormat(const FAudioBufferReadRef& InAudio, EPrivateToken InToken)
-			:	Super({WriteCast(InAudio)})
-			{
-			}
-			*/
-
 			/** Return writable audio buffer reference of center channel. */
 			FAudioBufferWriteRef GetCenter() { return GetBuffer<0>(); }
 
 			/** Return readable audio buffer reference of center channel. */
 			FAudioBufferReadRef GetCenter() const { return GetBuffer<0>(); }
-
-		private:
-
-
-			// Friendship with the data reference class gives it access to the 
-			// protected constructor for the scenario where a 
-			// TDataReadReference<*> is constructed with FAudioBufferReadRefs. 
-			// The constructor cannot be public as it would provide writable
-			// access to the passed in audio buffers, even though the passed 
-			// buffers explicitly were read references.
-			//friend class TDataReadReference<FMonoAudioFormat>;
 	};
-
-	// Template specialization for TDataReadReference<FMonoAudioFormat> constructor to use
-	// private token.
-	/*
-	template<>
-	template<>
-	TDataReadReference<FMonoAudioFormat>::TDataReadReference(const FAudioBufferReadRef& InReadBuffer)
-	:	FDataReference(InReadBuffer, FMonoAudioFormat::PrivateToken)
-	{
-	}
-	*/
 	
 
 	/** FStereoAudioFormat represents stereo audio containing two channels of 
@@ -449,26 +411,6 @@ namespace Metasound
 
 			/** Return readable audio buffer reference of right channel. */
 			FAudioBufferReadRef GetRight() const { return GetBuffer<1>(); }
-
-		protected:
-
-			// Friendship with the data reference class gives it access to the 
-			// protected constructor for the scenario where a 
-			// TDataReadReference<*> is constructed with FAudioBufferReadRefs. 
-			// The constructor cannot be public as it would provide writable
-			// access to the passed in audio buffers, even though the passed 
-			// buffers explicitly were read references.
-			friend class TDataReadReference<FStereoAudioFormat>;
-
-			// Construct a FStereoAudioFormat a readable buffer reference.
-			// 
-			// This constructor should only be used when it can be assured 
-			// that the constructed object will not provide writable access 
-			// to the passed in audio buffers. 
-			FStereoAudioFormat(const FAudioBufferReadRef& InLeftAudio, const FAudioBufferReadRef& InRightAudio)
-			:	Super({WriteCast(InLeftAudio), WriteCast(InRightAudio)})
-			{
-			}
 	};
 
 
