- Fix-up procedural sounds to support stopping source fade. Remove some restrictions requiring only one instance to be generated per USoundWaveProcedural instance.
- Rig up metasound OnPlay/OnFinished triggers
- Add typedefs for BopRead/WriteRef to begin deprecating Bop and using Trigger
- Fix-up Periodic Bop node to have activate/deactivate
- Rename to 'PeriodicTrigger'

diff --git a/Engine/Plugins/Runtime/Metasound/Source/MetasoundEngine/Private/MetasoundSource.cpp b/Engine/Plugins/Runtime/Metasound/Source/MetasoundEngine/Private/MetasoundSource.cpp
--- a/Engine/Plugins/Runtime/Metasound/Source/MetasoundEngine/Private/MetasoundSource.cpp
+++ b/Engine/Plugins/Runtime/Metasound/Source/MetasoundEngine/Private/MetasoundSource.cpp
@@ -26,6 +26,7 @@ UMetasoundSource::UMetasoundSource(const FObjectInitializer& ObjectInitializer)
 	: Super(ObjectInitializer)
 	, FMetasoundAssetBase(GetMonoSourceArchetype())
 {
+	bRequiresStopFade = true;
 	NumChannels = 1;
 	Duration = INDEFINITELY_LOOPING_DURATION;
 	bLooping = true;
@@ -93,9 +94,10 @@ ISoundGeneratorPtr UMetasoundSource::CreateSoundGenerator(const FSoundGeneratorI
 
 	Duration = INDEFINITELY_LOOPING_DURATION;
 	bLooping = true;
+	VirtualizationMode = EVirtualizationMode::PlayWhenSilent;
+
 	SampleRate = InParams.SampleRate;
 	const float BlockRate = 100.f; // Metasound graph gets evaluated 100 times per second.
-
 	FOperatorSettings InSettings(InParams.SampleRate, BlockRate);
 	FMetasoundEnvironment Environment;
 
@@ -150,19 +152,25 @@ ISoundGeneratorPtr UMetasoundSource::CreateSoundGenerator(const FSoundGeneratorI
 		}
 		else if (EMetasoundSourceAudioFormat::Mono == OutputFormat)
 		{
-			if (!Outputs.ContainsDataReadReference<FStereoAudioFormat>(GetAudioOutputName()))
+			if (!Outputs.ContainsDataReadReference<FMonoAudioFormat>(GetAudioOutputName()))
 			{
 				UE_LOG(LogMetasound, Warning, TEXT("MetasoundSource [%s] does not contain mono output [%s] in output"), *GetName(), *GetAudioOutputName());
 			}
 			OutputBuffers = Outputs.GetDataReadReferenceOrConstruct<FMonoAudioFormat>(GetAudioOutputName(), InSettings)->GetBuffers();
 		}
 
+		// References must be cached before moving the operator to the InitParams
+		FDataReferenceCollection Inputs = Operator->GetInputs();
+		FTriggerWriteRef PlayTrigger = Inputs.GetDataWriteReference<FBop>(GetOnPlayInputName());
+		FTriggerReadRef FinishTrigger = Outputs.GetDataReadReferenceOrConstruct<FBop>(GetIsFinishedOutputName(), InSettings, false);
+
 		// Create the FMetasoundGenerator.
 		FMetasoundGeneratorInitParams InitParams =
 		{
 			MoveTemp(Operator),
 			OutputBuffers,
-			Outputs.GetDataReadReferenceOrConstruct<FBop>(GetIsFinishedOutputName(), InSettings, false)
+			MoveTemp(PlayTrigger),
+			MoveTemp(FinishTrigger)
 		};
 
 		return ISoundGeneratorPtr(new FMetasoundGenerator(MoveTemp(InitParams)));
@@ -180,32 +188,31 @@ void UMetasoundSource::PostLoad()
 
 const TArray<FMetasoundArchetype>& UMetasoundSource::GetPreferredArchetypes() const 
 {
-	static TArray<FMetasoundArchetype> Preferred({GetMonoSourceArchetype(), GetStereoSourceArchetype()});
-
+	static const TArray<FMetasoundArchetype> Preferred({GetMonoSourceArchetype(), GetStereoSourceArchetype()});
 	return Preferred;
 }
 
 const FString& UMetasoundSource::GetOnPlayInputName()
 {
-	static FString BopInputName = FString(TEXT("On Play"));
-	return BopInputName;
+	static const FString TriggerInputName = TEXT("On Play");
+	return TriggerInputName;
 }
 
 const FString& UMetasoundSource::GetAudioOutputName()
 {
-	static FString AudioOutputName = FString(TEXT("Generated Audio"));
+	static const FString AudioOutputName = TEXT("Generated Audio");
 	return AudioOutputName;
 }
 
 const FString& UMetasoundSource::GetIsFinishedOutputName()
 {
-	static FString OnFinishedOutputName = FString(TEXT("On Finished"));
+	static const FString OnFinishedOutputName = TEXT("On Finished");
 	return OnFinishedOutputName;
 }
 
 const FString& UMetasoundSource::GetAudioDeviceHandleVariableName()
 {
-	static FString AudioDeviceHandleVarName = FString(TEXT("AudioDeviceHandle"));
+	static const FString AudioDeviceHandleVarName = TEXT("AudioDeviceHandle");
 	return AudioDeviceHandleVarName;
 }
 
@@ -215,13 +222,13 @@ const FMetasoundArchetype& UMetasoundSource::GetBaseArchetype()
 	{
 		FMetasoundArchetype Archetype;
 		
-		FMetasoundInputDescription OnPlayBop;
-		OnPlayBop.Name = UMetasoundSource::GetOnPlayInputName();
-		OnPlayBop.DisplayName = FText::FromString(OnPlayBop.Name);
-		OnPlayBop.TypeName = Metasound::Frontend::GetDataTypeName<Metasound::FBop>();
-		OnPlayBop.ToolTip = LOCTEXT("OnPlayBopToolTip", "Bop executed when this source is first played.");
+		FMetasoundInputDescription OnPlayTrigger;
+		OnPlayTrigger.Name = UMetasoundSource::GetOnPlayInputName();
+		OnPlayTrigger.DisplayName = FText::FromString(OnPlayTrigger.Name);
+		OnPlayTrigger.TypeName = Metasound::Frontend::GetDataTypeName<Metasound::FBop>();
+		OnPlayTrigger.ToolTip = LOCTEXT("OnPlayTriggerToolTip", "Trigger executed when this source is played.");
 
-		Archetype.RequiredInputs.Add(OnPlayBop);
+		Archetype.RequiredInputs.Add(OnPlayTrigger);
 
 		FMetasoundOutputDescription OnFinished;
 		OnFinished.Name = UMetasoundSource::GetIsFinishedOutputName();
diff --git a/Engine/Plugins/Runtime/Metasound/Source/MetasoundGenerator/Private/MetasoundGenerator.cpp b/Engine/Plugins/Runtime/Metasound/Source/MetasoundGenerator/Private/MetasoundGenerator.cpp
--- a/Engine/Plugins/Runtime/Metasound/Source/MetasoundGenerator/Private/MetasoundGenerator.cpp
+++ b/Engine/Plugins/Runtime/Metasound/Source/MetasoundGenerator/Private/MetasoundGenerator.cpp
@@ -7,7 +7,7 @@
 #include "MetasoundOscNode.h"
 #include "MetasoundAudioMultiplyNode.h"
 #include "MetasoundADSRNode.h"
-#include "MetasoundPeriodicBopNode.h"
+#include "MetasoundPeriodicTriggerNode.h"
 #include "MetasoundOutputNode.h"
 #include "MetasoundInputNode.h"
 #include "MetasoundOperatorBuilder.h"
@@ -17,10 +17,13 @@
 namespace Metasound
 {
 	FMetasoundGenerator::FMetasoundGenerator(FMetasoundGeneratorInitParams&& InParams)
-		: NumChannels(0)
+		: bIsPlaying(false)
+		, bIsFinished(false)
+		, NumChannels(0)
 		, NumFramesPerExecute(0)
 		, NumSamplesPerExecute(0)
-		, OnFinishedBopRef(InParams.BopOnFinishRef)
+		, OnPlayTriggerRef(InParams.TriggerOnPlayRef)
+		, OnFinishedTriggerRef(InParams.TriggerOnFinishRef)
 	{
 		SetGraph(MoveTemp(InParams));
 	}
@@ -56,7 +59,8 @@ namespace Metasound
 			GraphOutputAudio.Append(InParams.OutputBuffers.GetData(), InParams.OutputBuffers.Num());
 		}
 
-		OnFinishedBopRef = InParams.BopOnFinishRef;
+		OnPlayTriggerRef = InParams.TriggerOnPlayRef;
+		OnFinishedTriggerRef = InParams.TriggerOnFinishRef;
 
 		// The graph operator and graph audio output contain all the values needed
 		// by the sound generator.
@@ -100,6 +104,12 @@ namespace Metasound
 			return 0;
 		}
 
+		if (!bIsPlaying)
+		{
+			OnPlayTriggerRef->BopFrame(0);
+			bIsPlaying = true;
+		}
+
 		if (NumSamplesPerExecute < 1)
 		{
 			FMemory::Memset(OutAudio, 0, sizeof(float) * NumSamplesRemaining);
@@ -144,10 +154,16 @@ namespace Metasound
 			}
 		}
 
+		if (*OnFinishedTriggerRef)
+		{
+			bIsFinished = true;
+		}
+
+		OnPlayTriggerRef->AdvanceBlock();
+
 		return NumSamplesWritten;
 	}
 
-
 	int32 FMetasoundGenerator::GetDesiredNumSamplesToRenderPerCallback() const
 	{
 		// TODO: may improve performance if this number is increased. 
@@ -156,7 +172,7 @@ namespace Metasound
 
 	bool FMetasoundGenerator::IsFinished() const
 	{
-		return *OnFinishedBopRef;
+		return bIsFinished;
 	}
 
 	int32 FMetasoundGenerator::FillWithBuffer(const Audio::AlignedFloatBuffer& InBuffer, float* OutAudio, int32 MaxNumOutputSamples)
@@ -208,5 +224,3 @@ namespace Metasound
 		// TODO: memcpy for single channel. 
 	}
 } 
-
-
diff --git a/Engine/Plugins/Runtime/Metasound/Source/MetasoundGenerator/Public/MetasoundGenerator.h b/Engine/Plugins/Runtime/Metasound/Source/MetasoundGenerator/Public/MetasoundGenerator.h
--- a/Engine/Plugins/Runtime/Metasound/Source/MetasoundGenerator/Public/MetasoundGenerator.h
+++ b/Engine/Plugins/Runtime/Metasound/Source/MetasoundGenerator/Public/MetasoundGenerator.h
@@ -17,8 +17,8 @@ namespace Metasound
 	{
 		TUniquePtr<Metasound::IOperator> GraphOperator;
 		TArrayView<TDataReadReference<FAudioBuffer>> OutputBuffers;
-		// TODO: BopOnPlayRef;
-		FBopReadRef BopOnFinishRef;
+		FBopWriteRef TriggerOnPlayRef;
+		FBopReadRef TriggerOnFinishRef;
 	};
 
 	/** FMetasoundGenerator generates audio from a given metasound IOperator
@@ -119,14 +119,20 @@ namespace Metasound
 		
 		FExecuter RootExecuter;
 
+		bool bIsPlaying;
+		bool bIsFinished;
+
 		int32 NumChannels;
 		int32 NumFramesPerExecute;
 		int32 NumSamplesPerExecute;
 
 		TArray<FAudioBufferReadRef> GraphOutputAudio;
 
-		// This will be bopped when this metasound is finished playing.
-		FBopReadRef OnFinishedBopRef;
+		// Triggered when metasound is played
+		FBopWriteRef OnPlayTriggerRef;
+
+		// Triggered when metasound is finished
+		FBopReadRef OnFinishedTriggerRef;
 
 		Audio::AlignedFloatBuffer InterleavedAudioBuffer;
 
diff --git a/Engine/Plugins/Runtime/Metasound/Source/MetasoundGraphCore/Public/MetasoundDataReferenceCollection.h b/Engine/Plugins/Runtime/Metasound/Source/MetasoundGraphCore/Public/MetasoundDataReferenceCollection.h
--- a/Engine/Plugins/Runtime/Metasound/Source/MetasoundGraphCore/Public/MetasoundDataReferenceCollection.h
+++ b/Engine/Plugins/Runtime/Metasound/Source/MetasoundGraphCore/Public/MetasoundDataReferenceCollection.h
@@ -254,7 +254,7 @@ namespace Metasound
 				}
 				else
 				{
-					return FDataRefType(Forward<ConstructorArgTypes>(ConstructorArgs)...);
+					return FDataRefType::CreateNew(Forward<ConstructorArgTypes>(ConstructorArgs)...);
 				}
 			}
 
diff --git a/Engine/Plugins/Runtime/Metasound/Source/MetasoundStandardNodes/Private/MetasoundPeriodicBopNode.cpp b/Engine/Plugins/Runtime/Metasound/Source/MetasoundStandardNodes/Private/MetasoundPeriodicBopNode.cpp
deleted file mode 100644
diff --git a/Engine/Plugins/Runtime/Metasound/Source/MetasoundStandardNodes/Private/MetasoundPeriodicTriggerNode.cpp b/Engine/Plugins/Runtime/Metasound/Source/MetasoundStandardNodes/Private/MetasoundPeriodicTriggerNode.cpp
new file mode 100644
--- /dev/null
+++ b/Engine/Plugins/Runtime/Metasound/Source/MetasoundStandardNodes/Private/MetasoundPeriodicTriggerNode.cpp
@@ -0,0 +1,193 @@
+// Copyright Epic Games, Inc. All Rights Reserved.
+
+#include "MetasoundPeriodicTriggerNode.h"
+
+#include "CoreMinimal.h"
+
+#include "MetasoundBuilderInterface.h"
+#include "MetasoundDataReferenceCollection.h"
+#include "MetasoundExecutableOperator.h"
+#include "MetasoundFacade.h"
+#include "MetasoundNode.h"
+#include "MetasoundNodeInterface.h"
+#include "MetasoundNodeRegistrationMacro.h"
+#include "MetasoundOperatorInterface.h"
+#include "MetasoundPrimitives.h"
+
+#define LOCTEXT_NAMESPACE "MetasoundStandardNodes"
+
+namespace Metasound
+{
+	METASOUND_REGISTER_NODE(FPeriodicTriggerNode)
+
+	class FPeriodicTriggerOperator : public TExecutableOperator<FPeriodicTriggerOperator>
+	{
+		public:
+			static const FNodeInfo& GetNodeInfo();
+			static FVertexInterface DeclareVertexInterface();
+			static TUniquePtr<IOperator> CreateOperator(const FCreateOperatorParams& InParams, FBuildErrorArray& OutErrors);
+
+			static constexpr float MinimumPeriodSeconds = 0.001f;
+			static constexpr float MinimumPeriodSamples = 20.f;
+
+			FPeriodicTriggerOperator(const FOperatorSettings& InSettings, const FTriggerReadRef& InTriggerEnable, const FTriggerReadRef& InTriggerDisable, const FFloatTimeReadRef& InPeriod);
+
+			virtual FDataReferenceCollection GetInputs() const override;
+
+			virtual FDataReferenceCollection GetOutputs() const override;
+
+			void Execute();
+
+		private:
+			bool bEnabled;
+
+			FTriggerWriteRef TriggerOut;
+
+			FTriggerReadRef TriggerEnable;
+			FTriggerReadRef TriggerDisable;
+
+			FFloatTimeReadRef Period;
+
+			float ExecuteDurationInSamples;
+			float SampleRate;
+			float SampleCountdown;
+	};
+
+	FPeriodicTriggerOperator::FPeriodicTriggerOperator(const FOperatorSettings& InSettings, const FTriggerReadRef& InTriggerEnable, const FTriggerReadRef& InTriggerDisable, const FFloatTimeReadRef& InPeriod)
+	:	bEnabled(false)
+	,	TriggerOut(FTriggerWriteRef::CreateNew(InSettings))
+	,	TriggerEnable(InTriggerEnable)
+	,	TriggerDisable(InTriggerDisable)
+	,	Period(InPeriod)
+	,	ExecuteDurationInSamples(InSettings.GetNumFramesPerBlock())
+	,	SampleRate(InSettings.GetSampleRate())
+	,	SampleCountdown(0.f)
+	{
+	}
+
+	FDataReferenceCollection FPeriodicTriggerOperator::GetInputs() const
+	{
+		FDataReferenceCollection InputDataReferences;
+		InputDataReferences.AddDataReadReference(TEXT("Period"), FFloatTimeReadRef(Period));
+		InputDataReferences.AddDataReadReference(TEXT("Activate"), FTriggerReadRef(TriggerEnable));
+		InputDataReferences.AddDataReadReference(TEXT("Deactivate"), FTriggerReadRef(TriggerDisable));
+		return InputDataReferences;
+	}
+
+	FDataReferenceCollection FPeriodicTriggerOperator::GetOutputs() const
+	{
+		FDataReferenceCollection OutputDataReferences;
+		OutputDataReferences.AddDataReadReference(TEXT("Out"), FTriggerReadRef(TriggerOut));
+
+		return OutputDataReferences;
+	}
+
+	void FPeriodicTriggerOperator::Execute()
+	{
+		TriggerEnable->ExecuteBlock([](int32, int32) { },
+			[this](int32 StartFrame, int32 EndFrame)
+			{
+				bEnabled = true;
+			}
+		);
+
+		TriggerDisable->ExecuteBlock([](int32, int32) { },
+			[this](int32 StartFrame, int32 EndFrame)
+			{
+				bEnabled = false;
+			}
+		);
+
+		// Advance internal counter to get rid of old bops.
+		TriggerOut->AdvanceBlock();
+
+		if (bEnabled)
+		{
+			float PeriodInSamples = FMath::Max(Period->GetSeconds(), MinimumPeriodSeconds) * SampleRate;
+
+			PeriodInSamples = FMath::Max(PeriodInSamples, MinimumPeriodSamples);
+
+			while ((SampleCountdown - ExecuteDurationInSamples) <= 0.f)
+			{
+				uint32 Frame = FMath::RoundToInt(SampleCountdown);
+				TriggerOut->BopFrame(Frame);
+				SampleCountdown += PeriodInSamples;
+			}
+
+			SampleCountdown -= ExecuteDurationInSamples;
+		}
+	}
+
+
+	TUniquePtr<IOperator> FPeriodicTriggerOperator::CreateOperator(const FCreateOperatorParams& InParams, FBuildErrorArray& OutErrors)
+	{
+		const FPeriodicTriggerNode& PeriodicBopNode = static_cast<const FPeriodicTriggerNode&>(InParams.Node);
+
+		FFloatTimeReadRef Period = FFloatTimeReadRef::CreateNew(PeriodicBopNode.GetDefaultPeriodInSeconds(), ETimeResolution::Seconds);
+
+		if (InParams.InputDataReferences.ContainsDataReadReference<FFloatTime>(TEXT("Period")))
+		{
+			Period = InParams.InputDataReferences.GetDataReadReference<FFloatTime>(TEXT("Period"));
+		}
+
+		FTriggerReadRef TriggerEnable = InParams.InputDataReferences.GetDataReadReferenceOrConstruct<FBop>(TEXT("Activate"), InParams.OperatorSettings);
+		FTriggerReadRef TriggerDisable = InParams.InputDataReferences.GetDataReadReferenceOrConstruct<FBop>(TEXT("Deactivate"), InParams.OperatorSettings);
+
+		return MakeUnique<FPeriodicTriggerOperator>(InParams.OperatorSettings, TriggerEnable, TriggerDisable, Period);
+	}
+
+	FVertexInterface FPeriodicTriggerOperator::DeclareVertexInterface()
+	{
+		static const FVertexInterface Interface(
+			FInputVertexInterface(
+				TInputDataVertexModel<FFloatTime>(TEXT("Period"), LOCTEXT("PeriodTooltip", "The period to trigger in seconds.")),
+				TInputDataVertexModel<FBop>(TEXT("Activate"), LOCTEXT("TriggerEnableTooltip", "Enables executing periodic output triggers.")),
+				TInputDataVertexModel<FBop>(TEXT("Deactivate"), LOCTEXT("TriggerDisableTooltip", "Disables executing periodic output triggers."))
+			),
+			FOutputVertexInterface(
+				TOutputDataVertexModel<FBop>(TEXT("Out"), LOCTEXT("TriggerOutTooltip", "The periodically generated output trigger"))
+			)
+		);
+
+		return Interface;
+	}
+
+	const FNodeInfo& FPeriodicTriggerOperator::GetNodeInfo()
+	{
+		auto InitNodeInfo = []() -> FNodeInfo
+		{
+			FNodeInfo Info;
+			Info.ClassName = "PeriodicBop";
+			Info.MajorVersion = 1;
+			Info.MinorVersion = 0;
+			Info.Description = LOCTEXT("Metasound_PeriodicBopNodeDescription", "Emits a bop periodically based on the period duration given.");
+			Info.Author = PluginAuthor;
+			Info.PromptIfMissing = PluginNodeMissingPrompt;
+			Info.DefaultInterface = DeclareVertexInterface();
+
+			return Info;
+		};
+
+		static const FNodeInfo Info = InitNodeInfo();
+
+		return Info;
+	}
+
+
+	FPeriodicTriggerNode::FPeriodicTriggerNode(const FString& InName, float InDefaultPeriodInSeconds)
+	:	FNodeFacade(InName, TFacadeOperatorClass<FPeriodicTriggerOperator>())
+	,	DefaultPeriod(InDefaultPeriodInSeconds)
+	{
+	}
+
+	FPeriodicTriggerNode::FPeriodicTriggerNode(const FNodeInitData& InInitData)
+		: FPeriodicTriggerNode(InInitData.InstanceName, 1.0f)
+	{}
+
+	float FPeriodicTriggerNode::GetDefaultPeriodInSeconds() const
+	{
+		return DefaultPeriod;
+	}
+}
+
+#undef LOCTEXT_NAMESPACE //MetasoundPeriodicBopNode
diff --git a/Engine/Plugins/Runtime/Metasound/Source/MetasoundStandardNodes/Public/MetasoundBop.h b/Engine/Plugins/Runtime/Metasound/Source/MetasoundStandardNodes/Public/MetasoundBop.h
--- a/Engine/Plugins/Runtime/Metasound/Source/MetasoundStandardNodes/Public/MetasoundBop.h
+++ b/Engine/Plugins/Runtime/Metasound/Source/MetasoundStandardNodes/Public/MetasoundBop.h
@@ -167,4 +167,8 @@ namespace Metasound
 	};
 
 	DECLARE_METASOUND_DATA_REFERENCE_TYPES(FBop, METASOUNDSTANDARDNODES_API, FBopTypeInfo, FBopReadRef, FBopWriteRef);
+
+	using FTrigger = FBop;
+	using FTriggerReadRef = FBopReadRef;
+	using FTriggerWriteRef = FBopWriteRef;
 }
diff --git a/Engine/Plugins/Runtime/Metasound/Source/MetasoundStandardNodes/Public/MetasoundPeriodicBopNode.h b/Engine/Plugins/Runtime/Metasound/Source/MetasoundStandardNodes/Public/MetasoundPeriodicBopNode.h
deleted file mode 100644
diff --git a/Engine/Plugins/Runtime/Metasound/Source/MetasoundStandardNodes/Public/MetasoundPeriodicTriggerNode.h b/Engine/Plugins/Runtime/Metasound/Source/MetasoundStandardNodes/Public/MetasoundPeriodicTriggerNode.h
new file mode 100644
--- /dev/null
+++ b/Engine/Plugins/Runtime/Metasound/Source/MetasoundStandardNodes/Public/MetasoundPeriodicTriggerNode.h
@@ -0,0 +1,29 @@
+// Copyright Epic Games, Inc. All Rights Reserved.
+
+#pragma once
+
+#include "CoreMinimal.h"
+#include "MetasoundDataReferenceCollection.h"
+#include "MetasoundNode.h"
+#include "MetasoundBop.h"
+#include "MetasoundBuilderInterface.h"
+#include "MetasoundExecutableOperator.h"
+#include "MetasoundFacade.h"
+#include "MetasoundOperatorInterface.h"
+#include "MetasoundNodeInterface.h"
+#include "MetasoundTime.h"
+
+namespace Metasound
+{
+	class METASOUNDSTANDARDNODES_API FPeriodicTriggerNode : public FNodeFacade
+	{
+		public:
+			FPeriodicTriggerNode(const FString& InName, float InDefaultPeriodInSeconds);
+			FPeriodicTriggerNode(const FNodeInitData& InInitData);
+
+			float GetDefaultPeriodInSeconds() const;
+
+		private:
+			float DefaultPeriod;
+	};
+}
diff --git a/Engine/Source/Runtime/AudioMixer/Private/AudioMixerSource.cpp b/Engine/Source/Runtime/AudioMixer/Private/AudioMixerSource.cpp
--- a/Engine/Source/Runtime/AudioMixer/Private/AudioMixerSource.cpp
+++ b/Engine/Source/Runtime/AudioMixer/Private/AudioMixerSource.cpp
@@ -23,6 +23,21 @@ FAutoConsoleVariableRef CVarUseListenerOverrideForSpread(
 	TEXT("0: Use actual distance, 1: use listener override"),
 	ECVF_Default);
 
+static uint32 AudioMixerSourceFadeMinCVar = 512;
+static FAutoConsoleCommand GSetAudioMixerSourceFadeMin(
+	TEXT("au.SourceFadeMin"),
+	TEXT("Sets the length (in samples) of minimum fade when a sound source is stopped. Must be divisible by 4 (vectorization requirement). Ignored for some procedural source types. (Default: 512, Min: 4). \n"),
+	FConsoleCommandWithArgsDelegate::CreateStatic(
+		[](const TArray<FString>& Args)
+		{
+			if (Args.Num() > 0)
+			{
+				const int32 SourceFadeMin = FMath::Max(FCString::Atoi(*Args[0]), 4);
+				AudioMixerSourceFadeMinCVar = AlignArbitrary(SourceFadeMin, 4);
+			}
+		}
+	)
+);
 
 namespace Audio
 {
@@ -758,7 +773,7 @@ namespace Audio
 			StartFrame = FMath::Clamp<int32>((InWaveInstance->StartTime / SoundWave.Duration) * NumTotalFrames, 0, NumTotalFrames);
 		}
 
-		check(!MixerSourceBuffer.IsValid());		
+		check(!MixerSourceBuffer.IsValid());
 		MixerSourceBuffer = FMixerSourceBuffer::Create(AudioDevice->GetSampleRate(), *MixerBuffer, SoundWave, InWaveInstance->LoopingMode, bIsSeeking);
 		
 		if (!MixerSourceBuffer.IsValid())
@@ -865,11 +880,6 @@ namespace Audio
 		// we'll just not actually play the source here. Instead we'll call play when the sound finishes loading.
 		if (MixerSourceVoice && InitializationState == EMixerSourceInitializationState::Initialized)
 		{
-			if (WaveInstance && WaveInstance->WaveData && WaveInstance->WaveData->bProcedural)
-			{
-				WaveInstance->WaveData->bPlayingProcedural = true;
-			}
-
 			MixerSourceVoice->Play();
 		}
 
@@ -895,10 +905,18 @@ namespace Audio
 			return;
 		}
 
-		// Always stop procedural sounds immediately.
-		if (WaveInstance && WaveInstance->WaveData && WaveInstance->WaveData->bProcedural)
+		USoundWave* SoundWave = WaveInstance ? WaveInstance->WaveData : nullptr;
+
+		// If MarkPendingKill() was called, SoundWave can be null
+		if (!SoundWave)
+		{
+			StopNow();
+			return;
+		}
+
+		// Stop procedural sounds immediately that don't require fade
+		if (SoundWave->bProcedural && !SoundWave->bRequiresStopFade)
 		{
-			WaveInstance->WaveData->bPlayingProcedural = false;
 			StopNow();
 			return;
 		}
@@ -906,37 +924,27 @@ namespace Audio
 		if (bIsDone)
 		{
 			StopNow();
+			return;
 		}
-		else if (!bIsStopping)
+
+		if (Playing && !bIsStoppingVoicesEnabled)
 		{
-			// Otherwise, we need to do a quick fade-out of the sound and put the state
-			// of the sound into "stopping" mode. This prevents this source from
-			// being put into the "free" pool and prevents the source from freeing its resources
-			// until the sound has finished naturally (i.e. faded all the way out)
+			StopNow();
+			return;
+		}
 
-			// StopFade will stop a sound with a very small fade to avoid discontinuities
-			if (MixerSourceVoice && Playing)
-			{
-				// if MarkPendingKill() was called, WaveInstance->WaveData is null
-				if (!WaveInstance || !WaveInstance->WaveData)
-				{
-					StopNow();
-					return;
-				}
-				else if (bIsStoppingVoicesEnabled && !WaveInstance->WaveData->bProcedural)
-				{
-					// Let the wave instance know it's stopping
-					WaveInstance->SetStopping(true);
+		// Otherwise, we need to do a quick fade-out of the sound and put the state
+		// of the sound into "stopping" mode. This prevents this source from
+		// being put into the "free" pool and prevents the source from freeing its resources
+		// until the sound has finished naturally (i.e. faded all the way out)
 
-					// TODO: parameterize the number of fades
-					MixerSourceVoice->StopFade(512);
-					bIsStopping = true;
-				}
-				else
-				{
-					StopNow();
-				}
-			}
+		// Let the wave instance know it's stopping
+		if (!bIsStopping)
+		{
+			WaveInstance->SetStopping(true);
+
+			MixerSourceVoice->StopFade(AudioMixerSourceFadeMinCVar);
+			bIsStopping = true;
 			Paused = false;
 		}
 	}
diff --git a/Engine/Source/Runtime/AudioMixer/Private/AudioMixerSourceBuffer.cpp b/Engine/Source/Runtime/AudioMixer/Private/AudioMixerSourceBuffer.cpp
--- a/Engine/Source/Runtime/AudioMixer/Private/AudioMixerSourceBuffer.cpp
+++ b/Engine/Source/Runtime/AudioMixer/Private/AudioMixerSourceBuffer.cpp
@@ -60,19 +60,8 @@ namespace Audio
 	{
 		LLM_SCOPE(ELLMTag::AudioMixer);
 
-		// Prevent double-triggering procedural soundwaves
-		if (InWave.bProcedural && InWave.IsGeneratingAudio())
-		{
-			UE_LOG(LogAudioMixer, Warning,
-				TEXT("Procedural USoundWave is reinitializing even though it is actively "
-				"generating audio. Sound must be stopped before playing again."));
-
-			// Need to set the procedural sound wave as not looping to allow it to get stopped during sound wave parsing
-			InWave.bLooping = false;
-			return nullptr;
-		}
-
 		TSharedPtr<FMixerSourceBuffer, ESPMode::ThreadSafe> NewSourceBuffer = MakeShareable(new FMixerSourceBuffer(InSampleRate, InBuffer, InWave, InLoopingMode, bInIsSeeking, bInForceSyncDecode));
+
 		return NewSourceBuffer;
 	}
 
@@ -559,6 +548,11 @@ namespace Audio
 		return true;
 	}
 
+	bool FMixerSourceBuffer::IsGeneratorFinished() const
+	{
+		return bProcedural && SoundGenerator.IsValid() && SoundGenerator->IsFinished();
+	}
+
 	void FMixerSourceBuffer::EnsureAsyncTaskFinishes()
 	{
 		FScopeLock Lock(&DecodeTaskCritSec);
diff --git a/Engine/Source/Runtime/AudioMixer/Private/AudioMixerSourceBuffer.h b/Engine/Source/Runtime/AudioMixer/Private/AudioMixerSourceBuffer.h
--- a/Engine/Source/Runtime/AudioMixer/Private/AudioMixerSourceBuffer.h
+++ b/Engine/Source/Runtime/AudioMixer/Private/AudioMixerSourceBuffer.h
@@ -49,7 +49,7 @@ namespace Audio
 	/** Class which handles decoding audio for a particular source buffer. */
 	class FMixerSourceBuffer : public ISoundWaveClient
 	{
-	public:		
+	public:
 		static FMixerSourceBufferPtr Create(int32 InSampleRate, FMixerBuffer& InBuffer, USoundWave& InWave, ELoopingMode InLoopingMode, bool bInIsSeeking, bool bInForceSyncDecode = false);
 
 		~FMixerSourceBuffer();
@@ -96,6 +96,10 @@ namespace Audio
 		void OnBeginGenerate();
 		void OnEndGenerate();
 		void ClearWave() { SoundWave = nullptr; }
+
+		// Returns whether or not generator is finished (returns false if generator is invalid)
+		bool IsGeneratorFinished() const;
+
 	private:
 		FMixerSourceBuffer(int32 InSampleRate, FMixerBuffer& InBuffer, USoundWave& InWave, ELoopingMode InLoopingMode, bool bInIsSeeking, bool bInForceSyncDecode = false);
 
diff --git a/Engine/Source/Runtime/AudioMixer/Private/AudioMixerSourceManager.cpp b/Engine/Source/Runtime/AudioMixer/Private/AudioMixerSourceManager.cpp
--- a/Engine/Source/Runtime/AudioMixer/Private/AudioMixerSourceManager.cpp
+++ b/Engine/Source/Runtime/AudioMixer/Private/AudioMixerSourceManager.cpp
@@ -1661,6 +1661,14 @@ namespace Audio
 				if (ensure(SourceInfo.MixerSourceBuffer.IsValid()))
 				{
 					SourceInfo.MixerSourceBuffer->OnBufferEnd();
+					if (SourceInfo.MixerSourceBuffer->IsGeneratorFinished())
+					{
+						SourceInfo.bIsPlaying = false;
+						SourceInfo.bIsPaused = false;
+						SourceInfo.bIsActive = false;
+						SourceInfo.bIsStopping = false;
+						break;
+					}
 				}
 			}
 
@@ -2174,7 +2182,7 @@ namespace Audio
 			// Update volume fade information if we're stopping
 			if (SourceInfo.bIsStopping)
 			{
-				const int32 NumFadeFrames = FMath::Min(SourceInfo.VolumeFadeNumFrames - SourceInfo.VolumeFadeFramePosition, NumOutputFrames);
+				int32 NumFadeFrames = FMath::Min(SourceInfo.VolumeFadeNumFrames - SourceInfo.VolumeFadeFramePosition, NumOutputFrames);
 
 				SourceInfo.VolumeFadeFramePosition += NumFadeFrames;
 				SourceInfo.VolumeSourceDestination = SourceInfo.VolumeFadeSlope * (float) SourceInfo.VolumeFadeFramePosition + SourceInfo.VolumeFadeStart;
@@ -2204,7 +2212,8 @@ namespace Audio
 					}
 					VolumeDestination *= ModVolumeEnd;
 				}
-				Audio::FadeBufferFast(PreDistanceAttenBufferPtr, NumSamples, VolumeStart, VolumeDestination);
+
+				Audio::FadeBufferFast(PreDistanceAttenBufferPtr, NumFadeSamples, VolumeStart, VolumeDestination);
 
 				// Zero the rest of the buffer
 				if (NumFadeFrames < NumOutputFrames)
diff --git a/Engine/Source/Runtime/Engine/Classes/Sound/SoundWave.h b/Engine/Source/Runtime/Engine/Classes/Sound/SoundWave.h
--- a/Engine/Source/Runtime/Engine/Classes/Sound/SoundWave.h
+++ b/Engine/Source/Runtime/Engine/Classes/Sound/SoundWave.h
@@ -370,8 +370,8 @@ public:
 	/** Set to true for programmatically generated audio. */
 	uint8 bProcedural:1;
 	
-	/** Set to true if the source is procedural and currently playing */
-	uint8 bPlayingProcedural : 1;
+	/** Set to true if fade is required when sound is abruptly stopped. */
+	uint8 bRequiresStopFade:1;
 
 	/** Set to true of this is a bus sound source. This will result in the sound wave not generating audio for itself, but generate audio through instances. Used only in audio mixer. */
 	uint8 bIsSourceBus:1;
diff --git a/Engine/Source/Runtime/Engine/Private/AudioDevice.cpp b/Engine/Source/Runtime/Engine/Private/AudioDevice.cpp
--- a/Engine/Source/Runtime/Engine/Private/AudioDevice.cpp
+++ b/Engine/Source/Runtime/Engine/Private/AudioDevice.cpp
@@ -4659,18 +4659,6 @@ void FAudioDevice::AddNewActiveSoundInternal(const FActiveSound& NewActiveSound,
 	}
 #endif // !UE_BUILD_SHIPPING
 
-	const USoundWave* SoundWave = Cast<USoundWave>(Sound);
-	if (SoundWave && SoundWave->bProcedural && SoundWave->bPlayingProcedural)
-	{
-		FString SoundWaveName;
-		SoundWave->GetName(SoundWaveName);
-
-		UE_LOG(LogAudio, Warning, TEXT("Replaying a procedural sound '%s' without stopping the previous instance. Only one sound instance per procedural sound wave is supported."), *SoundWaveName);
-
-		ReportSoundFailedToStart(NewActiveSound.AudioComponentID, VirtualLoopToRetrigger);
-		return;
-	}
-
 	// Determine if sound is loop and eligible for virtualize prior to creating "live" active sound in next Concurrency check step
 	if (!VirtualLoopToRetrigger)
 	{
diff --git a/Engine/Source/Runtime/Engine/Private/SoundWave.cpp b/Engine/Source/Runtime/Engine/Private/SoundWave.cpp
--- a/Engine/Source/Runtime/Engine/Private/SoundWave.cpp
+++ b/Engine/Source/Runtime/Engine/Private/SoundWave.cpp
@@ -242,6 +242,9 @@ USoundWave::USoundWave(const FObjectInitializer& ObjectInitializer)
 	CachedSampleRateOverride = 0.0f;
 	CachedSoundWaveLoadingBehavior = ESoundWaveLoadingBehavior::Uninitialized;
 
+	bProcedural = false;
+	bRequiresStopFade = false;
+
 #if WITH_EDITOR
 	bWasStreamCachingEnabledOnLastCook = FPlatformCompressionUtilities::IsCurrentPlatformUsingStreamCaching();
 	bLoadedFromCookedData = false;
@@ -1979,6 +1982,15 @@ void USoundWave::Parse(FAudioDevice* AudioDevice, const UPTRINT NodeWaveInstance
 
 		WaveInstance = &HandleStart(ActiveSound, NodeWaveInstanceHash);
 	}
+	// If the procedural sound has notified as finished, honor it.
+	// Procedural sounds should always only have one instance max.
+	else
+	{
+		if (WaveInstance->bIsFinished && WaveInstance->bAlreadyNotifiedHook && bProcedural)
+		{
+			return;
+		}
+	}
 
 	// Looping sounds are never actually finished
 	if (bLooping || ParseParams.bLooping)
