Update ADSR node to work like a more traditional ADSR envelope

diff --git a/Engine/Plugins/Runtime/Metasound/Source/MetasoundStandardNodes/Private/MetasoundADSRNode.cpp b/Engine/Plugins/Runtime/Metasound/Source/MetasoundStandardNodes/Private/MetasoundADSRNode.cpp
--- a/Engine/Plugins/Runtime/Metasound/Source/MetasoundStandardNodes/Private/MetasoundADSRNode.cpp
+++ b/Engine/Plugins/Runtime/Metasound/Source/MetasoundStandardNodes/Private/MetasoundADSRNode.cpp
@@ -3,11 +3,13 @@
 #include "MetasoundADSRNode.h"
 
 #include "DSP/Envelope.h"
+#include "Internationalization/Text.h"
 #include "MetasoundAudioBuffer.h"
 #include "MetasoundBop.h"
 #include "MetasoundExecutableOperator.h"
 #include "MetasoundTime.h"
 #include "MetasoundNodeRegistrationMacro.h"
+#include "MetasoundPrimitives.h"
 
 #define LOCTEXT_NAMESPACE "MetasoundStandardNodes"
 
@@ -22,7 +24,7 @@ namespace Metasound
 			{
 				FFloatTimeReadRef Attack;
 				FFloatTimeReadRef Decay;
-				FFloatTimeReadRef Sustain;
+				FFloatReadRef SustainLevel;
 				FFloatTimeReadRef Release;
 			};
 
@@ -30,8 +32,7 @@ namespace Metasound
 			static FVertexInterface DeclareVertexInterface();
 			static TUniquePtr<IOperator> CreateOperator(const FCreateOperatorParams& InParams, FBuildErrorArray& OutErrors);
 
-
-			FADSROperator(const FOperatorSettings& InSettings, const FBopReadRef& InBop, const FADSRReferences& InADSRData);
+			FADSROperator(const FOperatorSettings& InSettings, const FBopReadRef& InTriggerAttack, const FBopReadRef& InTriggerRelease, const FADSRReferences& InADSRData);
 
 			virtual FDataReferenceCollection GetInputs() const override;
 			virtual FDataReferenceCollection GetOutputs() const override;
@@ -40,42 +41,52 @@ namespace Metasound
 		private:
 
 			void GenerateEnvelope(int32 InStartFrame, int32 InEndFrame);
+			void Start();
 
 			// TODO: write envelope gen for metasound more suited to this processing structure. 
 			Audio::FEnvelope Envelope;
 
-			FBopReadRef Bop;
-			bool bNotReleased;
-			float SampleRate;
-			FSampleTime TimeTillRelease;
+			FBopReadRef TriggerAttack;
+			FBopReadRef TriggerRelease;
+
+			FBopWriteRef TriggerAttackComplete;
+			FBopWriteRef TriggerDecayComplete;
+			FBopWriteRef TriggerReleaseComplete;
+
+			bool bReleased;
+
+			FSampleTime TimeUntilRelease;
 			FSampleTime TimePerBlock;
 			FADSRReferences ADSRReferences;
 			FAudioBufferWriteRef EnvelopeBuffer;
 	};
 
-	FADSROperator::FADSROperator(const FOperatorSettings& InSettings, const FBopReadRef& InBop, const FADSRReferences& InADSRData)
-	:	Bop(InBop)
-	,	bNotReleased(false)
-	,	SampleRate(InSettings.GetSampleRate())
-	,	TimeTillRelease(0, InSettings.GetSampleRate())
+	FADSROperator::FADSROperator(const FOperatorSettings& InSettings, const FBopReadRef& InTriggerAttack, const FBopReadRef& InTriggerRelease, const FADSRReferences& InADSRData)
+	:	TriggerAttack(InTriggerAttack)
+	,	TriggerRelease(InTriggerRelease)
+	,	TriggerAttackComplete(FBopWriteRef::CreateNew(InSettings))
+	,	TriggerDecayComplete(FBopWriteRef::CreateNew(InSettings))
+	,	TriggerReleaseComplete(FBopWriteRef::CreateNew(InSettings))
+	,	bReleased(false)
+	,	TimeUntilRelease(0, InSettings.GetSampleRate())
 	,	TimePerBlock(InSettings.GetNumFramesPerBlock(), InSettings.GetSampleRate())
 	,	ADSRReferences(InADSRData)
 	,	EnvelopeBuffer(FAudioBufferWriteRef::CreateNew(InSettings.GetNumFramesPerBlock()))
 	{
 		check(EnvelopeBuffer->Num() == InSettings.GetNumFramesPerBlock());
 
-
-		Envelope.Init(SampleRate);
+		Envelope.Init(InSettings.GetSampleRate());
 	}
 
 	FDataReferenceCollection FADSROperator::GetInputs() const
 	{
 		FDataReferenceCollection InputDataReferences;
-		InputDataReferences.AddDataReadReference(TEXT("Bop"), FBopReadRef(Bop));
 		InputDataReferences.AddDataReadReference(TEXT("Attack"), FFloatTimeReadRef(ADSRReferences.Attack));
 		InputDataReferences.AddDataReadReference(TEXT("Decay"), FFloatTimeReadRef(ADSRReferences.Decay));
-		InputDataReferences.AddDataReadReference(TEXT("Sustain"), FFloatTimeReadRef(ADSRReferences.Sustain));
 		InputDataReferences.AddDataReadReference(TEXT("Release"), FFloatTimeReadRef(ADSRReferences.Release));
+		InputDataReferences.AddDataReadReference(TEXT("Sustain Level"), FFloatReadRef(ADSRReferences.SustainLevel));
+		InputDataReferences.AddDataReadReference(TEXT("Trigger Attack"), FBopReadRef(TriggerAttack));
+		InputDataReferences.AddDataReadReference(TEXT("Trigger Release"), FBopReadRef(TriggerRelease));
 		return InputDataReferences;
 	}
 
@@ -83,85 +94,158 @@ namespace Metasound
 	{
 		FDataReferenceCollection OutputDataReferences;
 		OutputDataReferences.AddDataReadReference(TEXT("Envelope"), FAudioBufferReadRef(EnvelopeBuffer));
+		OutputDataReferences.AddDataReadReference(TEXT("Attack Complete"), FBopReadRef(TriggerAttackComplete));
+		OutputDataReferences.AddDataReadReference(TEXT("Decay Complete"), FBopReadRef(TriggerDecayComplete));
+		OutputDataReferences.AddDataReadReference(TEXT("Release Complete"), FBopReadRef(TriggerReleaseComplete));
 		return OutputDataReferences;
 	}
 
 	void FADSROperator::GenerateEnvelope(int32 InStartFrame, int32 InEndFrame)
 	{
 		float* EnvelopeData = EnvelopeBuffer->GetData();
+		const bool bWasDone = Envelope.IsDone();
+
+		Audio::FEnvelope::EEnvelopeState EnvState = Envelope.GetState();
+
+		int32 AttackCompleteFrame = -1;
+		int32 DecayCompleteFrame = -1;
+		int32 ReleaseCompleteFrame = -1;
 
-		const int32 ReleaseFrame = TimeTillRelease.GetNumSamples();
+		TimeUntilRelease -= TimePerBlock;
 
-		if (bNotReleased && (ReleaseFrame < InEndFrame))
+		auto GenerateLambda = [&] (int32 InFrame)
 		{
-			// If the envelope is expected to be released during these frames,
-			// need to split generation of envelope so envelope can be stopped.
-			for (int32 i = InStartFrame; i < ReleaseFrame; i++)
+			const float EnvValue = Envelope.Generate();
+			if (Envelope.GetState() != EnvState)
 			{
-				EnvelopeData[i] = Envelope.Generate();
-			}
+				switch (EnvState)
+				{
+				case Audio::FEnvelope::EEnvelopeState::Attack:
+					AttackCompleteFrame = InFrame;
+					break;
 
-			Envelope.Stop();
+				case Audio::FEnvelope::EEnvelopeState::Decay:
+					DecayCompleteFrame = InFrame;
+					break;
 
-			bNotReleased = false;
+				case Audio::FEnvelope::EEnvelopeState::Release:
+					ReleaseCompleteFrame = InFrame;
+					break;
 
-			for (int32 i = FMath::Max(InStartFrame, ReleaseFrame); i < InEndFrame; i++)
-			{
-				EnvelopeData[i] = Envelope.Generate();
+				default:
+					break;
+				}
+
+				EnvState = Envelope.GetState();
 			}
+
+			return EnvValue;
+		};
+
+		for (int32 i = InStartFrame; i < InEndFrame; i++)
+		{
+			EnvelopeData[i] = GenerateLambda(i);
 		}
-		else
+
+		if (AttackCompleteFrame >= 0)
 		{
-			for (int32 i = InStartFrame; i < InEndFrame; i++)
-			{
-				EnvelopeData[i] = Envelope.Generate();
-			}
+			TriggerAttackComplete->BopFrame(AttackCompleteFrame);
+		}
+
+		if (DecayCompleteFrame >= 0)
+		{
+			TriggerDecayComplete->BopFrame(DecayCompleteFrame);
 		}
-	}
 
+		if (Envelope.IsDone() && !bWasDone && ReleaseCompleteFrame >= 0)
+		{
+			TriggerReleaseComplete->BopFrame(ReleaseCompleteFrame);
+		}
+	}
 
 	void FADSROperator::Execute()
 	{
 		Envelope.SetAttackTime(ADSRReferences.Attack->GetMilliseconds());
 		Envelope.SetDecayTime(ADSRReferences.Decay->GetMilliseconds());
 		Envelope.SetReleaseTime(ADSRReferences.Release->GetMilliseconds());
+		Envelope.SetSustainGain(*ADSRReferences.SustainLevel);
 
-		if (bNotReleased)
-		{
-			// Keep track of when next release should happen
-			TimeTillRelease -= TimePerBlock;
-		}
-		
-		Bop->ExecuteBlock(
+		TriggerAttackComplete->AdvanceBlock();
+		TriggerDecayComplete->AdvanceBlock();
+		TriggerReleaseComplete->AdvanceBlock();
+
+		TriggerRelease->ExecuteBlock(
 			// OnPreBop
-			[&](int32 StartFrame, int32 EndFrame) { GenerateEnvelope(StartFrame, EndFrame); },
+			[&](int32 StartFrame, int32 EndFrame)
+			{
+			},
 			// OnBop
 			[&](int32 StartFrame, int32 EndFrame)
 			{
-				Envelope.Start();
-
-				// Capture attack, decay and sustain time (skip release).
-				TimeTillRelease = *ADSRReferences.Attack + *ADSRReferences.Decay + *ADSRReferences.Sustain;
-
-				bNotReleased = true;
+				if (!bReleased)
+				{
+					bReleased = true;
+					TimeUntilRelease = *(ADSRReferences.Release);
+
+					const int32 ReleaseFrame = TimeUntilRelease.GetNumSamples();
+					if (ReleaseFrame < EndFrame)
+					{
+						GenerateEnvelope(StartFrame, ReleaseFrame);
+						Envelope.Stop();
+						GenerateEnvelope(ReleaseFrame, EndFrame);
+					}
+					else
+					{
+						Envelope.Stop();
+						GenerateEnvelope(StartFrame, EndFrame);
+					}
+				}
+			}
+		);
 
+		TriggerAttack->ExecuteBlock(
+			// OnPreBop
+			[&](int32 StartFrame, int32 EndFrame)
+			{
 				GenerateEnvelope(StartFrame, EndFrame);
+			},
+			// OnBop
+			[&](int32 StartFrame, int32 EndFrame)
+			{
+				Start();
+				if (!Envelope.IsDone())
+				{
+					GenerateEnvelope(StartFrame, EndFrame);
+				}
 			}
 		);
 	}
 
+	void FADSROperator::Start()
+	{
+		if (Envelope.IsDone())
+		{
+			Envelope.Start();
+			bReleased = false;
+		}
+	}
+
 	FVertexInterface FADSROperator::DeclareVertexInterface()
 	{
 		static const FVertexInterface Interface(
 			FInputVertexInterface(
-				TInputDataVertexModel<FBop>(TEXT("Bop"), LOCTEXT("BopTooltip", "Trigger for envelope.")),
-				TInputDataVertexModel<FFloatTime>(TEXT("Attack"), LOCTEXT("AttackTooltip", "Attack time.")),
-				TInputDataVertexModel<FFloatTime>(TEXT("Decay"), LOCTEXT("DecayTooltip", "Decay time.")),
-				TInputDataVertexModel<FFloatTime>(TEXT("Sustain"), LOCTEXT("SustainTooltip", "Sustain time.")),
-				TInputDataVertexModel<FFloatTime>(TEXT("Release"), LOCTEXT("ReleaseTooltip", "Release time."))
+				TInputDataVertexModel<FFloatTime>(TEXT("Attack"), LOCTEXT("AttackDurationTooltip", "Attack duration.")),
+				TInputDataVertexModel<FFloatTime>(TEXT("Decay"), LOCTEXT("DecayDurationTooltip", "Decay duration.")),
+				TInputDataVertexModel<FFloatTime>(TEXT("Release"), LOCTEXT("ReleaseDurationTooltip", "Release duration.")),
+				TInputDataVertexModel<float>(TEXT("Sustain Level"), LOCTEXT("SustainLevelTooltip", "Sustain level [0.0f, 1.0f].")),
+				TInputDataVertexModel<FBop>(TEXT("Trigger Attack"), LOCTEXT("BopTooltip", "Trigger the envelope's attack.")),
+				TInputDataVertexModel<FBop>(TEXT("Trigger Release"), LOCTEXT("BopTooltip", "Trigger the envelope's release."))
 			),
 			FOutputVertexInterface(
-				TOutputDataVertexModel<FAudioBuffer>(TEXT("Envelope"), LOCTEXT("EnvelopeTooltip", "The output envelope"))
+				TOutputDataVertexModel<FAudioBuffer>(TEXT("Envelope"), LOCTEXT("EnvelopeTooltip", "The output envelope")),
+				TOutputDataVertexModel<FBop>(TEXT("Attack Complete"), LOCTEXT("AttackCompleteTooltip", "Triggered when the envelope attack is complete")),
+				TOutputDataVertexModel<FBop>(TEXT("Decay Complete"), LOCTEXT("DecayCompleteTooltip", "Triggered when the envelope decay is complete")),
+				TOutputDataVertexModel<FBop>(TEXT("Release Complete"), LOCTEXT("ReleaseCompleteTooltip", "Triggered when the envelope is released"))
 			)
 		);
 
@@ -189,40 +273,38 @@ namespace Metasound
 		return Info;
 	}
 
-
 	TUniquePtr<IOperator> FADSROperator::CreateOperator(const FCreateOperatorParams& InParams, FBuildErrorArray& OutErrors)
 	{
 		const FADSRNode& ADSRNode = static_cast<const FADSRNode&>(InParams.Node);
-
-
 		const FDataReferenceCollection& InputCollection = InParams.InputDataReferences;
 
 		// TODO: Could no-op this if the bop is not connected.
-		FBopReadRef Bop = InputCollection.GetDataReadReferenceOrConstruct<FBop>(TEXT("Bop"), InParams.OperatorSettings);
+		FBopReadRef TriggerAttack = InputCollection.GetDataReadReferenceOrConstruct<FBop>(TEXT("Trigger Attack"), InParams.OperatorSettings);
+		FBopReadRef TriggerRelease = InputCollection.GetDataReadReferenceOrConstruct<FBop>(TEXT("Trigger Release"), InParams.OperatorSettings);
 
 		// TODO: If none of these are connected, could pre-generate ADSR envelope and return a different operator. 
 		FADSRReferences ADSRReferences = 
-			{
-				InputCollection.GetDataReadReferenceOrConstruct<FFloatTime>(TEXT("Attack"), ADSRNode.GetDefaultAttackMs(), ETimeResolution::Milliseconds),
-				InputCollection.GetDataReadReferenceOrConstruct<FFloatTime>(TEXT("Decay"), ADSRNode.GetDefaultDecayMs(), ETimeResolution::Milliseconds),
-				InputCollection.GetDataReadReferenceOrConstruct<FFloatTime>(TEXT("Sustain"), ADSRNode.GetDefaultSustainMs(), ETimeResolution::Milliseconds),
-				InputCollection.GetDataReadReferenceOrConstruct<FFloatTime>(TEXT("Release"), ADSRNode.GetDefaultReleaseMs(), ETimeResolution::Milliseconds)
-			};
+		{
+			InputCollection.GetDataReadReferenceOrConstruct<FFloatTime>(TEXT("Attack"), ADSRNode.GetDefaultAttackMs(), ETimeResolution::Milliseconds),
+			InputCollection.GetDataReadReferenceOrConstruct<FFloatTime>(TEXT("Decay"), ADSRNode.GetDefaultDecayMs(), ETimeResolution::Milliseconds),
+			InputCollection.GetDataReadReferenceOrConstruct<float>(TEXT("Sustain Level"), ADSRNode.GetDefaultSustainLevel()),
+			InputCollection.GetDataReadReferenceOrConstruct<FFloatTime>(TEXT("Release"), ADSRNode.GetDefaultReleaseMs(), ETimeResolution::Milliseconds),
+		};
 	
-		return MakeUnique<FADSROperator>(InParams.OperatorSettings, Bop, ADSRReferences);
+		return MakeUnique<FADSROperator>(InParams.OperatorSettings, TriggerAttack, TriggerRelease, ADSRReferences);
 	}
 
-	FADSRNode::FADSRNode(const FString& InName, float InDefaultAttackMs, float InDefaultDecayMs, float InDefaultSustainMs, float InDefaultReleaseMs)
+	FADSRNode::FADSRNode(const FString& InName, float InDefaultAttackMs, float InDefaultDecayMs, float InDefaultSustainLevel, float InDefaultReleaseMs)
 	:	FNodeFacade(InName, TFacadeOperatorClass<FADSROperator>())
+	,	DefaultSustainLevel(InDefaultSustainLevel)
 	,	DefaultAttackMs(InDefaultAttackMs)
 	,	DefaultDecayMs(InDefaultDecayMs)
-	,	DefaultSustainMs(InDefaultSustainMs)
 	,	DefaultReleaseMs(InDefaultReleaseMs)
 	{
 	}
 
 	FADSRNode::FADSRNode(const FNodeInitData& InitData)
-		: FADSRNode(InitData.InstanceName, 10.0f, 20.0f, 50.0f, 20.0f)
+		: FADSRNode(InitData.InstanceName, 10.0f, 20.0f, 0.7f, 20.0f)
 	{
 	}
 
@@ -230,6 +312,10 @@ namespace Metasound
 	{
 	}
 
+	float FADSRNode::GetDefaultSustainLevel() const
+	{
+		return DefaultSustainLevel;
+	}
 
 	float FADSRNode::GetDefaultAttackMs() const
 	{
@@ -241,11 +327,6 @@ namespace Metasound
 		return DefaultDecayMs;
 	}
 
-	float FADSRNode::GetDefaultSustainMs() const
-	{
-		return DefaultSustainMs;
-	}
-
 	float FADSRNode::GetDefaultReleaseMs() const
 	{
 		return DefaultReleaseMs;
diff --git a/Engine/Plugins/Runtime/Metasound/Source/MetasoundStandardNodes/Public/MetasoundADSRNode.h b/Engine/Plugins/Runtime/Metasound/Source/MetasoundStandardNodes/Public/MetasoundADSRNode.h
--- a/Engine/Plugins/Runtime/Metasound/Source/MetasoundStandardNodes/Public/MetasoundADSRNode.h
+++ b/Engine/Plugins/Runtime/Metasound/Source/MetasoundStandardNodes/Public/MetasoundADSRNode.h
@@ -2,7 +2,6 @@
 #pragma once
 
 #include "CoreMinimal.h"
-#include "Internationalization/Text.h"
 
 #include "MetasoundFacade.h"
 
@@ -10,7 +9,7 @@ namespace Metasound
 {
 	/** FADSRNode
 	 *
-	 *  Creates an Attack, Decay Sustain, Release audio processer node. 
+	 *  Creates an Attack, Decay Sustain, Release audio processor node. 
 	 */
 	class METASOUNDSTANDARDNODES_API FADSRNode : public FNodeFacade
 	{
@@ -19,12 +18,12 @@ namespace Metasound
 			/** FADSR node constructor.
 			 *
 			 * @param InName - Name of this node.
-			 * @param InDefaultAttackMs - Default attack in milliseconds.
-			 * @param InDefaultDecayMs - Default decay in milliseconds.
-			 * @param InDefaultSustainMs - Default sustain in milliseconds.
-			 * @param InDefaultReleaseMs - Default release in milliseconds.
+			 * @param InDefaultAttackMs - Default attack duration in milliseconds.
+			 * @param InDefaultDecayMs - Default decay duration in milliseconds.
+			 * @param InDefaultSustainLevel - Default sustain level [0.0f, 1.0f].
+			 * @param InDefaultReleaseMs - Default release duration in milliseconds.
 			 */
-			FADSRNode(const FString& InName, float InDefaultAttackMs, float InDefaultDecayMs, float InDefaultSustainMs, float InDefaultReleaseMs);
+			FADSRNode(const FString& InName, float InDefaultAttackMs, float InDefaultDecayMs, float InDefaultSustainLevel, float InDefaultReleaseMs);
 
 			/**
 			 * Constructor used by the Metasound Frontend.
@@ -33,19 +32,23 @@ namespace Metasound
 
 			virtual ~FADSRNode();
 
+			/** Return default sustain gain [0.0f, 1.0f] */
+			float GetDefaultSustainLevel() const;
+
 			/** Return default attack time in milliseconds */
 			float GetDefaultAttackMs() const;
-			/** Return default attack time in milliseconds */
+
+			/** Return default decay time in milliseconds */
 			float GetDefaultDecayMs() const;
-			/** Return default attack time in milliseconds */
-			float GetDefaultSustainMs() const;
-			/** Return default attack time in milliseconds */
+
+			/** Return default release time in milliseconds */
 			float GetDefaultReleaseMs() const;
 
 		private:
+			float DefaultSustainLevel;
+
 			float DefaultAttackMs;
 			float DefaultDecayMs;
-			float DefaultSustainMs;
 			float DefaultReleaseMs;
 	};
 } // namespace Metasound
diff --git a/Engine/Source/Runtime/SignalProcessing/Private/Envelope.cpp b/Engine/Source/Runtime/SignalProcessing/Private/Envelope.cpp
--- a/Engine/Source/Runtime/SignalProcessing/Private/Envelope.cpp
+++ b/Engine/Source/Runtime/SignalProcessing/Private/Envelope.cpp
@@ -53,6 +53,11 @@ namespace Audio
 		}
 	}
 
+	FEnvelope::EEnvelopeState FEnvelope::GetState() const
+	{
+		return CurrentState;
+	}
+
 	void FEnvelope::SetSimulateAnalog(const bool bInSimulatingAnalog)
 	{
 		bIsSimulatingAnalog = bInSimulatingAnalog;
@@ -73,7 +78,7 @@ namespace Audio
 		Reset();
 
 		// Set the state back to attack no matter where it is
-		CurrentState = EEnvelopeState::Attack;
+		CurrentState = FEnvelope::EEnvelopeState::Attack;
 	}
 
 	// logic for one mono note interrupting another mono note (same voice)
@@ -379,26 +384,26 @@ namespace Audio
 
 	void FEnvelope::SetAttackTime(const float InAttackTimeMsec)
 	{
+		bChanged |= !FMath::IsNearlyEqual(AttackTimeMSec, InAttackTimeMsec);
 		AttackTimeMSec = InAttackTimeMsec;
-		bChanged = true;
 	}
 
 	void FEnvelope::SetDecayTime(const float InDecayTimeMsec)
 	{
+		bChanged |= !FMath::IsNearlyEqual(DecayTimeMsec, InDecayTimeMsec);
 		DecayTimeMsec = InDecayTimeMsec;
-		bChanged = true;
 	}
 
 	void FEnvelope::SetSustainGain(const float InSustainGain)
 	{
+		bChanged |= !FMath::IsNearlyEqual(SustainGain, InSustainGain);
 		SustainGain = InSustainGain;
-		bChanged = true;
 	}
 
 	void FEnvelope::SetReleaseTime(const float InReleaseTimeMsec)
 	{
+		bChanged |= !FMath::IsNearlyEqual(ReleaseTimeMsec, InReleaseTimeMsec);
 		ReleaseTimeMsec = InReleaseTimeMsec;
-		bChanged = true;
 	}
 
 	void FEnvelope::SetInvert(const bool bInInvert)
diff --git a/Engine/Source/Runtime/SignalProcessing/Public/DSP/Envelope.h b/Engine/Source/Runtime/SignalProcessing/Public/DSP/Envelope.h
--- a/Engine/Source/Runtime/SignalProcessing/Public/DSP/Envelope.h
+++ b/Engine/Source/Runtime/SignalProcessing/Public/DSP/Envelope.h
@@ -12,6 +12,17 @@ namespace Audio
 	class SIGNALPROCESSING_API FEnvelope
 	{
 	public:
+		// States for the envelope state machine
+		enum class EEnvelopeState
+		{
+			Off,
+			Attack,
+			Decay,
+			Sustain,
+			Release,
+			Shutdown
+		};
+		
 		FEnvelope();
 		virtual ~FEnvelope();
 
@@ -57,6 +68,8 @@ namespace Audio
 		// Optionally outputs the bias output (i.e. -1.0 to 1.0)
 		virtual float Generate(float* BiasedOutput = nullptr);
 
+		virtual EEnvelopeState GetState() const;
+
 		// Sets the envelope attack time in msec
 		virtual void SetAttackTime(const float InAttackTimeMsec);
 
@@ -86,18 +99,6 @@ namespace Audio
 		const FPatchSource GetModSourceBiasEnv() const { return BiasedEnvSource; }
 
 	protected:
-
-		// States for the envelope state machine
-		enum class EEnvelopeState
-		{
-			Off,
-			Attack,
-			Decay,
-			Sustain,
-			Release,
-			Shutdown
-		};
-
 		struct FEnvData
 		{
 			float Coefficient;
