Add DelayTrigger & Pipe Trigger Metasound nodes

diff --git a/Engine/Plugins/Runtime/Metasound/Source/MetasoundStandardNodes/Private/MetasoundTriggerDelayNode.cpp b/Engine/Plugins/Runtime/Metasound/Source/MetasoundStandardNodes/Private/MetasoundTriggerDelayNode.cpp
new file mode 100644
--- /dev/null
+++ b/Engine/Plugins/Runtime/Metasound/Source/MetasoundStandardNodes/Private/MetasoundTriggerDelayNode.cpp
@@ -0,0 +1,188 @@
+// Copyright Epic Games, Inc. All Rights Reserved.
+#include "MetasoundTriggerDelayNode.h"
+
+#include "CoreMinimal.h"
+
+#include "MetasoundBuilderInterface.h"
+#include "MetasoundDataReferenceCollection.h"
+#include "MetasoundExecutableOperator.h"
+#include "MetasoundFacade.h"
+#include "MetasoundNode.h"
+#include "MetasoundNodeInterface.h"
+#include "MetasoundNodeRegistrationMacro.h"
+#include "MetasoundOperatorInterface.h"
+#include "MetasoundPrimitives.h"
+
+#define LOCTEXT_NAMESPACE "MetasoundStandardNodes"
+
+namespace Metasound
+{
+	METASOUND_REGISTER_NODE(FTriggerDelayNode)
+
+	class FTriggerDelayOperator : public TExecutableOperator<FTriggerDelayOperator>
+	{
+		public:
+			static const FNodeInfo& GetNodeInfo();
+			static FVertexInterface DeclareVertexInterface();
+			static TUniquePtr<IOperator> CreateOperator(const FCreateOperatorParams& InParams, FBuildErrorArray& OutErrors);
+
+			FTriggerDelayOperator(const FOperatorSettings& InSettings, const FTriggerReadRef& InTriggerReset, const FTriggerReadRef& InTriggerIn, const FFloatTimeReadRef& InDelay);
+
+			virtual FDataReferenceCollection GetInputs() const override;
+			virtual FDataReferenceCollection GetOutputs() const override;
+
+			void Execute();
+			void Reset();
+
+		private:
+			int32 SamplesUntilTrigger;
+
+			FTriggerReadRef TriggerIn;
+			FTriggerReadRef TriggerReset;
+			FTriggerWriteRef TriggerOut;
+
+			FFloatTimeReadRef Delay;
+
+			float FramesPerBlock;
+			float SampleRate;
+	};
+
+	FTriggerDelayOperator::FTriggerDelayOperator(const FOperatorSettings& InSettings, const FTriggerReadRef& InTriggerReset, const FTriggerReadRef& InTriggerIn, const FFloatTimeReadRef& InDelay)
+	:	SamplesUntilTrigger(-1)
+	,	TriggerIn(InTriggerIn)
+	,	TriggerReset(InTriggerReset)
+	,	TriggerOut(FTriggerWriteRef::CreateNew(InSettings))
+	,	Delay(InDelay)
+	,	FramesPerBlock(InSettings.GetNumFramesPerBlock())
+	,	SampleRate(InSettings.GetSampleRate())
+	{
+	}
+
+	FDataReferenceCollection FTriggerDelayOperator::GetInputs() const
+	{
+		FDataReferenceCollection InputDataReferences;
+		InputDataReferences.AddDataReadReference(TEXT("Delay"), FFloatTimeReadRef(Delay));
+		InputDataReferences.AddDataReadReference(TEXT("In"), FTriggerReadRef(TriggerIn));
+		InputDataReferences.AddDataReadReference(TEXT("Reset"), FTriggerReadRef(TriggerIn));
+		return InputDataReferences;
+	}
+
+	FDataReferenceCollection FTriggerDelayOperator::GetOutputs() const
+	{
+		FDataReferenceCollection OutputDataReferences;
+		OutputDataReferences.AddDataReadReference(TEXT("Out"), FTriggerReadRef(TriggerOut));
+
+		return OutputDataReferences;
+	}
+
+	void FTriggerDelayOperator::Execute()
+	{
+		// Advance internal counter to get rid of old triggers.
+		TriggerOut->AdvanceBlock();
+
+		if (SamplesUntilTrigger > 0)
+		{
+			int32 SamplesRemaining = SamplesUntilTrigger - FramesPerBlock;
+			if (SamplesRemaining > 0.0f)
+			{
+				SamplesUntilTrigger -= FramesPerBlock;
+			}
+			else
+			{
+				TriggerOut->BopFrame(SamplesRemaining + (int32)FramesPerBlock);
+				Reset();
+			}
+		}
+
+		TriggerIn->ExecuteBlock(
+			[](int32 StartFrame, int32 EndFrame)
+			{
+			},
+			[this](int32 StartFrame, int32 EndFrame)
+			{
+				SamplesUntilTrigger = Delay->GetNumSamples(SampleRate);
+			}
+		);
+
+		TriggerReset->ExecuteBlock(
+			[](int32 StartFrame, int32 EndFrame)
+			{
+			},
+			[this](int32 StartFrame, int32 EndFrame)
+			{
+				Reset();
+			}
+		);
+	}
+
+
+	TUniquePtr<IOperator> FTriggerDelayOperator::CreateOperator(const FCreateOperatorParams& InParams, FBuildErrorArray& OutErrors)
+	{
+		const FTriggerDelayNode& DelayNode = static_cast<const FTriggerDelayNode&>(InParams.Node);
+
+		FFloatTimeReadRef Delay = InParams.InputDataReferences.GetDataReadReferenceOrConstruct<FFloatTime>(TEXT("Delay"), DelayNode.GetDefaultDelayInSeconds());
+		FTriggerReadRef TriggerIn = InParams.InputDataReferences.GetDataReadReferenceOrConstruct<FTrigger>(TEXT("In"), InParams.OperatorSettings);
+		FTriggerReadRef TriggerReset = InParams.InputDataReferences.GetDataReadReferenceOrConstruct<FTrigger>(TEXT("Reset"), InParams.OperatorSettings);
+
+		return MakeUnique<FTriggerDelayOperator>(InParams.OperatorSettings, TriggerReset, TriggerIn, Delay);
+	}
+
+	FVertexInterface FTriggerDelayOperator::DeclareVertexInterface()
+	{
+		static const FVertexInterface Interface(
+			FInputVertexInterface(
+				TInputDataVertexModel<FFloatTime>(TEXT("Delay"), LOCTEXT("TriggerDelayTooltip", "Time to delay and execute deferred trigger in seconds.")),
+				TInputDataVertexModel<FTrigger>(TEXT("In"), LOCTEXT("TriggerDelayInTooltip", "Triggers delay.")),
+				TInputDataVertexModel<FTrigger>(TEXT("Reset"), LOCTEXT("TriggerDelayResetInTooltip", "Resets the trigger delay, clearing the execution task if pending."))
+			),
+			FOutputVertexInterface(
+				TOutputDataVertexModel<FTrigger>(TEXT("Out"), LOCTEXT("TriggerOutTooltip", "The deferred output trigger"))
+			)
+		);
+
+		return Interface;
+	}
+
+	const FNodeInfo& FTriggerDelayOperator::GetNodeInfo()
+	{
+		auto InitNodeInfo = []() -> FNodeInfo
+		{
+			FNodeInfo Info;
+			Info.ClassName = "Trigger Delay";
+			Info.MajorVersion = 1;
+			Info.MinorVersion = 0;
+			Info.Description = LOCTEXT("Metasound_TriggerDelayNodeDescription", "Executes output trigger after the given delay time from the most recent execution of the input trigger .");
+			Info.Author = PluginAuthor;
+			Info.PromptIfMissing = PluginNodeMissingPrompt;
+			Info.DefaultInterface = DeclareVertexInterface();
+
+			return Info;
+		};
+
+		static const FNodeInfo Info = InitNodeInfo();
+		return Info;
+	}
+
+	void FTriggerDelayOperator::Reset()
+	{
+		SamplesUntilTrigger = -1;
+	}
+
+	FTriggerDelayNode::FTriggerDelayNode(const FString& InName, float InDefaultDelayInSeconds)
+	:	FNodeFacade(InName, TFacadeOperatorClass<FTriggerDelayOperator>())
+	,	DefaultDelay(InDefaultDelayInSeconds)
+	{
+	}
+
+	FTriggerDelayNode::FTriggerDelayNode(const FNodeInitData& InInitData)
+		: FTriggerDelayNode(InInitData.InstanceName, 0.0f)
+	{
+	}
+
+	float FTriggerDelayNode::GetDefaultDelayInSeconds() const
+	{
+		return DefaultDelay;
+	}
+}
+
+#undef LOCTEXT_NAMESPACE // MetasoundTriggerDelayNode
diff --git a/Engine/Plugins/Runtime/Metasound/Source/MetasoundStandardNodes/Private/MetasoundTriggerPipeNode.cpp b/Engine/Plugins/Runtime/Metasound/Source/MetasoundStandardNodes/Private/MetasoundTriggerPipeNode.cpp
new file mode 100644
--- /dev/null
+++ b/Engine/Plugins/Runtime/Metasound/Source/MetasoundStandardNodes/Private/MetasoundTriggerPipeNode.cpp
@@ -0,0 +1,181 @@
+// Copyright Epic Games, Inc. All Rights Reserved.
+#include "MetasoundTriggerPipeNode.h"
+
+#include "CoreMinimal.h"
+
+#include "MetasoundBop.h"
+#include "MetasoundBuilderInterface.h"
+#include "MetasoundDataReferenceCollection.h"
+#include "MetasoundExecutableOperator.h"
+#include "MetasoundFacade.h"
+#include "MetasoundNode.h"
+#include "MetasoundNodeInterface.h"
+#include "MetasoundNodeRegistrationMacro.h"
+#include "MetasoundOperatorInterface.h"
+#include "MetasoundPrimitives.h"
+#include "MetasoundTime.h"
+
+#define LOCTEXT_NAMESPACE "MetasoundStandardNodes"
+
+namespace Metasound
+{
+	METASOUND_REGISTER_NODE(FTriggerPipeNode)
+
+	class FTriggerPipeOperator : public TExecutableOperator<FTriggerPipeOperator>
+	{
+		public:
+			static const FNodeInfo& GetNodeInfo();
+			static FVertexInterface DeclareVertexInterface();
+			static TUniquePtr<IOperator> CreateOperator(const FCreateOperatorParams& InParams, FBuildErrorArray& OutErrors);
+
+			FTriggerPipeOperator(const FOperatorSettings& InSettings, const FTriggerReadRef& InTriggerReset, const FTriggerReadRef& InTriggerIn, const FFloatTimeReadRef& InDelay);
+
+			virtual FDataReferenceCollection GetInputs() const override;
+			virtual FDataReferenceCollection GetOutputs() const override;
+			void Execute();
+
+		private:
+			TArray<int32> SamplesUntilTrigger;
+
+			FTriggerReadRef TriggerIn;
+			FTriggerReadRef TriggerReset;
+			FTriggerWriteRef TriggerOut;
+
+			FFloatTimeReadRef Delay;
+
+			float FramesPerBlock;
+			float SampleRate;
+	};
+
+	FTriggerPipeOperator::FTriggerPipeOperator(const FOperatorSettings& InSettings, const FTriggerReadRef& InTriggerReset, const FTriggerReadRef& InTriggerIn, const FFloatTimeReadRef& InDelay)
+	:	TriggerIn(InTriggerIn)
+	,	TriggerReset(InTriggerReset)
+	,	TriggerOut(FTriggerWriteRef::CreateNew(InSettings))
+	,	Delay(InDelay)
+	,	FramesPerBlock(InSettings.GetNumFramesPerBlock())
+	,	SampleRate(InSettings.GetSampleRate())
+	{
+	}
+
+	FDataReferenceCollection FTriggerPipeOperator::GetInputs() const
+	{
+		FDataReferenceCollection InputDataReferences;
+		InputDataReferences.AddDataReadReference(TEXT("Delay"), FFloatTimeReadRef(Delay));
+		InputDataReferences.AddDataReadReference(TEXT("In"), FTriggerReadRef(TriggerIn));
+		InputDataReferences.AddDataReadReference(TEXT("Reset"), FTriggerReadRef(TriggerIn));
+		return InputDataReferences;
+	}
+
+	FDataReferenceCollection FTriggerPipeOperator::GetOutputs() const
+	{
+		FDataReferenceCollection OutputDataReferences;
+		OutputDataReferences.AddDataReadReference(TEXT("Out"), FTriggerReadRef(TriggerOut));
+
+		return OutputDataReferences;
+	}
+
+	void FTriggerPipeOperator::Execute()
+	{
+		// Advance internal counter to get rid of old triggers.
+		TriggerOut->AdvanceBlock();
+
+		for (int32 i = SamplesUntilTrigger.Num() - 1; i >= 0; --i)
+		{
+			int32 SamplesRemaining = SamplesUntilTrigger[i] - FramesPerBlock;
+			if (SamplesRemaining > 0.0f)
+			{
+				SamplesUntilTrigger[i] -= FramesPerBlock;
+			}
+			else
+			{
+				TriggerOut->BopFrame(SamplesRemaining + (int32)FramesPerBlock);
+				SamplesUntilTrigger.RemoveAtSwap(i);
+			}
+		}
+
+		TriggerIn->ExecuteBlock(
+			[](int32 StartFrame, int32 EndFrame)
+			{
+			},
+			[this](int32 StartFrame, int32 EndFrame)
+			{
+				SamplesUntilTrigger.AddUnique(Delay->GetNumSamples(SampleRate));
+			}
+		);
+
+		TriggerReset->ExecuteBlock(
+			[](int32 StartFrame, int32 EndFrame)
+			{
+			},
+			[this](int32 StartFrame, int32 EndFrame)
+			{
+				SamplesUntilTrigger.Empty();
+			}
+		);
+	}
+
+	TUniquePtr<IOperator> FTriggerPipeOperator::CreateOperator(const FCreateOperatorParams& InParams, FBuildErrorArray& OutErrors)
+	{
+		const FTriggerPipeNode& DelayNode = static_cast<const FTriggerPipeNode&>(InParams.Node);
+
+		FFloatTimeReadRef Delay = InParams.InputDataReferences.GetDataReadReferenceOrConstruct<FFloatTime>(TEXT("Delay"), DelayNode.GetDefaultDelayInSeconds());
+		FTriggerReadRef TriggerIn = InParams.InputDataReferences.GetDataReadReferenceOrConstruct<FTrigger>(TEXT("In"), InParams.OperatorSettings);
+		FTriggerReadRef TriggerReset = InParams.InputDataReferences.GetDataReadReferenceOrConstruct<FTrigger>(TEXT("Reset"), InParams.OperatorSettings);
+
+		return MakeUnique<FTriggerPipeOperator>(InParams.OperatorSettings, TriggerReset, TriggerIn, Delay);
+	}
+
+	FVertexInterface FTriggerPipeOperator::DeclareVertexInterface()
+	{
+		static const FVertexInterface Interface(
+			FInputVertexInterface(
+				TInputDataVertexModel<FFloatTime>(TEXT("Delay"), LOCTEXT("DelayTooltip", "Time to delay and execute deferred input trigger execution(s) in seconds.")),
+				TInputDataVertexModel<FTrigger>(TEXT("In"), LOCTEXT("TriggerInTooltip", "Trigger to execute at a future time by the given delay amount.")),
+				TInputDataVertexModel<FTrigger>(TEXT("Reset"), LOCTEXT("TriggerDelayResetInTooltip", "Resets the trigger delay, clearing any pending execution tasks."))
+			),
+			FOutputVertexInterface(
+				TOutputDataVertexModel<FTrigger>(TEXT("Out"), LOCTEXT("TriggerOutTooltip", "The deferred output trigger"))
+			)
+		);
+
+		return Interface;
+	}
+
+	const FNodeInfo& FTriggerPipeOperator::GetNodeInfo()
+	{
+		auto InitNodeInfo = []() -> FNodeInfo
+		{
+			FNodeInfo Info;
+			Info.ClassName = "Trigger Pipe";
+			Info.MajorVersion = 1;
+			Info.MinorVersion = 0;
+			Info.Description = LOCTEXT("Metasound_DelayNodeDescription", "Delays execution of the input trigger by the given delay for all input trigger executions.");
+			Info.Author = PluginAuthor;
+			Info.PromptIfMissing = PluginNodeMissingPrompt;
+			Info.DefaultInterface = DeclareVertexInterface();
+
+			return Info;
+		};
+
+		static const FNodeInfo Info = InitNodeInfo();
+		return Info;
+	}
+
+	FTriggerPipeNode::FTriggerPipeNode(const FString& InName, float InDefaultDelayInSeconds)
+	:	FNodeFacade(InName, TFacadeOperatorClass<FTriggerPipeOperator>())
+	,	DefaultDelay(InDefaultDelayInSeconds)
+	{
+	}
+
+	FTriggerPipeNode::FTriggerPipeNode(const FNodeInitData& InInitData)
+		: FTriggerPipeNode(InInitData.InstanceName, 0.0f)
+	{
+	}
+
+	float FTriggerPipeNode::GetDefaultDelayInSeconds() const
+	{
+		return DefaultDelay;
+	}
+}
+
+#undef LOCTEXT_NAMESPACE // MetasoundTriggerPipeNode
diff --git a/Engine/Plugins/Runtime/Metasound/Source/MetasoundStandardNodes/Public/MetasoundTriggerDelayNode.h b/Engine/Plugins/Runtime/Metasound/Source/MetasoundStandardNodes/Public/MetasoundTriggerDelayNode.h
new file mode 100644
--- /dev/null
+++ b/Engine/Plugins/Runtime/Metasound/Source/MetasoundStandardNodes/Public/MetasoundTriggerDelayNode.h
@@ -0,0 +1,29 @@
+// Copyright Epic Games, Inc. All Rights Reserved.
+
+#pragma once
+
+#include "CoreMinimal.h"
+#include "MetasoundNode.h"
+#include "MetasoundBop.h"
+#include "MetasoundBuilderInterface.h"
+#include "MetasoundExecutableOperator.h"
+#include "MetasoundFacade.h"
+#include "MetasoundOperatorInterface.h"
+#include "MetasoundNodeInterface.h"
+#include "MetasoundDataReferenceCollection.h"
+#include "MetasoundTime.h"
+
+namespace Metasound
+{
+	class METASOUNDSTANDARDNODES_API FTriggerDelayNode : public FNodeFacade
+	{
+		public:
+			FTriggerDelayNode(const FString& InName, float InDefaultDelayInSeconds);
+			FTriggerDelayNode(const FNodeInitData& InInitData);
+
+			float GetDefaultDelayInSeconds() const;
+
+		private:
+			float DefaultDelay;
+	};
+}
diff --git a/Engine/Plugins/Runtime/Metasound/Source/MetasoundStandardNodes/Public/MetasoundTriggerPipeNode.h b/Engine/Plugins/Runtime/Metasound/Source/MetasoundStandardNodes/Public/MetasoundTriggerPipeNode.h
new file mode 100644
--- /dev/null
+++ b/Engine/Plugins/Runtime/Metasound/Source/MetasoundStandardNodes/Public/MetasoundTriggerPipeNode.h
@@ -0,0 +1,27 @@
+// Copyright Epic Games, Inc. All Rights Reserved.
+
+#pragma once
+
+#include "CoreMinimal.h"
+#include "MetasoundNode.h"
+#include "MetasoundBuilderInterface.h"
+#include "MetasoundExecutableOperator.h"
+#include "MetasoundFacade.h"
+#include "MetasoundOperatorInterface.h"
+#include "MetasoundNodeInterface.h"
+#include "MetasoundDataReferenceCollection.h"
+
+namespace Metasound
+{
+	class METASOUNDSTANDARDNODES_API FTriggerPipeNode : public FNodeFacade
+	{
+		public:
+			FTriggerPipeNode(const FString& InName, float InDefaultDelayInSeconds);
+			FTriggerPipeNode(const FNodeInitData& InInitData);
+
+			float GetDefaultDelayInSeconds() const;
+
+		private:
+			float DefaultDelay;
+	};
+}
