Metasound Frontend Revamp
- New node name {namespace.name.variant}
- PointIDs for exact pin-to-pin connections
- TAccessPtr updates
- Controller/Handle interface
- Document Handle
- Frontend Transforms
- Input/Output node external registration
- Revamped document model

diff --git a/Engine/Plugins/Runtime/Metasound/Source/MetasoundEditor/Private/MetasoundDetailCustomization.cpp b/Engine/Plugins/Runtime/Metasound/Source/MetasoundEditor/Private/MetasoundDetailCustomization.cpp
--- a/Engine/Plugins/Runtime/Metasound/Source/MetasoundEditor/Private/MetasoundDetailCustomization.cpp
+++ b/Engine/Plugins/Runtime/Metasound/Source/MetasoundEditor/Private/MetasoundDetailCustomization.cpp
@@ -37,7 +37,7 @@ namespace Metasound
 			return FName(InBasePath.ToString() + TEXT(".") + InPropertyName.ToString());
 		}
 
-		template <typename T>
+		template <typename T, typename MetadataType>
 		void BuildIOFixedArray(IDetailLayoutBuilder& InDetailLayout, FName CategoryName, FName InPropertyName, const TSet<FString>& RequiredValues, bool bIsInput)
 		{
 			IDetailCategoryBuilder& CategoryBuilder = InDetailLayout.EditCategory(CategoryName);
@@ -48,25 +48,18 @@ namespace Metasound
 			ArrayHandle->GetNumElements(NumElements);
 			for (int32 i = 0; i < static_cast<int32>(NumElements); ++i)
 			{
-				TMap<FName, TSharedPtr<IPropertyHandle>> ChildProperties;
 				TSharedRef<IPropertyHandle> ArrayItemHandle = ArrayHandle->GetElement(i);
-				uint32 NumChildren = 0;
-				ArrayItemHandle->GetNumChildren(NumChildren);
-				for (int32 j = 0; j < static_cast<int32>(NumChildren); ++j)
-				{
-					TSharedPtr<IPropertyHandle> PropertyHandle = ArrayItemHandle->GetChildHandle(j);
-					FProperty* Property = PropertyHandle->GetProperty();
-					if (ensureAlways(Property))
-					{
-						const FName PropertyName = Property->GetFName();
-						ChildProperties.Add(PropertyName, PropertyHandle);
-					}
-				}
 
-				TSharedPtr<IPropertyHandle> TypeProperty = ChildProperties.FindChecked(GET_MEMBER_NAME_CHECKED(T, TypeName));
-				TSharedPtr<IPropertyHandle> ToolTipProperty = ChildProperties.FindChecked(GET_MEMBER_NAME_CHECKED(T, ToolTip));
-				TSharedPtr<IPropertyHandle> DisplayNameProperty = ChildProperties.FindChecked(GET_MEMBER_NAME_CHECKED(T, DisplayName));
-				TSharedPtr<IPropertyHandle> NameProperty = ChildProperties.FindChecked(GET_MEMBER_NAME_CHECKED(T, Name));
+				const FName TypeNamePropertyName = GET_MEMBER_NAME_CHECKED(T, TypeName);
+				const FName NamePropertyName = GET_MEMBER_NAME_CHECKED(T, TypeName);
+				const FName ToolTipPropertyName = GET_MEMBER_NAME_CHECKED(MetadataType, Description);
+				const FName DisplayNamePropertyName = GET_MEMBER_NAME_CHECKED(MetadataType, DisplayName);
+
+				TSharedPtr<IPropertyHandle> TypeProperty = ArrayItemHandle->GetChildHandle(TypeNamePropertyName);
+				TSharedPtr<IPropertyHandle> NameProperty = ArrayItemHandle->GetChildHandle(NamePropertyName);
+
+				TSharedPtr<IPropertyHandle> ToolTipProperty = ArrayItemHandle->GetChildHandle(ToolTipPropertyName, true /* bRecurse */);
+				TSharedPtr<IPropertyHandle> DisplayNameProperty = ArrayItemHandle->GetChildHandle(DisplayNamePropertyName, true /* bRecurse */);
 
 				FString Name;
 				const bool bNameFound = NameProperty->GetValue(Name) == FPropertyAccess::Success;
@@ -113,8 +106,18 @@ namespace Metasound
 
 				if (bIsInput)
 				{
-					TSharedPtr<IPropertyHandle> LiteralValueProperty = ChildProperties.FindChecked(GET_MEMBER_NAME_CHECKED(FMetasoundInputDescription, LiteralValue));
-					CategoryBuilder.AddProperty(LiteralValueProperty);
+					const FName DefaultsPropertyName = GET_MEMBER_NAME_CHECKED(FMetasoundFrontendClassInput, Defaults);
+					TSharedPtr<IPropertyHandle> DefaultsProperty = ArrayItemHandle->GetChildHandle(DefaultsPropertyName);
+
+					TSharedPtr<IPropertyHandleArray> DefaultsArrayHandle = DefaultsProperty->AsArray();
+
+					uint32 NumDefaults = 0;
+					DefaultsArrayHandle->GetNumElements(NumDefaults);
+					for (uint32 InputLiteralIndex = 0; InputLiteralIndex < NumDefaults; ++InputLiteralIndex)
+					{
+						TSharedPtr<IPropertyHandle> LiteralHandle = DefaultsArrayHandle->GetElement(InputLiteralIndex)->GetChildHandle(GET_MEMBER_NAME_CHECKED(FMetasoundFrontendVertexLiteral, Value));
+						CategoryBuilder.AddProperty(LiteralHandle);
+					}
 				}
 			}
 
@@ -133,13 +136,13 @@ namespace Metasound
 
 		FName FMetasoundDetailCustomization::GetMetadataRootClassPath() const
 		{
-			return Metasound::Editor::BuildChildPath(DocumentPropertyName, GET_MEMBER_NAME_CHECKED(FMetasoundDocument, RootClass));
+			return Metasound::Editor::BuildChildPath(DocumentPropertyName, GET_MEMBER_NAME_CHECKED(FMetasoundFrontendDocument, RootGraph));
 		}
 
 		FName FMetasoundDetailCustomization::GetMetadataPropertyPath() const
 		{
 			const FName RootClass = FName(GetMetadataRootClassPath());
-			return Metasound::Editor::BuildChildPath(RootClass, GET_MEMBER_NAME_CHECKED(FMetasoundClassDescription, Metadata));
+			return Metasound::Editor::BuildChildPath(RootClass, GET_MEMBER_NAME_CHECKED(FMetasoundFrontendClass, Metadata));
 		}
 
 		void FMetasoundDetailCustomization::CustomizeDetails(IDetailLayoutBuilder& DetailLayout)
@@ -152,11 +155,12 @@ namespace Metasound
 			TArray<TWeakObjectPtr<UObject>> ObjectsCustomized;
 			DetailLayout.GetObjectsBeingCustomized(ObjectsCustomized);
 
-			const FName AuthorPropertyPath = BuildChildPath(GetMetadataPropertyPath(), GET_MEMBER_NAME_CHECKED(FMetasoundClassMetadata, AuthorName));
-			const FName DescPropertyPath = BuildChildPath(GetMetadataPropertyPath(), GET_MEMBER_NAME_CHECKED(FMetasoundClassMetadata, MetasoundDescription));
-			const FName NodeTypePropertyPath = BuildChildPath(GetMetadataPropertyPath(), GET_MEMBER_NAME_CHECKED(FMetasoundClassMetadata, NodeType));
-			const FName MajorVersionPropertyPath = BuildChildPath(GetMetadataPropertyPath(), GET_MEMBER_NAME_CHECKED(FMetasoundClassMetadata, MajorVersion));
-			const FName MinorVersionPropertyPath = BuildChildPath(GetMetadataPropertyPath(), GET_MEMBER_NAME_CHECKED(FMetasoundClassMetadata, MinorVersion));
+			const FName AuthorPropertyPath = BuildChildPath(GetMetadataPropertyPath(), GET_MEMBER_NAME_CHECKED(FMetasoundFrontendClassMetadata, Author));
+			const FName DescPropertyPath = BuildChildPath(GetMetadataPropertyPath(), GET_MEMBER_NAME_CHECKED(FMetasoundFrontendClassMetadata, Description));
+			const FName NodeTypePropertyPath = BuildChildPath(GetMetadataPropertyPath(), GET_MEMBER_NAME_CHECKED(FMetasoundFrontendClassMetadata, Type));
+			const FName VersionPropertyPath = BuildChildPath(GetMetadataPropertyPath(), GET_MEMBER_NAME_CHECKED(FMetasoundFrontendClassMetadata, Version));
+			const FName MajorVersionPropertyPath = BuildChildPath(VersionPropertyPath, GET_MEMBER_NAME_CHECKED(FMetasoundFrontendVersionNumber, Major));
+			const FName MinorVersionPropertyPath = BuildChildPath(VersionPropertyPath, GET_MEMBER_NAME_CHECKED(FMetasoundFrontendVersionNumber, Minor));
 
 			TSharedPtr<IPropertyHandle> AuthorHandle = DetailLayout.GetProperty(AuthorPropertyPath);
 			TSharedPtr<IPropertyHandle> DescHandle = DetailLayout.GetProperty(DescPropertyPath);
@@ -183,22 +187,23 @@ namespace Metasound
 				FMetasoundAssetBase* MetasoundAsset = IMetasoundUObjectRegistry::Get().GetObjectAsAssetBase(Objects[0].Get());
 				check(MetasoundAsset);
 
-				Frontend::FGraphHandle GraphHandle = MetasoundAsset->GetRootGraphHandle();
-				for (const FMetasoundInputDescription& Desc : GraphHandle.GetRequiredInputs())
+				Frontend::FDocumentHandle DocumentHandle = MetasoundAsset->GetDocumentHandle();
+				for (const FMetasoundFrontendClassVertex& Desc : DocumentHandle->GetRequiredInputs())
 				{
 					RequiredInputs.Add(Desc.Name);
 				}
-				for (const FMetasoundOutputDescription& Desc : GraphHandle.GetRequiredOutputs())
+				for (const FMetasoundFrontendClassVertex& Desc : DocumentHandle->GetRequiredOutputs())
 				{
 					RequiredOutputs.Add(Desc.Name);
 				}
 			}
 
-			const FName InputsPropertyPath = BuildChildPath(GetMetadataRootClassPath(), GET_MEMBER_NAME_CHECKED(FMetasoundClassDescription, Inputs));
-			const FName OutputsPropertyPath = BuildChildPath(GetMetadataRootClassPath(), GET_MEMBER_NAME_CHECKED(FMetasoundClassDescription, Outputs));
+			const FName InterfacePropertyPath = BuildChildPath(GetMetadataRootClassPath(), GET_MEMBER_NAME_CHECKED(FMetasoundFrontendClass, Interface));
+			const FName InputsPropertyPath = BuildChildPath(InterfacePropertyPath, GET_MEMBER_NAME_CHECKED(FMetasoundFrontendClassInterface, Inputs));
+			const FName OutputsPropertyPath = BuildChildPath(InterfacePropertyPath, GET_MEMBER_NAME_CHECKED(FMetasoundFrontendClassInterface, Outputs));
 
-			BuildIOFixedArray<FMetasoundInputDescription>(DetailLayout, "Inputs", InputsPropertyPath, RequiredInputs, true /* bIsInput */);
-			BuildIOFixedArray<FMetasoundOutputDescription>(DetailLayout, "Outputs", OutputsPropertyPath, RequiredOutputs, false /* bIsInput */);
+			BuildIOFixedArray<FMetasoundFrontendClassInput, FMetasoundFrontendVertexMetadata>(DetailLayout, "Inputs", InputsPropertyPath, RequiredInputs, true /* bIsInput */);
+			BuildIOFixedArray<FMetasoundFrontendClassOutput, FMetasoundFrontendVertexMetadata>(DetailLayout, "Outputs", OutputsPropertyPath, RequiredOutputs, false /* bIsInput */);
 
 			// Hack to hide parent structs for nested metadata properties
 			DetailLayout.HideCategory("CustomView");
diff --git a/Engine/Plugins/Runtime/Metasound/Source/MetasoundEditor/Private/MetasoundEditor.cpp b/Engine/Plugins/Runtime/Metasound/Source/MetasoundEditor/Private/MetasoundEditor.cpp
--- a/Engine/Plugins/Runtime/Metasound/Source/MetasoundEditor/Private/MetasoundEditor.cpp
+++ b/Engine/Plugins/Runtime/Metasound/Source/MetasoundEditor/Private/MetasoundEditor.cpp
@@ -393,15 +393,15 @@ namespace Metasound
 				// TODO: use the same directory as the currently open metasound
 				const FString OutputPath = FString("/Game/ImportedMetasound/GeneratedMetasound");
 
-				FMetasoundDocument MetasoundDoc;
+				FMetasoundFrontendDocument MetasoundDoc;
 
 				if (Frontend::ImportJSONAssetToMetasound(InputPath, MetasoundDoc))
 				{
-					TArray<UClass*> ImportClasses = IMetasoundUObjectRegistry::Get().GetUClassesForArchetype(MetasoundDoc.Archetype.ArchetypeName);
+					TArray<UClass*> ImportClasses = IMetasoundUObjectRegistry::Get().GetUClassesForArchetype(MetasoundDoc.Archetype.Name);
 
 					if (ImportClasses.Num() < 1)
 					{
-						UE_LOG(LogTemp, Warning, TEXT("Cannot create UObject from Metasound document. No UClass supports archetype \"%s\""), *MetasoundDoc.Archetype.ArchetypeName.ToString());
+						UE_LOG(LogTemp, Warning, TEXT("Cannot create UObject from Metasound document. No UClass supports archetype \"%s\""), *MetasoundDoc.Archetype.Name.ToString());
 					}
 					else
 					{
@@ -410,7 +410,7 @@ namespace Metasound
 							for (UClass* Cls : ImportClasses)
 							{
 								// TODO: could do a modal dialog to give user choice of import type.
-								UE_LOG(LogTemp, Warning, TEXT("Duplicate UClass support archetype \"%s\" with UClass \"%s\""), *MetasoundDoc.Archetype.ArchetypeName.ToString(), *Cls->GetName());
+								UE_LOG(LogTemp, Warning, TEXT("Duplicate UClass support archetype \"%s\" with UClass \"%s\""), *MetasoundDoc.Archetype.Name.ToString(), *Cls->GetName());
 							}
 						}
 
@@ -438,7 +438,7 @@ namespace Metasound
 
 			// TODO: We could just make this an object.
 			const FString Path = FPaths::ProjectSavedDir() / TEXT("Metasounds") + FPaths::ChangeExtension(Metasound->GetPathName(), MetasoundExtension);
-			InMetasoundAsset->GetRootGraphHandle().ExportToJSONAsset(Path);
+			InMetasoundAsset->GetDocumentHandle()->ExportToJSONAsset(Path);
 		}
 
 		void FEditor::Play()
@@ -696,18 +696,24 @@ namespace Metasound
 			check(ParentMetasoundObject);
 			if (FMetasoundAssetBase* MetasoundAsset = IMetasoundUObjectRegistry::Get().GetObjectAsAssetBase(ParentMetasoundObject))
 			{
-				Frontend::FGraphHandle GraphHandle = MetasoundAsset->GetRootGraphHandle();
+				Frontend::FConstDocumentHandle DocumentHandle = MetasoundAsset->GetDocumentHandle();
 				for (FGraphPanelSelectionSet::TConstIterator NodeIt(SelectedNodes); NodeIt; ++NodeIt)
 				{
 					UMetasoundEditorGraphNode* Node = CastChecked<UMetasoundEditorGraphNode>(*NodeIt);
 					Frontend::FNodeHandle NodeHandle = Node->GetNodeHandle();
 
-					if (NodeHandle.GetNodeType() == EMetasoundClassType::Input)
+					if (NodeHandle->GetClassType() == EMetasoundFrontendClassType::Input)
 					{
-						if (GraphHandle.IsRequiredInput(NodeHandle.GetNodeName()))
+						auto IsRequiredInput = [&](const Frontend::FConstInputHandle& InputHandle) 
+						{ 
+							return DocumentHandle->IsRequiredInput(InputHandle->GetName()); 
+						};
+						TArray<Frontend::FConstInputHandle> NodeInputs = NodeHandle->GetConstInputs();
+
+						if (Frontend::FConstInputHandle* InputHandle = NodeInputs.FindByPredicate(IsRequiredInput))
 						{
 							FNotificationInfo Info(FText::Format(LOCTEXT("Metasounds_CannotDeleteRequiredInput",
-								"'Required Input '{0}' cannot be deleted."), GraphHandle.GetInputDisplayName(NodeHandle.GetNodeName())));
+								"'Required Input '{0}' cannot be deleted."), (*InputHandle)->GetDisplayName()));
 							Info.bFireAndForget = true;
 							Info.ExpireDuration = 2.0f;
 							Info.bUseThrobber = true;
@@ -716,12 +722,18 @@ namespace Metasound
 						}
 					}
 
-					if (NodeHandle.GetNodeType() == EMetasoundClassType::Output)
+					if (NodeHandle->GetClassType() == EMetasoundFrontendClassType::Output)
 					{
-						if (GraphHandle.IsRequiredOutput(NodeHandle.GetNodeName()))
+						auto IsRequiredOutput = [&](const Frontend::FConstOutputHandle& OutputHandle) 
+						{ 
+							return DocumentHandle->IsRequiredOutput(OutputHandle->GetName()); 
+						};
+						TArray<Frontend::FConstOutputHandle> NodeOutputs = NodeHandle->GetConstOutputs();
+
+						if (Frontend::FConstOutputHandle* OutputHandle = NodeOutputs.FindByPredicate(IsRequiredOutput))
 						{
 							FNotificationInfo Info(FText::Format(LOCTEXT("Metasounds_CannotDeleteRequiredOutput",
-								"'Required Output '{0}' cannot be deleted."), GraphHandle.GetOutputDisplayName(NodeHandle.GetNodeName())));
+								"'Required Output '{0}' cannot be deleted."), (*OutputHandle)->GetDisplayName()));
 							Info.bFireAndForget = true;
 							Info.ExpireDuration = 2.0f;
 							Info.bUseThrobber = true;
@@ -779,56 +791,51 @@ namespace Metasound
 			using namespace Metasound::Frontend;
 
 			// Export the selected nodes and place the text on the clipboard
-			const FGraphPanelSelectionSet SelectedNodes = MetasoundGraphEditor->GetSelectedNodes();
-
-			TMap<uint32, FNodeHandle> HandlesToCopy;
-			for (FGraphPanelSelectionSet::TConstIterator SelectedIter(SelectedNodes); SelectedIter; ++SelectedIter)
-			{
-				if (UMetasoundEditorGraphNode* Node = Cast<UMetasoundEditorGraphNode>(*SelectedIter))
-				{
-					HandlesToCopy.Add(Node->GetNodeID(), Node->GetNodeHandle());
-					Node->PrepareForCopying();
-				}
-			}
-
-			FMetasoundAssetBase* MetasoundAsset = Metasound::IMetasoundUObjectRegistry::Get().GetObjectAsAssetBase(Metasound);
-			check(MetasoundAsset);
-			FGraphHandle GraphHandle = MetasoundAsset->GetRootGraphHandle();
-
-			UMetasound* MetasoundCopy = NewObject<UMetasound>();
-			FGraphHandle CopyGraphHandle = InitMetasound(*MetasoundCopy);
-
-			GraphHandle.CopyGraph(CopyGraphHandle);
-
-			TArray<FNodeHandle> AllCopiedNodes = CopyGraphHandle.GetAllNodes();
-			for (FNodeHandle& Node : AllCopiedNodes)
-			{
-				if (!HandlesToCopy.Contains(Node.GetNodeID()))
-				{
-					if (Node.GetNodeType() == EMetasoundClassType::Input)
-					{
-						CopyGraphHandle.RemoveInput(Node.GetNodeName());
-					}
-					else if (Node.GetNodeType() == EMetasoundClassType::Output)
-					{
-						CopyGraphHandle.RemoveOutput(Node.GetNodeName());
-					}
-					else
-					{
-						CopyGraphHandle.RemoveNode(Node);
-					}
-				}
-			}
-
-			FString ExportedEdText;
-			FEdGraphUtilities::ExportNodesToText(SelectedNodes, ExportedEdText);
-
-			FMetasoundEditorData EditorData;
-			EditorData.GraphData = ExportedEdText;
-			CopyGraphHandle.SetEditorData(EditorData);
-
-			const FString JSONToCopy = CopyGraphHandle.ExportToJSON();
-			FPlatformApplicationMisc::ClipboardCopy(*JSONToCopy);
+			/* TODO: rethink how this is performed. Current implementation does not
+			 * handle case where nodes are not in the root graph. 
+			 *
+			 * GraphHandle should support exporting selection of nodes to a FMetasoundFrontendDocument as nodes on the root class. 
+			 */
+			 
+			//const FGraphPanelSelectionSet SelectedNodes = MetasoundGraphEditor->GetSelectedNodes();
+			//
+			//TMap<int32, FNodeHandle> HandlesToCopy;
+			//for (FGraphPanelSelectionSet::TConstIterator SelectedIter(SelectedNodes); SelectedIter; ++SelectedIter)
+			//{
+				//if (UMetasoundEditorGraphNode* Node = Cast<UMetasoundEditorGraphNode>(*SelectedIter))
+				//{
+					//HandlesToCopy.Add(Node->GetID(), Node->GetNodeHandle());
+					//Node->PrepareForCopying();
+				//}
+			//}
+//
+			//FMetasoundAssetBase* MetasoundAsset = Metasound::IMetasoundUObjectRegistry::Get().GetObjectAsAssetBase(Metasound);
+			//check(MetasoundAsset);
+			//FGraphHandle GraphHandle = MetasoundAsset->GetRootGraphHandle();
+//
+			//UMetasound* MetasoundCopy = NewObject<UMetasound>();
+			//FGraphHandle CopyGraphHandle = InitMetasound(*MetasoundCopy);
+//
+			//GraphHandle->CopyGraph(CopyGraphHandle);
+//
+			//TArray<FNodeHandle> AllCopiedNodes = CopyGraphHandle->GetAllNodes();
+			//for (FNodeHandle& Node : AllCopiedNodes)
+			//{
+				//if (!HandlesToCopy.Contains(Node->GetID()))
+				//{
+					//CopyGraphHandle->RemoveNode(Node);
+				//}
+			//}
+//
+			//FString ExportedEdText;
+			//FEdGraphUtilities::ExportNodesToText(SelectedNodes, ExportedEdText);
+//
+			//FMetasoundEditorData EditorData;
+			//EditorData.GraphData = ExportedEdText;
+			//CopyGraphHandle.SetEditorData(EditorData);
+//
+			//const FString JSONToCopy = CopyGraphHandle.ExportToJSON();
+			//FPlatformApplicationMisc::ClipboardCopy(*JSONToCopy);
 		}
 
 		bool FEditor::CanCopyNodes() const
@@ -902,7 +909,7 @@ namespace Metasound
 // 			FEdGraphUtilities::ImportNodesFromText(Graph, PastedDocument.EditorData.GraphData, PastedNodes);
 // 
 // 			// Clear out, no reason to hold a bunch of references to data no longer needed.
-// 			PastedDocument = FMetasoundDocument();
+// 			PastedDocument = FMetasoundFrontendDocument();
 // 
 // 			FVector2D AvgNodePosition = FVector2D::ZeroVector;
 // 			for (UEdGraphNode* Node : PastedNodes)
@@ -947,7 +954,7 @@ namespace Metasound
 // 			FMetasoundAssetBase* AssetBase = IMetasoundUObjectRegistry::Get().GetObjectAsAssetBase(Metasound);
 // 			check(AssetBase);
 // 
-// 			PastedDocument = FMetasoundDocument();
+// 			PastedDocument = FMetasoundFrontendDocument();
 // 			const bool bCanImportDocument = Metasound::Frontend::ImportJSONToMetasound(ClipboardContent, PastedDocument);
 // 			if (!bCanImportDocument)
 // 			{
diff --git a/Engine/Plugins/Runtime/Metasound/Source/MetasoundEditor/Private/MetasoundEditor.h b/Engine/Plugins/Runtime/Metasound/Source/MetasoundEditor/Private/MetasoundEditor.h
--- a/Engine/Plugins/Runtime/Metasound/Source/MetasoundEditor/Private/MetasoundEditor.h
+++ b/Engine/Plugins/Runtime/Metasound/Source/MetasoundEditor/Private/MetasoundEditor.h
@@ -8,6 +8,7 @@
 #include "IDetailsView.h"
 #include "IMetasoundEditor.h"
 #include "MetasoundFrontend.h"
+#include "MetasoundFrontendController.h"
 #include "Misc/NotifyHook.h"
 #include "SMetasoundPalette.h"
 #include "Textures/SlateIcon.h"
@@ -193,7 +194,7 @@ namespace Metasound
 			bool bPassedValidation = true;
 
 			/** Document used when pasting from clipboard to avoid deserializing twice */
-			FMetasoundDocument PastedDocument;
+			FMetasoundFrontendDocument PastedDocument;
 		};
 	} // namespace Editor
 } // namespace Metasound
diff --git a/Engine/Plugins/Runtime/Metasound/Source/MetasoundEditor/Private/MetasoundEditorGraphBuilder.cpp b/Engine/Plugins/Runtime/Metasound/Source/MetasoundEditor/Private/MetasoundEditorGraphBuilder.cpp
--- a/Engine/Plugins/Runtime/Metasound/Source/MetasoundEditor/Private/MetasoundEditorGraphBuilder.cpp
+++ b/Engine/Plugins/Runtime/Metasound/Source/MetasoundEditor/Private/MetasoundEditorGraphBuilder.cpp
@@ -10,6 +10,7 @@
 #include "MetasoundEditorGraphNode.h"
 #include "MetasoundEditorModule.h"
 #include "MetasoundFrontend.h"
+#include "MetasoundFrontendController.h"
 #include "MetasoundFrontendRegistries.h"
 #include "MetasoundUObjectRegistry.h"
 #include "Modules/ModuleManager.h"
@@ -43,7 +44,7 @@ namespace Metasound
 			UMetasoundEditorGraphNode* NewGraphNode = NodeCreator.CreateNode(bInSelectNewNode);
 			NodeCreator.Finalize();
 
-			NewGraphNode->SetNodeID(InNodeHandle.GetNodeID());
+			NewGraphNode->SetNodeID(InNodeHandle->GetID());
 			NewGraphNode->CreateNewGuid();
 			NewGraphNode->NodePosX = Location.X;
 			NewGraphNode->NodePosY = Location.Y;
@@ -67,7 +68,7 @@ namespace Metasound
 			FMetasoundAssetBase* MetasoundAsset = IMetasoundUObjectRegistry::Get().GetObjectAsAssetBase(&InMetasound);
 			check(MetasoundAsset);
 
-			return MetasoundAsset->GetRootGraphHandle().AddNewNode(InClassInfo);
+			return MetasoundAsset->GetRootGraphHandle()->AddNode(InClassInfo);
 		}
 
 		FString FGraphBuilder::GetDataTypeDisplayName(const FName& InDataTypeName)
@@ -106,7 +107,7 @@ namespace Metasound
 
 			int32 i = 1;
 			FString NewNodeName = NameBase + FString::Printf(TEXT("_%02d"), i);
-			while (GraphHandle.ContainsInputNodeWithName(NewNodeName))
+			while (GraphHandle->ContainsInputNodeWithName(NewNodeName))
 			{
 				NewNodeName = NameBase + FString::Printf(TEXT("_%02d"), ++i);
 			}
@@ -125,7 +126,7 @@ namespace Metasound
 
 			int32 i = 1;
 			FString NewNodeName = NameBase + FString::Printf(TEXT("_%02d"), i);
-			while (GraphHandle.ContainsOutputNodeWithName(NewNodeName))
+			while (GraphHandle->ContainsOutputNodeWithName(NewNodeName))
 			{
 				NewNodeName = NameBase + FString::Printf(TEXT("_%02d"), ++i);
 			}
@@ -141,89 +142,98 @@ namespace Metasound
 
 		Frontend::FNodeHandle FGraphBuilder::AddInputNodeHandle(UObject& InMetasound, const FString& InName, const FName InTypeName, const FText& InToolTip)
 		{
-			FMetasoundInputDescription Description;
-			Description.Name = InName;
-			Description.TypeName = InTypeName;
-			Description.ToolTip = InToolTip;
 
 			FMetasoundAssetBase* MetasoundAsset = IMetasoundUObjectRegistry::Get().GetObjectAsAssetBase(&InMetasound);
 			check(MetasoundAsset);
 
 			Frontend::FGraphHandle GraphHandle = MetasoundAsset->GetRootGraphHandle();
-			Frontend::FNodeHandle NodeHandle = GraphHandle.AddNewInput(Description);
-			if (!NodeHandle.IsValid())
+
+			FMetasoundFrontendClassInput Description;
+
+			Description.Name = InName;
+			Description.TypeName = InTypeName;
+			Description.Metadata.Description = InToolTip;
+			int32 PointID = GraphHandle->GetNewPointID();
+
+			Description.PointIDs.Add(PointID);
+			FMetasoundFrontendVertexLiteral DefaultValue;
+			DefaultValue.PointID = PointID;
+			Description.Defaults.Add(DefaultValue);
+
+			Frontend::FNodeHandle NodeHandle = GraphHandle->AddInputVertex(Description);
+			if (!NodeHandle->IsValid())
 			{
-				return Frontend::FNodeHandle::InvalidHandle();
+				return NodeHandle;
 			}
 
 			IMetasoundEditorModule& EditorModule = FModuleManager::GetModuleChecked<IMetasoundEditorModule>("MetasoundEditor");
 			FEditorDataType DataType = EditorModule.FindDataType(InTypeName);
 
-			Metasound::FDataTypeLiteralParam LiteralParam = Frontend::GetDefaultParamForDataType(InTypeName);
+			Metasound::FLiteral LiteralParam = Frontend::GetDefaultParamForDataType(InTypeName);
 			if (!ensureAlways(LiteralParam.IsValid()))
 			{
-				return Frontend::FNodeHandle::InvalidHandle();
+				return NodeHandle;
 			}
 
-			switch (LiteralParam.ConstructorArgType)
+			switch (LiteralParam.GetType())
 			{
-				case ELiteralArgType::Boolean:
+				case ELiteralType::Boolean:
 				{
-					GraphHandle.SetInputToLiteral(InName, false);
+					GraphHandle->SetDefaultInputToLiteral(InName, PointID, false);
 				}
 				break;
 
-				case ELiteralArgType::Float:
+				case ELiteralType::Float:
 				{
-					GraphHandle.SetInputToLiteral(InName, 0.0f);
+					GraphHandle->SetDefaultInputToLiteral(InName, PointID, 0.0f);
 				}
 				break;
 
-				case ELiteralArgType::Integer:
+				case ELiteralType::Integer:
 				{
-					GraphHandle.SetInputToLiteral(InName, 0);
+					GraphHandle->SetDefaultInputToLiteral(InName, PointID, 0);
 				}
 				break;
 
-				case ELiteralArgType::String:
+				case ELiteralType::String:
 				{
-					GraphHandle.SetInputToLiteral(InName, FString(TEXT("")));
+					GraphHandle->SetDefaultInputToLiteral(InName, PointID, FString(TEXT("")));
 				}
 				break;
 
-				case ELiteralArgType::UObjectProxy:
+				case ELiteralType::UObjectProxy:
 				{
 					UClass* ClassToUse = FMetasoundFrontendRegistryContainer::Get()->GetLiteralUClassForDataType(InTypeName);
 					if (ClassToUse)
 					{
-						GraphHandle.SetInputToLiteral(InName, ClassToUse->ClassDefaultObject);
+						GraphHandle->SetDefaultInputToLiteral(InName, PointID, ClassToUse->ClassDefaultObject);
 					}
 				}
 				break;
 
-				case ELiteralArgType::UObjectProxyArray:
+				case ELiteralType::UObjectProxyArray:
 				{
 					UClass* ClassToUse = FMetasoundFrontendRegistryContainer::Get()->GetLiteralUClassForDataType(InTypeName);
 					if (ClassToUse)
 					{
 						TArray<UObject*> ObjectArray;
 						ObjectArray.Add(ClassToUse->ClassDefaultObject);
-						GraphHandle.SetInputToLiteral(InName, ObjectArray);
+						GraphHandle->SetDefaultInputToLiteral(InName, PointID, ObjectArray);
 					}
-					GraphHandle.SetInputToLiteral(InName, FString(TEXT("")));
+					GraphHandle->SetDefaultInputToLiteral(InName, PointID, FString(TEXT("")));
 				}
 				break;
 
-				case ELiteralArgType::Invalid:
-				case ELiteralArgType::None:
+				case ELiteralType::Invalid:
+				case ELiteralType::None:
 				default:
 				{
-					static_assert(static_cast<int32>(ELiteralArgType::Invalid) == 7, "Possible missing ELiteralArgType case coverage");
+					static_assert(static_cast<int32>(ELiteralType::Invalid) == 7, "Possible missing ELiteralType case coverage");
 				}
 				break;
 			}
 
-			GraphHandle.SetInputDisplayName(InName, FText::FromString(InName));
+			GraphHandle->SetInputDisplayName(InName, FText::FromString(InName));
 			return NodeHandle;
 		}
 
@@ -235,18 +245,21 @@ namespace Metasound
 
 		Frontend::FNodeHandle FGraphBuilder::AddOutputNodeHandle(UObject& InMetasound, const FString& InName, const FName InTypeName, const FText& InToolTip)
 		{
-			FMetasoundOutputDescription Description;
-			Description.Name = InName;
-			Description.TypeName = InTypeName;
-			Description.ToolTip = InToolTip;
-
 			FMetasoundAssetBase* MetasoundAsset = IMetasoundUObjectRegistry::Get().GetObjectAsAssetBase(&InMetasound);
 			check(MetasoundAsset);
 
 			Frontend::FGraphHandle GraphHandle = MetasoundAsset->GetRootGraphHandle();
-			Frontend::FNodeHandle NodeHandle = GraphHandle.AddNewOutput(Description);
 
-			GraphHandle.SetOutputDisplayName(InName, FText::FromString(InName));
+
+			FMetasoundFrontendClassOutput Description;
+			Description.Name = InName;
+			Description.TypeName = InTypeName;
+			Description.Metadata.Description = InToolTip;
+			Description.PointIDs.Add(GraphHandle->GetNewPointID());
+
+			Frontend::FNodeHandle NodeHandle = GraphHandle->AddOutputVertex(Description);
+
+			GraphHandle->SetOutputDisplayName(InName, FText::FromString(InName));
 
 			return NodeHandle;
 		}
@@ -268,28 +281,9 @@ namespace Metasound
 			check(MetasoundAsset);
 
 			Frontend::FGraphHandle GraphHandle = MetasoundAsset->GetRootGraphHandle();
-			if (GraphHandle.IsValid() && NodeHandle.IsValid())
+			if (GraphHandle->IsValid() && NodeHandle->IsValid())
 			{
-				switch (NodeHandle.GetNodeType())
-				{
-					case EMetasoundClassType::Input:
-					{
-						GraphHandle.RemoveInput(NodeHandle.GetNodeName());
-					}
-					break;
-
-					case EMetasoundClassType::Output:
-					{
-						GraphHandle.RemoveOutput(NodeHandle.GetNodeName());
-					}
-					break;
-
-					case EMetasoundClassType::External:
-					{
-						GraphHandle.RemoveNode(NodeHandle);
-					}
-					break;
-				}
+				GraphHandle->RemoveNode(*NodeHandle);
 			}
 
 			InNode.PostEditChange();
@@ -315,22 +309,22 @@ namespace Metasound
 
 			struct FNodePair
 			{
-				FNodeHandle NodeHandle = FNodeHandle::InvalidHandle();
+				FNodeHandle NodeHandle = INodeController::GetInvalidHandle();
 				UEdGraphNode* GraphNode = nullptr;
 			};
 
 			TMap<uint32, FNodePair> NewIdNodeMap;
-			TArray<FNodeHandle> NodeHandles = GraphHandle.GetAllNodes();
+			TArray<FNodeHandle> NodeHandles = GraphHandle->GetNodes();
 			for (FNodeHandle& NodeHandle : NodeHandles)
 			{
-				const FNodeClassInfo ClassInfo = NodeHandle.GetClassInfo();
+				const EMetasoundFrontendClassType NodeType = NodeHandle->GetClassType();
 				FVector2D Location;
-				if (ClassInfo.NodeType == EMetasoundClassType::Input)
+				if (NodeType == EMetasoundFrontendClassType::Input)
 				{
 					Location = InputNodeLocation;
 					InputNodeLocation.Y += 100.0f;
 				}
-				else if (ClassInfo.NodeType == EMetasoundClassType::Output)
+				else if (NodeType == EMetasoundFrontendClassType::Output)
 				{
 					Location = OutputNodeLocation;
 					OutputNodeLocation.Y += 100.0f;
@@ -341,7 +335,7 @@ namespace Metasound
 					OpNodeLocation.Y += 100.0f;
 				}
 				UEdGraphNode* NewNode = AddNode(InMetasound, Location, NodeHandle, false /* bInSelectNewNode */);
-				NewIdNodeMap.Add(NodeHandle.GetNodeID(), FNodePair { NodeHandle, NewNode });
+				NewIdNodeMap.Add(NodeHandle->GetID(), FNodePair { NodeHandle, NewNode });
 			}
 
 			for (const TPair<uint32, FNodePair>& IdNodePair : NewIdNodeMap)
@@ -352,7 +346,7 @@ namespace Metasound
 				FNodeHandle NodeHandle = IdNodePair.Value.NodeHandle;
 				TArray<UEdGraphPin*> Pins = GraphNode->GetAllPins();
 
-				const TArray<FInputHandle> NodeInputs = NodeHandle.GetAllInputs();
+				const TArray<FInputHandle> NodeInputs = NodeHandle->GetInputs();
 
 				int32 InputIndex = 0;
 				for (UEdGraphPin* Pin : Pins)
@@ -361,11 +355,11 @@ namespace Metasound
 					{
 						case EEdGraphPinDirection::EGPD_Input:
 						{
-							FOutputHandle OutputHandle = NodeInputs[InputIndex].GetCurrentlyConnectedOutput();
-							if (OutputHandle.IsValid())
+							FOutputHandle OutputHandle = NodeInputs[InputIndex]->GetCurrentlyConnectedOutput();
+							if (OutputHandle->IsValid())
 							{
-								UEdGraphNode* OutputGraphNode = NewIdNodeMap.FindChecked(OutputHandle.GetOwningNodeID()).GraphNode;
-								UEdGraphPin* OutputPin = OutputGraphNode->FindPinChecked(OutputHandle.GetOutputName(), EEdGraphPinDirection::EGPD_Output);
+								UEdGraphNode* OutputGraphNode = NewIdNodeMap.FindChecked(OutputHandle->GetOwningNodeID()).GraphNode;
+								UEdGraphPin* OutputPin = OutputGraphNode->FindPinChecked(OutputHandle->GetName(), EEdGraphPinDirection::EGPD_Output);
 								Pin->MakeLinkTo(OutputPin);
 							}
 
@@ -384,7 +378,7 @@ namespace Metasound
 			InMetasound.MarkPackageDirty();
 		}
 
-		void FGraphBuilder::RebuildNodePins(UMetasoundEditorGraphNode& InGraphNode, Frontend::FNodeHandle& InNodeHandle, bool bInRecordTransaction)
+		void FGraphBuilder::RebuildNodePins(UMetasoundEditorGraphNode& InGraphNode, Frontend::FNodeHandle InNodeHandle, bool bInRecordTransaction)
 		{
 			const FScopedTransaction Transaction(LOCTEXT("RebuildMetasoundGraphNodePins", "Rebuild Metasound Pins"), bInRecordTransaction);
 
@@ -400,13 +394,13 @@ namespace Metasound
 
 			Frontend::FGraphHandle GraphHandle = MetasoundAsset->GetRootGraphHandle();
 
-			TArray<Frontend::FInputHandle> InputHandles = InNodeHandle.GetAllInputs();
+			TArray<Frontend::FInputHandle> InputHandles = InNodeHandle->GetInputs();
 			for (int32 i = 0; i < InputHandles.Num(); ++i)
 			{
 				AddPinToNode(InGraphNode, InputHandles[i]);
 			}
 
-			TArray<Frontend::FOutputHandle> OutputHandles = InNodeHandle.GetAllOutputs();
+			TArray<Frontend::FOutputHandle> OutputHandles = InNodeHandle->GetOutputs();
 			for (int32 i = 0; i < OutputHandles.Num(); ++i)
 			{
 				AddPinToNode(InGraphNode, OutputHandles[i]);
@@ -419,18 +413,18 @@ namespace Metasound
 		{
 			if (EEdGraphPinDirection::EGPD_Input == InEditorPin.Direction)
 			{
-				if (InEditorPin.GetName() == InInputHandle.GetInputName())
+				if (InEditorPin.GetName() == InInputHandle->GetName())
 				{
 					IMetasoundEditorModule& EditorModule = FModuleManager::GetModuleChecked<IMetasoundEditorModule>("MetasoundEditor");
 
-					if (InEditorPin.PinType == EditorModule.FindDataType(InInputHandle.GetInputType()).PinType)
+					if (InEditorPin.PinType == EditorModule.FindDataType(InInputHandle->GetDataType()).PinType)
 					{
 						UEdGraphNode* EditorNode = InEditorPin.GetOwningNode();
 						UMetasoundEditorGraphNode* MetasoundEditorNode = Cast<UMetasoundEditorGraphNode>(EditorNode);
 
 						if (nullptr != MetasoundEditorNode)
 						{
-							if  (MetasoundEditorNode->GetNodeHandle().GetNodeID() == InInputHandle.GetOwningNodeID())
+							if  (MetasoundEditorNode->GetNodeHandle()->GetID() == InInputHandle->GetOwningNodeID())
 							{
 								return true;
 							}
@@ -445,18 +439,18 @@ namespace Metasound
 		{
 			if (EEdGraphPinDirection::EGPD_Output == InEditorPin.Direction)
 			{
-				if (InEditorPin.GetName() == InOutputHandle.GetOutputName())
+				if (InEditorPin.GetName() == InOutputHandle->GetName())
 				{
 					IMetasoundEditorModule& EditorModule = FModuleManager::GetModuleChecked<IMetasoundEditorModule>("MetasoundEditor");
 
-					if (InEditorPin.PinType == EditorModule.FindDataType(InOutputHandle.GetOutputType()).PinType)
+					if (InEditorPin.PinType == EditorModule.FindDataType(InOutputHandle->GetDataType()).PinType)
 					{
 						UEdGraphNode* EditorNode = InEditorPin.GetOwningNode();
 						UMetasoundEditorGraphNode* MetasoundEditorNode = Cast<UMetasoundEditorGraphNode>(EditorNode);
 
 						if (nullptr != MetasoundEditorNode)
 						{
-							if  (MetasoundEditorNode->GetNodeHandle().GetNodeID() == InOutputHandle.GetOwningNodeID())
+							if  (MetasoundEditorNode->GetNodeHandle()->GetID() == InOutputHandle->GetOwningNodeID())
 							{
 								return true;
 							}
@@ -467,34 +461,34 @@ namespace Metasound
 			return false;
 		}
 
-		UEdGraphPin* FGraphBuilder::AddPinToNode(UMetasoundEditorGraphNode& InEditorNode, Frontend::FInputHandle& InInputHandle)
+		UEdGraphPin* FGraphBuilder::AddPinToNode(UMetasoundEditorGraphNode& InEditorNode, Frontend::FInputHandle InInputHandle)
 		{
 			IMetasoundEditorModule& EditorModule = FModuleManager::GetModuleChecked<IMetasoundEditorModule>("MetasoundEditor");
 
 			FEdGraphPinType PinType("MetasoundNode", NAME_None, nullptr, EPinContainerType::None, false, FEdGraphTerminalType());
-			UEdGraphPin* NewPin = InEditorNode.CreatePin(EGPD_Input, PinType, FName(*InInputHandle.GetInputName()));
+			UEdGraphPin* NewPin = InEditorNode.CreatePin(EGPD_Input, PinType, FName(*InInputHandle->GetName()));
 
 			if (ensureAlways(NewPin))
 			{
-				NewPin->PinToolTip = InInputHandle.GetInputTooltip().ToString();
-				NewPin->PinType = EditorModule.FindDataType(InInputHandle.GetInputType()).PinType;
+				NewPin->PinToolTip = InInputHandle->GetTooltip().ToString();
+				NewPin->PinType = EditorModule.FindDataType(InInputHandle->GetDataType()).PinType;
 			}
 
 			return NewPin;
 		}
 
-		UEdGraphPin* FGraphBuilder::AddPinToNode(UMetasoundEditorGraphNode& InEditorNode, Frontend::FOutputHandle& InOutputHandle)
+		UEdGraphPin* FGraphBuilder::AddPinToNode(UMetasoundEditorGraphNode& InEditorNode, Frontend::FOutputHandle InOutputHandle)
 		{
 			IMetasoundEditorModule& EditorModule = FModuleManager::GetModuleChecked<IMetasoundEditorModule>("MetasoundEditor");
 			FText DisplayName;
 
 			FEdGraphPinType PinType("MetasoundNode", NAME_None, nullptr, EPinContainerType::None, false, FEdGraphTerminalType());
-			UEdGraphPin* NewPin = InEditorNode.CreatePin(EGPD_Output, PinType, FName(*InOutputHandle.GetOutputName()));
+			UEdGraphPin* NewPin = InEditorNode.CreatePin(EGPD_Output, PinType, FName(*InOutputHandle->GetName()));
 
 			if (ensureAlways(NewPin))
 			{
-				NewPin->PinToolTip = InOutputHandle.GetOutputTooltip().ToString();
-				NewPin->PinType = EditorModule.FindDataType(InOutputHandle.GetOutputType()).PinType;
+				NewPin->PinToolTip = InOutputHandle->GetTooltip().ToString();
+				NewPin->PinType = EditorModule.FindDataType(InOutputHandle->GetDataType()).PinType;
 			}
 
 			return NewPin;
@@ -511,7 +505,7 @@ namespace Metasound
 
 			// Get all nodes from metasound graph
 			FGraphHandle GraphHandle = MetasoundAsset->GetRootGraphHandle();
-			TArray<FNodeHandle> Nodes = GraphHandle.GetAllNodes();
+			TArray<FNodeHandle> Nodes = GraphHandle->GetNodes();
 
 			// Get all nodes from metasound editor graph
 			UMetasoundEditorGraph* EditorGraph = CastChecked<UMetasoundEditorGraph>(MetasoundAsset->GetGraph());
@@ -529,7 +523,7 @@ namespace Metasound
 			struct FNodePair
 			{
 				UMetasoundEditorGraphNode* EditorNode = nullptr;
-				FNodeHandle Node = FNodeHandle::InvalidHandle();;
+				FNodeHandle Node = INodeController::GetInvalidHandle();;
 			};
 			TMap<int32, FNodePair> PairedNodes;
 
@@ -539,7 +533,7 @@ namespace Metasound
 				FNodeHandle Node = Nodes[i];
 				auto IsEditorNodeWithSameNodeID = [&](const UMetasoundEditorGraphNode* InEditorNode)
 				{
-					return InEditorNode->GetNodeID() == Node.GetNodeID();
+					return InEditorNode->GetNodeID() == Node->GetID();
 				};
 
 				UMetasoundEditorGraphNode* EditorNode = nullptr;
@@ -550,7 +544,7 @@ namespace Metasound
 
 				if (nullptr != EditorNode)
 				{
-					PairedNodes.Add(Node.GetNodeID(), { EditorNode, Node });
+					PairedNodes.Add(Node->GetID(), { EditorNode, Node });
 					EditorNodes.RemoveSwap(EditorNode);
 					Nodes.RemoveAtSwap(i);
 				}
@@ -572,15 +566,13 @@ namespace Metasound
 			for (FNodeHandle Node : Nodes)
 			{
 				UEdGraphNode* NewNode = AddNode(InMetasound, EditorGraph->GetGoodPlaceForNewNode(), Node, false /* bInSelectNewNode */);
-				PairedNodes.Add(Node.GetNodeID(), {Cast<UMetasoundEditorGraphNode>(NewNode), Node});
+				PairedNodes.Add(Node->GetID(), {Cast<UMetasoundEditorGraphNode>(NewNode), Node});
 			}
 
 			// Synchronize pins on node pairs.
 			for (const TPair<int32, FNodePair>& IdNodePair : PairedNodes)
 			{
 				UMetasoundEditorGraphNode* EditorNode = IdNodePair.Value.EditorNode;
-				FNodeHandle Node = IdNodePair.Value.Node;
-
 				bIsEditorGraphDirty |= SynchronizeNodePins(*IdNodePair.Value.EditorNode, IdNodePair.Value.Node);
 			}
 
@@ -600,8 +592,8 @@ namespace Metasound
 		{
 			bool bIsNodeDirty = false;
 
-			TArray<Frontend::FInputHandle> InputHandles = InNode.GetAllInputs();
-			TArray<Frontend::FOutputHandle> OutputHandles = InNode.GetAllOutputs();
+			TArray<Frontend::FInputHandle> InputHandles = InNode->GetInputs();
+			TArray<Frontend::FOutputHandle> OutputHandles = InNode->GetOutputs();
 			TArray<UEdGraphPin*> EditorPins = InEditorNode.Pins;
 
 			// Filter out pins which are not paired.
@@ -609,12 +601,12 @@ namespace Metasound
 			{
 				UEdGraphPin* Pin = EditorPins[i];
 
-				auto IsMatchingInputHandle = [&](const Frontend::FInputHandle& InputHandle)
+				auto IsMatchingInputHandle = [&](const Frontend::FInputHandle& InputHandle) -> bool
 				{
 					return IsMatchingInputHandleAndPin(InputHandle, *Pin);
 				};
 
-				auto IsMatchingOutputHandle = [&](const Frontend::FOutputHandle& OutputHandle)
+				auto IsMatchingOutputHandle = [&](const Frontend::FOutputHandle& OutputHandle) -> bool
 				{
 					return IsMatchingOutputHandleAndPin(OutputHandle, *Pin);
 				};
@@ -707,7 +699,7 @@ namespace Metasound
 				Frontend::FNodeHandle Node = MetasoundEditorNode->GetNodeHandle();
 
 				TArray<UEdGraphPin*> Pins = MetasoundEditorNode->GetAllPins();
-				TArray<Frontend::FInputHandle> NodeInputs = Node.GetAllInputs();
+				TArray<Frontend::FInputHandle> NodeInputs = Node->GetInputs();
 
 				for (Frontend::FInputHandle& NodeInput : NodeInputs)
 				{
@@ -727,8 +719,8 @@ namespace Metasound
 						// Remove pin so it isn't used twice. 
 						Pins.Remove(MatchingPin);
 
-						Frontend::FOutputHandle OutputHandle = NodeInput.GetCurrentlyConnectedOutput();
-						if (OutputHandle.IsValid())
+						Frontend::FOutputHandle OutputHandle = NodeInput->GetCurrentlyConnectedOutput();
+						if (OutputHandle->IsValid())
 						{
 							bool bAddLink = false;
 
@@ -746,8 +738,8 @@ namespace Metasound
 
 							if (bAddLink)
 							{
-								UMetasoundEditorGraphNode* OutputEditorNode = EditorNodesByID[OutputHandle.GetOwningNodeID()];
-								UEdGraphPin* OutputPin = OutputEditorNode->FindPinChecked(OutputHandle.GetOutputName(), EEdGraphPinDirection::EGPD_Output);
+								UMetasoundEditorGraphNode* OutputEditorNode = EditorNodesByID[OutputHandle->GetOwningNodeID()];
+								UEdGraphPin* OutputPin = OutputEditorNode->FindPinChecked(OutputHandle->GetName(), EEdGraphPinDirection::EGPD_Output);
 								MatchingPin->MakeLinkTo(OutputPin);
 								bIsNodeDirty = true;
 							}
diff --git a/Engine/Plugins/Runtime/Metasound/Source/MetasoundEditor/Private/MetasoundEditorGraphBuilder.h b/Engine/Plugins/Runtime/Metasound/Source/MetasoundEditor/Private/MetasoundEditorGraphBuilder.h
--- a/Engine/Plugins/Runtime/Metasound/Source/MetasoundEditor/Private/MetasoundEditorGraphBuilder.h
+++ b/Engine/Plugins/Runtime/Metasound/Source/MetasoundEditor/Private/MetasoundEditorGraphBuilder.h
@@ -2,7 +2,12 @@
 #pragma once
 
 #include "CoreMinimal.h"
+#include "EdGraph/EdGraph.h"
+#include "EdGraph/EdGraphNode.h"
+#include "EdGraph/EdGraphPin.h"
 #include "MetasoundFrontend.h"
+#include "MetasoundFrontendController.h"
+
 
 // Forward Declarations
 class UEdGraphNode;
@@ -12,13 +17,6 @@ class UMetasoundEditorGraphNode;
 
 namespace Metasound
 {
-	namespace Frontend
-	{
-		// Forward Declarations
-		class FNodeHandle;
-		struct FNodeClassInfo;
-	} // namespace Frontend
-
 	namespace Editor
 	{
 		class FGraphBuilder
@@ -62,15 +60,15 @@ namespace Metasound
 
 			static void RebuildGraph(UObject& InMetasound);
 
-			static void RebuildNodePins(UMetasoundEditorGraphNode& InGraphNode, Frontend::FNodeHandle& InNodeHandle, bool bInRecordTransaction = true);
+			static void RebuildNodePins(UMetasoundEditorGraphNode& InGraphNode, Frontend::FNodeHandle InNodeHandle, bool bInRecordTransaction = true);
 
 			// Adds an Input UEdGraphPin to a UMetasoundEditorGraphNode
-			static UEdGraphPin* AddPinToNode(UMetasoundEditorGraphNode& InEditorNode, Frontend::FInputHandle& InInputHandle);
+			static UEdGraphPin* AddPinToNode(UMetasoundEditorGraphNode& InEditorNode, Frontend::FInputHandle InInputHandle);
 
 			// Adds an Output UEdGraphPin to a UMetasoundEditorGraphNode
-			static UEdGraphPin* AddPinToNode(UMetasoundEditorGraphNode& InEditorNode, Frontend::FOutputHandle& InOutputHandle);
+			static UEdGraphPin* AddPinToNode(UMetasoundEditorGraphNode& InEditorNode, Frontend::FOutputHandle InOutputHandle);
 
-			// Adds and removes nodes, pins and connections so that the UEdGraph of the Metasound matches the FMetasoundDocumentModel
+			// Adds and removes nodes, pins and connections so that the UEdGraph of the Metasound matches the FMetasoundFrontendDocumentModel
 			//
 			// @return True if the UEdGraph was altered, false otherwise. 
 			static bool SynchronizeGraph(UObject& InMetasound);
@@ -81,7 +79,7 @@ namespace Metasound
 			static bool SynchronizeNodePins(UMetasoundEditorGraphNode& InEditorNode, Frontend::FNodeHandle InNode);
 
 			// Adds and removes connections so that the UEdGraph of the metasound has the same
-			// connections as the FMetasoundDocument graph.
+			// connections as the FMetasoundFrontendDocument graph.
 			//
 			// @return True if the UEdGraph was altered. False otherwise. 
 			static bool SynchronizeConnections(UObject& InMetasound);
diff --git a/Engine/Plugins/Runtime/Metasound/Source/MetasoundEditor/Private/MetasoundEditorGraphNode.cpp b/Engine/Plugins/Runtime/Metasound/Source/MetasoundEditor/Private/MetasoundEditorGraphNode.cpp
--- a/Engine/Plugins/Runtime/Metasound/Source/MetasoundEditor/Private/MetasoundEditorGraphNode.cpp
+++ b/Engine/Plugins/Runtime/Metasound/Source/MetasoundEditor/Private/MetasoundEditorGraphNode.cpp
@@ -96,7 +96,7 @@ Metasound::Frontend::FGraphHandle UMetasoundEditorGraphNode::GetRootGraphHandle(
 
 Metasound::Frontend::FNodeHandle UMetasoundEditorGraphNode::GetNodeHandle()
 {
-	return GetRootGraphHandle().GetNodeWithId(NodeID);
+	return GetRootGraphHandle()->GetNodeWithID(NodeID);
 }
 
 void UMetasoundEditorGraphNode::IteratePins(TUniqueFunction<void(UEdGraphPin* /* Pin */, int32 /* Index */)> InFunc, EEdGraphPinDirection InPinDirection)
@@ -115,25 +115,13 @@ void UMetasoundEditorGraphNode::AllocateDefaultPins()
 	using namespace Metasound;
 
 	ensureAlways(Pins.Num() == 0);
-
-	Frontend::FNodeHandle NodeHandle = Frontend::FNodeHandle::InvalidHandle();
-	if (NodeID != INDEX_NONE)
-	{
-		NodeHandle = GetNodeHandle();
-	}
-	Metasound::Editor::FGraphBuilder::RebuildNodePins(*this, NodeHandle, false /* bInRecordTransaction */);
+	Metasound::Editor::FGraphBuilder::RebuildNodePins(*this, GetNodeHandle(), false /* bInRecordTransaction */);
 }
 
 void UMetasoundEditorGraphNode::ReconstructNode()
 {
 	using namespace Metasound;
-
-	Frontend::FNodeHandle NodeHandle = Frontend::FNodeHandle::InvalidHandle();
-	if (NodeID != INDEX_NONE)
-	{
-		NodeHandle = GetNodeHandle();
-	}
-	Editor::FGraphBuilder::RebuildNodePins(*this, NodeHandle);
+	Editor::FGraphBuilder::RebuildNodePins(*this, GetNodeHandle());
 }
 
 void UMetasoundEditorGraphNode::AutowireNewNode(UEdGraphPin* FromPin)
@@ -188,17 +176,17 @@ bool UMetasoundEditorGraphNode::CanUserDeleteNode() const
 	FMetasoundAssetBase* MetasoundAsset = const_cast<FMetasoundAssetBase*>(ConstMetasoundAsset);
 	check(MetasoundAsset);
 
-	Metasound::Frontend::FGraphHandle GraphHandle = MetasoundAsset->GetRootGraphHandle();
-	Metasound::Frontend::FNodeHandle NodeHandle = GraphHandle.GetNodeWithId(NodeID);
+	Metasound::Frontend::FDocumentHandle DocumentHandle = MetasoundAsset->GetDocumentHandle();
+	Metasound::Frontend::FNodeHandle NodeHandle = DocumentHandle->GetRootGraph()->GetNodeWithID(NodeID);
 
-	const FString& NodeName = NodeHandle.GetNodeName();
+	const FString& NodeName = NodeHandle->GetNodeName();
 
-	if (GraphHandle.IsRequiredInput(NodeName))
+	if (DocumentHandle->IsRequiredInput(NodeName))
 	{
 		return false;
 	}
 
-	if (GraphHandle.IsRequiredOutput(NodeName))
+	if (DocumentHandle->IsRequiredOutput(NodeName))
 	{
 		return false;
 	}
@@ -232,27 +220,38 @@ FText UMetasoundEditorGraphNode::GetNodeTitle(ENodeTitleType::Type TitleType) co
 	check(MetasoundAsset);
 
 	FGraphHandle GraphHandle = MetasoundAsset->GetRootGraphHandle();
-	FNodeHandle NodeHandle = GraphHandle.GetNodeWithId(NodeID);
+	FConstNodeHandle NodeHandle = GraphHandle->GetNodeWithID(NodeID);
 
-	switch (NodeHandle.GetNodeType())
+	switch (NodeHandle->GetClassType())
 	{
-		case EMetasoundClassType::Input:
+		case EMetasoundFrontendClassType::Input:
 		{
-			return FText::Format(LOCTEXT("MetasoundGraphNode_TitleFormat", "Input {0}"), GraphHandle.GetInputDisplayName(NodeHandle.GetNodeName()));
+			// TODO: have INodeController determine display name
+			TArray<FConstInputHandle> InputHandles = NodeHandle->GetConstInputs();
+			if (InputHandles.Num() > 0)
+			{
+				return FText::Format(LOCTEXT("MetasoundGraphNode_InputTitleFormat", "Input {0}"), InputHandles[0]->GetDisplayName());
+			}
+			return LOCTEXT("MetasoundGraphNode_DefaultInputTitle", "Input");
 		}
 		break;
 
-		case EMetasoundClassType::Output:
+		case EMetasoundFrontendClassType::Output:
 		{
-			return FText::Format(LOCTEXT("MetasoundGraphNode_TitleFormat", "Output {0}"), GraphHandle.GetOutputDisplayName(NodeHandle.GetNodeName()));
+			TArray<FConstOutputHandle> OutputHandles = NodeHandle->GetConstOutputs();
+			if (OutputHandles.Num() > 0)
+			{
+				return FText::Format(LOCTEXT("MetasoundGraphNode_OutputTitleFormat", "Output {0}"), OutputHandles[0]->GetDisplayName());
+			}
+			return LOCTEXT("MetasoundGraphNode_DefaultOutputTitle", "Output");
 		}
 		break;
 
 		default:
-		case EMetasoundClassType::External:
-		case EMetasoundClassType::MetasoundGraph:
+		case EMetasoundFrontendClassType::External:
+		case EMetasoundFrontendClassType::Graph:
 		{
-			return FText::FromString(NodeHandle.GetNodeClassName());
+			return FText::FromString(NodeHandle->GetClassName());
 		}
 		break;
 	}
@@ -335,32 +334,8 @@ FText UMetasoundEditorGraphNode::GetTooltipText() const
 	// TODO: create FConstGraphHandle so that this const cast is unnecessary.
 	FMetasoundAssetBase* MetasoundAsset = const_cast<FMetasoundAssetBase*>(ConstMetasoundAsset);
 	check(MetasoundAsset);
-
-	FGraphHandle GraphHandle = MetasoundAsset->GetRootGraphHandle();
-	FNodeHandle NodeHandle = GraphHandle.GetNodeWithId(NodeID);
-
-	switch (NodeHandle.GetNodeType())
-	{
-		case EMetasoundClassType::Input:
-		{
-			return GraphHandle.GetInputToolTip(NodeHandle.GetNodeName());
-		}
-		break;
-
-		case EMetasoundClassType::Output:
-		{
-			return GraphHandle.GetOutputToolTip(NodeHandle.GetNodeName());
-		}
-		break;
-
-		default:
-		case EMetasoundClassType::External:
-		case EMetasoundClassType::MetasoundGraph:
-		{
-			return GenerateClassDescription(NodeHandle.GetClassInfo()).Metadata.MetasoundDescription;
-		}
-		break;
-	}
+	
+	return MetasoundAsset->GetRootGraphHandle()->GetNodeWithID(NodeID)->GetClassDescription();
 }
 
 FString UMetasoundEditorGraphNode::GetDocumentationExcerptName() const
diff --git a/Engine/Plugins/Runtime/Metasound/Source/MetasoundEditor/Private/MetasoundEditorGraphSchema.cpp b/Engine/Plugins/Runtime/Metasound/Source/MetasoundEditor/Private/MetasoundEditorGraphSchema.cpp
--- a/Engine/Plugins/Runtime/Metasound/Source/MetasoundEditor/Private/MetasoundEditorGraphSchema.cpp
+++ b/Engine/Plugins/Runtime/Metasound/Source/MetasoundEditor/Private/MetasoundEditorGraphSchema.cpp
@@ -341,43 +341,54 @@ const FPinConnectionResponse UMetasoundEditorGraphSchema::CanCreateConnection(co
 		return FPinConnectionResponse(CONNECT_RESPONSE_DISALLOW, LOCTEXT("ConnectionTypeIncorrect", "Connection pin types do not match"));
 	}
 
+	// TODO: use IDs to connect rather than names. Likely needs a UMetasoundEditorGraphPin
+	//int32 InputPointID = InputPin->GetMetasoundPointID();
+	//int32 InputPointID = OutputPin->GetMetasoundPointID();
+
 	UMetasoundEditorGraphNode* InputGraphNode = CastChecked<UMetasoundEditorGraphNode>(InputPin->GetOwningNode());
 	Metasound::Frontend::FNodeHandle InputNodeHandle = InputGraphNode->GetNodeHandle();
-	Metasound::Frontend::FInputHandle InputHandle = InputNodeHandle.GetInputWithName(InputPin->GetName());
+	TArray<Metasound::Frontend::FInputHandle> InputHandles = InputNodeHandle->GetInputsWithVertexName(InputPin->GetName());
 
 	UMetasoundEditorGraphNode* OutputGraphNode = CastChecked<UMetasoundEditorGraphNode>(OutputPin->GetOwningNode());
 	Metasound::Frontend::FNodeHandle OutputNodeHandle = OutputGraphNode->GetNodeHandle();
-	Metasound::Frontend::FOutputHandle OutputHandle = OutputNodeHandle.GetOutputWithName(OutputPin->GetName());
-
-	// TODO: Implement YesWithConverterNode to provide conversion options
-	Metasound::Frontend::FConnectability Connectability = InputHandle.CanConnectTo(OutputHandle);
-	if (Connectability.Connectable != Metasound::Frontend::FConnectability::EConnectable::Yes)
-	{
-		const FName InputType = InputHandle.GetInputType();
-		const FName OutputType = OutputHandle.GetOutputType();
-		return FPinConnectionResponse(CONNECT_RESPONSE_DISALLOW, FText::Format(
-			LOCTEXT("ConnectionTypeIncompatibleFormat", "Output pin of type '{0}' cannot be connected to input pin of type '{1}'"),
-			FText::FromName(OutputType),
-			FText::FromName(InputType)
-		));
-	}
+	TArray<Metasound::Frontend::FOutputHandle> OutputHandles = OutputNodeHandle->GetOutputsWithVertexName(OutputPin->GetName());
 
-	// Break existing connections on inputs only - multiple output connections are acceptable
-	if (InputPin->LinkedTo.Num() > 0)
+	if (ensure((InputHandles.Num() == 1) && (OutputHandles.Num() == 1)))
 	{
-		ECanCreateConnectionResponse ReplyBreakOutputs;
-		if (InputPin == PinA)
+		Metasound::Frontend::FInputHandle InputHandle = InputHandles[0];
+		Metasound::Frontend::FOutputHandle OutputHandle = OutputHandles[0];
+		// TODO: Implement YesWithConverterNode to provide conversion options
+		Metasound::Frontend::FConnectability Connectability = InputHandle->CanConnectTo(*OutputHandle);
+		if (Connectability.Connectable != Metasound::Frontend::FConnectability::EConnectable::Yes)
 		{
-			ReplyBreakOutputs = CONNECT_RESPONSE_BREAK_OTHERS_A;
+			const FName InputType = InputHandle->GetDataType();
+			const FName OutputType = OutputHandle->GetDataType();
+			return FPinConnectionResponse(CONNECT_RESPONSE_DISALLOW, FText::Format(
+				LOCTEXT("ConnectionTypeIncompatibleFormat", "Output pin of type '{0}' cannot be connected to input pin of type '{1}'"),
+				FText::FromName(OutputType),
+				FText::FromName(InputType)
+			));
 		}
-		else
+
+		// Break existing connections on inputs only - multiple output connections are acceptable
+		if (InputPin->LinkedTo.Num() > 0)
 		{
-			ReplyBreakOutputs = CONNECT_RESPONSE_BREAK_OTHERS_B;
+			ECanCreateConnectionResponse ReplyBreakOutputs;
+			if (InputPin == PinA)
+			{
+				ReplyBreakOutputs = CONNECT_RESPONSE_BREAK_OTHERS_A;
+			}
+			else
+			{
+				ReplyBreakOutputs = CONNECT_RESPONSE_BREAK_OTHERS_B;
+			}
+			return FPinConnectionResponse(ReplyBreakOutputs, LOCTEXT("ConnectionReplace", "Replace existing connections"));
 		}
-		return FPinConnectionResponse(ReplyBreakOutputs, LOCTEXT("ConnectionReplace", "Replace existing connections"));
+
+		return FPinConnectionResponse(CONNECT_RESPONSE_MAKE, TEXT(""));
 	}
 
-	return FPinConnectionResponse(CONNECT_RESPONSE_MAKE, TEXT(""));
+	return FPinConnectionResponse(CONNECT_RESPONSE_DISALLOW, LOCTEXT("ConnectionInternalError", "Internal error. Metasound node vertex handle mismatch."));
 }
 
 bool UMetasoundEditorGraphSchema::TryCreateConnection(UEdGraphPin* PinA, UEdGraphPin* PinB) const
@@ -394,17 +405,23 @@ bool UMetasoundEditorGraphSchema::TryCreateConnection(UEdGraphPin* PinA, UEdGrap
 	{
 		UMetasoundEditorGraphNode* InputGraphNode = CastChecked<UMetasoundEditorGraphNode>(InputPin->GetOwningNode());
 		Metasound::Frontend::FNodeHandle InputNodeHandle = InputGraphNode->GetNodeHandle();
-		Metasound::Frontend::FInputHandle InputHandle = InputNodeHandle.GetInputWithName(InputPin->GetName());
+		TArray<Metasound::Frontend::FInputHandle> InputHandles = InputNodeHandle->GetInputsWithVertexName(InputPin->GetName());
 
 		UMetasoundEditorGraphNode* OutputGraphNode = CastChecked<UMetasoundEditorGraphNode>(OutputPin->GetOwningNode());
 		Metasound::Frontend::FNodeHandle OutputNodeHandle = OutputGraphNode->GetNodeHandle();
-		Metasound::Frontend::FOutputHandle OutputHandle = OutputNodeHandle.GetOutputWithName(OutputPin->GetName());
+		TArray<Metasound::Frontend::FOutputHandle> OutputHandles = OutputNodeHandle->GetOutputsWithVertexName(OutputPin->GetName());
 
-		// TODO: Implement YesWithConverterNode with selected conversion option
-		if (!ensure(InputHandle.Connect(OutputHandle)))
+		if (ensure((InputHandles.Num() == 1) && (OutputHandles.Num() == 1)))
 		{
-			InputPin->BreakLinkTo(PinB);
-			return false;
+			Metasound::Frontend::FInputHandle InputHandle = InputHandles[0];
+			Metasound::Frontend::FOutputHandle OutputHandle = OutputHandles[0];
+
+			// TODO: Implement YesWithConverterNode with selected conversion option
+			if (!ensure(InputHandle->Connect(*OutputHandle)))
+			{
+				InputPin->BreakLinkTo(PinB);
+				return false;
+			}
 		}
 	}
 	return bModified;
@@ -424,28 +441,42 @@ FText UMetasoundEditorGraphSchema::GetPinDisplayName(const UEdGraphPin* Pin) con
 
 	UMetasoundEditorGraphNode* Node = CastChecked<UMetasoundEditorGraphNode>(Pin->GetOwningNode());
 	FNodeHandle NodeHandle = Node->GetNodeHandle();
-	const EMetasoundClassType ClassType = NodeHandle.GetClassInfo().NodeType;
+	const EMetasoundFrontendClassType ClassType = NodeHandle->GetClassType();
 
 	switch (ClassType)
 	{
-		case EMetasoundClassType::Input:
+		case EMetasoundFrontendClassType::Input:
 		{
-			FGraphHandle GraphHandle = Node->GetRootGraphHandle();
-			return GraphHandle.GetInputDisplayName(NodeHandle.GetNodeName());
+			TArray<FOutputHandle> OutputHandles = NodeHandle->GetOutputsWithVertexName(Pin->GetName());
+			if (ensure(OutputHandles.Num() > 0))
+			{
+				return OutputHandles[0]->GetDisplayName();
+			}
+			else
+			{
+				return Super::GetPinDisplayName(Pin);
+			}
 		}
 
-		case EMetasoundClassType::Output:
+		case EMetasoundFrontendClassType::Output:
 		{
-			FGraphHandle GraphHandle = Node->GetRootGraphHandle();
-			return GraphHandle.GetOutputDisplayName(NodeHandle.GetNodeName());
+			TArray<FInputHandle> InputHandles = NodeHandle->GetInputsWithVertexName(Pin->GetName());
+			if (ensure(InputHandles.Num() > 0))
+			{
+				return InputHandles[0]->GetDisplayName();
+			}
+			else
+			{
+				return Super::GetPinDisplayName(Pin);
+			}
 		}
 
-		case EMetasoundClassType::External:
-		case EMetasoundClassType::MetasoundGraph:
-		case EMetasoundClassType::Invalid:
+		case EMetasoundFrontendClassType::External:
+		case EMetasoundFrontendClassType::Graph:
+		case EMetasoundFrontendClassType::Invalid:
 		default:
 		{
-			static_assert(static_cast<int32>(EMetasoundClassType::Invalid) == 4, "Possible missing EMetasoundClassType case coverage");
+			static_assert(static_cast<int32>(EMetasoundFrontendClassType::Invalid) == 4, "Possible missing EMetasoundFrontendClassType case coverage");
 			return Super::GetPinDisplayName(Pin);
 		}
 	}
@@ -466,21 +497,21 @@ void UMetasoundEditorGraphSchema::BreakNodeLinks(UEdGraphNode& TargetNode) const
 	Super::BreakNodeLinks(TargetNode);
 
 	FNodeHandle NodeHandle = CastChecked<UMetasoundEditorGraphNode>(&TargetNode)->GetNodeHandle();
-	const uint32 NodeID = NodeHandle.GetNodeID();
+	const uint32 NodeID = NodeHandle->GetID();
 
 	FGraphHandle GraphHandle = CastChecked<UMetasoundEditorGraphNode>(&TargetNode)->GetRootGraphHandle();
-	TArray<FNodeHandle> AllNodes = GraphHandle.GetAllNodes();
+	TArray<FNodeHandle> AllNodes = GraphHandle->GetNodes();
 	for (Metasound::Frontend::FNodeHandle& IterNode : AllNodes)
 	{
-		if (NodeID != IterNode.GetNodeID())
+		if (NodeID != IterNode->GetID())
 		{
-			TArray<FInputHandle> Inputs = IterNode.GetAllInputs();
+			TArray<FInputHandle> Inputs = IterNode->GetInputs();
 			for (FInputHandle& Input : Inputs)
 			{
-				FOutputHandle Output = Input.GetCurrentlyConnectedOutput();
-				if (Output.GetOwningNodeID() == NodeID)
+				FOutputHandle Output = Input->GetCurrentlyConnectedOutput();
+				if (Output->GetOwningNodeID() == NodeID)
 				{
-					Input.Disconnect(Output);
+					Input->Disconnect(*Output);
 				}
 			}
 		}
@@ -493,11 +524,12 @@ void UMetasoundEditorGraphSchema::BreakPinLinks(UEdGraphPin& TargetPin, bool bSe
 
 	const FScopedTransaction Transaction(NSLOCTEXT("UnrealEd", "GraphEd_BreakPinLinks", "Break Pin Links"));
 
+	TArray<FInputHandle> InputHandles;
+
 	if (TargetPin.Direction == EGPD_Input)
 	{
 		FNodeHandle NodeHandle = CastChecked<UMetasoundEditorGraphNode>(TargetPin.GetOwningNode())->GetNodeHandle();
-		FInputHandle InputHandle = NodeHandle.GetInputWithName(TargetPin.GetName());
-		InputHandle.Disconnect();
+		InputHandles = NodeHandle->GetInputsWithVertexName(TargetPin.GetName());
 	}
 	else
 	{
@@ -505,11 +537,15 @@ void UMetasoundEditorGraphSchema::BreakPinLinks(UEdGraphPin& TargetPin, bool bSe
 		for (UEdGraphPin* Pin : TargetPin.LinkedTo)
 		{
 			FNodeHandle NodeHandle = CastChecked<UMetasoundEditorGraphNode>(Pin->GetOwningNode())->GetNodeHandle();
-			FInputHandle InputHandle = NodeHandle.GetInputWithName(Pin->GetName());
-			InputHandle.Disconnect();
+			InputHandles.Append(NodeHandle->GetInputsWithVertexName(Pin->GetName()));
 		}
 	}
 
+	for (FInputHandle Handle : InputHandles)
+	{
+		Handle->Disconnect();
+	}
+
 	Super::BreakPinLinks(TargetPin, bSendsNodeNotifcation);
 }
 
@@ -575,11 +611,11 @@ void UMetasoundEditorGraphSchema::GetAllMetasoundActions(FGraphActionMenuBuilder
 	const TArray<FNodeClassInfo> ClassInfos = GetAllAvailableNodeClasses();
 	for (const FNodeClassInfo& ClassInfo : ClassInfos)
 	{
-		const FMetasoundClassDescription ClassDescription = GenerateClassDescription(ClassInfo);
-		const FMetasoundClassMetadata Metadata = ClassDescription.Metadata;
-		const FText Tooltip = Metadata.AuthorName.IsEmpty()
-			? Metadata.MetasoundDescription
-			: FText::Format(LOCTEXT("MetasoundTooltipAuthorFormat", "{0}\nAuthor: {1}"), Metadata.MetasoundDescription, Metadata.AuthorName);
+		const FMetasoundFrontendClass ClassDescription = GenerateClassDescription(ClassInfo);
+		const FMetasoundFrontendClassMetadata Metadata = ClassDescription.Metadata;
+		const FText Tooltip = Metadata.Author.IsEmpty()
+			? Metadata.Description
+			: FText::Format(LOCTEXT("MetasoundTooltipAuthorFormat", "{0}\nAuthor: {1}"), Metadata.Description, Metadata.Author);
 
 		TSharedPtr<FMetasoundGraphSchemaAction_NewNode> NewNodeAction = MakeShared<FMetasoundGraphSchemaAction_NewNode>
 		(
diff --git a/Engine/Plugins/Runtime/Metasound/Source/MetasoundEditor/Private/MetasoundEditorModule.cpp b/Engine/Plugins/Runtime/Metasound/Source/MetasoundEditor/Private/MetasoundEditorModule.cpp
--- a/Engine/Plugins/Runtime/Metasound/Source/MetasoundEditor/Private/MetasoundEditorModule.cpp
+++ b/Engine/Plugins/Runtime/Metasound/Source/MetasoundEditor/Private/MetasoundEditorModule.cpp
@@ -109,37 +109,37 @@ namespace Metasound
 					FName PinType = NAME_None;
 					switch (RegistryInfo.PreferredLiteralType)
 					{
-						case ELiteralArgType::Boolean:
+						case ELiteralType::Boolean:
 						{
 							PinType = FGraphBuilder::PinPrimitiveBoolean;
 						}
 						break;
 
-						case ELiteralArgType::Float:
+						case ELiteralType::Float:
 						{
 							PinType = FGraphBuilder::PinPrimitiveFloat;
 						}
 						break;
 
-						case ELiteralArgType::Integer:
+						case ELiteralType::Integer:
 						{
 							PinType = FGraphBuilder::PinPrimitiveInteger;
 						}
 						break;
 
-						case ELiteralArgType::String:
+						case ELiteralType::String:
 						{
 							PinType = FGraphBuilder::PinPrimitiveString;
 						}
 						break;
 
-						case ELiteralArgType::UObjectProxy:
+						case ELiteralType::UObjectProxy:
 						{
 							PinType = FGraphBuilder::PinPrimitiveUObject;
 						}
 						break;
 
-						case ELiteralArgType::UObjectProxyArray:
+						case ELiteralType::UObjectProxyArray:
 						{
 							PinType = FGraphBuilder::PinPrimitiveUObjectArray;
 						}
@@ -147,10 +147,10 @@ namespace Metasound
 
 						// Register atypical primitives
 						default:
-						case ELiteralArgType::None:
-						case ELiteralArgType::Invalid:
+						case ELiteralType::None:
+						case ELiteralType::Invalid:
 						{
-							static_assert(static_cast<int32>(ELiteralArgType::Invalid) == 7, "Possible missing binding of pin category to primitive type");
+							static_assert(static_cast<int32>(ELiteralType::Invalid) == 7, "Possible missing binding of pin category to primitive type");
 						}
 						break;
 					}
@@ -177,8 +177,8 @@ namespace Metasound
 					FOnGetDetailCustomizationInstance::CreateLambda([]() { return MakeShared<FMetasoundDetailCustomization>(UMetasoundSource::GetDocumentPropertyName()); }));
 
 				PropertyModule.RegisterCustomPropertyTypeLayout(
-					"MetasoundLiteralDescription",
-					FOnGetPropertyTypeCustomizationInstance::CreateLambda([]() { return MakeShared<FMetasoundLiteralDescriptionDetailCustomization>(); }));
+					"MetasoundFrontendLiteral",
+					FOnGetPropertyTypeCustomizationInstance::CreateLambda([]() { return MakeShared<FMetasoundFrontendLiteralDetailCustomization>(); }));
 
 				StyleSet = MakeShared<FSlateStyle>();
 
diff --git a/Engine/Plugins/Runtime/Metasound/Source/MetasoundEditor/Private/MetasoundFactory.cpp b/Engine/Plugins/Runtime/Metasound/Source/MetasoundEditor/Private/MetasoundFactory.cpp
--- a/Engine/Plugins/Runtime/Metasound/Source/MetasoundEditor/Private/MetasoundFactory.cpp
+++ b/Engine/Plugins/Runtime/Metasound/Source/MetasoundEditor/Private/MetasoundFactory.cpp
@@ -21,13 +21,13 @@ UObject* UMetasoundFactory::FactoryCreateNew(UClass* InClass, UObject* InParent,
 {
 	UMetasound* NewMetasound = NewObject<UMetasound>(InParent, Name, Flags);
 
-	FMetasoundClassMetadata Metadata;
+	FMetasoundFrontendClassMetadata Metadata;
 
-	Metadata.NodeName = NewMetasound->GetName();
-	Metadata.MajorVersion = 1;
-	Metadata.MinorVersion = 0;
-	Metadata.NodeType = EMetasoundClassType::MetasoundGraph;
-	Metadata.AuthorName = FText::FromString(UKismetSystemLibrary::GetPlatformUserName());
+	Metadata.Name = FMetasoundFrontendClassName{TEXT(""), NewMetasound->GetName(), TEXT("")};
+	Metadata.Version.Major = 1;
+	Metadata.Version.Minor = 0;
+	Metadata.Type = EMetasoundFrontendClassType::Graph;
+	Metadata.Author = FText::FromString(UKismetSystemLibrary::GetPlatformUserName());
 
 	NewMetasound->SetMetadata(Metadata);
 
@@ -50,13 +50,13 @@ UObject* UMetasoundSourceFactory::FactoryCreateNew(UClass* InClass, UObject* InP
 {
 	UMetasoundSource* MetasoundSource = NewObject<UMetasoundSource>(InParent, Name, Flags);
 
-	FMetasoundClassMetadata Metadata;
+	FMetasoundFrontendClassMetadata Metadata;
 
-	Metadata.NodeName = MetasoundSource->GetName();
-	Metadata.MajorVersion = 1;
-	Metadata.MinorVersion = 0;
-	Metadata.NodeType = EMetasoundClassType::MetasoundGraph;
-	Metadata.AuthorName = FText::FromString(UKismetSystemLibrary::GetPlatformUserName());
+	Metadata.Name = FMetasoundFrontendClassName{TEXT(""), MetasoundSource->GetName(), TEXT("")};
+	Metadata.Version.Major = 1;
+	Metadata.Version.Minor = 0;
+	Metadata.Type = EMetasoundFrontendClassType::Graph;
+	Metadata.Author = FText::FromString(UKismetSystemLibrary::GetPlatformUserName());
 
 	MetasoundSource->SetMetadata(Metadata);
 
diff --git a/Engine/Plugins/Runtime/Metasound/Source/MetasoundEditor/Private/MetasoundLiteralDescriptionDetailCustomization.cpp b/Engine/Plugins/Runtime/Metasound/Source/MetasoundEditor/Private/MetasoundLiteralDescriptionDetailCustomization.cpp
--- a/Engine/Plugins/Runtime/Metasound/Source/MetasoundEditor/Private/MetasoundLiteralDescriptionDetailCustomization.cpp
+++ b/Engine/Plugins/Runtime/Metasound/Source/MetasoundEditor/Private/MetasoundLiteralDescriptionDetailCustomization.cpp
@@ -9,8 +9,8 @@
 #include "IDetailChildrenBuilder.h"
 #include "IDetailGroup.h"
 #include "Layout/Visibility.h"
-#include "MetasoundFrontendDataLayout.h"
 #include "MetasoundFrontend.h"
+#include "MetasoundFrontendDocument.h"
 #include "MetasoundFrontendRegistries.h"
 #include "Misc/AssertionMacros.h"
 #include "Misc/Attribute.h"
@@ -27,62 +27,62 @@ namespace Metasound
 {
 	namespace Editor
 	{
-		TSharedPtr<IPropertyHandle> GetLiteralHandleForType(TSharedRef<IPropertyHandle> StructPropertyHandle, EMetasoundLiteralType LiteralType)
+		TSharedPtr<IPropertyHandle> GetLiteralHandleForType(TSharedRef<IPropertyHandle> StructPropertyHandle, EMetasoundFrontendLiteralType LiteralType)
 		{
 			switch (LiteralType)
 			{
-			case EMetasoundLiteralType::Bool:
+			case EMetasoundFrontendLiteralType::Bool:
 			{
-				return StructPropertyHandle->GetChildHandle(GET_MEMBER_NAME_CHECKED(FMetasoundLiteralDescription, AsBool));
+				return StructPropertyHandle->GetChildHandle(GET_MEMBER_NAME_CHECKED(FMetasoundFrontendLiteral, AsBool));
 			}
 
-			case EMetasoundLiteralType::Float:
+			case EMetasoundFrontendLiteralType::Float:
 			{
-				return StructPropertyHandle->GetChildHandle(GET_MEMBER_NAME_CHECKED(FMetasoundLiteralDescription, AsFloat));
+				return StructPropertyHandle->GetChildHandle(GET_MEMBER_NAME_CHECKED(FMetasoundFrontendLiteral, AsFloat));
 			}
 
-			case EMetasoundLiteralType::Integer:
+			case EMetasoundFrontendLiteralType::Integer:
 			{
-				return StructPropertyHandle->GetChildHandle(GET_MEMBER_NAME_CHECKED(FMetasoundLiteralDescription, AsInteger));
+				return StructPropertyHandle->GetChildHandle(GET_MEMBER_NAME_CHECKED(FMetasoundFrontendLiteral, AsInteger));
 			}
 
-			case EMetasoundLiteralType::String:
+			case EMetasoundFrontendLiteralType::String:
 			{
-				return StructPropertyHandle->GetChildHandle(GET_MEMBER_NAME_CHECKED(FMetasoundLiteralDescription, AsString));
+				return StructPropertyHandle->GetChildHandle(GET_MEMBER_NAME_CHECKED(FMetasoundFrontendLiteral, AsString));
 			}
 
-			case EMetasoundLiteralType::UObject:
+			case EMetasoundFrontendLiteralType::UObject:
 			{
-				return StructPropertyHandle->GetChildHandle(GET_MEMBER_NAME_CHECKED(FMetasoundLiteralDescription, AsUObject));
+				return StructPropertyHandle->GetChildHandle(GET_MEMBER_NAME_CHECKED(FMetasoundFrontendLiteral, AsUObject));
 			}
 
-			case EMetasoundLiteralType::UObjectArray:
+			case EMetasoundFrontendLiteralType::UObjectArray:
 			{
-				return StructPropertyHandle->GetChildHandle(GET_MEMBER_NAME_CHECKED(FMetasoundLiteralDescription, AsUObjectArray));
+				return StructPropertyHandle->GetChildHandle(GET_MEMBER_NAME_CHECKED(FMetasoundFrontendLiteral, AsUObjectArray));
 			}
 
-			case EMetasoundLiteralType::None:
+			case EMetasoundFrontendLiteralType::None:
 			default:
-				static_assert(static_cast<int32>(EMetasoundLiteralType::Invalid) == 7, "Possible missing case coverage for EMetasoundLiteralType");
+				static_assert(static_cast<int32>(EMetasoundFrontendLiteralType::Invalid) == 7, "Possible missing case coverage for EMetasoundFrontendLiteralType");
 				return TSharedPtr<IPropertyHandle>();
 			}
 		}
 
 		TSharedPtr<IPropertyHandle> GetActiveLiteralHandle(TSharedRef<IPropertyHandle> StructPropertyHandle, TSharedRef<IPropertyHandle> TypeEnumHandle)
 		{
-			uint8 Value = static_cast<uint8>(EMetasoundLiteralType::None);
+			uint8 Value = static_cast<uint8>(EMetasoundFrontendLiteralType::None);
 			if (TypeEnumHandle->GetValue(Value) == FPropertyAccess::Result::Success)
 			{
-				const EMetasoundLiteralType LiteralType = static_cast<EMetasoundLiteralType>(Value);
+				const EMetasoundFrontendLiteralType LiteralType = static_cast<EMetasoundFrontendLiteralType>(Value);
 				return GetLiteralHandleForType(StructPropertyHandle, LiteralType);
 			}
 
 			return TSharedPtr<IPropertyHandle>();
 		}
 
-		bool IsLiteralTypeActive(TSharedRef<IPropertyHandle> TypeEnumHandle, EMetasoundLiteralType LiteralType)
+		bool IsLiteralTypeActive(TSharedRef<IPropertyHandle> TypeEnumHandle, EMetasoundFrontendLiteralType LiteralType)
 		{
-			uint8 Value = static_cast<uint8>(EMetasoundLiteralType::None);
+			uint8 Value = static_cast<uint8>(EMetasoundFrontendLiteralType::None);
 			TypeEnumHandle->GetValue(Value);
 
 			return Value == static_cast<uint8>(LiteralType);
@@ -91,26 +91,26 @@ namespace Metasound
 } // namespace Metasound
 
 
-void FMetasoundLiteralDescriptionDetailCustomization::CustomizeHeader(TSharedRef<IPropertyHandle> StructPropertyHandle, FDetailWidgetRow& HeaderRow, IPropertyTypeCustomizationUtils& StructCustomizationUtils)
+void FMetasoundFrontendLiteralDetailCustomization::CustomizeHeader(TSharedRef<IPropertyHandle> StructPropertyHandle, FDetailWidgetRow& HeaderRow, IPropertyTypeCustomizationUtils& StructCustomizationUtils)
 {
 }
 
-void FMetasoundLiteralDescriptionDetailCustomization::CustomizeChildren(TSharedRef<IPropertyHandle> StructPropertyHandle, IDetailChildrenBuilder& ChildBuilder, IPropertyTypeCustomizationUtils& StructCustomizationUtils)
+void FMetasoundFrontendLiteralDetailCustomization::CustomizeChildren(TSharedRef<IPropertyHandle> StructPropertyHandle, IDetailChildrenBuilder& ChildBuilder, IPropertyTypeCustomizationUtils& StructCustomizationUtils)
 {
-	TSharedRef<IPropertyHandle> TypeEnumHandle = StructPropertyHandle->GetChildHandle(GET_MEMBER_NAME_CHECKED(FMetasoundLiteralDescription, LiteralType)).ToSharedRef();
+	TSharedRef<IPropertyHandle> TypeEnumHandle = StructPropertyHandle->GetChildHandle(GET_MEMBER_NAME_CHECKED(FMetasoundFrontendLiteral, Type)).ToSharedRef();
 
 
 	FName OwningDataTypeName;
 	UClass* OwningDataTypeClass = nullptr;
-	EMetasoundLiteralType PreferredLiteralType = EMetasoundLiteralType::None;
+	EMetasoundFrontendLiteralType PreferredLiteralType = EMetasoundFrontendLiteralType::None;
 
 	// Grab the preferred uclass for the owning data type.
-	TSharedPtr<IPropertyHandle> InputDescription = StructPropertyHandle->GetParentHandle();
+	TSharedPtr<IPropertyHandle> InputDescription = StructPropertyHandle->GetParentHandle()->GetParentHandle()->GetParentHandle();
 
 	// TODO: Ensure that StructPropertyHandle is an FMetasoundInputDescription.
 	if (InputDescription)
 	{
-		TSharedPtr<IPropertyHandle> DataTypeNameHandle = InputDescription->GetChildHandle(GET_MEMBER_NAME_CHECKED(FMetasoundInputDescription, TypeName));
+		TSharedPtr<IPropertyHandle> DataTypeNameHandle = InputDescription->GetChildHandle(GET_MEMBER_NAME_CHECKED(FMetasoundFrontendClassInput, TypeName));
 		if (DataTypeNameHandle)
 		{
 			
@@ -126,7 +126,7 @@ void FMetasoundLiteralDescriptionDetailCustomization::CustomizeChildren(TSharedR
 
 	struct FTypeEditorInfo
 	{
-		EMetasoundLiteralType LiteralType;
+		EMetasoundFrontendLiteralType LiteralType;
 		TSharedRef<IPropertyHandle> TypeEnumHandle;
 		TSharedPtr<IPropertyHandle> PropertyHandle;
 		FText DisplayName;
@@ -134,7 +134,7 @@ void FMetasoundLiteralDescriptionDetailCustomization::CustomizeChildren(TSharedR
 
 		TSharedRef<SWidget> CreatePropertyValueWidget(bool bDisplayDefaultPropertyButtons = true) const
 		{
-			if (LiteralType == EMetasoundLiteralType::UObject)
+			if (LiteralType == EMetasoundFrontendLiteralType::UObject)
 			{
 				auto ValidateAsset = [CachedDataTypeUClass = OwningDataTypeUClass](const UObject* InObject, FText& OutReason)  -> bool
 				{
@@ -153,11 +153,11 @@ void FMetasoundLiteralDescriptionDetailCustomization::CustomizeChildren(TSharedR
 					// If this asset selector was set to no asset, clear out the literal description.
 					if (InObject)
 					{
-						CachedEnumHandle->SetValue(static_cast<uint8>(EMetasoundLiteralType::UObject));
+						CachedEnumHandle->SetValue(static_cast<uint8>(EMetasoundFrontendLiteralType::UObject));
 					}
 					else
 					{
-						CachedEnumHandle->SetValue(static_cast<uint8>(EMetasoundLiteralType::None));
+						CachedEnumHandle->SetValue(static_cast<uint8>(EMetasoundFrontendLiteralType::None));
 					}
 				};
 
@@ -182,7 +182,7 @@ void FMetasoundLiteralDescriptionDetailCustomization::CustomizeChildren(TSharedR
 					.DisplayThumbnail(true)
 					.NewAssetFactories(FactoriesToUse);
 			}
-			else if (LiteralType == EMetasoundLiteralType::UObjectArray)
+			else if (LiteralType == EMetasoundFrontendLiteralType::UObjectArray)
 			{
 				// TODO: Implement.
 				return SNullWidget::NullWidget;
@@ -194,59 +194,59 @@ void FMetasoundLiteralDescriptionDetailCustomization::CustomizeChildren(TSharedR
 		}
 	};
 
-	static_assert(static_cast<int32>(EMetasoundLiteralType::Invalid) == 7, "Possible missing property descriptor for literal customization display");
+	static_assert(static_cast<int32>(EMetasoundFrontendLiteralType::Invalid) == 7, "Possible missing property descriptor for literal customization display");
 	const TArray<FTypeEditorInfo> TypePropertyInfo =
 	{
 		FTypeEditorInfo
 		{
-			EMetasoundLiteralType::Bool,
+			EMetasoundFrontendLiteralType::Bool,
 			TypeEnumHandle,
-			StructPropertyHandle->GetChildHandle(GET_MEMBER_NAME_CHECKED(FMetasoundLiteralDescription, AsBool)),
+			StructPropertyHandle->GetChildHandle(GET_MEMBER_NAME_CHECKED(FMetasoundFrontendLiteral, AsBool)),
 			LOCTEXT("Metasound_LiteralDisplayNameBool", "Bool"),
 			nullptr
 		},
 
 		FTypeEditorInfo
 		{
-			EMetasoundLiteralType::Integer,
+			EMetasoundFrontendLiteralType::Integer,
 			TypeEnumHandle,
-			StructPropertyHandle->GetChildHandle(GET_MEMBER_NAME_CHECKED(FMetasoundLiteralDescription, AsInteger)),
+			StructPropertyHandle->GetChildHandle(GET_MEMBER_NAME_CHECKED(FMetasoundFrontendLiteral, AsInteger)),
 			LOCTEXT("Metasound_LiteralDisplayNameInteger", "Int32"),
 			nullptr
 		},
 
 		FTypeEditorInfo
 		{
-			EMetasoundLiteralType::Float,
+			EMetasoundFrontendLiteralType::Float,
 			TypeEnumHandle,
-			StructPropertyHandle->GetChildHandle(GET_MEMBER_NAME_CHECKED(FMetasoundLiteralDescription, AsFloat)),
+			StructPropertyHandle->GetChildHandle(GET_MEMBER_NAME_CHECKED(FMetasoundFrontendLiteral, AsFloat)),
 			LOCTEXT("Metasound_LiteralDisplayNameFloat", "Float"),
 			nullptr
 		},
 
 		FTypeEditorInfo
 		{
-			EMetasoundLiteralType::String,
+			EMetasoundFrontendLiteralType::String,
 			TypeEnumHandle,
-			StructPropertyHandle->GetChildHandle(GET_MEMBER_NAME_CHECKED(FMetasoundLiteralDescription, AsString)),
+			StructPropertyHandle->GetChildHandle(GET_MEMBER_NAME_CHECKED(FMetasoundFrontendLiteral, AsString)),
 			LOCTEXT("Metasound_LiteralDisplayNameString", "String"),
 			nullptr
 		},
 
 		FTypeEditorInfo
 		{
-			EMetasoundLiteralType::UObject,
+			EMetasoundFrontendLiteralType::UObject,
 			TypeEnumHandle,
-			StructPropertyHandle->GetChildHandle(GET_MEMBER_NAME_CHECKED(FMetasoundLiteralDescription, AsUObject)),
+			StructPropertyHandle->GetChildHandle(GET_MEMBER_NAME_CHECKED(FMetasoundFrontendLiteral, AsUObject)),
 			LOCTEXT("Metasound_LiteralDisplayNameUObject", "UObject"),
 			OwningDataTypeClass
 		},
 
 		FTypeEditorInfo
 		{
-			EMetasoundLiteralType::UObjectArray,
+			EMetasoundFrontendLiteralType::UObjectArray,
 			TypeEnumHandle,
-			StructPropertyHandle->GetChildHandle(GET_MEMBER_NAME_CHECKED(FMetasoundLiteralDescription, AsUObjectArray)),
+			StructPropertyHandle->GetChildHandle(GET_MEMBER_NAME_CHECKED(FMetasoundFrontendLiteral, AsUObjectArray)),
 			LOCTEXT("Metasound_LiteralDisplayNameUObjectArray", "UObjectArray"),
 			OwningDataTypeClass
 		},
@@ -297,4 +297,4 @@ void FMetasoundLiteralDescriptionDetailCustomization::CustomizeChildren(TSharedR
 			}));
 	}
 }
-#undef LOCTEXT_NAMESPACE // MetasoundEditor
\ No newline at end of file
+#undef LOCTEXT_NAMESPACE // MetasoundEditor
diff --git a/Engine/Plugins/Runtime/Metasound/Source/MetasoundEditor/Private/MetasoundLiteralDescriptionDetailCustomization.h b/Engine/Plugins/Runtime/Metasound/Source/MetasoundEditor/Private/MetasoundLiteralDescriptionDetailCustomization.h
--- a/Engine/Plugins/Runtime/Metasound/Source/MetasoundEditor/Private/MetasoundLiteralDescriptionDetailCustomization.h
+++ b/Engine/Plugins/Runtime/Metasound/Source/MetasoundEditor/Private/MetasoundLiteralDescriptionDetailCustomization.h
@@ -12,11 +12,11 @@ class FDetailWidgetRow;
 class IPropertyTypeCustomizationUtils;
 
 
-class FMetasoundLiteralDescriptionDetailCustomization : public IPropertyTypeCustomization
+class FMetasoundFrontendLiteralDetailCustomization : public IPropertyTypeCustomization
 {
 public:
 	//~ Begin IPropertyTypeCustomization
 	virtual void CustomizeHeader(TSharedRef<IPropertyHandle> StructPropertyHandle, FDetailWidgetRow& HeaderRow, IPropertyTypeCustomizationUtils& StructCustomizationUtils) override;
 	virtual void CustomizeChildren(TSharedRef<IPropertyHandle> StructPropertyHandle, IDetailChildrenBuilder& ChildBuilder, IPropertyTypeCustomizationUtils& StructCustomizationUtils) override;
 	//~ End IPropertyTypeCustomization
-};
\ No newline at end of file
+};
diff --git a/Engine/Plugins/Runtime/Metasound/Source/MetasoundEditor/Public/IMetasoundEditor.h b/Engine/Plugins/Runtime/Metasound/Source/MetasoundEditor/Public/IMetasoundEditor.h
--- a/Engine/Plugins/Runtime/Metasound/Source/MetasoundEditor/Public/IMetasoundEditor.h
+++ b/Engine/Plugins/Runtime/Metasound/Source/MetasoundEditor/Public/IMetasoundEditor.h
@@ -15,9 +15,9 @@ class IMetasoundEditor : public FAssetEditorToolkit
 	virtual void SetSelection(const TArray<UObject*>& SelectedObjects) = 0;
 	virtual bool GetBoundsForSelectedNodes(FSlateRect& Rect, float Padding) = 0;
 
-	// Rebuilds all nodes and connections from root FMetasoundDocument. Results in loss of node layout.
+	// Rebuilds all nodes and connections from root FMetasoundFrontendDocument. Results in loss of node layout.
 	virtual bool RebuildGraph() const = 0;
 
-	// Synchronizes the nodes and connections of the graph to the root FMetasoundDocument.
+	// Synchronizes the nodes and connections of the graph to the root FMetasoundFrontendDocument.
 	virtual bool SynchronizeGraph() const = 0;
 };
diff --git a/Engine/Plugins/Runtime/Metasound/Source/MetasoundEditor/Public/MetasoundEditorGraph.h b/Engine/Plugins/Runtime/Metasound/Source/MetasoundEditor/Public/MetasoundEditorGraph.h
--- a/Engine/Plugins/Runtime/Metasound/Source/MetasoundEditor/Public/MetasoundEditorGraph.h
+++ b/Engine/Plugins/Runtime/Metasound/Source/MetasoundEditor/Public/MetasoundEditorGraph.h
@@ -9,7 +9,7 @@
 #include "MetasoundEditorGraph.generated.h"
 
 // Forward Declarations
-struct FMetasoundDocument;
+struct FMetasoundFrontendDocument;
 
 
 UCLASS(MinimalAPI)
diff --git a/Engine/Plugins/Runtime/Metasound/Source/MetasoundEditor/Public/MetasoundEditorGraphNode.h b/Engine/Plugins/Runtime/Metasound/Source/MetasoundEditor/Public/MetasoundEditorGraphNode.h
--- a/Engine/Plugins/Runtime/Metasound/Source/MetasoundEditor/Public/MetasoundEditorGraphNode.h
+++ b/Engine/Plugins/Runtime/Metasound/Source/MetasoundEditor/Public/MetasoundEditorGraphNode.h
@@ -6,7 +6,7 @@
 #include "EdGraph/EdGraphNode.h"
 #include "UObject/ObjectMacros.h"
 #include "MetasoundFrontend.h"
-#include "MetasoundFrontendDataLayout.h"
+#include "MetasoundFrontendController.h"
 #include "MetasoundFrontendRegistries.h"
 
 #include "MetasoundEditorGraphNode.generated.h"
diff --git a/Engine/Plugins/Runtime/Metasound/Source/MetasoundEngine/Private/Metasound.cpp b/Engine/Plugins/Runtime/Metasound/Source/MetasoundEngine/Private/Metasound.cpp
--- a/Engine/Plugins/Runtime/Metasound/Source/MetasoundEngine/Private/Metasound.cpp
+++ b/Engine/Plugins/Runtime/Metasound/Source/MetasoundEngine/Private/Metasound.cpp
@@ -8,7 +8,7 @@
 
 UMetasound::UMetasound(const FObjectInitializer& ObjectInitializer)
 	: Super(ObjectInitializer)
-	, FMetasoundAssetBase(FMetasoundArchetype())
+	, FMetasoundAssetBase(FMetasoundFrontendArchetype())
 {
 }
 
@@ -42,31 +42,31 @@ void UMetasound::SetGraph(UEdGraph* InGraph)
 #endif // WITH_EDITORONLY_DATA
 
 // Returns document object responsible for serializing asset
-Metasound::Frontend::TAccessPtr<FMetasoundDocument> UMetasound::GetDocument()
+Metasound::Frontend::TAccessPtr<FMetasoundFrontendDocument> UMetasound::GetDocument()
 {
-	return Metasound::Frontend::MakeAccessPtr(MetasoundDocumentAccessPoint, MetasoundDocument);
+	return Metasound::Frontend::MakeAccessPtr(MetasoundDocument.AccessPoint, MetasoundDocument);
 }
 
 // Returns document object responsible for serializing asset
-Metasound::Frontend::TAccessPtr<const FMetasoundDocument> UMetasound::GetDocument() const
+Metasound::Frontend::TAccessPtr<const FMetasoundFrontendDocument> UMetasound::GetDocument() const
 {
-	return Metasound::Frontend::MakeAccessPtr<const FMetasoundDocument>(MetasoundDocumentAccessPoint, MetasoundDocument);
+	return Metasound::Frontend::MakeAccessPtr<const FMetasoundFrontendDocument>(MetasoundDocument.AccessPoint, MetasoundDocument);
 }
 
-const TArray<FMetasoundArchetype>& UMetasound::GetPreferredArchetypes() const
+const TArray<FMetasoundFrontendArchetype>& UMetasound::GetPreferredArchetypes() const
 {
 	// Not preferred archetypes for a basic UMetasound.
-	static const TArray<FMetasoundArchetype> Preferred;
+	static const TArray<FMetasoundFrontendArchetype> Preferred;
 	return Preferred;
 }
 
-bool UMetasound::IsArchetypeSupported(const FMetasoundArchetype& InArchetype) const
+bool UMetasound::IsArchetypeSupported(const FMetasoundFrontendArchetype& InArchetype) const
 {
 	// All archetypes are supported.
 	return true;
 }
 
-const FMetasoundArchetype& UMetasound::GetPreferredArchetype(const FMetasoundDocument& InDocument) const
+const FMetasoundFrontendArchetype& UMetasound::GetPreferredArchetype(const FMetasoundFrontendDocument& InDocument) const
 {
 	// Prefer to keep original archetype.
 	return InDocument.Archetype;
@@ -74,7 +74,7 @@ const FMetasoundArchetype& UMetasound::GetPreferredArchetype(const FMetasoundDoc
 
 // This can be used to update the metadata (name, author, etc) for this metasound.
 // @param InMetadata may be updated with any corrections we do to the input metadata.
-void UMetasound::SetMetadata(FMetasoundClassMetadata& InMetadata)
+void UMetasound::SetMetadata(FMetasoundFrontendClassMetadata& InMetadata)
 {
 	FMetasoundAssetBase::SetMetadata(InMetadata);
 	MarkPackageDirty();
diff --git a/Engine/Plugins/Runtime/Metasound/Source/MetasoundEngine/Private/MetasoundEngineModule.cpp b/Engine/Plugins/Runtime/Metasound/Source/MetasoundEngine/Private/MetasoundEngineModule.cpp
--- a/Engine/Plugins/Runtime/Metasound/Source/MetasoundEngine/Private/MetasoundEngineModule.cpp
+++ b/Engine/Plugins/Runtime/Metasound/Source/MetasoundEngine/Private/MetasoundEngineModule.cpp
@@ -11,7 +11,7 @@
 
 DEFINE_LOG_CATEGORY(LogMetasoundEngine);
 
-REGISTER_METASOUND_DATATYPE(Metasound::FWaveAsset, "Primitive:WaveAsset", Metasound::ELiteralArgType::UObjectProxy, USoundWave);
+REGISTER_METASOUND_DATATYPE(Metasound::FWaveAsset, "Primitive:WaveAsset", Metasound::ELiteralType::UObjectProxy, USoundWave);
 
 class FMetasoundEngineModule : public IMetasoundEngineModule
 {
diff --git a/Engine/Plugins/Runtime/Metasound/Source/MetasoundEngine/Private/MetasoundSource.cpp b/Engine/Plugins/Runtime/Metasound/Source/MetasoundEngine/Private/MetasoundSource.cpp
--- a/Engine/Plugins/Runtime/Metasound/Source/MetasoundEngine/Private/MetasoundSource.cpp
+++ b/Engine/Plugins/Runtime/Metasound/Source/MetasoundEngine/Private/MetasoundSource.cpp
@@ -36,9 +36,9 @@ UMetasoundSource::UMetasoundSource(const FObjectInitializer& ObjectInitializer)
 
 }
 
-void UMetasoundSource::SetMetadata(FMetasoundClassMetadata& InMetadata)
+void UMetasoundSource::SetMetadata(FMetasoundFrontendClassMetadata& InMetadata)
 {
-	RootMetasoundDocument.RootClass.Metadata = InMetadata;
+	RootMetasoundDocument.RootGraph.Metadata = InMetadata;
 }
 
 #if WITH_EDITOR
@@ -120,9 +120,9 @@ ISoundGeneratorPtr UMetasoundSource::CreateSoundGenerator(const FSoundGeneratorI
 	TArray<IOperatorBuilder::FBuildErrorPtr> BuildErrors;
 
 	Frontend::FGraphHandle RootGraph = GetRootGraphHandle();
-	ensureAlways(RootGraph.IsValid());
+	ensureAlways(RootGraph->IsValid());
 
-	TUniquePtr<IOperator> Operator = RootGraph.BuildOperator(InSettings, Environment, BuildErrors);
+	TUniquePtr<IOperator> Operator = RootGraph->BuildOperator(InSettings, Environment, BuildErrors);
 
 	if (!Operator.IsValid())
 	{
@@ -161,7 +161,7 @@ ISoundGeneratorPtr UMetasoundSource::CreateSoundGenerator(const FSoundGeneratorI
 
 		// References must be cached before moving the operator to the InitParams
 		FDataReferenceCollection Inputs = Operator->GetInputs();
-		FTriggerWriteRef PlayTrigger = Inputs.GetDataWriteReference<FBop>(GetOnPlayInputName());
+		FTriggerWriteRef PlayTrigger = Inputs.GetDataWriteReferenceOrConstruct<FBop>(GetOnPlayInputName(), InSettings, false);
 		FTriggerReadRef FinishTrigger = Outputs.GetDataReadReferenceOrConstruct<FBop>(GetIsFinishedOutputName(), InSettings, false);
 
 		// Create the FMetasoundGenerator.
@@ -186,9 +186,10 @@ void UMetasoundSource::PostLoad()
 	ConformDocumentToArchetype();
 }
 
-const TArray<FMetasoundArchetype>& UMetasoundSource::GetPreferredArchetypes() const 
+const TArray<FMetasoundFrontendArchetype>& UMetasoundSource::GetPreferredArchetypes() const 
 {
-	static const TArray<FMetasoundArchetype> Preferred({GetMonoSourceArchetype(), GetStereoSourceArchetype()});
+	static const TArray<FMetasoundFrontendArchetype> Preferred({GetMonoSourceArchetype(), GetStereoSourceArchetype()});
+
 	return Preferred;
 }
 
@@ -216,85 +217,89 @@ const FString& UMetasoundSource::GetAudioDeviceHandleVariableName()
 	return AudioDeviceHandleVarName;
 }
 
-const FMetasoundArchetype& UMetasoundSource::GetBaseArchetype()
+const FMetasoundFrontendArchetype& UMetasoundSource::GetBaseArchetype()
 {
-	auto CreateBaseArchetype = []() -> FMetasoundArchetype
+	auto CreateBaseArchetype = []() -> FMetasoundFrontendArchetype
 	{
-		FMetasoundArchetype Archetype;
+		FMetasoundFrontendArchetype Archetype;
 		
-		FMetasoundInputDescription OnPlayTrigger;
+		FMetasoundFrontendClassVertex OnPlayTrigger;
 		OnPlayTrigger.Name = UMetasoundSource::GetOnPlayInputName();
-		OnPlayTrigger.DisplayName = FText::FromString(OnPlayTrigger.Name);
+		OnPlayTrigger.Metadata.DisplayName = FText::FromString(OnPlayTrigger.Name);
 		OnPlayTrigger.TypeName = Metasound::Frontend::GetDataTypeName<Metasound::FBop>();
-		OnPlayTrigger.ToolTip = LOCTEXT("OnPlayTriggerToolTip", "Trigger executed when this source is played.");
+		OnPlayTrigger.Metadata.Description = LOCTEXT("OnPlayTriggerToolTip", "Trigger executed when this source is played.");
+		OnPlayTrigger.PointIDs.Add(2);
 
-		Archetype.RequiredInputs.Add(OnPlayTrigger);
+		Archetype.Interface.Inputs.Add(OnPlayTrigger);
 
-		FMetasoundOutputDescription OnFinished;
+		FMetasoundFrontendClassVertex OnFinished;
 		OnFinished.Name = UMetasoundSource::GetIsFinishedOutputName();
-		OnFinished.DisplayName = FText::FromString(OnFinished.Name);
+		OnFinished.Metadata.DisplayName = FText::FromString(OnFinished.Name);
 		OnFinished.TypeName = Metasound::Frontend::GetDataTypeName<Metasound::FBop>();
-		OnFinished.ToolTip = LOCTEXT("OnFinishedToolTip", "Bop executed to initiate stopping the source.");
+		OnFinished.Metadata.Description = LOCTEXT("OnFinishedToolTip", "Bop executed to initiate stopping the source.");
+		OnFinished.PointIDs.Add(3);
 
-		Archetype.RequiredOutputs.Add(OnFinished);
+		Archetype.Interface.Outputs.Add(OnFinished);
 
-		FMetasoundEnvironmentVariableDescription AudioDeviceHandle;
+		FMetasoundFrontendEnvironmentVariable AudioDeviceHandle;
 		AudioDeviceHandle.Name = UMetasoundSource::GetAudioDeviceHandleVariableName();
-		AudioDeviceHandle.DisplayName = FText::FromString(AudioDeviceHandle.Name);
-		AudioDeviceHandle.ToolTip = LOCTEXT("AudioDeviceHandleToolTip", "Audio device handle");
+		AudioDeviceHandle.Metadata.DisplayName = FText::FromString(AudioDeviceHandle.Name);
+		AudioDeviceHandle.Metadata.Description = LOCTEXT("AudioDeviceHandleToolTip", "Audio device handle");
 
-		Archetype.EnvironmentVariables.Add(AudioDeviceHandle);
+		Archetype.Interface.Environment.Add(AudioDeviceHandle);
 
 		return Archetype;
 	};
 
-	static const FMetasoundArchetype BaseArchetype = CreateBaseArchetype();
+	static const FMetasoundFrontendArchetype BaseArchetype = CreateBaseArchetype();
 
 	return BaseArchetype;
 }
 
-const FMetasoundArchetype& UMetasoundSource::GetMonoSourceArchetype()
+const FMetasoundFrontendArchetype& UMetasoundSource::GetMonoSourceArchetype()
 {
-	auto CreateMonoArchetype = []() -> FMetasoundArchetype
+	auto CreateMonoArchetype = []() -> FMetasoundFrontendArchetype
 	{
-		FMetasoundArchetype Archetype = GetBaseArchetype();
-		Archetype.ArchetypeName = FName(TEXT("MonoSource"));
+		FMetasoundFrontendArchetype Archetype = GetBaseArchetype();
+		Archetype.Name = FName(TEXT("MonoSource"));
 
-		FMetasoundOutputDescription GeneratedAudio;
+		FMetasoundFrontendClassVertex GeneratedAudio;
 		GeneratedAudio.Name = UMetasoundSource::GetAudioOutputName();
-		GeneratedAudio.DisplayName = LOCTEXT("GeneratedMono", "Mono");
 		GeneratedAudio.TypeName = Metasound::Frontend::GetDataTypeName<Metasound::FMonoAudioFormat>();
-		GeneratedAudio.ToolTip = LOCTEXT("GeneratedAudioToolTip", "The resulting output audio from this source.");
+		GeneratedAudio.Metadata.DisplayName = LOCTEXT("GeneratedMono", "Mono");
+		GeneratedAudio.Metadata.Description = LOCTEXT("GeneratedAudioToolTip", "The resulting output audio from this source.");
+		GeneratedAudio.PointIDs.Add(4);
 
-		Archetype.RequiredOutputs.Add(GeneratedAudio);
+		Archetype.Interface.Outputs.Add(GeneratedAudio);
 
 		return Archetype;
 	};
 
-	static const FMetasoundArchetype MonoArchetype = CreateMonoArchetype();
+	static const FMetasoundFrontendArchetype MonoArchetype = CreateMonoArchetype();
 
 	return MonoArchetype;
 }
 
-const FMetasoundArchetype& UMetasoundSource::GetStereoSourceArchetype()
+const FMetasoundFrontendArchetype& UMetasoundSource::GetStereoSourceArchetype()
 {
-	auto CreateStereoArchetype = []() -> FMetasoundArchetype
+	auto CreateStereoArchetype = []() -> FMetasoundFrontendArchetype
 	{
-		FMetasoundArchetype Archetype = GetBaseArchetype();
-		Archetype.ArchetypeName = FName(TEXT("StereoSource"));
+		FMetasoundFrontendArchetype Archetype = GetBaseArchetype();
+		Archetype.Name = FName(TEXT("StereoSource"));
 
-		FMetasoundOutputDescription GeneratedAudio;
+		FMetasoundFrontendClassVertex GeneratedAudio;
 		GeneratedAudio.Name = UMetasoundSource::GetAudioOutputName();
-		GeneratedAudio.DisplayName = LOCTEXT("GeneratedStereo", "Stereo");
 		GeneratedAudio.TypeName = Metasound::Frontend::GetDataTypeName<Metasound::FStereoAudioFormat>();
-		GeneratedAudio.ToolTip = LOCTEXT("GeneratedAudioToolTip", "The resulting output audio from this source.");
+		GeneratedAudio.Metadata.DisplayName = LOCTEXT("GeneratedStereo", "Stereo");
+		GeneratedAudio.Metadata.Description = LOCTEXT("GeneratedAudioToolTip", "The resulting output audio from this source.");
+		GeneratedAudio.PointIDs.Add(5);
 
-		Archetype.RequiredOutputs.Add(GeneratedAudio);
+		Archetype.Interface.Outputs.Add(GeneratedAudio);
 
 		return Archetype;
 	};
 
-	static const FMetasoundArchetype StereoArchetype = CreateStereoArchetype();
+	static const FMetasoundFrontendArchetype StereoArchetype = CreateStereoArchetype();
 
 	return StereoArchetype;
 }
diff --git a/Engine/Plugins/Runtime/Metasound/Source/MetasoundEngine/Private/MetasoundUObjectRegistry.cpp b/Engine/Plugins/Runtime/Metasound/Source/MetasoundEngine/Private/MetasoundUObjectRegistry.cpp
--- a/Engine/Plugins/Runtime/Metasound/Source/MetasoundEngine/Private/MetasoundUObjectRegistry.cpp
+++ b/Engine/Plugins/Runtime/Metasound/Source/MetasoundEngine/Private/MetasoundUObjectRegistry.cpp
@@ -5,6 +5,7 @@
 #include "Algo/Copy.h"
 #include "AssetRegistry/AssetRegistryModule.h"
 #include "Async/Async.h"
+#include "Async/TaskGraphInterfaces.h"
 #include "CoreMinimal.h"
 #include "UObject/Object.h"
 
@@ -46,7 +47,7 @@ namespace Metasound
 				return Classes;
 			}
 
-			UObject* NewObject(UClass* InClass, const FMetasoundDocument& InDocument, const FString& InPath) const override
+			UObject* NewObject(UClass* InClass, const FMetasoundFrontendDocument& InDocument, const FString& InPath) const override
 			{
 				auto IsChildClassOfRegisteredClass = [&](const IMetasoundUObjectRegistryEntry* Entry)
 				{
@@ -57,7 +58,7 @@ namespace Metasound
 
 				for (const IMetasoundUObjectRegistryEntry* Entry : EntriesForClass)
 				{
-					if (Entry->GetArchetypeName() == InDocument.Archetype.ArchetypeName)
+					if (Entry->GetArchetypeName() == InDocument.Archetype.Name)
 					{
 						return NewObject(*Entry, InDocument, InPath);
 					}
@@ -91,7 +92,7 @@ namespace Metasound
 			}
 
 		private:
-			UObject* NewObject(const IMetasoundUObjectRegistryEntry& InEntry, const FMetasoundDocument& InDocument, const FString& InPath) const
+			UObject* NewObject(const IMetasoundUObjectRegistryEntry& InEntry, const FMetasoundFrontendDocument& InDocument, const FString& InPath) const
 			{
 				UPackage* PackageToSaveTo = nullptr;
 
@@ -114,7 +115,7 @@ namespace Metasound
 					PackageToSaveTo = GetTransientPackage();
 				}
 
-				UObject* NewMetasoundObject = InEntry.NewObject(PackageToSaveTo, *InDocument.RootClass.Metadata.NodeName);
+				UObject* NewMetasoundObject = InEntry.NewObject(PackageToSaveTo, *InDocument.RootGraph.Metadata.Name.GetFullName());
 				FMetasoundAssetBase* NewAssetBase = InEntry.Cast(NewMetasoundObject);
 				if (ensure(nullptr != NewAssetBase))
 				{
diff --git a/Engine/Plugins/Runtime/Metasound/Source/MetasoundEngine/Public/Metasound.h b/Engine/Plugins/Runtime/Metasound/Source/MetasoundEngine/Public/Metasound.h
--- a/Engine/Plugins/Runtime/Metasound/Source/MetasoundEngine/Public/Metasound.h
+++ b/Engine/Plugins/Runtime/Metasound/Source/MetasoundEngine/Public/Metasound.h
@@ -3,7 +3,7 @@
 
 #include "CoreMinimal.h"
 #include "MetasoundFrontend.h"
-#include "MetasoundFrontendDataLayout.h"
+#include "MetasoundFrontendDocument.h"
 #include "MetasoundAssetBase.h"
 
 
@@ -30,15 +30,13 @@ class METASOUNDENGINE_API UMetasound : public UObject, public FMetasoundAssetBas
 
 protected:
 	UPROPERTY(EditAnywhere, Category = CustomView)
-	FMetasoundDocument MetasoundDocument;
-
-	Metasound::Frontend::FAccessPoint MetasoundDocumentAccessPoint;
+	FMetasoundFrontendDocument MetasoundDocument;
 
 	// Returns document object responsible for serializing asset
-	Metasound::Frontend::TAccessPtr<FMetasoundDocument> GetDocument() override;
+	Metasound::Frontend::TAccessPtr<FMetasoundFrontendDocument> GetDocument() override;
 
 	// Returns document object responsible for serializing asset
-	Metasound::Frontend::TAccessPtr<const FMetasoundDocument> GetDocument() const override;
+	Metasound::Frontend::TAccessPtr<const FMetasoundFrontendDocument> GetDocument() const override;
 
 #if WITH_EDITORONLY_DATA
 	UPROPERTY()
@@ -77,15 +75,15 @@ public:
 		return const_cast<UObject*>(CastChecked<const UObject>(this));
 	}
 
-	const TArray<FMetasoundArchetype>& GetPreferredArchetypes() const override;
+	const TArray<FMetasoundFrontendArchetype>& GetPreferredArchetypes() const override;
 
-	bool IsArchetypeSupported(const FMetasoundArchetype& InArchetype) const override;
+	bool IsArchetypeSupported(const FMetasoundFrontendArchetype& InArchetype) const override;
 
-	const FMetasoundArchetype& GetPreferredArchetype(const FMetasoundDocument& InDocument) const override;
+	const FMetasoundFrontendArchetype& GetPreferredArchetype(const FMetasoundFrontendDocument& InDocument) const override;
 
 	// Updates the Metasound's metadata (name, author, etc).
 	// @param InMetadata Metadata containing corrections to the class metadata.
-	void SetMetadata(FMetasoundClassMetadata& InMetadata) override;
+	void SetMetadata(FMetasoundFrontendClassMetadata& InMetadata) override;
 
 	void PostLoad() override;
 
