Metasound Frontend Revamp

diff --git a/Engine/Plugins/Runtime/Metasound/Source/MetasoundEngine/Public/MetasoundSource.h b/Engine/Plugins/Runtime/Metasound/Source/MetasoundEngine/Public/MetasoundSource.h
--- a/Engine/Plugins/Runtime/Metasound/Source/MetasoundEngine/Public/MetasoundSource.h
+++ b/Engine/Plugins/Runtime/Metasound/Source/MetasoundEngine/Public/MetasoundSource.h
@@ -3,7 +3,7 @@
 
 #include "CoreMinimal.h"
 #include "MetasoundFrontend.h"
-#include "MetasoundFrontendDataLayout.h"
+#include "MetasoundFrontendDocument.h"
 #include "MetasoundAssetBase.h"
 
 #include "Sound/SoundWaveProcedural.h"
@@ -34,14 +34,7 @@ class METASOUNDENGINE_API UMetasoundSource : public USoundWaveProcedural, public
 
 protected:
 	UPROPERTY(EditAnywhere, Category = CustomView)
-	FMetasoundDocument RootMetasoundDocument;
-
-	// MetasoundSourceAccessPoint allows the RootMetasoundDocument to be safely
-	// accessed via the TAccessPtr. The lifetime of MetasoundSourceAccessPoint is
-	// tied to the lifetime of RootMetasoundDocument by their coexistence as members
-	// of the same class. When MetasoundSourceAccessPoint is destructed, the documents
-	// access ptr becomes invalid. 
-	Metasound::Frontend::FAccessPoint MetasoundSourceAccessPoint;
+	FMetasoundFrontendDocument RootMetasoundDocument;
 
 #if WITH_EDITORONLY_DATA
 	UPROPERTY()
@@ -112,7 +105,7 @@ public:
 
 	// Updates the Metasound's metadata (name, author, etc).
 	// @param InMetadata Metadata containing corrections to the class metadata.
-	void SetMetadata(FMetasoundClassMetadata& InMetadata) override;
+	void SetMetadata(FMetasoundFrontendClassMetadata& InMetadata) override;
 
 	bool IsPlayable() const override;
 	bool SupportsSubtitles() const override;
@@ -121,22 +114,22 @@ public:
 	void PostLoad() override;
 
 	// Get the most up to date archetype for metasound sources.
-	const TArray<FMetasoundArchetype>& GetPreferredArchetypes() const override;
+	const TArray<FMetasoundFrontendArchetype>& GetPreferredArchetypes() const override;
 
 protected:
 
-	Metasound::Frontend::TAccessPtr<FMetasoundDocument> GetDocument() override
+	Metasound::Frontend::TAccessPtr<FMetasoundFrontendDocument> GetDocument() override
 	{
 		// Return document using FAccessPoint to inform the TAccessPtr when the 
 		// object is no longer valid.
-		return Metasound::Frontend::MakeAccessPtr<FMetasoundDocument>(MetasoundSourceAccessPoint, RootMetasoundDocument);
+		return Metasound::Frontend::MakeAccessPtr<FMetasoundFrontendDocument>(RootMetasoundDocument.AccessPoint, RootMetasoundDocument);
 	}
 
-	Metasound::Frontend::TAccessPtr<const FMetasoundDocument> GetDocument() const override
+	Metasound::Frontend::TAccessPtr<const FMetasoundFrontendDocument> GetDocument() const override
 	{
 		// Return document using FAccessPoint to inform the TAccessPtr when the 
 		// object is no longer valid.
-		return Metasound::Frontend::MakeAccessPtr<const FMetasoundDocument>(MetasoundSourceAccessPoint, RootMetasoundDocument);
+		return Metasound::Frontend::MakeAccessPtr<const FMetasoundFrontendDocument>(RootMetasoundDocument.AccessPoint, RootMetasoundDocument);
 	}
 
 private:
@@ -145,7 +138,7 @@ private:
 	static const FString& GetAudioOutputName();
 	static const FString& GetIsFinishedOutputName();
 	static const FString& GetAudioDeviceHandleVariableName();
-	static const FMetasoundArchetype& GetBaseArchetype();
-	static const FMetasoundArchetype& GetMonoSourceArchetype();
-	static const FMetasoundArchetype& GetStereoSourceArchetype();
+	static const FMetasoundFrontendArchetype& GetBaseArchetype();
+	static const FMetasoundFrontendArchetype& GetMonoSourceArchetype();
+	static const FMetasoundFrontendArchetype& GetStereoSourceArchetype();
 };
diff --git a/Engine/Plugins/Runtime/Metasound/Source/MetasoundEngine/Public/MetasoundUObjectRegistry.h b/Engine/Plugins/Runtime/Metasound/Source/MetasoundEngine/Public/MetasoundUObjectRegistry.h
--- a/Engine/Plugins/Runtime/Metasound/Source/MetasoundEngine/Public/MetasoundUObjectRegistry.h
+++ b/Engine/Plugins/Runtime/Metasound/Source/MetasoundEngine/Public/MetasoundUObjectRegistry.h
@@ -10,7 +10,7 @@ namespace Metasound
 {
 	/** Interface for an entry into the Metasound-UObject Registry. 
 	 *
-	 * An entry provides information linking a FMetasoundArchetype to a UClass.
+	 * An entry provides information linking a FMetasoundFrontendArchetype to a UClass.
 	 * It also provides methods for accessing the FMetasoundAssetBase from a UObject.
 	 */
 	class IMetasoundUObjectRegistryEntry
@@ -126,7 +126,7 @@ namespace Metasound
 	/** IMetaoundUObjectRegistry contains IMetasoundUObjectRegistryEntrys. 
 	 *
 	 * Registered UObject classes can utilize the Metasound Editor. It also enables
-	 * the creation of a UObject directly from a FMetasoundDocument.
+	 * the creation of a UObject directly from a FMetasoundFrontendDocument.
 	 */
 	class METASOUNDENGINE_API IMetasoundUObjectRegistry
 	{
@@ -145,17 +145,17 @@ namespace Metasound
 			{
 				static_assert(std::is_base_of<FMetasoundAssetBase, UClassType>::value, "UClass must be derived from FMetasoundAssetBase");
 
-				const TArray<FMetasoundArchetype>& PreferredArchetypes = GetDefault<UClassType>()->GetPreferredArchetypes();
+				const TArray<FMetasoundFrontendArchetype>& PreferredArchetypes = GetDefault<UClassType>()->GetPreferredArchetypes();
 
-				for (const FMetasoundArchetype& Arch : PreferredArchetypes)
+				for (const FMetasoundFrontendArchetype& Arch : PreferredArchetypes)
 				{
-					IMetasoundUObjectRegistry::RegisterUClassArchetype<UClassType>(Arch.ArchetypeName);
+					IMetasoundUObjectRegistry::RegisterUClassArchetype<UClassType>(Arch.Name);
 				}
 			}
 
 			/** Register an archetype for the UClass. 
 			 *
-			 * @param InArchetypeName - An FMetasoundDocument to associate with the UClass.
+			 * @param InArchetypeName - An FMetasoundFrontendDocument to associate with the UClass.
 			 */
 			template<typename UClassType>
 			static void RegisterUClassArchetype(const FName& InArchetypeName)
@@ -174,12 +174,12 @@ namespace Metasound
 			/** Creates a new object from a metasound doucment.
 			 *
 			 * @param InClass - A registered UClass to create.
-			 * @param InDocucument - The FMetasoundDocument to use when creating the class.
+			 * @param InDocucument - The FMetasoundFrontendDocument to use when creating the class.
 			 * @param InPath - If in editor, the created asset will be stored at this content path.
 			 *
 			 * @return A new object. A nullptr on error.
 			 */
-			virtual UObject* NewObject(UClass* InClass, const FMetasoundDocument& InDocument, const FString& InPath) const = 0;
+			virtual UObject* NewObject(UClass* InClass, const FMetasoundFrontendDocument& InDocument, const FString& InPath) const = 0;
 
 			/** Returns true if the InObject is of a class or child class which is registered. */
 			virtual bool IsRegisteredClass(UObject* InObject) const = 0;
diff --git a/Engine/Plugins/Runtime/Metasound/Source/MetasoundFrontend/Private/MetasoundArchetype.cpp b/Engine/Plugins/Runtime/Metasound/Source/MetasoundFrontend/Private/MetasoundArchetype.cpp
--- a/Engine/Plugins/Runtime/Metasound/Source/MetasoundFrontend/Private/MetasoundArchetype.cpp
+++ b/Engine/Plugins/Runtime/Metasound/Source/MetasoundFrontend/Private/MetasoundArchetype.cpp
@@ -1,47 +1,43 @@
 // Copyright Epic Games, Inc. All Rights Reserved.
 
+
 #include "MetasoundArchetype.h"
 
+#include "Algo/ForEach.h"
+#include "MetasoundFrontendDocument.h"
+
 namespace Metasound
 {
 	namespace Frontend
 	{
 		namespace MetasoundArchetypeIntrinsics
 		{
-
-			bool IsLessThanInputDescription(const FMetasoundInputDescription& InputA, const FMetasoundInputDescription& InputB)
+			bool IsLessThanVertex(const FMetasoundFrontendVertex& VertexA, const FMetasoundFrontendVertex& VertexB)
 			{
-				// For archetypes we can ignore display name, tooltip and literalvalue.
-				if (InputA.Name == InputB.Name)
+				// For archetypes we can ignore metadata when comparing
+				if (VertexA.Name == VertexB.Name)
 				{
-					return InputA.TypeName.FastLess(InputB.TypeName);
+					return VertexA.TypeName.FastLess(VertexB.TypeName);
 				}
-				return InputA.Name < InputB.Name;
+				return VertexA.Name < VertexB.Name;
 			}
 
-			bool IsLessThanOutputDescription(const FMetasoundOutputDescription& OutputA, const FMetasoundOutputDescription& OutputB)
+			bool IsLessThanEnvironmentVariable(const FMetasoundFrontendEnvironmentVariable& EnvironmentA, const FMetasoundFrontendEnvironmentVariable& EnvironmentB)
 			{
-				// For archetypes we can ignore display name and tooltip
-				if (OutputA.Name == OutputB.Name)
+				// For archetypes we can ignore metadata when comparing
+				if (EnvironmentA.Name == EnvironmentB.Name)
 				{
-					return OutputA.TypeName.FastLess(OutputB.TypeName);
+					return EnvironmentA.TypeName.FastLess(EnvironmentB.TypeName);
 				}
-				return OutputA.Name < OutputB.Name;
+				return EnvironmentA.Name < EnvironmentB.Name;
 			}
 
-			bool IsLessThanEnvironmentDescription(const FMetasoundEnvironmentVariableDescription& EnvironmentA, const FMetasoundEnvironmentVariableDescription& EnvironmentB)
+			bool IsEquivalentEnvironmentVariable(const FMetasoundFrontendEnvironmentVariable& EnvironmentA, const FMetasoundFrontendEnvironmentVariable& EnvironmentB)
 			{
 				// For archetypes we can ignore display name and tooltip
-				return (EnvironmentA.Name < EnvironmentB.Name);
+				return (EnvironmentA.Name == EnvironmentB.Name) && (EnvironmentA.TypeName == EnvironmentB.TypeName);
 			}
 
-			bool IsEquivalentEnvironmentDescription(const FMetasoundEnvironmentVariableDescription& EnvironmentA, const FMetasoundEnvironmentVariableDescription& EnvironmentB)
-			{
-				// For archetypes we can ignore display name and tooltip
-				return (EnvironmentA.Name == EnvironmentB.Name);
-			}
-
-
 
 			template<typename Type>
 			TArray<const Type*> MakePointerToArrayElements(const TArray<Type>& InArray)
@@ -54,13 +50,13 @@ namespace Metasound
 				return PointerArray;
 			}
 
-			template<typename DescriptionType, typename ComparisonType>
-			void DescriptionSetDifference(const TArray<DescriptionType>& InDescriptionsA, const TArray<DescriptionType>& InDescriptionsB, TArray<const DescriptionType*>& OutUniqueDescriptionsA, TArray<const DescriptionType*>& OutUniqueDescriptionsB, ComparisonType Compare)
+			template<typename FrontendTypeA, typename FrontendTypeB, typename ComparisonType>
+			void SetDifference(const TArray<FrontendTypeA>& InSetA, const TArray<FrontendTypeB>& InSetB, TArray<const FrontendTypeA*>& OutUniqueA, TArray<const FrontendTypeB*>& OutUniqueB, ComparisonType Compare)
 			{
-				TArray<const DescriptionType*> SortedA = MakePointerToArrayElements(InDescriptionsA);
+				TArray<const FrontendTypeA*> SortedA = MakePointerToArrayElements(InSetA);
 				SortedA.Sort(Compare);
 
-				TArray<const DescriptionType*> SortedB = MakePointerToArrayElements(InDescriptionsB);
+				TArray<const FrontendTypeB*> SortedB = MakePointerToArrayElements(InSetB);
 				SortedB.Sort(Compare);
 
 				int32 IndexA = 0;
@@ -70,12 +66,12 @@ namespace Metasound
 				{
 					if (Compare(*SortedA[IndexA], *SortedB[IndexB]))
 					{
-						OutUniqueDescriptionsA.Add(SortedA[IndexA]);
+						OutUniqueA.Add(SortedA[IndexA]);
 						IndexA++;
 					}
 					else if (Compare(*SortedB[IndexB], *SortedA[IndexA]))
 					{
-						OutUniqueDescriptionsB.Add(SortedB[IndexB]);
+						OutUniqueB.Add(SortedB[IndexB]);
 						IndexB++;
 					}
 					else
@@ -87,71 +83,71 @@ namespace Metasound
 				}
 			}
 
-			template<typename DescriptionType, typename ComparisonType>
-			bool IsDescriptionSetEquivalent(const TArray<DescriptionType>& InDescriptionsA, const TArray<DescriptionType>& InDescriptionsB, ComparisonType Compare)
+			template<typename FrontendTypeA, typename FrontendTypeB, typename ComparisonType>
+			bool IsSetEquivalent(const TArray<FrontendTypeA>& InSetA, const TArray<FrontendTypeB>& InSetB, ComparisonType Compare)
 			{
-				TArray<const DescriptionType*> UniqueA;
-				TArray<const DescriptionType*> UniqueB;
+				TArray<const FrontendTypeA*> UniqueA;
+				TArray<const FrontendTypeB*> UniqueB;
 
-				DescriptionSetDifference(InDescriptionsA, InDescriptionsB, UniqueA, UniqueB, Compare);
+				SetDifference(InSetA, InSetB, UniqueA, UniqueB, Compare);
 
 				return (UniqueA.Num() == 0) && (UniqueB.Num() == 0);
 			}
 			
-			template<typename DescriptionType, typename ComparisonType>
-			bool IsDescriptionSetIncluded(const TArray<DescriptionType>& InSubset, const TArray<DescriptionType>& InSuperset, ComparisonType Compare)
+			template<typename FrontendTypeA, typename FrontendTypeB, typename ComparisonType>
+			bool IsSetIncluded(const TArray<FrontendTypeA>& InSubset, const TArray<FrontendTypeB>& InSuperset, ComparisonType Compare)
 			{
-				TArray<const DescriptionType*> UniqueSubset;
-				TArray<const DescriptionType*> UniqueSuperset;
+				TArray<const FrontendTypeA*> UniqueSubset;
+				TArray<const FrontendTypeB*> UniqueSuperset;
 
-				DescriptionSetDifference(InSubset, InSuperset, UniqueSubset, UniqueSuperset, Compare);
+				SetDifference(InSubset, InSuperset, UniqueSubset, UniqueSuperset, Compare);
 
 				return UniqueSubset.Num() == 0;
 			}
 
-			template<typename DescriptionType, typename ComparisonType>
-			int32 DescriptionSetDifferenceCount(const TArray<DescriptionType>& InDescriptionsA, const TArray<DescriptionType>& InDescriptionsB, ComparisonType Compare)
+			template<typename FrontendTypeA, typename FrontendTypeB, typename ComparisonType>
+			int32 SetDifferenceCount(const TArray<FrontendTypeA>& InSetA, const TArray<FrontendTypeB>& InSetB, ComparisonType Compare)
 			{
-				TArray<const DescriptionType*> UniqueA;
-				TArray<const DescriptionType*> UniqueB;
+				TArray<const FrontendTypeA*> UniqueA;
+				TArray<const FrontendTypeB*> UniqueB;
 
-				DescriptionSetDifference(InDescriptionsA, InDescriptionsB, UniqueA, UniqueB, Compare);
+				SetDifference(InSetA, InSetB, UniqueA, UniqueB, Compare);
 
 				return UniqueA.Num() + UniqueB.Num();
 			}
 		}
 
-		bool IsSubsetOfArchetype(const FMetasoundArchetype& InSubsetArchetype, const FMetasoundArchetype& InSupersetArchetype)
+		bool IsSubsetOfArchetype(const FMetasoundFrontendArchetype& InSubsetArchetype, const FMetasoundFrontendArchetype& InSupersetArchetype)
 		{
 			using namespace MetasoundArchetypeIntrinsics;
 
-			const bool bIsInputSetIncluded = IsDescriptionSetIncluded(InSubsetArchetype.RequiredInputs, InSupersetArchetype.RequiredInputs, &IsLessThanInputDescription);
+			const bool bIsInputSetIncluded = IsSetIncluded(InSubsetArchetype.Interface.Inputs, InSupersetArchetype.Interface.Inputs, &IsLessThanVertex);
 
-			const bool bIsOutputSetIncluded = IsDescriptionSetIncluded(InSubsetArchetype.RequiredOutputs, InSupersetArchetype.RequiredOutputs, &IsLessThanOutputDescription);
+			const bool bIsOutputSetIncluded = IsSetIncluded(InSubsetArchetype.Interface.Outputs, InSupersetArchetype.Interface.Outputs, &IsLessThanVertex);
 
-			const bool bIsEnvironmentVariableSetIncluded = IsDescriptionSetIncluded(InSubsetArchetype.EnvironmentVariables, InSupersetArchetype.EnvironmentVariables, &IsLessThanEnvironmentDescription);
+			const bool bIsEnvironmentVariableSetIncluded = IsSetIncluded(InSubsetArchetype.Interface.Environment, InSupersetArchetype.Interface.Environment, &IsLessThanEnvironmentVariable);
 
 			return bIsInputSetIncluded && bIsOutputSetIncluded && bIsEnvironmentVariableSetIncluded;
 		}
 
-		bool IsEquivalentArchetype(const FMetasoundArchetype& InInputArchetype, const FMetasoundArchetype& InTargetArchetype)
+		bool IsEquivalentArchetype(const FMetasoundFrontendArchetype& InInputArchetype, const FMetasoundFrontendArchetype& InTargetArchetype)
 		{
 			using namespace MetasoundArchetypeIntrinsics;
 
-			const bool bIsInputSetEquivalent = IsDescriptionSetEquivalent(InTargetArchetype.RequiredInputs, InInputArchetype.RequiredInputs, &IsLessThanInputDescription);
+			const bool bIsInputSetEquivalent = IsSetEquivalent(InTargetArchetype.Interface.Inputs, InInputArchetype.Interface.Inputs, &IsLessThanVertex);
 
-			const bool bIsOutputSetEquivalent = IsDescriptionSetEquivalent(InTargetArchetype.RequiredOutputs, InInputArchetype.RequiredOutputs, &IsLessThanOutputDescription);
+			const bool bIsOutputSetEquivalent = IsSetEquivalent(InTargetArchetype.Interface.Outputs, InInputArchetype.Interface.Outputs, &IsLessThanVertex);
 
-			const bool bIsEnvironmentVariableSetEquivalent = IsDescriptionSetEquivalent(InTargetArchetype.EnvironmentVariables, InInputArchetype.EnvironmentVariables, &IsLessThanEnvironmentDescription);
+			const bool bIsEnvironmentVariableSetEquivalent = IsSetEquivalent(InTargetArchetype.Interface.Environment, InInputArchetype.Interface.Environment, &IsLessThanEnvironmentVariable);
 
 			return bIsInputSetEquivalent && bIsOutputSetEquivalent && bIsEnvironmentVariableSetEquivalent;
 		}
 
-		bool IsEqualArchetype(const FMetasoundArchetype& InInputArchetype, const FMetasoundArchetype& InTargetArchetype)
+		bool IsEqualArchetype(const FMetasoundFrontendArchetype& InInputArchetype, const FMetasoundFrontendArchetype& InTargetArchetype)
 		{
-			bool bIsMetadataEqual = InInputArchetype.ArchetypeName == InTargetArchetype.ArchetypeName;
-			bIsMetadataEqual &= InInputArchetype.MajorVersion == InTargetArchetype.MajorVersion;
-			bIsMetadataEqual &= InInputArchetype.MinorVersion == InTargetArchetype.MinorVersion;
+			bool bIsMetadataEqual = InInputArchetype.Name == InTargetArchetype.Name;
+			bIsMetadataEqual &= InInputArchetype.Version.Major == InTargetArchetype.Version.Major;
+			bIsMetadataEqual &= InInputArchetype.Version.Minor== InTargetArchetype.Version.Minor;
 
 			if (bIsMetadataEqual)
 			{
@@ -163,73 +159,78 @@ namespace Metasound
 			return false;
 		}
 
-		int32 InputOutputDifferenceCount(const FMetasoundDocument& InDocument, const FMetasoundArchetype& InArchetype)
+		int32 InputOutputDifferenceCount(const FMetasoundFrontendDocument& InDocument, const FMetasoundFrontendArchetype& InArchetype)
 		{
 			using namespace MetasoundArchetypeIntrinsics;
 
-			int32 DiffCount = DescriptionSetDifferenceCount(InDocument.RootClass.Inputs, InArchetype.RequiredInputs, &IsLessThanInputDescription);
+			int32 DiffCount = SetDifferenceCount(InDocument.RootGraph.Interface.Inputs, InArchetype.Interface.Inputs, &IsLessThanVertex);
 
-			DiffCount += DescriptionSetDifferenceCount(InDocument.RootClass.Outputs, InArchetype.RequiredOutputs, &IsLessThanOutputDescription);
+			DiffCount += SetDifferenceCount(InDocument.RootGraph.Interface.Outputs, InArchetype.Interface.Outputs, &IsLessThanVertex);
 
 			return DiffCount;
 		}
 
-		int32 InputOutputDifferenceCount(const FMetasoundArchetype& InArchetypeA, const FMetasoundArchetype& InArchetypeB)
+		int32 InputOutputDifferenceCount(const FMetasoundFrontendArchetype& InArchetypeA, const FMetasoundFrontendArchetype& InArchetypeB)
 		{
 			using namespace MetasoundArchetypeIntrinsics;
 
-			int32 DiffCount = DescriptionSetDifferenceCount(InArchetypeA.RequiredInputs, InArchetypeB.RequiredInputs, &IsLessThanInputDescription);
+			int32 DiffCount = SetDifferenceCount(InArchetypeA.Interface.Inputs, InArchetypeB.Interface.Inputs, &IsLessThanVertex);
 
-			DiffCount += DescriptionSetDifferenceCount(InArchetypeA.RequiredOutputs, InArchetypeB.RequiredOutputs, &IsLessThanOutputDescription);
+			DiffCount += SetDifferenceCount(InArchetypeA.Interface.Outputs, InArchetypeB.Interface.Outputs, &IsLessThanVertex);
 
 			return DiffCount;
 		}
 
-		void GatherRequiredEnvironmentVariables(const FMetasoundDocument& InDocument, TArray<FMetasoundEnvironmentVariableDescription>& OutEnvironmentVariables)
+		void GatherRequiredEnvironmentVariables(const FMetasoundFrontendDocument& InDocument, TArray<FMetasoundFrontendEnvironmentVariable>& OutEnvironmentVariables)
 		{
-			using namespace MetasoundArchetypeIntrinsics;
-
-			OutEnvironmentVariables = InDocument.RootClass.EnvironmentVariables;
-
-			for (const FMetasoundClassDescription& MetasoundClass : InDocument.Dependencies)
+			auto GatherRequiredEnvironmentVariablesFromClass = [&](const FMetasoundFrontendClass& InClass)
 			{
-				for (const FMetasoundEnvironmentVariableDescription& EnvVar : MetasoundClass.EnvironmentVariables)
-				{
-					auto IsEquivalentEnvVar = [&](const FMetasoundEnvironmentVariableDescription& OtherEnvVar) 
-					{ 
-						return IsEquivalentEnvironmentDescription(EnvVar, OtherEnvVar); 
-					};
+				using namespace MetasoundArchetypeIntrinsics;
 
-					// Basically same as TArray::AddUnique except uses the `IsEquivalentEnvironmentDescription` instead of `operator==` to test for uniqueness.
-					if (nullptr == OutEnvironmentVariables.FindByPredicate(IsEquivalentEnvVar))
+				for (const FMetasoundFrontendClassEnvironmentVariable& EnvVar : InClass.Interface.Environment)
+				{
+					if (EnvVar.bIsRequired)
 					{
-						OutEnvironmentVariables.Add(EnvVar);
+						auto IsEquivalentEnvVar = [&](const FMetasoundFrontendEnvironmentVariable& OtherEnvVar) 
+						{ 
+							return IsEquivalentEnvironmentVariable(EnvVar, OtherEnvVar); 
+						};
+
+						// Basically same as TArray::AddUnique except uses the `IsEquivalentEnvironmentVariable` instead of `operator==` to test for uniqueness.
+						if (nullptr == OutEnvironmentVariables.FindByPredicate(IsEquivalentEnvVar))
+						{
+							OutEnvironmentVariables.Add(EnvVar);
+						}
 					}
 				}
-			}
+			};
+
+			GatherRequiredEnvironmentVariablesFromClass(InDocument.RootGraph);
+			Algo::ForEach(InDocument.Dependencies, GatherRequiredEnvironmentVariablesFromClass);
+			Algo::ForEach(InDocument.Subgraphs, GatherRequiredEnvironmentVariablesFromClass);
 		}
 
-		const FMetasoundArchetype* FindMostSimilarArchetypeSupportingEnvironment(const FMetasoundDocument& InDocument, const TArray<FMetasoundArchetype>& InCandidateArchetypes)
+		const FMetasoundFrontendArchetype* FindMostSimilarArchetypeSupportingEnvironment(const FMetasoundFrontendDocument& InDocument, const TArray<FMetasoundFrontendArchetype>& InCandidateArchetypes)
 		{
 			using namespace MetasoundArchetypeIntrinsics;
 
-			const FMetasoundArchetype* Result = nullptr;
+			const FMetasoundFrontendArchetype* Result = nullptr;
 
-			TArray<const FMetasoundArchetype*> CandidateArchs = MakePointerToArrayElements(InCandidateArchetypes);
+			TArray<const FMetasoundFrontendArchetype*> CandidateArchs = MakePointerToArrayElements(InCandidateArchetypes);
 
-			TArray<FMetasoundEnvironmentVariableDescription> AllDocumentEnvironmentVariables;
+			TArray<FMetasoundFrontendEnvironmentVariable> AllDocumentEnvironmentVariables;
 			GatherRequiredEnvironmentVariables(InDocument, AllDocumentEnvironmentVariables);
 
 			// Remove all archetypes which do not provide all environment variables. 
-			auto DoesNotProvideAllEnvironmentVariables = [&](const FMetasoundArchetype* CandidateArch)
+			auto DoesNotProvideAllEnvironmentVariables = [&](const FMetasoundFrontendArchetype* CandidateArch)
 			{
-				return !IsDescriptionSetIncluded(AllDocumentEnvironmentVariables, CandidateArch->EnvironmentVariables, &IsLessThanEnvironmentDescription);
+				return !IsSetIncluded(AllDocumentEnvironmentVariables, CandidateArch->Interface.Environment, &IsLessThanEnvironmentVariable);
 			};
 
 			CandidateArchs.RemoveAll(DoesNotProvideAllEnvironmentVariables);
 
 			// Return the archetype with the least amount of differences.
-			auto DifferencesFromDocument = [&](const FMetasoundArchetype* CandidateArch) -> int32
+			auto DifferencesFromDocument = [&](const FMetasoundFrontendArchetype* CandidateArch) -> int32
 			{ 
 				return InputOutputDifferenceCount(InDocument, *CandidateArch); 
 			};
diff --git a/Engine/Plugins/Runtime/Metasound/Source/MetasoundFrontend/Private/MetasoundAssetBase.cpp b/Engine/Plugins/Runtime/Metasound/Source/MetasoundFrontend/Private/MetasoundAssetBase.cpp
--- a/Engine/Plugins/Runtime/Metasound/Source/MetasoundFrontend/Private/MetasoundAssetBase.cpp
+++ b/Engine/Plugins/Runtime/Metasound/Source/MetasoundFrontend/Private/MetasoundAssetBase.cpp
@@ -6,7 +6,10 @@
 #include "HAL/FileManager.h"
 #include "IStructSerializerBackend.h"
 #include "MetasoundArchetype.h"
+#include "MetasoundFrontendController.h"
+#include "MetasoundFrontendTransform.h"
 #include "MetasoundJsonBackend.h"
+#include "MetasoundLog.h"
 #include "StructSerializer.h"
 
 const FString FMetasoundAssetBase::FileExtension(TEXT(".metasound"));
@@ -15,23 +18,29 @@ FMetasoundAssetBase::FMetasoundAssetBase()
 {
 }
 
-FMetasoundAssetBase::FMetasoundAssetBase(const FMetasoundArchetype& InArchetype)
+FMetasoundAssetBase::FMetasoundAssetBase(const FMetasoundFrontendArchetype& InArchetype)
 :	Archetype(InArchetype)
 {
 }
 
-void FMetasoundAssetBase::SetMetadata(FMetasoundClassMetadata& InMetadata)
+void FMetasoundAssetBase::SetMetadata(FMetasoundFrontendClassMetadata& InMetadata)
 {
-	InMetadata.NodeType = EMetasoundClassType::MetasoundGraph;
-	GetDocumentChecked().RootClass.Metadata = InMetadata;
+	FMetasoundFrontendDocument& Doc = GetDocumentChecked();
+	Doc.RootGraph.Metadata = InMetadata;
+
+	if (Doc.RootGraph.Metadata.Type != EMetasoundFrontendClassType::Graph)
+	{
+		UE_LOG(LogMetasound, Display, TEXT("Forcing class type to EMetasoundFrontendClassType::Graph on root graph metadata"));
+		Doc.RootGraph.Metadata.Type = EMetasoundFrontendClassType::Graph;
+	}
 }
 
-const FMetasoundArchetype& FMetasoundAssetBase::GetArchetype() const
+const FMetasoundFrontendArchetype& FMetasoundAssetBase::GetArchetype() const
 {
 	return Archetype;
 }
 
-bool FMetasoundAssetBase::SetArchetype(const FMetasoundArchetype& InArchetype)
+bool FMetasoundAssetBase::SetArchetype(const FMetasoundFrontendArchetype& InArchetype)
 {
 	if (IsArchetypeSupported(InArchetype))
 	{
@@ -45,9 +54,9 @@ bool FMetasoundAssetBase::SetArchetype(const FMetasoundArchetype& InArchetype)
 	return false;
 }
 
-bool FMetasoundAssetBase::IsArchetypeSupported(const FMetasoundArchetype& InArchetype) const
+bool FMetasoundAssetBase::IsArchetypeSupported(const FMetasoundFrontendArchetype& InArchetype) const
 {
-	auto IsEqualArchetype = [&](const FMetasoundArchetype& SupportedArchetype)
+	auto IsEqualArchetype = [&](const FMetasoundFrontendArchetype& SupportedArchetype)
 	{
 		return Metasound::Frontend::IsEqualArchetype(InArchetype, SupportedArchetype);
 	};
@@ -55,7 +64,7 @@ bool FMetasoundAssetBase::IsArchetypeSupported(const FMetasoundArchetype& InArch
 	return Algo::AnyOf(GetPreferredArchetypes(), IsEqualArchetype);
 }
 
-const FMetasoundArchetype& FMetasoundAssetBase::GetPreferredArchetype(const FMetasoundDocument& InDocument) const
+const FMetasoundFrontendArchetype& FMetasoundAssetBase::GetPreferredArchetype(const FMetasoundFrontendDocument& InDocument) const
 {
 	// Default to archetype provided in case it is supported. 
 	if (IsArchetypeSupported(InDocument.Archetype))
@@ -64,7 +73,7 @@ const FMetasoundArchetype& FMetasoundAssetBase::GetPreferredArchetype(const FMet
 	}
 
 	// If existing archetype is not supported, get the most similar that still supports the documents environment.
-	const FMetasoundArchetype* SimilarArchetype = Metasound::Frontend::FindMostSimilarArchetypeSupportingEnvironment(InDocument, GetPreferredArchetypes());
+	const FMetasoundFrontendArchetype* SimilarArchetype = Metasound::Frontend::FindMostSimilarArchetypeSupportingEnvironment(InDocument, GetPreferredArchetypes());
 
 	if (nullptr != SimilarArchetype)
 	{
@@ -75,12 +84,12 @@ const FMetasoundArchetype& FMetasoundAssetBase::GetPreferredArchetype(const FMet
 	return GetArchetype();
 }
 
-void FMetasoundAssetBase::SetDocument(const FMetasoundDocument& InDocument, bool bForceUpdateArchetype)
+void FMetasoundAssetBase::SetDocument(const FMetasoundFrontendDocument& InDocument, bool bForceUpdateArchetype)
 {
-	FMetasoundDocument& Document = GetDocumentChecked();
+	FMetasoundFrontendDocument& Document = GetDocumentChecked();
 	Document = InDocument;
 
-	FMetasoundArchetype NewArch = InDocument.Archetype;
+	FMetasoundFrontendArchetype NewArch = InDocument.Archetype;
 
 	if (bForceUpdateArchetype || (!IsArchetypeSupported(NewArch)))
 	{
@@ -92,64 +101,38 @@ void FMetasoundAssetBase::SetDocument(const FMetasoundDocument& InDocument, bool
 
 void FMetasoundAssetBase::ConformDocumentToArchetype()
 {
-	FMetasoundDocument& Document = GetDocumentChecked();
+	FMetasoundFrontendDocument& Document = GetDocumentChecked();
 	Document.Archetype = GetArchetype();
-	GetRootGraphHandle().FixDocumentToMatchArchetype();
+
+	Metasound::Frontend::FMatchRootGraphToArchetype Transform;
+	Transform.Transform(GetDocumentHandle());
 }
 
-FMetasoundClassMetadata FMetasoundAssetBase::GetMetadata()
+FMetasoundFrontendClassMetadata FMetasoundAssetBase::GetMetadata()
 {
-	return GetDocumentChecked().RootClass.Metadata;
+	return GetDocumentChecked().RootGraph.Metadata;
 }
 
-Metasound::Frontend::FDescriptionAccessPoint FMetasoundAssetBase::GetGraphAccessPoint() 
+/*
+Metasound::Frontend::FDescriptionAccessPoint FMetasoundAssetBase::GetGraphAccessPoint()
 {
 	return Metasound::Frontend::FDescriptionAccessPoint(GetDocument());
 }
+*/
 
-Metasound::Frontend::FGraphHandle FMetasoundAssetBase::GetRootGraphHandle()
+Metasound::Frontend::FDocumentHandle FMetasoundAssetBase::GetDocumentHandle()
 {
-	using FDescPath = Metasound::Frontend::FDescPath;
-	using FGraphHandle = Metasound::Frontend::FGraphHandle;
-	using FHandleInitParams = Metasound::Frontend::FHandleInitParams;
-	using EFromDocument = Metasound::Frontend::Path::EFromDocument;
-	using EFromClass = Metasound::Frontend::Path::EFromClass;
-
-	FDescPath PathToGraph = FDescPath()[EFromDocument::ToRootClass][EFromClass::ToGraph];
-	FHandleInitParams InitParams = { GetGraphAccessPoint(), PathToGraph, INDEX_NONE, GetDocumentChecked().RootClass.UniqueID, MakeWeakObjectPtr(GetOwningAsset()) };
-	return FGraphHandle(GetPrivateToken(), InitParams);
+	return Metasound::Frontend::IDocumentController::CreateDocumentHandle(GetDocument());
 }
 
-TArray<Metasound::Frontend::FGraphHandle> FMetasoundAssetBase::GetAllSubgraphHandles()
+Metasound::Frontend::FGraphHandle FMetasoundAssetBase::GetRootGraphHandle()
 {
-	using FDescPath = Metasound::Frontend::FDescPath;
-	using FGraphHandle = Metasound::Frontend::FGraphHandle;
-	using FHandleInitParams = Metasound::Frontend::FHandleInitParams;
-	using EFromDocument = Metasound::Frontend::Path::EFromDocument;
-	using EFromClass = Metasound::Frontend::Path::EFromClass;
-
-	TArray<FGraphHandle> OutArray;
-
-	FDescPath RootPathForDependencyGraphs = FDescPath()[EFromDocument::ToDependencies];
-
-	TArray<FMetasoundClassDescription>& DependenciesList = GetDocumentChecked().Dependencies;
-	for (FMetasoundClassDescription& Dependency : DependenciesList)
-	{
-		const bool bIsSubgraph = Dependency.Metadata.NodeType == EMetasoundClassType::MetasoundGraph && Dependency.Graph.Nodes.Num() > 0;
-		if (bIsSubgraph)
-		{
-			FDescPath SubgraphPath = RootPathForDependencyGraphs[Dependency.UniqueID][EFromClass::ToGraph];
-			FHandleInitParams InitParams = { GetGraphAccessPoint(), SubgraphPath, INDEX_NONE, Dependency.UniqueID, MakeWeakObjectPtr(GetOwningAsset()) };
-			OutArray.Emplace(GetPrivateToken(), InitParams);
-		}
-	}
-
-	return OutArray;
+	return GetDocumentHandle()->GetRootGraph();
 }
 
 bool FMetasoundAssetBase::ImportFromJSON(const FString& InJSON)
 {
-	Metasound::Frontend::TAccessPtr<FMetasoundDocument> Document = GetDocument();
+	Metasound::Frontend::TAccessPtr<FMetasoundFrontendDocument> Document = GetDocument();
 	if (ensure(Document.IsValid()))
 	{
 		return Metasound::Frontend::ImportJSONToMetasound(InJSON, *Document);
@@ -159,7 +142,7 @@ bool FMetasoundAssetBase::ImportFromJSON(const FString& InJSON)
 
 bool FMetasoundAssetBase::ImportFromJSONAsset(const FString& InAbsolutePath)
 {
-	Metasound::Frontend::TAccessPtr<FMetasoundDocument> Document = GetDocument();
+	Metasound::Frontend::TAccessPtr<FMetasoundFrontendDocument> Document = GetDocument();
 	if (ensure(Document.IsValid()))
 	{
 		return Metasound::Frontend::ImportJSONAssetToMetasound(InAbsolutePath, *Document);
@@ -167,17 +150,17 @@ bool FMetasoundAssetBase::ImportFromJSONAsset(const FString& InAbsolutePath)
 	return false;
 }
 
-FMetasoundDocument& FMetasoundAssetBase::GetDocumentChecked()
+FMetasoundFrontendDocument& FMetasoundAssetBase::GetDocumentChecked()
 {
-	Metasound::Frontend::TAccessPtr<FMetasoundDocument> DocAccessPtr = GetDocument();
+	Metasound::Frontend::TAccessPtr<FMetasoundFrontendDocument> DocAccessPtr = GetDocument();
 
 	check(DocAccessPtr.IsValid());
 	return *DocAccessPtr;
 }
 
-const FMetasoundDocument& FMetasoundAssetBase::GetDocumentChecked() const
+const FMetasoundFrontendDocument& FMetasoundAssetBase::GetDocumentChecked() const
 {
-	Metasound::Frontend::TAccessPtr<const FMetasoundDocument> DocAccessPtr = GetDocument();
+	Metasound::Frontend::TAccessPtr<const FMetasoundFrontendDocument> DocAccessPtr = GetDocument();
 
 	check(DocAccessPtr.IsValid());
 	return *DocAccessPtr;
