Metasound Frontend Revamp

diff --git a/Engine/Plugins/Runtime/Metasound/Source/MetasoundFrontend/Private/MetasoundFrontend.cpp b/Engine/Plugins/Runtime/Metasound/Source/MetasoundFrontend/Private/MetasoundFrontend.cpp
--- a/Engine/Plugins/Runtime/Metasound/Source/MetasoundFrontend/Private/MetasoundFrontend.cpp
+++ b/Engine/Plugins/Runtime/Metasound/Source/MetasoundFrontend/Private/MetasoundFrontend.cpp
@@ -6,6 +6,7 @@
 #include "HAL/IConsoleManager.h"
 #include "MetasoundAudioBuffer.h"
 #include "MetasoundDataTypeRegistrationMacro.h"
+#include "MetasoundFrontendDocument.h"
 #include "MetasoundFrontendGraph.h"
 #include "MetasoundFrontendRegistries.h"
 #include "MetasoundJsonBackend.h"
@@ -29,8 +30,6 @@ namespace Metasound
 {
 	namespace Frontend
 	{
-		const FHandleInitParams::EPrivateToken FHandleInitParams::PrivateToken = FHandleInitParams::EPrivateToken::Token;
-
 		TArray<FNodeClassInfo> GetAllAvailableNodeClasses()
 		{
 			TArray<FNodeClassInfo> OutClasses;
@@ -40,7 +39,7 @@ namespace Metasound
 			{
 				FNodeClassInfo ClassInfo;
 				ClassInfo.NodeName = NodeClassTuple.Key.NodeName.ToString();
-				ClassInfo.NodeType = EMetasoundClassType::External;
+				ClassInfo.NodeType = EMetasoundFrontendClassType::External;
 				ClassInfo.LookupKey = NodeClassTuple.Key;
 
 				OutClasses.Add(ClassInfo);
@@ -51,114 +50,134 @@ namespace Metasound
 
 
 		// gets all metadata (name, description, author, what to say if it's missing) for a given node.
-		FMetasoundClassMetadata GenerateClassMetadata(const FNodeClassInfo& InInfo)
+		FMetasoundFrontendClassMetadata GenerateClassMetadata(const FNodeClassInfo& InInfo)
 		{
 			return GenerateClassDescription(InInfo).Metadata;
 		}
 
-		FMetasoundClassDescription GenerateClassDescription(const FNodeInfo& InNodeMetadata)
+		FMetasoundFrontendClass GenerateClassDescription(const FNodeInfo& InNodeMetadata, EMetasoundFrontendClassType ClassType)
 		{
-			FMetasoundClassDescription ClassDescription;
-
-			// Set metadata for class description.
-			ClassDescription.Metadata = FMetasoundClassMetadata
+			FMetasoundFrontendClassName ClassName
 			{
+				TEXT(""), // Namespace
 				InNodeMetadata.ClassName.ToString(),
+				TEXT("") // Variant
+			};
+
+			FMetasoundFrontendVersionNumber ClassVersion
+			{
 				InNodeMetadata.MajorVersion,
-				InNodeMetadata.MinorVersion,
-				EMetasoundClassType::External,
+				InNodeMetadata.MinorVersion
+			};
+
+			FMetasoundFrontendClass ClassDescription;
+
+			// Set metadata for class description.
+			ClassDescription.Metadata = FMetasoundFrontendClassMetadata
+			{
+				ClassName,
+				ClassVersion,
+				ClassType,
 				InNodeMetadata.Description,
 				InNodeMetadata.PromptIfMissing,
-				InNodeMetadata.Author
+				InNodeMetadata.Author,
+				InNodeMetadata.Keywords	
 			};
 
+			int32 RunningPointID = 0;
+
 			// External metasounds aren't dependent on any other nodes by definition, so all we need to do
 			// is populate the Input and Output sets.
 			for (auto& InputTuple : InNodeMetadata.DefaultInterface.GetInputInterface())
 			{
-				FMetasoundInputDescription InputDescription;
-				InputDescription.Name = InputTuple.Value.GetVertexName();
-				InputDescription.TypeName = InputTuple.Value.GetDataTypeName();
-				InputDescription.ToolTip = InputTuple.Value.GetDescription();
+				// TODO: lots to be added here. 
+				FMetasoundFrontendClassInput ClassInput;
 
-				ClassDescription.Inputs.Add(InputDescription);
+				ClassInput.Name = InputTuple.Value.GetVertexName();
+				ClassInput.TypeName = InputTuple.Value.GetDataTypeName();
+				ClassInput.PointIDs.Add(RunningPointID++);
+				ClassInput.Metadata.DisplayName = FText::FromString(InputTuple.Value.GetVertexName());
+				ClassInput.Metadata.Description = InputTuple.Value.GetDescription();
+
+				ClassDescription.Interface.Inputs.Add(ClassInput);
 			}
 
 			for (auto& OutputTuple : InNodeMetadata.DefaultInterface.GetOutputInterface())
 			{
-				FMetasoundOutputDescription OutputDescription;
-				OutputDescription.Name = OutputTuple.Value.GetVertexName();
-				OutputDescription.TypeName = OutputTuple.Value.GetDataTypeName();
-				OutputDescription.ToolTip = OutputTuple.Value.GetDescription();
+				// TODO: lots to be added here. 
+				FMetasoundFrontendClassOutput ClassOutput;
+
+				ClassOutput.Name = OutputTuple.Value.GetVertexName();
+				ClassOutput.TypeName = OutputTuple.Value.GetDataTypeName();
+				ClassOutput.PointIDs.Add(RunningPointID++);
+				ClassOutput.Metadata.DisplayName = FText::FromString(OutputTuple.Value.GetVertexName());
+				ClassOutput.Metadata.Description = OutputTuple.Value.GetDescription();
 
-				ClassDescription.Outputs.Add(OutputDescription);
+				ClassDescription.Interface.Outputs.Add(ClassOutput);
 			}
 
 			for (auto& EnvTuple : InNodeMetadata.DefaultInterface.GetEnvironmentInterface())
 			{
-				FMetasoundEnvironmentVariableDescription EnvironmentDescription;
-				EnvironmentDescription.Name = EnvTuple.Value.GetVertexName();
-				EnvironmentDescription.ToolTip = EnvTuple.Value.GetDescription();
-
-				ClassDescription.EnvironmentVariables.Add(EnvironmentDescription);
-			}
+				// TODO: lots to be added here. 
+				FMetasoundFrontendClassEnvironmentVariable EnvVar;
 
+				EnvVar.Name = EnvTuple.Value.GetVertexName();
+				EnvVar.Metadata.DisplayName = FText::FromString(EnvTuple.Value.GetVertexName());
+				EnvVar.Metadata.Description = EnvTuple.Value.GetDescription();
+				EnvVar.bIsRequired = true;
 
-			// Populate lookup data.
-			FNodeRegistryKey Key = FMetasoundFrontendRegistryContainer::GetRegistryKey(InNodeMetadata);
-
-			ClassDescription.ExternalNodeClassLookupInfo.ExternalNodeClassName = Key.NodeName;
-			ClassDescription.ExternalNodeClassLookupInfo.ExternalNodeClassHash = Key.NodeHash;
+				ClassDescription.Interface.Environment.Add(EnvVar);
+			}
 
 			return ClassDescription;
 		}
 
-		FMetasoundClassDescription GenerateClassDescription(const FNodeClassInfo& InInfo)
+		FMetasoundFrontendClass GenerateClassDescription(const FNodeClassInfo& InInfo)
 		{
 			auto& Registry = GetExternalNodeRegistry();
 			if (Registry.Contains(InInfo.LookupKey))
 			{
-				return Registry[InInfo.LookupKey].CreateClassDescription();
+				return Registry[InInfo.LookupKey].CreateFrontendClass();
 			}
 			else
 			{
 				ensureAlwaysMsgf(false, TEXT("Tried to get Class Description for unknown node %s!"), *InInfo.NodeName);
-				return FMetasoundClassDescription();
+				return FMetasoundFrontendClass();
 			}
 		}
 
-		EMetasoundLiteralType GetMetasoundLiteralType(Metasound::ELiteralArgType InLiteralType)
+		EMetasoundFrontendLiteralType GetMetasoundLiteralType(Metasound::ELiteralType InLiteralType)
 		{
 			switch (InLiteralType)
 			{
-				case Metasound::ELiteralArgType::Boolean:
+				case Metasound::ELiteralType::Boolean:
 				{
-					return EMetasoundLiteralType::Bool;
+					return EMetasoundFrontendLiteralType::Bool;
 				}
-				case Metasound::ELiteralArgType::Integer:
+				case Metasound::ELiteralType::Integer:
 				{
-					return EMetasoundLiteralType::Integer;
+					return EMetasoundFrontendLiteralType::Integer;
 				}
-				case Metasound::ELiteralArgType::Float:
+				case Metasound::ELiteralType::Float:
 				{
-					return EMetasoundLiteralType::Float;
+					return EMetasoundFrontendLiteralType::Float;
 				}
-				case Metasound::ELiteralArgType::String:
+				case Metasound::ELiteralType::String:
 				{
-					return EMetasoundLiteralType::String;
+					return EMetasoundFrontendLiteralType::String;
 				}
-				case Metasound::ELiteralArgType::UObjectProxy:
+				case Metasound::ELiteralType::UObjectProxy:
 				{
-					return EMetasoundLiteralType::UObject;
+					return EMetasoundFrontendLiteralType::UObject;
 				}
-				case Metasound::ELiteralArgType::UObjectProxyArray:
+				case Metasound::ELiteralType::UObjectProxyArray:
 				{
-					return EMetasoundLiteralType::UObjectArray;
+					return EMetasoundFrontendLiteralType::UObjectArray;
 				}
-				case Metasound::ELiteralArgType::None:
+				case Metasound::ELiteralType::None:
 				default:
 				{
-					return EMetasoundLiteralType::None;
+					return EMetasoundFrontendLiteralType::None;
 				}
 			}
 		}
@@ -173,7 +192,7 @@ namespace Metasound
 			return FMetasoundFrontendRegistryContainer::Get()->GetInfoForDataType(InDataType, OutInfo);
 		}
 
-		bool ImportJSONToMetasound(const FString& InJSON, FMetasoundDocument& OutMetasoundDocument)
+		bool ImportJSONToMetasound(const FString& InJSON, FMetasoundFrontendDocument& OutMetasoundDocument)
 		{
 			TArray<uint8> ReadBuffer;
 			ReadBuffer.SetNumUninitialized(InJSON.Len() * sizeof(ANSICHAR));
@@ -187,7 +206,7 @@ namespace Metasound
 			return DeserializeResult && !MemReader.IsError();
 		}
 
-		bool ImportJSONAssetToMetasound(const FString& InPath, FMetasoundDocument& OutMetasoundDocument)
+		bool ImportJSONAssetToMetasound(const FString& InPath, FMetasoundFrontendDocument& OutMetasoundDocument)
 		{
 			if (TUniquePtr<FArchive> FileReader = TUniquePtr<FArchive>(IFileManager::Get().CreateFileReader(*InPath)))
 			{
@@ -201,2081 +220,6 @@ namespace Metasound
 			return false;
 		}
 
-		void FGraphHandle::FixDocumentToMatchArchetype()
-		{
-			// TODO: Also check if this is the root class.
-			if (!IsValid())
-			{
-				return;
-			}
-
-			FMetasoundDocument& DocumentToFixUp = *OwningDocument;
-
-			// Add any missing inputs from the Required Inputs list:
-			const TArray<FMetasoundInputDescription>& RequiredInputs = DocumentToFixUp.Archetype.RequiredInputs;
-			TArray<FMetasoundInputDescription>& CurrentInputs = DocumentToFixUp.RootClass.Inputs;
-			for (const FMetasoundInputDescription& RequiredInput : RequiredInputs)
-			{
-				auto IsEqualInputName = [&](const FMetasoundInputDescription& InputDesc)
-				{
-					return RequiredInput.Name == InputDesc.Name;
-				};
-
-				if (const FMetasoundInputDescription* InputWithSameName = CurrentInputs.FindByPredicate(IsEqualInputName))
-				{
-					if (InputWithSameName->TypeName != RequiredInput.TypeName)
-					{
-						// Swap existing input because types do not match.
-						RemoveInput(InputWithSameName->Name);
-						AddNewInput(RequiredInput);
-					}
-				}
-				else
-				{
-					AddNewInput(RequiredInput);
-				}
-			}
-
-			// Add any missing outputs from the Required Outputs list:
-			const TArray<FMetasoundOutputDescription>& RequiredOutputs = DocumentToFixUp.Archetype.RequiredOutputs;
-			TArray<FMetasoundOutputDescription>& CurrentOutputs = DocumentToFixUp.RootClass.Outputs;
-			for (const FMetasoundOutputDescription& RequiredOutput : RequiredOutputs)
-			{
-				auto IsEqualOutputName = [&](const FMetasoundOutputDescription& OutputDesc)
-				{
-					return RequiredOutput.Name == OutputDesc.Name;
-				};
-
-				if (const FMetasoundOutputDescription* OutputWithSameName = CurrentOutputs.FindByPredicate(IsEqualOutputName))
-				{
-					if (OutputWithSameName->TypeName != RequiredOutput.TypeName)
-					{
-						// Swap existing input because types do not match.
-						RemoveOutput(OutputWithSameName->Name);
-						AddNewOutput(RequiredOutput);
-					}
-				}
-				else
-				{
-					AddNewOutput(RequiredOutput);
-				}
-			}
-		}
-
-		FInputHandle::FInputHandle(FHandleInitParams::EPrivateToken PrivateToken, const FHandleInitParams& InParams, const FString& InputName, EDefaultTag InTag)
-			: ITransactable(MetasoundUndoRollLimitCvar, InParams.InOwningAsset)
-			, NodePtr(InParams.InAccessPoint, InParams.InPath)
-			, NodeClass(InParams.InAccessPoint, Path::GetDependencyPath(InParams.DependencyID))
-			, InputPtr(InParams.InAccessPoint, NodeClass.GetPath()[Path::EFromClass::ToInputs][*InputName])
-			, OutputNodePtr(TDescriptionPtr<FMetasoundOutputDescription>::CreateInvalid())
-		{
-			if (InParams.InAccessPoint.IsValid())
-			{
-				// Test both pointers to the graph and it's owning class description.
-				ensureAlwaysMsgf(NodePtr.IsValid() && NodeClass.IsValid() && InputPtr.IsValid(), TEXT("Tried to build FGraphHandle with Invalid Path: %s"), *Path::GetPrintableString(InParams.InPath));
-			}
-		}
-
-		FInputHandle::FInputHandle(FHandleInitParams::EPrivateToken PrivateToken, const FHandleInitParams& InParams, const FString& InputName, EOutputNodeTag InTag)
-			: ITransactable(MetasoundUndoRollLimitCvar, InParams.InOwningAsset)
-			, NodePtr(InParams.InAccessPoint, InParams.InPath)
-			, NodeClass(TDescriptionPtr<FMetasoundClassDescription>::CreateInvalid())
-			, InputPtr(TDescriptionPtr<FMetasoundInputDescription>::CreateInvalid())
-			, OutputNodePtr(InParams.InAccessPoint, Path::GetOutputDescriptionPath(InParams.InPath, InputName))
-		{
-			if (InParams.InAccessPoint.IsValid())
-			{
-				// Test both pointers to the graph and it's owning class description.
-				ensureAlwaysMsgf(NodePtr.IsValid() && OutputNodePtr.IsValid(), TEXT("Tried to build FGraphHandle with Invalid Path: %s"), *Path::GetPrintableString(InParams.InPath));
-			}
-		}
-
-		FInputHandle FInputHandle::InvalidHandle()
-		{
-			FDescPath NullPath = FDescPath();
-			FString NullString = FString();
-			FHandleInitParams InitParams = { FDescriptionAccessPoint::CreateInvalid(), NullPath , FMetasoundNodeDescription::InvalidID, FMetasoundClassDescription::InvalidID, nullptr };
-			return FInputHandle(FHandleInitParams::PrivateToken, InitParams, NullString, FInputHandle::Default);
-		}
-
-		bool FInputHandle::IsValid() const
-		{
-			return (NodePtr.IsValid() && NodeClass.IsValid() && InputPtr.IsValid())
-				|| (NodePtr.IsValid() && OutputNodePtr.IsValid());
-		}
-
-		bool FInputHandle::IsConnected() const
-		{
-			if (!IsValid())
-			{
-				return false;
-			}
-
-			const FMetasoundNodeConnectionDescription* Connection = GetConnectionDescription();
-			if (Connection)
-			{
-				return Connection->NodeID != FMetasoundNodeConnectionDescription::DisconnectedNodeID;
-			}
-			else
-			{
-				return false;
-			}
-		}
-
-		FName FInputHandle::GetInputType() const
-		{
-			if (!IsValid())
-			{
-				return FName();
-			}
-
-			if (InputPtr.IsValid())
-			{
-				return InputPtr->TypeName;
-			}
-			else if (OutputNodePtr.IsValid())
-			{
-				return OutputNodePtr->TypeName;
-			}
-			else
-			{
-				checkNoEntry();
-				return FName();
-			}
-		}
-
-		FString FInputHandle::GetInputName() const
-		{
-			if (!IsValid())
-			{
-				return FString();
-			}
-
-			if (InputPtr.IsValid())
-			{
-				return InputPtr->Name;
-			}
-			else if (OutputNodePtr.IsValid())
-			{
-				return OutputNodePtr->Name;
-			}
-			else
-			{
-				checkNoEntry();
-				return FString();
-			}
-		}
-
-		FText FInputHandle::GetInputTooltip() const
-		{
-			if (!IsValid())
-			{
-				return FText();
-			}
-
-			if (InputPtr.IsValid())
-			{
-				return InputPtr->ToolTip;
-			}
-			else if (OutputNodePtr.IsValid())
-			{
-				return OutputNodePtr->ToolTip;
-			}
-			else
-			{
-				checkNoEntry();
-				return FText();
-			}
-		}
-
-		int32 FInputHandle::GetOwningNodeID() const
-		{
-			if (!IsValid())
-			{
-				return FMetasoundNodeDescription::InvalidID;
-			}
-
-			return NodePtr->UniqueID;
-		}
-
-		FOutputHandle FInputHandle::GetCurrentlyConnectedOutput() const
-		{
-			if (!IsValid())
-			{
-				return FOutputHandle::InvalidHandle();
-			}
-
-			if (!IsConnected())
-			{
-				return FOutputHandle::InvalidHandle();
-			}
-
-			const FMetasoundNodeConnectionDescription* Connection = GetConnectionDescription();
-			check(Connection);
-
-			FString OutputName = Connection->OutputName;
-			int32 OutputNodeID = Connection->NodeID;
-
-			// All node connections are in the same graph, so we just need to go up one level to the Nodes array
-			// and look up the node by it's unique ID.
-			FDescPath OutputNodePath = (NodePtr.GetPath() << 1)[OutputNodeID];
-			TDescriptionPtr<FMetasoundNodeDescription> OtherNodePtr(NodePtr.GetAccessPoint(), OutputNodePath);
-
-			int32 OutputNodeDependencyId = FMetasoundClassDescription::InvalidID;
-			EMetasoundClassType OutputNodeType = EMetasoundClassType::Invalid;
-			if (OtherNodePtr.IsValid())
-			{
-				OutputNodeDependencyId = OtherNodePtr->DependencyID;
-				OutputNodeType = OtherNodePtr->ObjectTypeOfNode;
-			}
-
-			if (EMetasoundClassType::Input == OutputNodeType)
-			{
-				FHandleInitParams InitParams = { NodePtr.GetAccessPoint(), OutputNodePath, OutputNodeID, FMetasoundClassDescription::InvalidID, OwningAsset };
-				return FOutputHandle(FHandleInitParams::PrivateToken, InitParams, OutputName, FOutputHandle::InputNode);
-			}
-			else if (EMetasoundClassType::Invalid != OutputNodeType)
-			{
-				FHandleInitParams InitParams = { NodePtr.GetAccessPoint(), OutputNodePath, OutputNodeID, OutputNodeDependencyId, OwningAsset };
-				return FOutputHandle(FHandleInitParams::PrivateToken, InitParams, OutputName, FOutputHandle::Default);
-			}
-			else
-			{
-				checkNoEntry();
-				return FOutputHandle::InvalidHandle();
-			}
-		}
-
-		FConnectability FInputHandle::CanConnectTo(const FOutputHandle& InHandle) const
-		{
-			FConnectability OutConnectability;
-			OutConnectability.Connectable = FConnectability::EConnectable::No;
-
-			if (InHandle.GetOutputType() == GetInputType())
-			{
-				OutConnectability.Connectable = FConnectability::EConnectable::Yes;
-				return OutConnectability;
-			}
-
-			OutConnectability.PossibleConverterNodeClasses = FMetasoundFrontendRegistryContainer::Get()->GetPossibleConverterNodes(InHandle.GetOutputType(), GetInputType());
-			if (OutConnectability.PossibleConverterNodeClasses.Num() > 0)
-			{
-				OutConnectability.Connectable = FConnectability::EConnectable::YesWithConverterNode;
-				return OutConnectability;
-			}
-
-			return OutConnectability;
-		}
-
-		bool FInputHandle::Connect(FOutputHandle& InHandle)
-		{
-			if (!IsValid() || !InHandle.IsValid())
-			{
-				return false;
-			}
-
-			if (!ensureAlwaysMsgf(InHandle.GetOutputType() == GetInputType(), TEXT("Tried to connect incompatible types!")))
-			{
-				return false;
-			}
-
-			int32 OutputNodeID = InHandle.GetOwningNodeID();
-			FString OutputName = InHandle.GetOutputName();
-
-			FMetasoundNodeConnectionDescription* Connection = GetConnectionDescription();
-
-			if (!Connection)
-			{
-				TArray<FMetasoundNodeConnectionDescription>& Connections = NodePtr->InputConnections;
-				FMetasoundNodeConnectionDescription NewConnection;
-				NewConnection.InputName = GetInputName();
-				Connection = &Connections.Add_GetRef(NewConnection);
-			}
-
-			Connection->NodeID = OutputNodeID;
-			Connection->OutputName = OutputName;
-
-			return true;
-		}
-
-		bool FInputHandle::Disconnect()
-		{
-			if (!IsValid())
-			{
-				return false;
-			}
-
-			TArray<FMetasoundNodeConnectionDescription>& Connections = NodePtr->InputConnections;
-			for (int32 i = 0; i < Connections.Num(); ++i)
-			{
-				if (Connections[i].InputName == GetInputName())
-				{
-					Connections.RemoveAtSwap(i);
-					return true;
-				}
-			}
-
-			return false;
-		}
-
-		bool FInputHandle::Disconnect(FOutputHandle& InHandle)
-		{
-			if (!IsValid() || !InHandle.IsValid())
-			{
-				return false;
-			}
-
-			if (!ensureAlwaysMsgf(InHandle.GetOutputType() == GetInputType(), TEXT("Tried to disconnect incompatible types!")))
-			{
-				return false;
-			}
-
-
-			FMetasoundNodeConnectionDescription* Connection = GetConnectionDescription();
-			if (!ensure(Connection))
-			{
-				return false;
-			}
-
-			const int32 OutputNodeID = InHandle.GetOwningNodeID();
-			for (int32 i = 0; i < NodePtr->InputConnections.Num(); ++i)
-			{
-				if (NodePtr->InputConnections[i].NodeID == OutputNodeID)
-				{
-					NodePtr->InputConnections.RemoveAtSwap(i);
-					return true;
-				}
-			}
-
-			return false;
-		}
-
-		bool FInputHandle::ConnectWithConverterNode(FOutputHandle& InHandle, const FConverterNodeInfo& InNodeClassName)
-		{
-			FDescPath OwningGraphPath = Path::GetOuterGraphPath(NodePtr.GetPath());
-			FHandleInitParams GraphHandleInitParams =
-			{
-				NodePtr.GetAccessPoint(),
-				OwningGraphPath,
-				FMetasoundNodeDescription::InvalidID,
-				FMetasoundClassDescription::InvalidID,
-				OwningAsset
-			};
-
-			FGraphHandle OwningGraphHandle = FGraphHandle(FHandleInitParams::PrivateToken, GraphHandleInitParams);
-			check(OwningGraphHandle.IsValid());
-
-			// Generate the converter node.
-			FNodeHandle Converter = OwningGraphHandle.AddNewNode(InNodeClassName.NodeKey);
-
-			check(Converter.IsValid());
-
-			FInputHandle ConverterInput = Converter.GetInputWithName(InNodeClassName.PreferredConverterInputPin);
-			FOutputHandle ConverterOutput = Converter.GetOutputWithName(InNodeClassName.PreferredConverterOutputPin);
-
-			check(ConverterInput.IsValid() && ConverterOutput.IsValid());
-
-			// Connect the output InHandle to the converter, than connect the converter to this input.
-			ensureAlways(ConverterInput.Connect(InHandle));
-			return ensureAlways(Connect(ConverterOutput));
-		}
-
-		const FMetasoundNodeConnectionDescription* FInputHandle::GetConnectionDescription() const
-		{
-			if (!IsValid())
-			{
-				return nullptr;
-			}
-
-			const TArray<FMetasoundNodeConnectionDescription>& Connections = NodePtr->InputConnections;
-
-			for (const FMetasoundNodeConnectionDescription& Connection : Connections)
-			{
-				if (Connection.InputName == GetInputName())
-				{
-					return &Connection;
-				}
-			}
-
-			return nullptr;
-		}
-
-		FMetasoundNodeConnectionDescription* FInputHandle::GetConnectionDescription()
-		{
-			if (!IsValid())
-			{
-				return nullptr;
-			}
-
-			TArray<FMetasoundNodeConnectionDescription>& Connections = NodePtr->InputConnections;
-
-			for (FMetasoundNodeConnectionDescription& Connection : Connections)
-			{
-				if (Connection.InputName == GetInputName())
-				{
-					return &Connection;
-				}
-			}
-
-			return nullptr;
-		}
-
-		FOutputHandle::FOutputHandle(FHandleInitParams::EPrivateToken InToken /* unused */, const FHandleInitParams& InParams, const FString& InOutputName, FOutputHandle::EDefaultTag InTag)
-			: ITransactable(MetasoundUndoRollLimitCvar, InParams.InOwningAsset)
-			, NodePtr(InParams.InAccessPoint, InParams.InPath)
-			, NodeClass(InParams.InAccessPoint, Path::GetDependencyPath(InParams.DependencyID))
-			, OutputPtr(InParams.InAccessPoint, NodeClass.GetPath()[Path::EFromClass::ToOutputs][*InOutputName])
-			, InputNodePtr(TDescriptionPtr<FMetasoundInputDescription>::CreateInvalid())
-		{
-			if (InParams.InAccessPoint.IsValid())
-			{
-				// Test both pointers to the graph and it's owning class description.
-				ensureAlwaysMsgf(NodePtr.IsValid() && NodeClass.IsValid() && OutputPtr.IsValid(), TEXT("Tried to build FGraphHandle with Invalid Path: %s"), *Path::GetPrintableString(InParams.InPath));
-			}
-		}
-
-		FOutputHandle::FOutputHandle(FHandleInitParams::EPrivateToken InToken, const FHandleInitParams& InParams, const FString& InOutputName, FOutputHandle::EInputNodeTag InTag)
-			: ITransactable(MetasoundUndoRollLimitCvar, InParams.InOwningAsset)
-			, NodePtr(InParams.InAccessPoint, InParams.InPath)
-			, NodeClass(TDescriptionPtr<FMetasoundClassDescription>::CreateInvalid())
-			, OutputPtr(TDescriptionPtr<FMetasoundOutputDescription>::CreateInvalid())
-			, InputNodePtr(InParams.InAccessPoint, Path::GetInputDescriptionPath(InParams.InPath, InOutputName))
-		{
-			if (InParams.InAccessPoint.IsValid())
-			{
-				// Test both pointers to the graph and it's owning class description.
-				ensureAlwaysMsgf(NodePtr.IsValid() && InputNodePtr.IsValid(), TEXT("Tried to build FGraphHandle with Invalid Path: %s"), *Path::GetPrintableString(InParams.InPath));
-			}
-		}
-
-		FOutputHandle FOutputHandle::InvalidHandle()
-		{
-			FString NullString;
-			FDescPath NullPath;
-			FHandleInitParams InitParams = { FDescriptionAccessPoint::CreateInvalid(), NullPath, FMetasoundNodeDescription::InvalidID, FMetasoundClassDescription::InvalidID, nullptr };
-			return FOutputHandle(FHandleInitParams::PrivateToken, InitParams, NullString, FOutputHandle::Default);
-		}
-
-		bool FOutputHandle::IsValid() const
-		{
-			return (NodePtr.IsValid() && NodeClass.IsValid() && OutputPtr.IsValid())
-				|| (NodePtr.IsValid() && InputNodePtr.IsValid());
-		}
-
-		FName FOutputHandle::GetOutputType() const
-		{
-			if (!IsValid())
-			{
-				return FName();
-			}
-
-			if (OutputPtr.IsValid())
-			{
-				return OutputPtr->TypeName;
-			}
-			else if (InputNodePtr.IsValid())
-			{
-				return InputNodePtr->TypeName;
-			}
-			else
-			{
-				checkNoEntry();
-				return FName();
-			}
-		}
-
-		FString FOutputHandle::GetOutputName() const
-		{
-			if (!IsValid())
-			{
-				return FString();
-			}
-
-			if (OutputPtr.IsValid())
-			{
-				return OutputPtr->Name;
-			}
-			else if (InputNodePtr.IsValid())
-			{
-				return InputNodePtr->Name;
-			}
-			else
-			{
-				checkNoEntry();
-				return FString();
-			}
-		}
-
-		FText FOutputHandle::GetOutputTooltip() const
-		{
-			if (!IsValid())
-			{
-				return FText();
-			}
-
-			if (OutputPtr.IsValid())
-			{
-				return OutputPtr->ToolTip;
-			}
-			else if (InputNodePtr.IsValid())
-			{
-				return InputNodePtr->ToolTip;
-			}
-			else
-			{
-				checkNoEntry();
-				return FText();
-			}
-		}
-
-		int32 FOutputHandle::GetOwningNodeID() const
-		{
-			if (!IsValid())
-			{
-				return FMetasoundNodeDescription::InvalidID;
-			}
-
-			return NodePtr->UniqueID;
-		}
-
-		FConnectability FOutputHandle::CanConnectTo(const FInputHandle& InHandle) const
-		{
-			return InHandle.CanConnectTo(*this);
-		}
-
-		bool FOutputHandle::Connect(FInputHandle& InHandle)
-		{
-			if (!IsValid() || !InHandle.IsValid())
-			{
-				return false;
-			}
-
-			return InHandle.Connect(*this);
-		}
-
-		bool FOutputHandle::ConnectWithConverterNode(FInputHandle& InHandle, const FConverterNodeInfo& InNodeClassName)
-		{
-			return InHandle.ConnectWithConverterNode(*this, InNodeClassName);
-		}
-
-		bool FOutputHandle::Disconnect(FInputHandle& InHandle)
-		{
-			if (!IsValid() || !InHandle.IsValid())
-			{
-				return false;
-			}
-
-			return InHandle.Disconnect(*this);
-		}
-
-		TDescriptionPtr<FMetasoundClassDescription> FNodeHandle::GetNodeClassDescriptionForNodeHandle(const FHandleInitParams& InitParams, EMetasoundClassType InNodeClassType)
-		{
-			if (InNodeClassType != EMetasoundClassType::Input && InNodeClassType != EMetasoundClassType::Output)
-			{
-				return TDescriptionPtr<FMetasoundClassDescription>(InitParams.InAccessPoint, Path::GetDependencyPath(InitParams.DependencyID));
-			}
-			else
-			{
-				// input nodes and output nodes don't have class descriptions.
-				return TDescriptionPtr<FMetasoundClassDescription>::CreateInvalid();
-			}
-		}
-
-		FNodeHandle::FNodeHandle(FHandleInitParams::EPrivateToken PrivateToken, const FHandleInitParams& InParams, EMetasoundClassType InNodeClassType)
-			: ITransactable(MetasoundUndoRollLimitCvar, InParams.InOwningAsset)
-			, NodePtr(InParams.InAccessPoint, InParams.InPath)
-			, NodeClass(GetNodeClassDescriptionForNodeHandle(InParams, InNodeClassType))
-			, NodeClassType(InNodeClassType)
-		{
-			if (InParams.InAccessPoint.IsValid())
-			{
-				// Test both pointers to the graph and it's owning class description.
-				ensureAlwaysMsgf(IsValid(), TEXT("Tried to build FGraphHandle with Invalid Path: %s"), *Path::GetPrintableString(InParams.InPath));
-			}
-		}
-
-		FNodeHandle FNodeHandle::InvalidHandle()
-		{
-			FDescPath InvalidPath;
-			FString InvalidClassName;
-			FHandleInitParams InitParams = { FDescriptionAccessPoint::CreateInvalid(), InvalidPath, FMetasoundNodeDescription::InvalidID, FMetasoundClassDescription::InvalidID, nullptr };
-			return FNodeHandle(FHandleInitParams::PrivateToken, InitParams, EMetasoundClassType::Invalid);
-		}
-
-		bool FNodeHandle::IsValid() const
-		{
-			const bool bNeedsNodeClass = NodeClassType == EMetasoundClassType::External || NodeClassType == EMetasoundClassType::MetasoundGraph;
-			return NodePtr.IsValid() && (!bNeedsNodeClass || NodeClass.IsValid());
-		}
-
-		TArray<FInputHandle> FNodeHandle::GetAllInputs()
-		{
-			TArray<FInputHandle> OutArray;
-
-			if (!IsValid() || NodeClassType == EMetasoundClassType::Input)
-			{
-				return OutArray;
-			}
-
-			if (NodeClassType == EMetasoundClassType::Output)
-			{
-				// Output nodes only have one input- the outgoing parameter.
-				FDescPath NodePath = NodePtr.GetPath();
-
-				const FString& NodeClassName = NodePtr->Name;
-
-				FHandleInitParams InitParams = { NodePtr.GetAccessPoint(), NodePath, NodePtr->UniqueID, FMetasoundClassDescription::InvalidID, OwningAsset };
-				OutArray.Emplace(FHandleInitParams::PrivateToken, InitParams, NodeClassName, FInputHandle::OutputNode);
-			}
-			else
-			{
-				// Iterate over our input descriptions and emplace a new handle for each of them.
-
-				TArray<FMetasoundInputDescription>& InputDescriptions = NodeClass->Inputs;
-				for (FMetasoundInputDescription& InputDescription : InputDescriptions)
-				{
-					FDescPath NodePath = NodePtr.GetPath();
-					FString ClassName = GetNodeClassName();
-					FHandleInitParams InitParams = { NodePtr.GetAccessPoint(), NodePath, NodePtr->UniqueID, NodeClass->UniqueID, OwningAsset };
-					OutArray.Emplace(FHandleInitParams::PrivateToken, InitParams, InputDescription.Name, FInputHandle::Default);
-				}
-			}
-
-			return OutArray;
-		}
-
-		TArray<FOutputHandle> FNodeHandle::GetAllOutputs()
-		{
-			TArray<FOutputHandle> OutArray;
-
-			if (!IsValid() || NodeClassType == EMetasoundClassType::Output)
-			{
-				return OutArray;
-			}
-
-			if (NodeClassType == EMetasoundClassType::Input)
-			{
-				// Input nodes nodes only have one output- the incoming parameter.
-				FDescPath NodePath = NodePtr.GetPath();
-
-				const FString& NodeClassName = NodePtr->Name;
-
-				FHandleInitParams InitParams = { NodePtr.GetAccessPoint(), NodePath, NodePtr->UniqueID, FMetasoundClassDescription::InvalidID, OwningAsset };
-				OutArray.Emplace(FHandleInitParams::PrivateToken, InitParams, NodeClassName, FOutputHandle::InputNode);
-			}
-			else
-			{
-				TArray<FMetasoundOutputDescription>& OutputDescriptions = NodeClass->Outputs;
-
-				for (FMetasoundOutputDescription& OutputDescription : OutputDescriptions)
-				{
-					FDescPath NodePath = NodePtr.GetPath();
-					FString NodeClassName = GetNodeClassName();
-					FHandleInitParams InitParams = { NodePtr.GetAccessPoint(), NodePath, NodePtr->UniqueID, NodeClass->UniqueID, OwningAsset };
-					OutArray.Emplace(FHandleInitParams::PrivateToken, InitParams, OutputDescription.Name, FOutputHandle::Default);
-				}
-			}
-
-			return OutArray;
-		}
-
-		FInputHandle FNodeHandle::GetInputWithName(const FString& InName)
-		{
-			if (!IsValid() || NodeClassType == EMetasoundClassType::Input)
-			{
-				return FInputHandle::InvalidHandle();
-			}
-
-			if (NodeClassType == EMetasoundClassType::Output)
-			{
-				const FString& NodeClassName = NodePtr->Name;
-				ensureAlwaysMsgf(InName == NodeClassName, TEXT("An output node's input connection should always be the same as it's class name!"));
-
-				FDescPath NodePath = NodePtr.GetPath();
-				FHandleInitParams InitParams = { NodePtr.GetAccessPoint(), NodePath, NodePtr->UniqueID, FMetasoundClassDescription::InvalidID, OwningAsset };
-				return FInputHandle(FHandleInitParams::PrivateToken, InitParams, NodeClassName, FInputHandle::OutputNode);
-			}
-
-			TArray<FMetasoundInputDescription>& InputDescriptions = NodeClass->Inputs;
-
-			for (FMetasoundInputDescription& InputDescription : InputDescriptions)
-			{
-				if (InputDescription.Name == InName)
-				{
-					FString ClassName = GetNodeClassName();
-					FDescPath NodePath = NodePtr.GetPath();
-					FHandleInitParams InitParams = { NodePtr.GetAccessPoint(), NodePath, NodePtr->UniqueID, NodeClass->UniqueID, OwningAsset };
-					return FInputHandle(FHandleInitParams::PrivateToken, InitParams, InputDescription.Name, FInputHandle::Default);
-				}
-			}
-
-			ensureAlwaysMsgf(false, TEXT("Couldn't find an input with this name on this node!"));
-			return FInputHandle::InvalidHandle();
-		}
-
-		FOutputHandle FNodeHandle::GetOutputWithName(const FString& InName)
-		{
-			if (!IsValid() || NodeClassType == EMetasoundClassType::Output)
-			{
-				return FOutputHandle::InvalidHandle();
-			}
-
-			// all input nodes have one connectable output, which is the input param they represent.
-			if (NodeClassType == EMetasoundClassType::Input)
-			{
-				const FString& NodeClassName = NodePtr->Name;
-				ensureAlwaysMsgf(InName == NodeClassName, TEXT("An input node's output connection should always be the same as it's class name!"));
-				FDescPath NodePath = NodePtr.GetPath();
-				FHandleInitParams InitParams = { NodePtr.GetAccessPoint(), NodePath, FMetasoundNodeDescription::InvalidID, FMetasoundClassDescription::InvalidID, OwningAsset };
-				return FOutputHandle(FHandleInitParams::PrivateToken, InitParams, NodeClassName, FOutputHandle::InputNode);
-			}
-
-			TArray<FMetasoundOutputDescription>& OutputDescriptions = NodeClass->Outputs;
-
-			for (FMetasoundOutputDescription& OutputDescription : OutputDescriptions)
-			{
-				if (OutputDescription.Name == InName)
-				{
-					FString NodeClassName = GetNodeClassName();
-					FDescPath NodePath = NodePtr.GetPath();
-					FHandleInitParams InitParams = { NodePtr.GetAccessPoint(), NodePath, FMetasoundNodeDescription::InvalidID, NodeClass->UniqueID, OwningAsset };
-					return FOutputHandle(FHandleInitParams::PrivateToken, InitParams, OutputDescription.Name, FOutputHandle::Default);
-				}
-			}
-
-			ensureAlwaysMsgf(false, TEXT("Couldn't find an output with this name on this node!"));
-			return FOutputHandle::InvalidHandle();
-		}
-
-		EMetasoundClassType FNodeHandle::GetNodeType() const
-		{
-			if (!IsValid())
-			{
-				return EMetasoundClassType::Invalid;
-			}
-
-			if (NodeClassType == EMetasoundClassType::Input || NodeClassType == EMetasoundClassType::Output)
-			{
-				return NodeClassType;
-			}
-
-			return NodeClass->Metadata.NodeType;
-		}
-
-		const FString& FNodeHandle::GetNodeClassName() const
-		{
-			if (!IsValid())
-			{
-				static const FString DefaultClassName;
-				return DefaultClassName;
-			}
-
-			if (NodeClassType == EMetasoundClassType::Input)
-			{
-				static const FString InputClassName = TEXT("Input");
-				return InputClassName;
-			}
-			else if (NodeClassType == EMetasoundClassType::Output)
-			{
-				static const FString OutputClassName = TEXT("Output");
-				return OutputClassName;
-			}
-
-			return NodeClass->Metadata.NodeName;
-		}
-
-		FNodeClassInfo FNodeHandle::GetClassInfo() const
-		{
-			FNodeClassInfo ClassInfo;
-
-			if (IsValid())
-			{
-
-				if (NodeClassType == EMetasoundClassType::Input)
-				{
-					ClassInfo.NodeName = TEXT("Input");
-					ClassInfo.NodeType = NodeClassType;
-					ClassInfo.LookupKey.NodeHash = 0;
-					ClassInfo.LookupKey.NodeName = FName();
-				}
-				else if (NodeClassType == EMetasoundClassType::Output)
-				{
-					ClassInfo.NodeName = TEXT("Output");
-					ClassInfo.NodeType = NodeClassType;
-					ClassInfo.LookupKey.NodeHash = 0;
-					ClassInfo.LookupKey.NodeName = FName();
-				}
-				else
-				{
-					ClassInfo.NodeName = NodeClass->Metadata.NodeName;
-					ClassInfo.NodeType = NodeClass->Metadata.NodeType;
-					ClassInfo.LookupKey.NodeHash = NodeClass->ExternalNodeClassLookupInfo.ExternalNodeClassHash;
-					ClassInfo.LookupKey.NodeName = NodeClass->ExternalNodeClassLookupInfo.ExternalNodeClassName;
-				}
-			}
-
-			return ClassInfo;
-		}
-
-		void FNodeHandle::GetContainedGraph(FGraphHandle& OutGraph)
-		{
-			if (!IsValid())
-			{
-				OutGraph = FGraphHandle::InvalidHandle();
-			}
-
-			if (!ensureAlwaysMsgf(GetNodeType() == EMetasoundClassType::MetasoundGraph, TEXT("Tried to get the Metasound Graph for a node that was not a Metasound graph.")))
-			{
-				OutGraph = FGraphHandle::InvalidHandle();
-			}
-
-			FDescPath ContainedGraphPath = NodeClass.GetPath()[Path::EFromClass::ToGraph];
-			FHandleInitParams InitParams = { NodeClass.GetAccessPoint(), ContainedGraphPath, FMetasoundNodeDescription::InvalidID, NodeClass->UniqueID, OwningAsset };
-			// Todo: link this up to look for externally implemented graphs as well.
-			OutGraph = FGraphHandle(FHandleInitParams::PrivateToken, InitParams);
-		}
-
-		int32 FNodeHandle::GetNodeID() const
-		{
-			if (!IsValid())
-			{
-				return FMetasoundNodeDescription::InvalidID;
-			}
-
-			return NodePtr->UniqueID;
-		}
-
-		const FString& FNodeHandle::GetNodeName() const
-		{
-			if (!IsValid())
-			{
-				static const FString DefaultName = "InvalidNodeHandle";
-				return DefaultName;
-			}
-
-			return NodePtr->Name;
-		}
-
-		int32 FNodeHandle::GetNodeID(const FDescPath& InNodePath)
-		{
-			if (!ensureAlwaysMsgf(InNodePath.Path.Num() != 0, TEXT("Tried to get a node ID from an empty path.")))
-			{
-				return FMetasoundNodeDescription::InvalidID;
-			}
-
-			const Path::FElement& LastElementInPath = InNodePath.Path.Last();
-			
-			if (!ensureAlwaysMsgf(LastElementInPath.CurrentDescType == Path::EDescType::Node, TEXT("Tried to get the node ID for a path that was not set up for a node.")))
-			{
-				return FMetasoundNodeDescription::InvalidID;
-			}
-
-			return LastElementInPath.LookupID;
-		}
-
-		FGraphHandle::FGraphHandle(FHandleInitParams::EPrivateToken PrivateToken, const FHandleInitParams& InParams)
-			: ITransactable(MetasoundUndoRollLimitCvar, InParams.InOwningAsset)
-			, GraphPtr(InParams.InAccessPoint, InParams.InPath)
-			, GraphsClassDeclaration(InParams.InAccessPoint, Path::GetOwningClassDescription(InParams.InPath))
-			, OwningDocument(InParams.InAccessPoint, FDescPath())
-		{
-			if (InParams.InAccessPoint.IsValid())
-			{
-				// Test both pointers to the graph and it's owning class description.
-				ensureAlwaysMsgf(GraphPtr.IsValid() && GraphsClassDeclaration.IsValid(), TEXT("Tried to build FGraphHandle with Invalid Path: %s"), *Path::GetPrintableString(InParams.InPath));
-			}
-		}
-
-		int32 FGraphHandle::FindNewUniqueNodeId()
-		{
-			// Assumption here is that we will never need more than ten thousand nodes,
-			// and four digits are easy enough to read/remember when looking at metasound graph documents.
-			static const int32 NodeIDMax = 9999;
-
-			if (!IsValid())
-			{
-				return FMetasoundNodeDescription::InvalidID;
-			}
-
-			TArray<FMetasoundNodeDescription>& Nodes = GraphPtr->Nodes;
-
-			if (!ensureAlwaysMsgf(((int32)Nodes.Num()) < NodeIDMax, TEXT("Too many nodes to guarantee a unique node ID. Increase the value of NodeIDMax.")))
-			{
-				return FMetasoundNodeDescription::InvalidID;
-			}
-
-			while (int32 RandomID = FMath::RandRange(1, NodeIDMax))
-			{
-				// Scan through the nodes in this graph to see if they match this ID.
-				// If it does, generate a new random ID.
-				bool bFoundMatch = false;
-				for (FMetasoundNodeDescription& Node : Nodes)
-				{
-					if (Node.UniqueID == RandomID)
-					{
-						bFoundMatch = true;
-						break;
-					}
-				}
-
-				if (!bFoundMatch)
-				{
-					return RandomID;
-				}
-			}
-
-			return FMetasoundNodeDescription::InvalidID;
-		}
-
-		int32 FGraphHandle::FindNewUniqueDependencyId()
-		{
-			// Assumption here is that we will never need more than ten thousand dependencies,
-			// and four digits are easy enough to read/remember when looking at metasound graph documents.
-			static const int32 DependencyIDMax = 9999;
-
-			if (!IsValid())
-			{
-				return FMetasoundClassDescription::InvalidID;
-			}
-
-			TArray<FMetasoundClassDescription>& Dependencies = OwningDocument->Dependencies;
-
-			if (!ensureAlwaysMsgf(Dependencies.Num() < DependencyIDMax, TEXT("Too many nodes to guarantee a unique node ID. Increase the value of NodeIDMax.")))
-			{
-				return FMetasoundClassDescription::InvalidID;
-			}
-
-			while (int32 RandomID = FMath::RandRange(1, DependencyIDMax))
-			{
-				// Scan through the nodes in this graph to see if they match this ID.
-				// If it does, generate a new random ID.
-				bool bFoundMatch = false;
-				for (FMetasoundClassDescription& Dependency : Dependencies)
-				{
-					if (Dependency.UniqueID == RandomID)
-					{
-						bFoundMatch = true;
-						break;
-					}
-				}
-
-				if (!bFoundMatch)
-				{
-					return RandomID;
-				}
-			}
-
-			ensureAlwaysMsgf(false, TEXT("Failed to generate a new unique ID"));
-			return FMetasoundClassDescription::InvalidID;
-		}
-
-		FMetasoundLiteralDescription* FGraphHandle::GetLiteralDescriptionForInput(const FString& InInputName, FName& OutDataType) const
-		{
-			if (!IsValid())
-			{
-				return nullptr;
-			}
-
-			// scan through our inputs to find a match.
-			TArray<FMetasoundInputDescription>& Inputs = GraphsClassDeclaration->Inputs;
-			for (FMetasoundInputDescription& Input : Inputs)
-			{
-				if (Input.Name == InInputName)
-				{
-					OutDataType = Input.TypeName;
-					return &(Input.LiteralValue);
-				}
-			}
-
-			ensureAlwaysMsgf(false, TEXT("Couldn't find Input of name %s in this Metasoud graph!"), *InInputName);
-			return nullptr;
-		}
-
-		bool FGraphHandle::GetDataTypeForInput(const FString& InInputName, FName& OutDataType)
-		{
-			if (!IsValid())
-			{
-				return false;
-			}
-
-			// scan through our inputs to find a match.
-			TArray<FMetasoundInputDescription>& Inputs = GraphsClassDeclaration->Inputs;
-			for (FMetasoundInputDescription& Input : Inputs)
-			{
-				if (Input.Name == InInputName)
-				{
-					OutDataType = Input.TypeName;
-					return true;
-				}
-			}
-
-			ensureAlwaysMsgf(false, TEXT("Couldn't find Input of name %s in this Metasoud graph!"), *InInputName);
-			return false;
-		}
-
-		FGraphHandle FGraphHandle::InvalidHandle()
-		{
-			FDescPath InvalidPath;
-			FString InvalidClassName;
-			FHandleInitParams InitParams = { FDescriptionAccessPoint::CreateInvalid(), InvalidPath, FMetasoundNodeDescription::InvalidID, FMetasoundClassDescription::InvalidID, nullptr };
-			return FGraphHandle(FHandleInitParams::Token, InitParams);
-		}
-
-		bool FGraphHandle::IsValid() const
-		{
-			return OwningDocument.IsValid() && GraphPtr.IsValid() && GraphsClassDeclaration.IsValid();
-		}
-
-		void FGraphHandle::CopyGraph(FGraphHandle& InOther)
-		{
-			*InOther.GraphPtr.Get() = *GraphPtr.Get();
-			*InOther.GraphsClassDeclaration.Get() = *GraphsClassDeclaration.Get();
-			*InOther.OwningDocument.Get() = *OwningDocument.Get();
-		}
-
-		bool FGraphHandle::IsRequiredInput(const FString& InInputName) const
-		{
-			const TArray<FMetasoundInputDescription>& RequiredInputs = OwningDocument->Archetype.RequiredInputs;
-			for (const FMetasoundInputDescription& RequiredInput : RequiredInputs)
-			{
-				if (InInputName == RequiredInput.Name)
-				{
-					return true;
-				}
-			}
-
-			return false;
-		}
-
-		bool FGraphHandle::IsRequiredOutput(const FString& InOutputName) const
-		{
-			const TArray<FMetasoundOutputDescription>& RequiredOutputs = OwningDocument->Archetype.RequiredOutputs;
-			for (const FMetasoundOutputDescription& RequiredOutput : RequiredOutputs)
-			{
-				if (InOutputName == RequiredOutput.Name)
-				{
-					return true;
-				}
-			}
-
-			return false;
-		}
-
-		const TArray<FMetasoundInputDescription>& FGraphHandle::GetRequiredInputs() const
-		{
-			static const TArray<FMetasoundInputDescription> InvalidDesc;
-			if (!OwningDocument.IsValid())
-			{
-				return InvalidDesc;
-			}
-
-			return OwningDocument->Archetype.RequiredInputs;
-		}
-
-		const TArray<FMetasoundOutputDescription>& FGraphHandle::GetRequiredOutputs() const
-		{
-			static const TArray<FMetasoundOutputDescription> InvalidDesc;
-			if (!OwningDocument.IsValid())
-			{
-				return InvalidDesc;
-			}
-
-			return OwningDocument->Archetype.RequiredOutputs;
-		}
-
-		TArray<FNodeHandle> FGraphHandle::GetAllNodes()
-		{
-			TArray<FNodeHandle> OutArray;
-
-			if (!IsValid())
-			{
-				return OutArray;
-			}
-			
-			const TArray<FMetasoundNodeDescription>& NodeDescriptions = GraphPtr->Nodes;
-			for (const FMetasoundNodeDescription& NodeDescription : NodeDescriptions)
-			{
-				FDescPath NodePath = GraphPtr.GetPath()[Path::EFromGraph::ToNodes][NodeDescription.UniqueID];
-
-				FHandleInitParams InitParams = { GraphPtr.GetAccessPoint(), NodePath, NodeDescription.UniqueID, NodeDescription.DependencyID, OwningAsset };
-				OutArray.Emplace(FHandleInitParams::PrivateToken, InitParams, NodeDescription.ObjectTypeOfNode);
-			}
-
-			return OutArray;
-		}
-
-		FNodeHandle FGraphHandle::GetNodeWithId(int32 InNodeId) const
-		{
-			if (!IsValid())
-			{
-				return FNodeHandle::InvalidHandle();
-			}
-
-			const TArray<FMetasoundNodeDescription>& NodeDescriptions = GraphPtr->Nodes;
-			for (const FMetasoundNodeDescription& NodeDescription : NodeDescriptions)
-			{
-				if (NodeDescription.UniqueID == InNodeId)
-				{
-					FDescPath NodePath = GraphPtr.GetPath()[Path::EFromGraph::ToNodes][InNodeId];
-					FHandleInitParams InitParams = { GraphPtr.GetAccessPoint(), NodePath, NodeDescription.UniqueID, NodeDescription.DependencyID, OwningAsset };
-					return FNodeHandle(FHandleInitParams::PrivateToken, InitParams, NodeDescription.ObjectTypeOfNode);
-				}
-			}
-
-			ensureAlwaysMsgf(false, TEXT("Couldn't find node in graph with ID %u!"), InNodeId);
-			return FNodeHandle::InvalidHandle();
-		}
-
-		TArray<FNodeHandle> FGraphHandle::GetOutputNodes()
-		{
-			TArray<FNodeHandle> OutArray;
-
-			if (!IsValid())
-			{
-				return OutArray;
-			}
-
-			const TArray<FMetasoundNodeDescription>& NodeDescriptions = GraphPtr->Nodes;
-			for (const FMetasoundNodeDescription& NodeDescription : NodeDescriptions)
-			{
-				if (NodeDescription.ObjectTypeOfNode == EMetasoundClassType::Output)
-				{
-					FDescPath NodePath = GraphPtr.GetPath()[Path::EFromGraph::ToNodes][NodeDescription.UniqueID];
-					FHandleInitParams InitParams = { GraphPtr.GetAccessPoint(), NodePath, NodeDescription.UniqueID, NodeDescription.DependencyID, OwningAsset };
-					OutArray.Emplace(FHandleInitParams::PrivateToken, InitParams, NodeDescription.ObjectTypeOfNode);
-				}
-			}
-
-			return OutArray;
-		}
-
-		TArray<FNodeHandle> FGraphHandle::GetInputNodes()
-		{
-			TArray<FNodeHandle> OutArray;
-
-			if (!IsValid())
-			{
-				return OutArray;
-			}
-
-			const TArray<FMetasoundNodeDescription>& NodeDescriptions = GraphPtr->Nodes;
-			for (const FMetasoundNodeDescription& NodeDescription : NodeDescriptions)
-			{
-				if (NodeDescription.ObjectTypeOfNode == EMetasoundClassType::Input)
-				{
-					FDescPath NodePath = GraphPtr.GetPath()[Path::EFromGraph::ToNodes][NodeDescription.UniqueID];
-					FHandleInitParams InitParams = { GraphPtr.GetAccessPoint(), NodePath, NodeDescription.UniqueID, NodeDescription.DependencyID, OwningAsset };
-					OutArray.Emplace(FHandleInitParams::PrivateToken, InitParams, NodeDescription.ObjectTypeOfNode);
-				}
-			}
-
-			return OutArray;
-		}
-
-		bool FGraphHandle::ContainsOutputNodeWithName(const FString& InName) const
-		{
-			if (!IsValid())
-			{
-				return false;
-			}
-
-			const TArray<FMetasoundNodeDescription>& NodeDescriptions = GraphPtr->Nodes;
-			for (const FMetasoundNodeDescription& NodeDescription : NodeDescriptions)
-			{
-				if (NodeDescription.ObjectTypeOfNode == EMetasoundClassType::Output && NodeDescription.Name == InName)
-				{
-					return true;
-				}
-			}
-
-			return false;
-		}
-
-		bool FGraphHandle::ContainsInputNodeWithName(const FString& InName) const
-		{
-			if (!IsValid())
-			{
-				return false;
-			}
-
-			const TArray<FMetasoundNodeDescription>& NodeDescriptions = GraphPtr->Nodes;
-			for (const FMetasoundNodeDescription& NodeDescription : NodeDescriptions)
-			{
-				if (NodeDescription.ObjectTypeOfNode == EMetasoundClassType::Input && NodeDescription.Name == InName)
-				{
-					return true;
-				}
-			}
-
-			return false;
-		}
-
-		FNodeHandle FGraphHandle::GetOutputNodeWithName(const FString& InName)
-		{
-			if (!IsValid())
-			{
-				return FNodeHandle::InvalidHandle();
-			}
-
-			const TArray<FMetasoundNodeDescription>& NodeDescriptions = GraphPtr->Nodes;
-			for (const FMetasoundNodeDescription& NodeDescription : NodeDescriptions)
-			{
-				if (NodeDescription.ObjectTypeOfNode == EMetasoundClassType::Output && NodeDescription.Name == InName)
-				{
-					FDescPath NodePath = GraphPtr.GetPath()[Path::EFromGraph::ToNodes][NodeDescription.UniqueID];
-
-					FHandleInitParams InitParams = { GraphPtr.GetAccessPoint(), NodePath, NodeDescription.UniqueID, NodeDescription.DependencyID, OwningAsset };
-					return FNodeHandle(FHandleInitParams::PrivateToken, InitParams, NodeDescription.ObjectTypeOfNode);
-				}
-			}
-
-			ensureAlwaysMsgf(false, TEXT("Tried to get output node %s, but it didn't exist"), *InName);
-			return FNodeHandle::InvalidHandle();
-		}
-
-		FNodeHandle FGraphHandle::GetInputNodeWithName(const FString& InName)
-		{
-			if (!IsValid())
-			{
-				return FNodeHandle::InvalidHandle();
-			}
-
-			const TArray<FMetasoundNodeDescription>& NodeDescriptions = GraphPtr->Nodes;
-			for (const FMetasoundNodeDescription& NodeDescription : NodeDescriptions)
-			{
-				if (NodeDescription.ObjectTypeOfNode == EMetasoundClassType::Input && NodeDescription.Name == InName)
-				{
-					FDescPath NodePath = GraphPtr.GetPath()[Path::EFromGraph::ToNodes][NodeDescription.UniqueID];
-
-					// todo: add special enum for input/output nodes
-					FHandleInitParams InitParams = { GraphPtr.GetAccessPoint(), NodePath, NodeDescription.UniqueID, FMetasoundClassDescription::InvalidID, OwningAsset };
-					return FNodeHandle(FHandleInitParams::PrivateToken, InitParams, NodeDescription.ObjectTypeOfNode);
-				}
-			}
-
-			ensureAlwaysMsgf(false, TEXT("Tried to get output node %s, but it didn't exist"), *InName);
-			return FNodeHandle::InvalidHandle();
-		}
-
-		FNodeHandle FGraphHandle::AddNewInput(const FMetasoundInputDescription& InDescription)
-		{
-			if (!IsValid())
-			{
-				return FNodeHandle::InvalidHandle();
-			}
-			
-			// @todo: verify that InDescription.TypeName is a valid Metasound type.
-
-			const int32 NewUniqueId = FindNewUniqueNodeId();
-			
-			if (!ensureAlwaysMsgf(NewUniqueId != FMetasoundNodeDescription::InvalidID, TEXT("FindNewUniqueNodeId failed!")))
-			{
-				return FNodeHandle::InvalidHandle();
-			}
-
-			TArray<FMetasoundInputDescription>& Inputs = GraphsClassDeclaration->Inputs;
-
-			// Sanity check that this input has a unique name.
-			for (const FMetasoundInputDescription& Input : Inputs)
-			{
-				if (!ensureAlwaysMsgf(Input.Name != InDescription.Name, TEXT("Tried to add a new input with a name that already exists!")))
-				{
-					return FNodeHandle::InvalidHandle();
-				}
-			}
-
-			// Add the input to this node's class description.
-			Inputs.Add(InDescription);
-			ClearLiteralForInput(InDescription.Name);
-
-			FMetasoundNodeDescription NewNodeDescription;
-			NewNodeDescription.Name = InDescription.Name;
-			NewNodeDescription.UniqueID = NewUniqueId;
-			NewNodeDescription.ObjectTypeOfNode = EMetasoundClassType::Input;
-
-			GraphPtr->Nodes.Add(NewNodeDescription);
-
-			FDescPath NodePath = GraphPtr.GetPath()[Path::EFromGraph::ToNodes][NewNodeDescription.UniqueID];
-			FHandleInitParams InitParams = { GraphPtr.GetAccessPoint(), NodePath, NewUniqueId, FMetasoundClassDescription::InvalidID, OwningAsset };
-
-			// todo: add special enum for input and output nodes
-			return FNodeHandle(FHandleInitParams::PrivateToken, InitParams, NewNodeDescription.ObjectTypeOfNode);
-		}
-
-		bool FGraphHandle::RemoveInput(const FString& InInputName)
-		{
-			if (!IsValid())
-			{
-				return false;
-			}
-
-			TArray<FMetasoundInputDescription>& Inputs = GraphsClassDeclaration->Inputs;
-			int32 IndexOfInputToRemove = -1;
-
-			for (int32 InputIndex = 0; InputIndex < Inputs.Num(); InputIndex++)
-			{
-				if (Inputs[InputIndex].Name == InInputName)
-				{
-					IndexOfInputToRemove = InputIndex;
-					break;
-				}
-			}
-
-			if (!ensureAlwaysMsgf(IndexOfInputToRemove >= 0, TEXT("Tried to remove an Input that didn't exist: %s"), *InInputName))
-			{
-				return false;
-			}
-
-			// find the corresponding node handle to delete.
-			FNodeHandle InputNode = GetInputNodeWithName(InInputName);
-
-			// If we found the input declared in the class description but couldn't find it in the graph,
-			// something has gone terribly wrong. Remove the input from the description, but still ensure.
-			if (!ensureAlwaysMsgf(InputNode.IsValid(), TEXT(R"(Couldn't find an input node with name %s, even though we found the input listed as a dependency.
-				This indicates the underlying FMetasoundClassDescription is corrupted.
-				Removing the Input in the class dependency to resolve...)"), *InInputName))
-			{
-				Inputs.RemoveAt(IndexOfInputToRemove);
-				return true;
-			}
-
-			// Finally, remove the node, and remove the input.
-			if (!ensureAlwaysMsgf(RemoveNodeInternal(InputNode), TEXT("Call to RemoveNodeInternal failed.")))
-			{
-				return false;
-			}
-
-			Inputs.RemoveAt(IndexOfInputToRemove);
-			return true;
-		}
-
-		const FMetasoundEditorData& FGraphHandle::GetEditorData() const
-		{
-			return OwningDocument->EditorData;
-		}
-
-		void FGraphHandle::SetEditorData(const FMetasoundEditorData& InEditorData)
-		{
-			OwningDocument->EditorData = InEditorData;
-		}
-
-		FNodeHandle FGraphHandle::AddNewOutput(const FMetasoundOutputDescription& InDescription)
-		{
-			if (!IsValid())
-			{
-				return FNodeHandle::InvalidHandle();
-			}
-
-			// @todo: verify that InDescription.TypeName is a valid Metasound type.
-
-			const int32 NewUniqueId = FindNewUniqueNodeId();
-
-			if (!ensureAlwaysMsgf(NewUniqueId != FMetasoundNodeDescription::InvalidID, TEXT("FindNewUniqueNodeId failed")))
-			{
-				return FNodeHandle::InvalidHandle();
-			}
-
-
-			TArray<FMetasoundOutputDescription>& Outputs = GraphsClassDeclaration->Outputs;
-
-			// Sanity check that this input has a unique name.
-			for (const FMetasoundOutputDescription& Output : Outputs)
-			{
-				if (!ensureAlwaysMsgf(Output.Name != InDescription.Name, TEXT("Tried to add a new output with a name that already exists!")))
-				{
-					return FNodeHandle::InvalidHandle();
-				}
-			}
-
-			// Add the output to this node's class description.
-			Outputs.Add(InDescription);
-
-			// Add a node for this output to the graph description.
-			FMetasoundNodeDescription NewNodeDescription;
-			NewNodeDescription.Name = InDescription.Name;
-			NewNodeDescription.UniqueID = NewUniqueId;
-			NewNodeDescription.ObjectTypeOfNode = EMetasoundClassType::Output;
-
-			GraphPtr->Nodes.Add(NewNodeDescription);
-
-			FDescPath NodePath = GraphPtr.GetPath()[Path::EFromGraph::ToNodes][NewNodeDescription.UniqueID];
-
-			// todo: add special enum or class for input/output nodes
-			FHandleInitParams InitParams = { GraphPtr.GetAccessPoint(), NodePath, NewNodeDescription.UniqueID, NewNodeDescription.DependencyID, OwningAsset };
-			return FNodeHandle(FHandleInitParams::PrivateToken, InitParams, NewNodeDescription.ObjectTypeOfNode);
-		}
-
-		bool FGraphHandle::RemoveOutput(const FString& OutputName)
-		{
-			if (!IsValid())
-			{
-				return false;
-			}
-
-			TArray<FMetasoundOutputDescription>& Outputs = GraphsClassDeclaration->Outputs;
-			int32 IndexOfOutputToRemove = -1;
-
-			for (int32 OutputIndex = 0; OutputIndex < Outputs.Num(); OutputIndex++)
-			{
-				if (Outputs[OutputIndex].Name == OutputName)
-				{
-					IndexOfOutputToRemove = OutputIndex;
-					break;
-				}
-			}
-
-			if (!ensureAlwaysMsgf(IndexOfOutputToRemove >= 0, TEXT("Tried to remove an Output that didn't exist: %s"), *OutputName))
-			{
-				return false;
-			}
-
-			// find the corresponding node handle to delete.
-			FNodeHandle OutputNode = GetOutputNodeWithName(OutputName);
-
-			// If we found the output declared in the class description but couldn't find it in the graph,
-			// something has gone terribly wrong. Remove the output from the description, but still ensure.
-			if (!ensureAlwaysMsgf(OutputNode.IsValid(), TEXT(R"(Couldn't find an output node with name %s, even though we found the output listed as a dependency.
-				This indicates the underlying FMetasoundClassDescription is corrupted.
-				Removing the Output in the class dependency to resolve...)"), *OutputName))
-			{
-				Outputs.RemoveAt(IndexOfOutputToRemove);
-				return true;
-			}
-
-			// Finally, remove the node, and remove the output.
-			if (!ensureAlwaysMsgf(RemoveNodeInternal(OutputNode), TEXT("Call to RemoveNodeInternal failed.")))
-			{
-				return false;
-			}
-
-			Outputs.RemoveAt(IndexOfOutputToRemove);
-			return true;
-		}
-
-		ELiteralArgType FGraphHandle::GetPreferredLiteralTypeForInput(const FString& InInputName)
-		{
-			FName DataType;
-			if (GetDataTypeForInput(InInputName, DataType))
-			{
-				FMetasoundFrontendRegistryContainer* Registry = FMetasoundFrontendRegistryContainer::Get();
-
-				return Registry->GetDesiredLiteralTypeForDataType(DataType);
-			}
-			else
-			{
-				return ELiteralArgType::Invalid;
-			}
-		}
-
-		UClass* FGraphHandle::GetSupportedClassForInput(const FString& InInputName)
-		{
-			FName DataType;
-			if (GetDataTypeForInput(InInputName, DataType))
-			{
-				FMetasoundFrontendRegistryContainer* Registry = FMetasoundFrontendRegistryContainer::Get();
-
-				return Registry->GetLiteralUClassForDataType(DataType);
-			}
-			else
-			{
-				return nullptr;
-			}
-		}
-
-		bool FGraphHandle::SetInputToLiteral(const FString& InInputName, bool bInValue)
-		{
-			FName DataType;
-			if (FMetasoundLiteralDescription* Literal = GetLiteralDescriptionForInput(InInputName, DataType))
-			{
-				if (!ensureAlwaysMsgf(DoesDataTypeSupportLiteralType(DataType, ELiteralArgType::Boolean), TEXT("Tried to set Data Type %s to an unsupported literal type (Boolean)")))
-				{
-					return false;
-				}
-
-				SetLiteralDescription(*Literal, bInValue);
-				return true;
-			}
-
-			return false;
-		}
-
-		bool FGraphHandle::SetInputToLiteral(const FString& InInputName, int32 InValue)
-		{
-			FName DataType;
-			if (FMetasoundLiteralDescription* Literal = GetLiteralDescriptionForInput(InInputName, DataType))
-			{
-				if (!ensureAlwaysMsgf(DoesDataTypeSupportLiteralType(DataType, ELiteralArgType::Integer), TEXT("Tried to set Data Type %s to an unsupported literal type (Integer)")))
-				{
-					return false;
-				}
-
-				SetLiteralDescription(*Literal, InValue);
-				return true;
-			}
-
-			return false;
-		}
-
-		bool FGraphHandle::SetInputToLiteral(const FString& InInputName, float InValue)
-		{
-			FName DataType;
-			if (FMetasoundLiteralDescription* Literal = GetLiteralDescriptionForInput(InInputName, DataType))
-			{
-				if (!ensureAlwaysMsgf(DoesDataTypeSupportLiteralType(DataType, ELiteralArgType::Float), TEXT("Tried to set Data Type %s to an unsupported literal type (Float)")))
-				{
-					return false;
-				}
-
-				SetLiteralDescription(*Literal, InValue);
-				return true;
-			}
-
-			return false;
-		}
-
-		bool FGraphHandle::SetInputToLiteral(const FString& InInputName, const FString& InValue)
-		{
-			FName DataType;
-			if (FMetasoundLiteralDescription* Literal = GetLiteralDescriptionForInput(InInputName, DataType))
-			{
-				if (!ensureAlwaysMsgf(DoesDataTypeSupportLiteralType(DataType, ELiteralArgType::String), TEXT("Tried to set Data Type %s to an unsupported literal type (String)")))
-				{
-					return false;
-				}
-
-				SetLiteralDescription(*Literal, InValue);
-				return true;
-			}
-
-			return false;
-		}
-
-		bool FGraphHandle::SetInputToLiteral(const FString& InInputName, UObject* InValue)
-		{
-			FName DataType;
-			if (FMetasoundLiteralDescription* Literal = GetLiteralDescriptionForInput(InInputName, DataType))
-			{
-				if (!ensureAlwaysMsgf(DoesDataTypeSupportLiteralType(DataType, ELiteralArgType::UObjectProxy), TEXT("Tried to set Data Type %s to an unsupported literal type (String)")))
-				{
-					return false;
-				}
-
-				SetLiteralDescription(*Literal, InValue);
-				return true;
-			}
-
-			return false;
-		
-		
-		}
-
-		bool FGraphHandle::SetInputToLiteral(const FString& InInputName, TArray<UObject*> InValue)
-		{
-			FName DataType;
-			if (FMetasoundLiteralDescription* Literal = GetLiteralDescriptionForInput(InInputName, DataType))
-			{
-				if (!ensureAlwaysMsgf(DoesDataTypeSupportLiteralType(DataType, ELiteralArgType::UObjectProxyArray), TEXT("Tried to set Data Type %s to an unsupported literal type (String)")))
-				{
-					return false;
-				}
-
-				SetLiteralDescription(*Literal, InValue);
-				return true;
-			}
-
-			return false;
-		}
-
-		bool FGraphHandle::ClearLiteralForInput(const FString& InInputName)
-		{
-			FName DataType;
-			if (FMetasoundLiteralDescription* Literal = GetLiteralDescriptionForInput(InInputName, DataType))
-			{
-
-				ClearLiteralDescription(*Literal);
-				return true;
-			}
-
-			return false;
-		}
-
-		FNodeHandle FGraphHandle::AddNewNode(const FNodeClassInfo& InNodeClass)
-		{
-			if (!IsValid())
-			{
-				return FNodeHandle::InvalidHandle();
-			}
-
-			const int32 NodeID = FindNewUniqueNodeId();
-			if (!ensureAlwaysMsgf(NodeID != FMetasoundNodeDescription::InvalidID, TEXT("Call to FindNewUniqueNodeId failed!")))
-			{
-				return FNodeHandle::InvalidHandle();
-			}
-
-			auto IsMatchingDependency = [&](const FMetasoundClassDescription& InDependency) -> bool
-			{
-				// TODO: may need to add Input/Ouptut node types to external lookup.
-				if ((InNodeClass.NodeType == EMetasoundClassType::External) && (InDependency.Metadata.NodeType == EMetasoundClassType::External))
-				{
-					// For external nodes, we rely on the external node lookup info.
-					const bool bIsLookupInfoEqual = (InDependency.ExternalNodeClassLookupInfo.ExternalNodeClassName == InNodeClass.LookupKey.NodeName) && (InDependency.ExternalNodeClassLookupInfo.ExternalNodeClassHash == InNodeClass.LookupKey.NodeHash);
-					return bIsLookupInfoEqual;
-				}
-				else if ((InNodeClass.NodeType == EMetasoundClassType::MetasoundGraph) && (InDependency.Metadata.NodeType == EMetasoundClassType::MetasoundGraph))
-				{
-					// TODO: This dependency lookup should also take into consideration version info, and possibly even more.
-					// For graph nodes, we only rely on the name.
-					return InNodeClass.LookupKey.NodeName == InDependency.ExternalNodeClassLookupInfo.ExternalNodeClassName;
-				}
-
-				return false;
-			};
-
-			// Scan dependency list to see if the class already exists in dependency list.
-			const FMetasoundClassDescription* Dependency = OwningDocument->Dependencies.FindByPredicate(IsMatchingDependency);
-
-			// Add dependency to owning document if it does not exist.
-			if (nullptr == Dependency)
-			{
-				FMetasoundClassDescription NewDependencyClassDescription = GenerateClassDescription(InNodeClass);
-				NewDependencyClassDescription.UniqueID = FindNewUniqueDependencyId();
-
-				Dependency = &(OwningDocument->Dependencies.Add_GetRef(NewDependencyClassDescription));
-			}
-
-			// Add dependency ID to graph
-			GraphsClassDeclaration->DependencyIDs.AddUnique(Dependency->UniqueID);
-
-			// Add a node for this output to the graph description.
-			FMetasoundNodeDescription NewNodeDescription;
-			NewNodeDescription.Name = InNodeClass.NodeName;
-			NewNodeDescription.UniqueID = NodeID;
-			NewNodeDescription.DependencyID = Dependency->UniqueID;
-			NewNodeDescription.ObjectTypeOfNode = InNodeClass.NodeType;
-
-			GraphPtr->Nodes.Add(NewNodeDescription);
-
-			// Create node handle.
-			FDescPath NodePath = GraphPtr.GetPath()[Path::EFromGraph::ToNodes][NewNodeDescription.UniqueID];
-			FHandleInitParams InitParams = { GraphPtr.GetAccessPoint(), NodePath, NodeID, Dependency->UniqueID, OwningAsset };
-
-			return FNodeHandle(FHandleInitParams::PrivateToken, InitParams, NewNodeDescription.ObjectTypeOfNode);
-		}
-
-		FNodeHandle FGraphHandle::AddNewNode(const FNodeRegistryKey& InNodeClass)
-		{
-			// Construct a FNodeClassInfo from this lookup key.
-			FNodeClassInfo ClassInfo;
-			ClassInfo.LookupKey = InNodeClass;
-			ClassInfo.NodeName = InNodeClass.NodeName.ToString();
-			ClassInfo.NodeType = EMetasoundClassType::External;
-
-			return AddNewNode(ClassInfo);
-		}
-
-		bool FGraphHandle::RemoveNode(const FNodeHandle& InNode)
-		{
-			if (!ensureAlwaysMsgf(InNode.GetNodeType() != EMetasoundClassType::Input && InNode.GetNodeType() != EMetasoundClassType::Output,
-				TEXT("Inputs and outputs must be removed explicitly using 'RemoveInput' or 'RemoveOutput').")))
-			{
-				return false;
-			}
-
-			return RemoveNodeInternal(InNode);
-		}
-
-		bool FGraphHandle::RemoveNodeInternal(const FNodeHandle& InNode)
-		{
-			if (!IsValid())
-			{
-				return false;
-			}
-
-			// First, find the node in our nodes list,
-			// while also checking to see if this is the only node of this class left in this graph.
-			TArray<FMetasoundNodeDescription>& Nodes = GraphPtr->Nodes;
-			int32 IndexOfNodeToRemove = -1;
-			int32 NodesOfClass = 0;
-
-			FString NodeClassName = InNode.GetNodeClassName();
-
-			for (int32 NodeIndex = 0; NodeIndex < Nodes.Num(); NodeIndex++)
-			{
-				if (Nodes[NodeIndex].Name == NodeClassName)
-				{
-					NodesOfClass++;
-				}
-				
-				if (Nodes[NodeIndex].UniqueID == InNode.GetNodeID())
-				{
-					IndexOfNodeToRemove = NodeIndex;
-				}
-
-				// If we've found the matching node,
-				// and have found that there is more than one node of this class,
-				// we have found all the info we need.
-				if (IndexOfNodeToRemove > 0 && NodesOfClass > 1)
-				{
-					break;
-				}
-			}
-
-			if (!ensureAlwaysMsgf(IndexOfNodeToRemove >= 0, TEXT(R"(Couldn't find node corresponding to handle (%s ID: %u).
-				Are you sure this FNodeHandle was generated from this FGraphHandle?)"),
-				*InNode.GetNodeClassName(),
-				InNode.GetNodeType()))
-			{
-				return false;
-			}
-
-			if (InNode.GetNodeType() == EMetasoundClassType::Input || InNode.GetNodeType() == EMetasoundClassType::Output)
-			{
-				Nodes.RemoveAt(IndexOfNodeToRemove);
-				return true;
-			}
-
-			// This should never hit based on the logic above.
-			if (!ensureAlwaysMsgf(NodesOfClass > 0, TEXT("Found node with matching ID (%u) but mismatched class (%s). Likely means that the underlying class description was corrupted."),
-				InNode.GetNodeID(),
-				*InNode.GetNodeClassName()))
-			{
-				return false;
-			}
-
-			// If this node was the only node of this class remaining in the graph,
-			// Remove its ID as a dependency for the graph.
-			if (NodesOfClass < 2)
-			{
-				TArray<int32>& DependencyIDs = GraphsClassDeclaration->DependencyIDs;
-				int32 IndexOfDependencyToRemove = -1;
-
-				int32 UniqueIDForThisDependency = FMetasoundClassDescription::InvalidID;
-				TArray<FMetasoundClassDescription>& DependencyClasses = OwningDocument->Dependencies;
-
-				// scan the owning document's depenency classes for a dependency with this name.
-				int32 IndexOfDependencyInDocument = -1;
-				for (int32 Index = 0; Index < DependencyClasses.Num(); Index++)
-				{
-					if (DependencyClasses[Index].Metadata.NodeName == NodeClassName)
-					{
-						IndexOfDependencyInDocument = Index;
-						UniqueIDForThisDependency = DependencyClasses[Index].UniqueID;
-						break;
-					}
-				}
-
-				if (!ensureAlwaysMsgf(IndexOfDependencyInDocument >= 0,
-					TEXT("Couldn't find node class %s in the list of dependencies for this document, but found it in the nodes list. This likely means that the underlying class description is corrupted.)"), *NodeClassName))
-				{
-					return false;
-				}
-
-				for (int32 DependencyIndex = 0; DependencyIndex < DependencyIDs.Num(); DependencyIndex++)
-				{
-					const bool bIsDependencyForNodeToRemove = DependencyIDs[DependencyIndex] == UniqueIDForThisDependency;
-
-					if (bIsDependencyForNodeToRemove)
-					{
-						IndexOfDependencyToRemove = DependencyIndex;
-						break;
-					}
-				}
-
-				if (ensureAlwaysMsgf(IndexOfDependencyToRemove >= 0, TEXT(R"(Couldn't find node class %s in the list of dependencies for this graph, but found it in the nodes list.
-				This likely means that the underlying class description is corrupted.)"), *NodeClassName))
-				{
-					DependencyIDs.RemoveAt(IndexOfDependencyToRemove);
-				}
-
-				// Finally, check to see if there are any classes remaining in this document that depend on this class, and remove it from our dependencies list.
-				bool bFoundUsageOfDependencyInClass = false;
-
-				for (FMetasoundClassDescription& Dependency : DependencyClasses)
-				{
-					if (Dependency.DependencyIDs.Contains(UniqueIDForThisDependency))
-					{
-						bFoundUsageOfDependencyInClass = true;
-					}
-				}
-
-				// Also scan the root graph for this document, which lives outside of the Dependencies list.
-				FMetasoundClassDescription& RootClass = OwningDocument->RootClass;
-				if (RootClass.DependencyIDs.Contains(UniqueIDForThisDependency))
-				{
-					bFoundUsageOfDependencyInClass = true;
-				}
-
-				if (!bFoundUsageOfDependencyInClass)
-				{
-					// we can safely delete this dependency from the document.
-					DependencyClasses.RemoveAt(IndexOfDependencyInDocument);
-				}
-			}
-
-			// Finally, remove the node from the nodes list.
-			Nodes.RemoveAt(IndexOfNodeToRemove);
-			return true;
-		}
-
-		FMetasoundClassMetadata FGraphHandle::GetGraphMetadata()
-		{
-			if (!IsValid())
-			{
-				return FMetasoundClassMetadata();
-			}
-
-			return GraphsClassDeclaration->Metadata;
-		}
-
-		FString FGraphHandle::ExportToJSON() const
-		{
-			TArray<uint8> WriterBuffer;
-			FMemoryWriter MemWriter(WriterBuffer);
-
-			Metasound::TJsonStructSerializerBackend<Metasound::DefaultCharType> Backend(MemWriter, EStructSerializerBackendFlags::Default);
-			FStructSerializer::Serialize<FMetasoundDocument>(*OwningDocument, Backend);
-
-			MemWriter.Close();
-
-			// null terminator
-			WriterBuffer.AddZeroed(sizeof(ANSICHAR));
-
-			FString Output;
-			Output.AppendChars(reinterpret_cast<ANSICHAR*>(WriterBuffer.GetData()), WriterBuffer.Num() / sizeof(ANSICHAR));
-			return Output;
-		}
-
-		bool FGraphHandle::ExportToJSONAsset(const FString& InAbsolutePath) const
-		{
-			if (!IsValid())
-			{
-				return false;
-			}
-
-			if (TUniquePtr<FArchive> FileWriter = TUniquePtr<FArchive>(IFileManager::Get().CreateFileWriter(*InAbsolutePath)))
-			{
-				
-				TJsonStructSerializerBackend<DefaultCharType> Backend(*FileWriter, EStructSerializerBackendFlags::Default);
-				FStructSerializer::Serialize<FMetasoundDocument>(*OwningDocument, Backend);
-		
-				FileWriter->Close();
-
-				return true;
-			}
-			else
-			{
-				ensureAlwaysMsgf(false, TEXT("Failed to create a filewriter with the given path."));
-				return false;
-			}
-		}
-
-		bool FGraphHandle::InflateNodeDirectlyIntoGraph(FNodeHandle& InNode)
-		{
-			ensureAlwaysMsgf(false, TEXT("Implement me!"));
-			// nontrivial but required anyways for graph inflation (UEAU-475)
-
-			//step 0: get the node's FMetasoundClassDescription
-			//step 1: check if the node is itself a metasound
-			//step 2: get the FMetasoundGraphDescription for the node from the Dependencies list.
-			//step 3: create new unique IDs for each node in the subgraph.
-			//step 4: add nodes from subgraph to the current graph.
-			//step 5: rebuild connections for new nodes in current graph based on the new IDs.
-			//step 6: delete Input nodes and Output nodes from the subgraph, and rebuild connections from this graph directly to the nodes in the subgraph.
-			return false;
-		}
-
-		TTuple<FGraphHandle, FNodeHandle> FGraphHandle::CreateEmptySubgraphNode(const FMetasoundClassMetadata& InInfo)
-		{
-			auto BuildInvalidTupleHandle = [&]() -> TTuple<FGraphHandle, FNodeHandle>
-			{
-				return TTuple<FGraphHandle, FNodeHandle>(FGraphHandle::InvalidHandle(), FNodeHandle::InvalidHandle());
-			};
-
-			if (!IsValid())
-			{
-				return BuildInvalidTupleHandle();
-			}
-
-			// Sanity check that the given name isn't already in our graph's dependency list.
-			TArray<FMetasoundClassDescription>& Dependencies = OwningDocument->Dependencies;
-			for (FMetasoundClassDescription& Dependency : Dependencies)
-			{
-				if (!ensureAlwaysMsgf(Dependency.Metadata.NodeName != InInfo.NodeName, TEXT("Tried to create a new subgraph with name %s but there was already a dependency named that in the graph."), *InInfo.NodeName))
-				{
-					return BuildInvalidTupleHandle();
-				}
-			}
-
-			// Create a new class in this graph's dependencies list:
-			int32 NewUniqueIDForGraph = FindNewUniqueDependencyId();
-			if (!ensureAlwaysMsgf(NewUniqueIDForGraph != FMetasoundClassDescription::InvalidID, TEXT("Call to FindNewUniqueNodeId failed!")))
-			{
-				return BuildInvalidTupleHandle();
-			}
-
-			FMetasoundClassDescription& NewGraphClass = Dependencies.Add_GetRef(FMetasoundClassDescription());
-			NewGraphClass.Metadata = InInfo;
-			NewGraphClass.Metadata.NodeType = EMetasoundClassType::MetasoundGraph;
-			NewGraphClass.UniqueID = NewUniqueIDForGraph;
-
-			// Add the new subgraph's ID as a dependency for the current graph:
-			GraphsClassDeclaration->DependencyIDs.Add(NewUniqueIDForGraph);
-
-			// Generate a new FGraphHandle for this subgraph:
-			FDescPath PathForNewGraph = FDescPath()[Path::EFromDocument::ToDependencies][NewUniqueIDForGraph][Path::EFromClass::ToGraph];
-			FHandleInitParams InitParams = { GraphsClassDeclaration.GetAccessPoint(), PathForNewGraph, FMetasoundNodeDescription::InvalidID, NewUniqueIDForGraph, OwningAsset };
-			FGraphHandle SubgraphHandle = FGraphHandle(FHandleInitParams::PrivateToken, InitParams);
-
-			// Create the node for this subgraph in the current graph:
-			const int32 NewUniqueId = FindNewUniqueNodeId();
-			if (!ensureAlwaysMsgf(NewUniqueId != FMetasoundNodeDescription::InvalidID, TEXT("Call to FindNewUniqueNodeId failed!")))
-			{
-				return BuildInvalidTupleHandle();
-			}
-
-			// Add a node for this output to the graph description.
-			FMetasoundNodeDescription NewNodeDescription;
-			NewNodeDescription.Name = InInfo.NodeName;
-			NewNodeDescription.UniqueID = NewUniqueId;
-			NewNodeDescription.DependencyID = NewUniqueIDForGraph;
-			NewNodeDescription.ObjectTypeOfNode = InInfo.NodeType;
-
-			GraphPtr->Nodes.Add(NewNodeDescription);
-
-			FDescPath NodePath = GraphPtr.GetPath()[Path::EFromGraph::ToNodes][NewNodeDescription.UniqueID];
-			FHandleInitParams NodeInitParams = { GraphPtr.GetAccessPoint(), NodePath, NewUniqueId, NewUniqueIDForGraph, OwningAsset };
-			FNodeHandle SubgraphNode = FNodeHandle(FHandleInitParams::PrivateToken, NodeInitParams, NewNodeDescription.ObjectTypeOfNode);
-
-			return TTuple<FGraphHandle, FNodeHandle>(SubgraphHandle, SubgraphNode);
-		}
-
-		TUniquePtr<IOperator> FGraphHandle::BuildOperator(const FOperatorSettings& InSettings, const FMetasoundEnvironment& InEnvironment, TArray<IOperatorBuilder::FBuildErrorPtr>& OutBuildErrors) const
-		{
-			if (!IsValid())
-			{
-				return TUniquePtr<IOperator>(nullptr);
-			}
-
-			// TODO: Implement inflation step here.
-
-			if (!(GraphsClassDeclaration.IsValid() && OwningDocument.IsValid()))
-			{
-				// Need valid description of this class and list of dependencies
-				// in order to build operator.
-				return TUniquePtr<IOperator>(nullptr);
-			}
-
-			// TODO: bubble up errors. 
-			FFrontendGraphBuilder GraphBuilder;
-			TUniquePtr<FFrontendGraph> Graph = GraphBuilder.CreateGraph(*GraphsClassDeclaration, OwningDocument->Dependencies);
-
-			if (!Graph.IsValid())
-			{
-				return TUniquePtr<IOperator>(nullptr);
-			}
-
-			// Step 5: Invoke Operator Builder
-			FOperatorBuilder OperatorBuilder(FOperatorBuilderSettings::GetDefaultSettings());
-
-			return OperatorBuilder.BuildGraphOperator(*Graph, InSettings, InEnvironment, OutBuildErrors);
-		}
-
-		const FText& FGraphHandle::GetInputDisplayName(FString InputName) const
-		{
-			const FMetasoundDocument& RootMetasoundDocument = *OwningDocument;
-
-			for (const FMetasoundInputDescription& Desc : RootMetasoundDocument.RootClass.Inputs)
-			{
-				if (Desc.Name == InputName)
-				{
-					return Desc.DisplayName;
-				}
-			}
-
-			return FText::GetEmpty();
-		}
-
-		const FText& FGraphHandle::GetInputToolTip(FString InputName) const
-		{
-			const FMetasoundDocument& RootMetasoundDocument = *OwningDocument;
-
-			for (const FMetasoundInputDescription& Desc : RootMetasoundDocument.RootClass.Inputs)
-			{
-				if (Desc.Name == InputName)
-				{
-					return Desc.ToolTip;
-				}
-			}
-
-			return FText::GetEmpty();
-		}
-
-		const FText& FGraphHandle::GetOutputDisplayName(FString OutputName) const
-		{
-			const FMetasoundDocument& RootMetasoundDocument = *OwningDocument;
-
-			for (const FMetasoundOutputDescription& Desc : RootMetasoundDocument.RootClass.Outputs)
-			{
-				if (Desc.Name == OutputName)
-				{
-					return Desc.DisplayName;
-				}
-			}
-
-			return FText::GetEmpty();
-		}
-
-		const FText& FGraphHandle::GetOutputToolTip(FString OutputName) const
-		{
-			const FMetasoundDocument& RootMetasoundDocument = *OwningDocument;
-
-			for (const FMetasoundOutputDescription& Desc : RootMetasoundDocument.RootClass.Outputs)
-			{
-				if (Desc.Name == OutputName)
-				{
-					return Desc.ToolTip;
-				}
-			}
-
-			return FText::GetEmpty();
-		}
-
-		void FGraphHandle::SetInputDisplayName(FString InName, const FText& InDisplayName)
-		{
-			FMetasoundDocument& RootMetasoundDocument = *OwningDocument;
-			for (FMetasoundInputDescription& Desc : RootMetasoundDocument.RootClass.Inputs)
-			{
-				if (Desc.Name == InName)
-				{
-					Desc.DisplayName = InDisplayName;
-					break;
-				}
-			}
-		}
-
-		void FGraphHandle::SetOutputDisplayName(FString InName, const FText& InDisplayName)
-		{
-			FMetasoundDocument& RootMetasoundDocument = *OwningDocument;
-			for (FMetasoundOutputDescription& Desc : RootMetasoundDocument.RootClass.Outputs)
-			{
-				if (Desc.Name == InName)
-				{
-					Desc.DisplayName = InDisplayName;
-					break;
-				}
-			}
-		}
-
-		FGraphHandle GetGraphHandleForClass(const FMetasoundClassDescription& InClass)
-		{
-			ensureAlwaysMsgf(false, TEXT("Implement Me!"));
-
-			// to implement this, we'll need to
-			// step 1. add tags for metasound asset UObject types that make their node name/inputs/outputs asset registry searchable 
-			// step 2. search the asset registry for the assets.
-			// step 3. Consider a runtime implementation for this using soft object paths.
-			return FGraphHandle::InvalidHandle();
-		}
 	}
 }
 
@@ -2288,12 +232,12 @@ class FMetasoundFrontendModule : public IModuleInterface
 };
 
 
-REGISTER_METASOUND_DATATYPE(bool, "Primitive:Bool", ::Metasound::ELiteralArgType::Boolean)
-REGISTER_METASOUND_DATATYPE(int32, "Primitive:Int32", ::Metasound::ELiteralArgType::Integer)
-REGISTER_METASOUND_DATATYPE(int64, "Primitive:Int64", ::Metasound::ELiteralArgType::Integer)
-REGISTER_METASOUND_DATATYPE(float, "Primitive:Float", ::Metasound::ELiteralArgType::Float)
-REGISTER_METASOUND_DATATYPE(double, "Primitive:Double", ::Metasound::ELiteralArgType::Float)
-REGISTER_METASOUND_DATATYPE(FString, "Primitive:String", ::Metasound::ELiteralArgType::String)
+REGISTER_METASOUND_DATATYPE(bool, "Primitive:Bool", ::Metasound::ELiteralType::Boolean)
+REGISTER_METASOUND_DATATYPE(int32, "Primitive:Int32", ::Metasound::ELiteralType::Integer)
+REGISTER_METASOUND_DATATYPE(int64, "Primitive:Int64", ::Metasound::ELiteralType::Integer)
+REGISTER_METASOUND_DATATYPE(float, "Primitive:Float", ::Metasound::ELiteralType::Float)
+REGISTER_METASOUND_DATATYPE(double, "Primitive:Double", ::Metasound::ELiteralType::Float)
+REGISTER_METASOUND_DATATYPE(FString, "Primitive:String", ::Metasound::ELiteralType::String)
 
 REGISTER_METASOUND_DATATYPE(Metasound::FAudioBuffer, "Audio:Buffer")
 REGISTER_METASOUND_DATATYPE(Metasound::FSendAddress, "Send:Address")
