Metasound Frontend Revamp

diff --git a/Engine/Plugins/Runtime/Metasound/Source/MetasoundFrontend/Private/MetasoundFrontendBaseClasses.cpp b/Engine/Plugins/Runtime/Metasound/Source/MetasoundFrontend/Private/MetasoundFrontendBaseClasses.cpp
--- a/Engine/Plugins/Runtime/Metasound/Source/MetasoundFrontend/Private/MetasoundFrontendBaseClasses.cpp
+++ b/Engine/Plugins/Runtime/Metasound/Source/MetasoundFrontend/Private/MetasoundFrontendBaseClasses.cpp
@@ -21,7 +21,7 @@ namespace Metasound
 			return FMetasoundFrontendRegistryContainer::Get()->ConstructInputNode(InInputType, MoveTemp(InParams));
 		}
 
-		TUniquePtr<INode> ConstructOutputNode(const FName& InOutputType, const FOutputNodeConstrutorParams& InParams)
+		TUniquePtr<INode> ConstructOutputNode(const FName& InOutputType, const FOutputNodeConstructorParams& InParams)
 		{
 			return FMetasoundFrontendRegistryContainer::Get()->ConstructOutputNode(InOutputType, InParams);
 		}
@@ -31,9 +31,10 @@ namespace Metasound
 			return FMetasoundFrontendRegistryContainer::Get()->ConstructExternalNode(InNodeType, InNodeHash, InInitData);
 		}
 
-		void SetLiteralDescription(FMetasoundLiteralDescription& OutDescription, bool InValue)
+#if 0
+		void SetLiteralDescription(FMetasoundFrontendLiteral& OutDescription, bool InValue)
 		{
-			OutDescription.LiteralType = EMetasoundLiteralType::Bool;
+			OutDescription.LiteralType = EMetasoundFrontendLiteralType::Bool;
 			OutDescription.AsBool = InValue;
 			OutDescription.AsInteger = 0;
 			OutDescription.AsFloat = 0.0f;
@@ -42,9 +43,9 @@ namespace Metasound
 			OutDescription.AsUObjectArray.Empty();
 		}
 
-		void SetLiteralDescription(FMetasoundLiteralDescription& OutDescription, int32 InValue)
+		void SetLiteralDescription(FMetasoundFrontendLiteral& OutDescription, int32 InValue)
 		{
-			OutDescription.LiteralType = EMetasoundLiteralType::Integer;
+			OutDescription.LiteralType = EMetasoundFrontendLiteralType::Integer;
 			OutDescription.AsBool = false;
 			OutDescription.AsInteger = InValue;
 			OutDescription.AsFloat = 0.0f;
@@ -53,9 +54,9 @@ namespace Metasound
 			OutDescription.AsUObjectArray.Empty();
 		}
 
-		void SetLiteralDescription(FMetasoundLiteralDescription& OutDescription, float InValue)
+		void SetLiteralDescription(FMetasoundFrontendLiteral& OutDescription, float InValue)
 		{
-			OutDescription.LiteralType = EMetasoundLiteralType::Float;
+			OutDescription.LiteralType = EMetasoundFrontendLiteralType::Float;
 			OutDescription.AsBool = false;
 			OutDescription.AsInteger = 0;
 			OutDescription.AsFloat = InValue;
@@ -64,9 +65,9 @@ namespace Metasound
 			OutDescription.AsUObjectArray.Empty();
 		}
 
-		void SetLiteralDescription(FMetasoundLiteralDescription& OutDescription, const FString& InValue)
+		void SetLiteralDescription(FMetasoundFrontendLiteral& OutDescription, const FString& InValue)
 		{
-			OutDescription.LiteralType = EMetasoundLiteralType::String;
+			OutDescription.LiteralType = EMetasoundFrontendLiteralType::String;
 			OutDescription.AsBool = false;
 			OutDescription.AsInteger = 0;
 			OutDescription.AsFloat = 0.0f;
@@ -75,9 +76,9 @@ namespace Metasound
 			OutDescription.AsUObjectArray.Empty();
 		}
 
-		void SetLiteralDescription(FMetasoundLiteralDescription& OutDescription, UObject* InValue)
+		void SetLiteralDescription(FMetasoundFrontendLiteral& OutDescription, UObject* InValue)
 		{
-			OutDescription.LiteralType = EMetasoundLiteralType::UObject;
+			OutDescription.LiteralType = EMetasoundFrontendLiteralType::UObject;
 			OutDescription.AsBool = false;
 			OutDescription.AsInteger = 0;
 			OutDescription.AsFloat = 0.0f;
@@ -87,9 +88,9 @@ namespace Metasound
 
 		}
 
-		void SetLiteralDescription(FMetasoundLiteralDescription& OutDescription, const TArray<UObject*>& InValue)
+		void SetLiteralDescription(FMetasoundFrontendLiteral& OutDescription, const TArray<UObject*>& InValue)
 		{
-			OutDescription.LiteralType = EMetasoundLiteralType::UObjectArray;
+			OutDescription.LiteralType = EMetasoundFrontendLiteralType::UObjectArray;
 			OutDescription.AsBool = false;
 			OutDescription.AsInteger = 0;
 			OutDescription.AsFloat = 0.0f;
@@ -98,9 +99,9 @@ namespace Metasound
 			OutDescription.AsUObjectArray = InValue;
 		}
 
-		void ClearLiteralDescription(FMetasoundLiteralDescription& OutDescription)
+		void ClearLiteralDescription(FMetasoundFrontendLiteral& OutDescription)
 		{
-			OutDescription.LiteralType = EMetasoundLiteralType::None;
+			OutDescription.LiteralType = EMetasoundFrontendLiteralType::None;
 			OutDescription.AsBool = false;
 			OutDescription.AsInteger = 0;
 			OutDescription.AsFloat = 0.0f;
@@ -108,190 +109,191 @@ namespace Metasound
 			OutDescription.AsUObject = nullptr;
 			OutDescription.AsUObjectArray.Empty();
 		}
+#endif
 
-		FDataTypeLiteralParam GetLiteralParamForDataType(FName InDataType, const FMetasoundLiteralDescription& InDescription)
+		FLiteral GetLiteralParamForDataType(FName InDataType, const FMetasoundFrontendLiteral& InDescription)
 		{
-			EMetasoundLiteralType LiteralType = InDescription.LiteralType;
+			EMetasoundFrontendLiteralType LiteralType = InDescription.Type;
 			FMetasoundFrontendRegistryContainer* Registry = FMetasoundFrontendRegistryContainer::Get();
 
 			switch (LiteralType)
 			{
-				case EMetasoundLiteralType::None:
+				case EMetasoundFrontendLiteralType::None:
 				{
 					return GetDefaultParamForDataType(InDataType);
 				}
-				case EMetasoundLiteralType::Bool:
+				case EMetasoundFrontendLiteralType::Bool:
 				{
-					if (!Registry->DoesDataTypeSupportLiteralType(InDataType, ELiteralArgType::Boolean))
+					if (!Registry->DoesDataTypeSupportLiteralType(InDataType, ELiteralType::Boolean))
 					{
-						return FDataTypeLiteralParam::InvalidParam();
+						return FLiteral::CreateInvalid();
 					}
 					else
 					{
-						return FDataTypeLiteralParam(InDescription.AsBool);
+						return FLiteral(InDescription.AsBool);
 					}
 				}
-				case EMetasoundLiteralType::Float:
+				case EMetasoundFrontendLiteralType::Float:
 				{
-					if (!Registry->DoesDataTypeSupportLiteralType(InDataType, ELiteralArgType::Float))
+					if (!Registry->DoesDataTypeSupportLiteralType(InDataType, ELiteralType::Float))
 					{
-						return FDataTypeLiteralParam::InvalidParam();
+						return FLiteral::CreateInvalid();
 					}
 					else
 					{
-						return FDataTypeLiteralParam(InDescription.AsFloat);
+						return FLiteral(InDescription.AsFloat);
 					}
 				}
-				case EMetasoundLiteralType::Integer:
+				case EMetasoundFrontendLiteralType::Integer:
 				{
-					if (!Registry->DoesDataTypeSupportLiteralType(InDataType, ELiteralArgType::Integer))
+					if (!Registry->DoesDataTypeSupportLiteralType(InDataType, ELiteralType::Integer))
 					{
-						return FDataTypeLiteralParam::InvalidParam();
+						return FLiteral::CreateInvalid();
 					}
 					else
 					{
-						return FDataTypeLiteralParam(InDescription.AsInteger);
+						return FLiteral(InDescription.AsInteger);
 					}
 				}
-				case EMetasoundLiteralType::String:
+				case EMetasoundFrontendLiteralType::String:
 				{
-					if (!Registry->DoesDataTypeSupportLiteralType(InDataType, ELiteralArgType::String))
+					if (!Registry->DoesDataTypeSupportLiteralType(InDataType, ELiteralType::String))
 					{
-						return FDataTypeLiteralParam::InvalidParam();
+						return FLiteral::CreateInvalid();
 					}
 					else
 					{
-						return FDataTypeLiteralParam(InDescription.AsString);
+						return FLiteral(InDescription.AsString);
 					}
 				}
-				case EMetasoundLiteralType::UObject:
+				case EMetasoundFrontendLiteralType::UObject:
 				{
-					if (!Registry->DoesDataTypeSupportLiteralType(InDataType, ELiteralArgType::UObjectProxy))
+					if (!Registry->DoesDataTypeSupportLiteralType(InDataType, ELiteralType::UObjectProxy))
 					{
-						return FDataTypeLiteralParam::InvalidParam();
+						return FLiteral::CreateInvalid();
 					}
 					else
 					{
 						return Registry->GenerateLiteralForUObject(InDataType, InDescription.AsUObject);
 					}
 				}
-				case EMetasoundLiteralType::UObjectArray:
+				case EMetasoundFrontendLiteralType::UObjectArray:
 				{
-					if (!Registry->DoesDataTypeSupportLiteralType(InDataType, ELiteralArgType::UObjectProxyArray))
+					if (!Registry->DoesDataTypeSupportLiteralType(InDataType, ELiteralType::UObjectProxyArray))
 					{
-						return FDataTypeLiteralParam::InvalidParam();
+						return FLiteral::CreateInvalid();
 					}
 					else
 					{
 						return Registry->GenerateLiteralForUObjectArray(InDataType, InDescription.AsUObjectArray);
 					}
 				}
-				case EMetasoundLiteralType::Invalid:
+				case EMetasoundFrontendLiteralType::Invalid:
 				default:
 				{
-					return FDataTypeLiteralParam::InvalidParam();
+					return FLiteral::CreateInvalid();
 				}
 			}
 		}
 
-		bool DoesDataTypeSupportLiteralType(FName InDataType, EMetasoundLiteralType InLiteralType)
+		bool DoesDataTypeSupportLiteralType(FName InDataType, EMetasoundFrontendLiteralType InLiteralType)
 		{
 			switch (InLiteralType)
 			{
-			case EMetasoundLiteralType::None:
-				return DoesDataTypeSupportLiteralType(InDataType, ELiteralArgType::None);
-			case EMetasoundLiteralType::Bool:
-				return DoesDataTypeSupportLiteralType(InDataType, ELiteralArgType::Boolean);
-			case EMetasoundLiteralType::Float:
-				return DoesDataTypeSupportLiteralType(InDataType, ELiteralArgType::Float);
-			case EMetasoundLiteralType::Integer:
-				return DoesDataTypeSupportLiteralType(InDataType, ELiteralArgType::Integer);
-			case EMetasoundLiteralType::String:
-				return DoesDataTypeSupportLiteralType(InDataType, ELiteralArgType::String);
-			case EMetasoundLiteralType::UObject:
-				return DoesDataTypeSupportLiteralType(InDataType, ELiteralArgType::UObjectProxy);
-			case EMetasoundLiteralType::UObjectArray:
-				return DoesDataTypeSupportLiteralType(InDataType, ELiteralArgType::UObjectProxyArray);
-			case EMetasoundLiteralType::Invalid:
+			case EMetasoundFrontendLiteralType::None:
+				return DoesDataTypeSupportLiteralType(InDataType, ELiteralType::None);
+			case EMetasoundFrontendLiteralType::Bool:
+				return DoesDataTypeSupportLiteralType(InDataType, ELiteralType::Boolean);
+			case EMetasoundFrontendLiteralType::Float:
+				return DoesDataTypeSupportLiteralType(InDataType, ELiteralType::Float);
+			case EMetasoundFrontendLiteralType::Integer:
+				return DoesDataTypeSupportLiteralType(InDataType, ELiteralType::Integer);
+			case EMetasoundFrontendLiteralType::String:
+				return DoesDataTypeSupportLiteralType(InDataType, ELiteralType::String);
+			case EMetasoundFrontendLiteralType::UObject:
+				return DoesDataTypeSupportLiteralType(InDataType, ELiteralType::UObjectProxy);
+			case EMetasoundFrontendLiteralType::UObjectArray:
+				return DoesDataTypeSupportLiteralType(InDataType, ELiteralType::UObjectProxyArray);
+			case EMetasoundFrontendLiteralType::Invalid:
 			default:
 				return false;
 			}
 		}
 
-		bool DoesDataTypeSupportLiteralType(FName InDataType, ELiteralArgType InLiteralType)
+		bool DoesDataTypeSupportLiteralType(FName InDataType, ELiteralType InLiteralType)
 		{
 			FMetasoundFrontendRegistryContainer* Registry = FMetasoundFrontendRegistryContainer::Get();
 
 			return Registry->DoesDataTypeSupportLiteralType(InDataType, InLiteralType);
 		}
 
-		Metasound::FDataTypeLiteralParam GetLiteralParam(const FMetasoundLiteralDescription& InDescription)
+		Metasound::FLiteral GetLiteralParam(const FMetasoundFrontendLiteral& InDescription)
 		{
-			EMetasoundLiteralType LiteralType = InDescription.LiteralType;
+			EMetasoundFrontendLiteralType LiteralType = InDescription.Type;
 
 			switch (LiteralType)
 			{
-				case EMetasoundLiteralType::Bool:
+				case EMetasoundFrontendLiteralType::Bool:
 				{
-					return FDataTypeLiteralParam(InDescription.AsBool);
+					return FLiteral(InDescription.AsBool);
 				}
-				case EMetasoundLiteralType::Float:
+				case EMetasoundFrontendLiteralType::Float:
 				{
-					return FDataTypeLiteralParam(InDescription.AsFloat);
+					return FLiteral(InDescription.AsFloat);
 				}
-				case EMetasoundLiteralType::Integer:
+				case EMetasoundFrontendLiteralType::Integer:
 				{
-					return FDataTypeLiteralParam(InDescription.AsInteger);
+					return FLiteral(InDescription.AsInteger);
 				}
-				case EMetasoundLiteralType::String:
+				case EMetasoundFrontendLiteralType::String:
 				{
-					return FDataTypeLiteralParam(InDescription.AsString);
+					return FLiteral(InDescription.AsString);
 				}
-				case EMetasoundLiteralType::UObject:
-				case EMetasoundLiteralType::UObjectArray:
-				case EMetasoundLiteralType::None:
-				case EMetasoundLiteralType::Invalid:
+				case EMetasoundFrontendLiteralType::UObject:
+				case EMetasoundFrontendLiteralType::UObjectArray:
+				case EMetasoundFrontendLiteralType::None:
+				case EMetasoundFrontendLiteralType::Invalid:
 				default:
 				{
-					return FDataTypeLiteralParam::InvalidParam();
+					return FLiteral::CreateInvalid();
 				}
 			}
 		}
 
-		Metasound::FDataTypeLiteralParam GetDefaultParamForDataType(FName InDataType)
+		Metasound::FLiteral GetDefaultParamForDataType(FName InDataType)
 		{
 			FMetasoundFrontendRegistryContainer* Registry = FMetasoundFrontendRegistryContainer::Get();
 
-			ELiteralArgType DesiredArgType = Registry->GetDesiredLiteralTypeForDataType(InDataType);
+			ELiteralType DesiredArgType = Registry->GetDesiredLiteralTypeForDataType(InDataType);
 
 			switch (DesiredArgType)
 			{
-				case Metasound::ELiteralArgType::Boolean:
+				case Metasound::ELiteralType::Boolean:
 				{
-					return FDataTypeLiteralParam(false);
+					return FLiteral(false);
 				}
-				case Metasound::ELiteralArgType::Integer:
+				case Metasound::ELiteralType::Integer:
 				{
-					return FDataTypeLiteralParam(0);
+					return FLiteral(0);
 				}
-				case Metasound::ELiteralArgType::Float:
+				case Metasound::ELiteralType::Float:
 				{
-					return FDataTypeLiteralParam(0.0f);
+					return FLiteral(0.0f);
 				}
-				case Metasound::ELiteralArgType::String:
+				case Metasound::ELiteralType::String:
 				{
-					return FDataTypeLiteralParam(FString());
+					return FLiteral(FString());
 				}
-				case Metasound::ELiteralArgType::UObjectProxy:
-				case Metasound::ELiteralArgType::UObjectProxyArray:
-				case Metasound::ELiteralArgType::None:
+				case Metasound::ELiteralType::UObjectProxy:
+				case Metasound::ELiteralType::UObjectProxyArray:
+				case Metasound::ELiteralType::None:
 				{
-					return FDataTypeLiteralParam();
+					return FLiteral();
 				}
-				case Metasound::ELiteralArgType::Invalid:
+				case Metasound::ELiteralType::Invalid:
 				default:
 				{
-					return FDataTypeLiteralParam::InvalidParam();
+					return FLiteral::CreateInvalid();
 				}
 			}
 
@@ -506,12 +508,14 @@ namespace Metasound
 			return true;
 		}
 
-		FDescriptionAccessPoint::FDescriptionAccessPoint(TAccessPtr<FMetasoundDocument> InRootDocumentPtr)
+#if 0
+
+		FDescriptionAccessPoint::FDescriptionAccessPoint(TAccessPtr<FMetasoundFrontendDocument> InRootDocumentPtr)
 			: RootDocumentPtr(InRootDocumentPtr)
 		{
 		}
 
-		FMetasoundDocument& FDescriptionAccessPoint::GetRootChecked() const
+		FMetasoundFrontendDocument& FDescriptionAccessPoint::GetRootChecked() const
 		{
 			check(RootDocumentPtr.IsValid());
 
@@ -527,7 +531,7 @@ namespace Metasound
 
 			FDescPath CurrentPath = InPathFromRoot;
 			
-			OutPtr.Set<FMetasoundDocument*>(RootDocumentPtr.Get());
+			OutPtr.Set<FMetasoundFrontendDocument*>(RootDocumentPtr.Get());
 			FDescriptionUnwindStep CurrentStep = { OutPtr, Path::EDescType::Document };
 
 			while (CurrentPath.Num() != 0 && CurrentStep.Type != Path::EDescType::Invalid)
@@ -541,7 +545,7 @@ namespace Metasound
 		}
 
 
-		FDescriptionAccessPoint::FDescriptionUnwindStep FDescriptionAccessPoint::GoToNextFromDocument(FMetasoundDocument& InDocument, FDescPath& InPath, const Path::FElement& InNext) const
+		FDescriptionAccessPoint::FDescriptionUnwindStep FDescriptionAccessPoint::GoToNextFromDocument(FMetasoundFrontendDocument& InDocument, FDescPath& InPath, const Path::FElement& InNext) const
 		{
 			check(RootDocumentPtr.IsValid());
 			switch (InNext.CurrentDescType)
@@ -549,14 +553,14 @@ namespace Metasound
 				case Path::EDescType::Document:
 				{
 					FDescriptionUnwindStep UnwindStep;
-					UnwindStep.DescriptionStructPtr.Set<FMetasoundDocument*>(&InDocument);
+					UnwindStep.DescriptionStructPtr.Set<FMetasoundFrontendDocument*>(&InDocument);
 					UnwindStep.Type = Path::EDescType::Document;
 					return UnwindStep;
 				}
 				case Path::EDescType::Class:
 				{
 					FDescriptionUnwindStep UnwindStep;
-					UnwindStep.DescriptionStructPtr.Set<FMetasoundClassDescription*>(&(InDocument.RootClass));
+					UnwindStep.DescriptionStructPtr.Set<FMetasoundFrontendClass*>(&(InDocument.RootGraph));
 					UnwindStep.Type = Path::EDescType::Class;
 					return UnwindStep;
 				}
@@ -579,26 +583,26 @@ namespace Metasound
 						return FDescriptionUnwindStep::CreateInvalid();
 					}
 
-					if (DependencyID == FMetasoundClassDescription::RootClassID)
+					if (DependencyID == FMetasoundFrontendClass::RootClassID)
 					{
 						FDescriptionUnwindStep UnwindStep;
-						UnwindStep.DescriptionStructPtr.Set<FMetasoundClassDescription*>(&(InDocument.RootClass));
+						UnwindStep.DescriptionStructPtr.Set<FMetasoundFrontendClass*>(&(InDocument.RootGraph));
 						UnwindStep.Type = Path::EDescType::Class;
 						return UnwindStep;
 					}
 
-					TArray<FMetasoundClassDescription>& DependenciesList = InDocument.Dependencies;
+					TArray<FMetasoundFrontendClass>& DependenciesList = InDocument.Dependencies;
 
 					// Dependencies can be looked up by ID or by name.
 					if (DependencyID != INDEX_NONE)
 					{
 						// Scan the dependencies list for the matching lookup ID.
-						for (FMetasoundClassDescription& Dependency : DependenciesList)
+						for (FMetasoundFrontendClass& Dependency : DependenciesList)
 						{
 							if (Dependency.UniqueID == DependencyID)
 							{
 								FDescriptionUnwindStep UnwindStep;
-								UnwindStep.DescriptionStructPtr.Set<FMetasoundClassDescription*>(&Dependency);
+								UnwindStep.DescriptionStructPtr.Set<FMetasoundFrontendClass*>(&Dependency);
 								UnwindStep.Type = Path::EDescType::Class;
 								return UnwindStep;
 							}
@@ -609,12 +613,12 @@ namespace Metasound
 						// TODO: remove this chunk of code in the "else{}" block. Not sure if it has to be here. 
 						checkNoEntry();
 						// fall back to scanning the dependencies list for the matching lookup name.
-						for (FMetasoundClassDescription& Dependency : DependenciesList)
+						for (FMetasoundFrontendClass& Dependency : DependenciesList)
 						{
 							if (Dependency.Metadata.NodeName == DependencyName)
 							{
 								FDescriptionUnwindStep UnwindStep;
-								UnwindStep.DescriptionStructPtr.Set<FMetasoundClassDescription*>(&Dependency);
+								UnwindStep.DescriptionStructPtr.Set<FMetasoundFrontendClass*>(&Dependency);
 								UnwindStep.Type = Path::EDescType::Class;
 								return UnwindStep;
 							}
@@ -632,14 +636,14 @@ namespace Metasound
 			return FDescriptionUnwindStep::CreateInvalid();
 		}
 
-		FDescriptionAccessPoint::FDescriptionUnwindStep FDescriptionAccessPoint::GoToNextFromClass(FMetasoundClassDescription& InClassDescription, FDescPath& InPath, const Path::FElement& InNext) const
+		FDescriptionAccessPoint::FDescriptionUnwindStep FDescriptionAccessPoint::GoToNextFromClass(FMetasoundFrontendClass& InClassDescription, FDescPath& InPath, const Path::FElement& InNext) const
 		{
 			switch (InNext.CurrentDescType)
 			{
 				case Path::EDescType::Graph:
 				{
 					FDescriptionUnwindStep UnwindStep;
-					UnwindStep.DescriptionStructPtr.Set<FMetasoundGraphDescription*>(&(InClassDescription.Graph));
+					UnwindStep.DescriptionStructPtr.Set<FMetasoundFrontendGraph*>(&(InClassDescription.Graph));
 					UnwindStep.Type = Path::EDescType::Graph;
 					return UnwindStep;
 				}
@@ -663,13 +667,13 @@ namespace Metasound
 					FString& InputName = InputElement.LookupName;
 
 					// Scan the inputs list for the lookup name.
-					TArray<FMetasoundInputDescription>& InputsList = InClassDescription.Inputs;
-					for (FMetasoundInputDescription& Input : InputsList)
+					TArray<FMetasoundFrontendClassInput>& InputsList = InClassDescription.Inputs;
+					for (FMetasoundFrontendClassInput& Input : InputsList)
 					{
 						if (Input.Name == InputName)
 						{
 							FDescriptionUnwindStep UnwindStep;
-							UnwindStep.DescriptionStructPtr.Set<FMetasoundInputDescription*>(&Input);
+							UnwindStep.DescriptionStructPtr.Set<FMetasoundFrontendClassInput*>(&Input);
 							UnwindStep.Type = Path::EDescType::Input;
 							return UnwindStep;
 						}
@@ -692,13 +696,13 @@ namespace Metasound
 					FString& OutputName = OutputElement.LookupName;
 
 					// Scan the outputs list for the lookup name.
-					TArray<FMetasoundOutputDescription>& OutputsList = InClassDescription.Outputs;
-					for (FMetasoundOutputDescription& Output : OutputsList)
+					TArray<FMetasoundFrontendClassOutput>& OutputsList = InClassDescription.Outputs;
+					for (FMetasoundFrontendClassOutput& Output : OutputsList)
 					{
 						if (Output.Name == OutputElement.LookupName)
 						{
 							FDescriptionUnwindStep UnwindStep;
-							UnwindStep.DescriptionStructPtr.Set<FMetasoundOutputDescription*>(&Output);
+							UnwindStep.DescriptionStructPtr.Set<FMetasoundFrontendClassOutput*>(&Output);
 							UnwindStep.Type = Path::EDescType::Output;
 							return UnwindStep;
 						}
@@ -711,7 +715,7 @@ namespace Metasound
 				case Path::EDescType::Metadata:
 				{
 					FDescriptionUnwindStep UnwindStep;
-					UnwindStep.DescriptionStructPtr.Set<FMetasoundClassMetadata*>(&(InClassDescription.Metadata));
+					UnwindStep.DescriptionStructPtr.Set<FMetasoundFrontendClassMetadata*>(&(InClassDescription.Metadata));
 					UnwindStep.Type = Path::EDescType::Metadata;
 
 					return UnwindStep;
@@ -740,13 +744,13 @@ namespace Metasound
 			{
 				case Path::EDescType::Document:
 				{
-					FMetasoundDocument* Document = InElement.DescriptionStructPtr.Get<FMetasoundDocument*>();
+					FMetasoundFrontendDocument* Document = InElement.DescriptionStructPtr.Get<FMetasoundFrontendDocument*>();
 					return GoToNextFromDocument(*Document, InPath, NextStep);
 					break;
 				}
 				case Path::EDescType::Class:
 				{
-					FMetasoundClassDescription* ClassDescription = InElement.DescriptionStructPtr.Get<FMetasoundClassDescription*>();
+					FMetasoundFrontendClass* ClassDescription = InElement.DescriptionStructPtr.Get<FMetasoundFrontendClass*>();
 					return GoToNextFromClass(*ClassDescription, InPath, NextStep);
 					break;
 				}
@@ -757,7 +761,7 @@ namespace Metasound
 						return FDescriptionUnwindStep::CreateInvalid();
 					}
 
-					FMetasoundGraphDescription* GraphDescription = InElement.DescriptionStructPtr.Get<FMetasoundGraphDescription*>();
+					FMetasoundFrontendGraph* GraphDescription = InElement.DescriptionStructPtr.Get<FMetasoundFrontendGraph*>();
 				
 					if (!ensureAlwaysMsgf(InPath.Path.Num() != 0, TEXT("Incomplete path! path stopped at Nodes list without specifying a node ID.")))
 					{
@@ -774,13 +778,13 @@ namespace Metasound
 				
 					int32 NodeID = NodeElement.LookupID;
 
-					TArray<FMetasoundNodeDescription>& NodeList = GraphDescription->Nodes;
-					for (FMetasoundNodeDescription& Node : NodeList)
+					TArray<FMetasoundFrontendNode>& NodeList = GraphDescription->Nodes;
+					for (FMetasoundFrontendNode& Node : NodeList)
 					{
 						if (Node.UniqueID == NodeID)
 						{
 							FDescriptionUnwindStep UnwindStep;
-							UnwindStep.DescriptionStructPtr.Set<FMetasoundNodeDescription*>(&Node);
+							UnwindStep.DescriptionStructPtr.Set<FMetasoundFrontendNode*>(&Node);
 							UnwindStep.Type = Path::EDescType::Node;
 							return UnwindStep;
 						}
@@ -959,6 +963,7 @@ namespace Metasound
 
 			return OutString;
 		}
+#endif
 
 		
 
diff --git a/Engine/Plugins/Runtime/Metasound/Source/MetasoundFrontend/Private/MetasoundFrontendController.cpp b/Engine/Plugins/Runtime/Metasound/Source/MetasoundFrontend/Private/MetasoundFrontendController.cpp
new file mode 100644
--- /dev/null
+++ b/Engine/Plugins/Runtime/Metasound/Source/MetasoundFrontend/Private/MetasoundFrontendController.cpp
@@ -0,0 +1,44 @@
+// Copyright Epic Games, Inc. All Rights Reserved.
+
+#include "MetasoundFrontendController.h"
+
+#include "MetasoundFrontendBaseClasses.h"
+#include "MetasoundFrontendInvalidController.h"
+#include "MetasoundFrontendStandardController.h"
+
+namespace Metasound
+{
+	namespace Frontend
+	{
+		FOutputHandle IOutputController::GetInvalidHandle()
+		{
+			return FInvalidOutputController::GetInvalid();
+		}
+
+		FInputHandle IInputController::GetInvalidHandle()
+		{
+			return FInvalidInputController::GetInvalid();
+		}
+
+		FNodeHandle INodeController::GetInvalidHandle()
+		{
+			return FInvalidNodeController::GetInvalid();
+		}
+
+		FGraphHandle IGraphController::GetInvalidHandle()
+		{
+			return FInvalidGraphController::GetInvalid();
+		}
+
+		FDocumentHandle IDocumentController::GetInvalidHandle()
+		{
+			return FInvalidDocumentController::GetInvalid();
+		}
+			
+		FDocumentHandle IDocumentController::CreateDocumentHandle(TAccessPtr<FMetasoundFrontendDocument> InDocument)
+		{
+			// Create using standard document controller.
+			return FDocumentController::CreateDocumentHandle(InDocument);
+		}
+	}
+}
diff --git a/Engine/Plugins/Runtime/Metasound/Source/MetasoundFrontend/Private/MetasoundFrontendDocument.cpp b/Engine/Plugins/Runtime/Metasound/Source/MetasoundFrontend/Private/MetasoundFrontendDocument.cpp
new file mode 100644
--- /dev/null
+++ b/Engine/Plugins/Runtime/Metasound/Source/MetasoundFrontend/Private/MetasoundFrontendDocument.cpp
@@ -0,0 +1,142 @@
+// Copyright Epic Games, Inc. All Rights Reserved.
+
+#include "MetasoundFrontendDocument.h"
+
+#include "MetasoundFrontend.h"
+#include "MetasoundFrontendRegistries.h"
+
+void FMetasoundFrontendLiteral::Set(bool InValue)
+{
+	Clear();
+	AsBool = InValue;
+	Type = EMetasoundFrontendLiteralType::Bool;
+}
+
+void FMetasoundFrontendLiteral::Set(int32 InValue)
+{
+	Clear();
+	AsInteger = InValue;
+	Type = EMetasoundFrontendLiteralType::Integer;
+}
+
+void FMetasoundFrontendLiteral::Set(float InValue)
+{
+	Clear();
+	AsFloat = InValue;
+	Type = EMetasoundFrontendLiteralType::Float;
+}
+
+void FMetasoundFrontendLiteral::Set(const FString& InValue)
+{
+	Clear();
+	AsString = InValue;
+	Type = EMetasoundFrontendLiteralType::String;
+}
+
+void FMetasoundFrontendLiteral::Set(UObject* InValue)
+{
+	Clear();
+	AsUObject = InValue;
+	Type = EMetasoundFrontendLiteralType::UObject;
+}
+
+void FMetasoundFrontendLiteral::Set(const TArray<UObject*>& InValue)
+{
+	Clear();
+	AsUObjectArray = InValue;
+	Type = EMetasoundFrontendLiteralType::UObjectArray;
+}
+
+void FMetasoundFrontendLiteral::Clear()
+{
+	Type = EMetasoundFrontendLiteralType::None;
+
+	AsBool = false;
+	AsInteger = 0;
+	AsFloat = 0.0f;
+	AsString.Empty();
+	AsUObject = nullptr;
+	AsUObjectArray.Empty();
+}
+
+FMetasoundFrontendNodeInterface::FMetasoundFrontendNodeInterface(const FMetasoundFrontendClassInterface& InClassInterface)
+{
+	for (const FMetasoundFrontendClassInput& Input : InClassInterface.Inputs)
+	{
+		Inputs.Add(Input);
+	}
+
+	for (const FMetasoundFrontendClassOutput& Output : InClassInterface.Outputs)
+	{
+		Outputs.Add(Output);
+	}
+
+	for (const FMetasoundFrontendClassEnvironmentVariable& EnvVar : InClassInterface.Environment)
+	{
+		FMetasoundFrontendVertex EnvVertex;
+		EnvVertex.Name = EnvVar.Name;
+		EnvVertex.TypeName = EnvVar.TypeName;
+
+		Environment.Add(MoveTemp(EnvVertex));
+	}
+}
+
+FMetasoundFrontendNode::FMetasoundFrontendNode(const FMetasoundFrontendClass& InClass)
+: ClassID(InClass.ID)
+, Name(InClass.Metadata.Name.Name)
+, Interface(InClass.Interface)
+{
+
+}
+
+bool FMetasoundFrontendVertex::IsFunctionalEquivalent(const FMetasoundFrontendVertex& InLHS, const FMetasoundFrontendVertex& InRHS)
+{
+	return (InLHS.Name == InRHS.Name) && (InLHS.TypeName == InRHS.TypeName) && (InLHS.PointIDs.Num() == InRHS.PointIDs.Num());
+}
+
+bool FMetasoundFrontendVertexBehavior::IsFunctionalEquivalent(const FMetasoundFrontendVertexBehavior& InLHS, const FMetasoundFrontendVertexBehavior& InRHS)
+{
+	return (InLHS.Type == InRHS.Type) && (InLHS.ArrayMin == InRHS.ArrayMin) && (InLHS.ArrayMax == InRHS.ArrayMax);
+}
+
+bool FMetasoundFrontendClassVertex::IsFunctionalEquivalent(const FMetasoundFrontendClassVertex& InLHS, const FMetasoundFrontendClassVertex& InRHS) 
+{
+	return FMetasoundFrontendVertex::IsFunctionalEquivalent(InLHS, InRHS) && FMetasoundFrontendVertexBehavior::IsFunctionalEquivalent(InLHS.Behavior, InRHS.Behavior);
+}
+
+FString FMetasoundFrontendClassName::GetFullName() const
+{
+	//return FString::Format(TEXT("{0}.{1}.{2}"), {Namespace, Name, Variant});
+	return Name;
+}
+
+bool operator==(const FMetasoundFrontendClassName& InLHS, const FMetasoundFrontendClassName& InRHS)
+{
+	return (InLHS.Namespace == InRHS.Namespace) && (InLHS.Name == InRHS.Name) && (InLHS.Variant == InRHS.Variant);
+}
+
+FMetasoundFrontendClassInput::FMetasoundFrontendClassInput(const FMetasoundFrontendClassVertex& InOther)
+:	FMetasoundFrontendClassVertex(InOther)
+{
+	EMetasoundFrontendLiteralType DefaultType = Metasound::Frontend::GetMetasoundLiteralType(FMetasoundFrontendRegistryContainer::Get()->GetDesiredLiteralTypeForDataType(InOther.TypeName));
+
+	for (int32 PointID : PointIDs)
+	{
+		FMetasoundFrontendVertexLiteral Default;
+		Default.PointID = PointID;
+		Default.Value.Type = DefaultType;
+		Defaults.Add(Default);
+	}
+}
+
+FMetasoundFrontendGraphClass::FMetasoundFrontendGraphClass()
+{
+	Metadata.Type = EMetasoundFrontendClassType::Graph;
+}
+
+FMetasoundFrontendDocument::FMetasoundFrontendDocument()
+{
+	RootGraph.ID = 0;
+	RootGraph.Metadata.Type = EMetasoundFrontendClassType::Graph;
+}
+
diff --git a/Engine/Plugins/Runtime/Metasound/Source/MetasoundFrontend/Private/MetasoundFrontendGraph.cpp b/Engine/Plugins/Runtime/Metasound/Source/MetasoundFrontend/Private/MetasoundFrontendGraph.cpp
--- a/Engine/Plugins/Runtime/Metasound/Source/MetasoundFrontend/Private/MetasoundFrontendGraph.cpp
+++ b/Engine/Plugins/Runtime/Metasound/Source/MetasoundFrontend/Private/MetasoundFrontendGraph.cpp
@@ -7,8 +7,8 @@
 #include "Algo/Transform.h"
 #include "CoreMinimal.h"
 #include "MetasoundFrontend.h"
-#include "MetasoundFrontendDataLayout.h"
 #include "MetasoundGraph.h"
+#include "MetasoundLog.h"
 #include "MetasoundNodeInterface.h"
 
 namespace Metasound
@@ -143,196 +143,377 @@ namespace Metasound
 	}
 
 
-	TUniquePtr<INode> FFrontendGraphBuilder::CreateInputNode(const FMetasoundNodeDescription& InNode, const FMetasoundInputDescription& InDescription) const
+	TUniquePtr<INode> FFrontendGraphBuilder::CreateInputNode(const FMetasoundFrontendNode& InNode, const FMetasoundFrontendClass& InClass, const FMetasoundFrontendClassInput& InOwningGraphClassInput) const
 	{
-		if (!ensureAlwaysMsgf(Frontend::DoesDataTypeSupportLiteralType(InDescription.TypeName, InDescription.LiteralValue.LiteralType), TEXT("[%s] cannot be constructed with the provided literal type."), *InDescription.TypeName.ToString()))
+		const FMetasoundFrontendVertexLiteral* VertexLiteral = FindInputLiteralForInputNode(InNode, InClass, InOwningGraphClassInput);
+
+		if (nullptr != VertexLiteral)
 		{
-			return TUniquePtr<INode>(nullptr);
-		}
+			if (ensure(InNode.Interface.Inputs.Num() == 1))
+			{
+				const FMetasoundFrontendVertex& InputVertex = InNode.Interface.Inputs[0];
+
+				const bool IsLiteralParsableByDataType = Frontend::DoesDataTypeSupportLiteralType(InputVertex.TypeName, VertexLiteral->Value.Type);
 
-		FDataTypeLiteralParam LiteralParam = Frontend::GetLiteralParamForDataType(InDescription.TypeName, InDescription.LiteralValue);
+				if (IsLiteralParsableByDataType)
+				{
+					FLiteral Literal = Frontend::GetLiteralParamForDataType(InputVertex.TypeName, VertexLiteral->Value);
+
+					FInputNodeConstructorParams InitParams =
+					{
+						InNode.Name,
+						InputVertex.Name,
+						MoveTemp(Literal)
+					};
 
-		FInputNodeConstructorParams InitParams =
+					// TODO: create input node using external class definition
+					return FMetasoundFrontendRegistryContainer::Get()->ConstructInputNode(InputVertex.TypeName, MoveTemp(InitParams));
+				}
+				else
+				{
+					UE_LOG(LogMetasound, Error, TEXT("Cannot create input node [NodeID:%d]. [Vertex:%s] cannot be constructed with the provided literal type."), InNode.ID, *InputVertex.Name);
+				}
+			}
+		}
+		else
 		{
-			InDescription.Name,
-			InDescription.Name,
-			MoveTemp(LiteralParam)
-		};
+			UE_LOG(LogMetasound, Error, TEXT("Cannot create input node [NodeID:%d]. No default literal set for input node."), InNode.ID);
+		}
 
-		return FMetasoundFrontendRegistryContainer::Get()->ConstructInputNode(InDescription.TypeName, MoveTemp(InitParams));
+		return TUniquePtr<INode>(nullptr);
 	}
 
-	TUniquePtr<INode> FFrontendGraphBuilder::CreateOutputNode(const FMetasoundNodeDescription& InNode, const FMetasoundOutputDescription& InDescription) const
+	TUniquePtr<INode> FFrontendGraphBuilder::CreateOutputNode(const FMetasoundFrontendNode& InNode, const FMetasoundFrontendClass& InClass) const
 	{
-		FOutputNodeConstrutorParams InitParams =
+		check(InClass.Metadata.Type == EMetasoundFrontendClassType::Output);
+		check(InNode.ClassID == InClass.ID);
+
+		if (ensure(InNode.Interface.Outputs.Num() == 1))
 		{
-			InDescription.Name,
-			InDescription.Name,
-		};
+			const FMetasoundFrontendVertex& OutputVertex = InNode.Interface.Outputs[0];
+
+			FOutputNodeConstructorParams InitParams =
+			{
+				InNode.Name,
+				OutputVertex.Name,
+			};
 
-		return FMetasoundFrontendRegistryContainer::Get()->ConstructOutputNode(InDescription.TypeName, InitParams);
+			// TODO: create output node using external class definition
+			return FMetasoundFrontendRegistryContainer::Get()->ConstructOutputNode(OutputVertex.TypeName, InitParams);
+		}
+		return TUniquePtr<INode>(nullptr);
 	}
 
-	TUniquePtr<INode> FFrontendGraphBuilder::CreateExternalNode(const FMetasoundNodeDescription& InNode, const FMetasoundClassDescription& InClass) const
+	TUniquePtr<INode> FFrontendGraphBuilder::CreateExternalNode(const FMetasoundFrontendNode& InNode, const FMetasoundFrontendClass& InClass) const
 	{
-		check(InClass.Metadata.NodeType == EMetasoundClassType::External);
-		check(InNode.ObjectTypeOfNode == EMetasoundClassType::External);
-
-		// We found a match. now we just need to create the input node.
-		const FMetasoundExternalClassLookupInfo& LookupInfo = InClass.ExternalNodeClassLookupInfo;
+		check(InClass.Metadata.Type == EMetasoundFrontendClassType::External);
+		check(InNode.ClassID == InClass.ID);
 
 		FNodeInitData InitData;
 
 		InitData.InstanceName.Append(InNode.Name);
 		InitData.InstanceName.AppendChar('_');
-		InitData.InstanceName.AppendInt(InNode.UniqueID);
-		
+		InitData.InstanceName.AppendInt(InNode.ID);
+
 		// Copy over our initialization params.
+		/*
 		for (auto& StaticParamTuple : InNode.StaticParameters)
 		{
-			FDataTypeLiteralParam LiteralParam = Frontend::GetLiteralParam(StaticParamTuple.Value);
+			FLiteral LiteralParam = Frontend::GetLiteralParam(StaticParamTuple.Value);
 
 			if (LiteralParam.IsValid())
 			{
 				InitData.ParamMap.Add(StaticParamTuple.Key, MoveTemp(LiteralParam));
 			}
 		}
+		*/
+		
+		// TODO: handle check to see if node interface conforms to class interface here. 
+		// TODO: check to see if external object supports class interface.
+
+		Metasound::Frontend::FNodeRegistryKey Key = FMetasoundFrontendRegistryContainer::GetRegistryKey(InClass.Metadata);
 
-		return FMetasoundFrontendRegistryContainer::Get()->ConstructExternalNode(LookupInfo.ExternalNodeClassName, LookupInfo.ExternalNodeClassHash, InitData);
+		return FMetasoundFrontendRegistryContainer::Get()->ConstructExternalNode(Key.NodeName, Key.NodeHash, InitData);
 	}
 
-	// TODO: add errors here. Most will be a "PromptIfMissing"...
-	void FFrontendGraphBuilder::AddExternalNodesToGraph(const TArray<FMetasoundNodeDescription>& InNodes, const TMap<int32, const FMetasoundClassDescription*>& InClasses, FFrontendGraph& OutGraph) const
+
+	const FMetasoundFrontendClassInput* FFrontendGraphBuilder::FindClassInputForInputNode(const FMetasoundFrontendGraphClass& InOwningGraph, const FMetasoundFrontendNode& InInputNode, int32& OutClassInputIndex) const
 	{
-		for (const FMetasoundNodeDescription& Node : InNodes)
+		OutClassInputIndex = INDEX_NONE;
+
+		// TODO: assumes input node has exactly one input 
+		if (ensure(InInputNode.Interface.Inputs.Num() == 1))
 		{
-			if (Node.ObjectTypeOfNode == EMetasoundClassType::External)
+			const FName& TypeName = InInputNode.Interface.Inputs[0].TypeName;
+
+			auto IsMatchingInput = [&](const FMetasoundFrontendClassInput& GraphInput)
 			{
-				const FMetasoundClassDescription* ClassDesc = InClasses.FindRef(Node.DependencyID);
+				return (InInputNode.ID == GraphInput.NodeID);
+			};
 
-				if (ensure(nullptr != ClassDesc))
-				{
-					OutGraph.AddNode(Node.UniqueID, CreateExternalNode(Node, *ClassDesc));
-				}
+			OutClassInputIndex = InOwningGraph.Interface.Inputs.IndexOfByPredicate(IsMatchingInput);
+			if (INDEX_NONE != OutClassInputIndex)
+			{
+				return &InOwningGraph.Interface.Inputs[OutClassInputIndex];
 			}
 		}
+		return nullptr;
 	}
 
-	void FFrontendGraphBuilder::AddEdgesToGraph(const FMetasoundGraphDescription& InGraphDescription, FFrontendGraph& OutGraph) const
+	const FMetasoundFrontendClassOutput* FFrontendGraphBuilder::FindClassOutputForOutputNode(const FMetasoundFrontendGraphClass& InOwningGraph, const FMetasoundFrontendNode& InOutputNode, int32& OutClassOutputIndex) const
 	{
-		for (const FMetasoundNodeDescription& NodeDescription : InGraphDescription.Nodes)
+		OutClassOutputIndex = INDEX_NONE;
+
+		// TODO: assumes input node has exactly one input 
+		if (ensure(InOutputNode.Interface.Outputs.Num() == 1))
 		{
-			const INode* ToNode = OutGraph.FindNode(NodeDescription.UniqueID);
+			const FName& TypeName = InOutputNode.Interface.Outputs[0].TypeName;
 
-			if (nullptr == ToNode)
+			auto IsMatchingOutput = [&](const FMetasoundFrontendClassOutput& GraphOutput)
 			{
-				// TODO: log errror. 
-				continue;
-			}
+				return (InOutputNode.ID == GraphOutput.NodeID);
+			};
 
-			for (const FMetasoundNodeConnectionDescription& InputConnection : NodeDescription.InputConnections)
+			OutClassOutputIndex = InOwningGraph.Interface.Outputs.IndexOfByPredicate(IsMatchingOutput);
+			if (INDEX_NONE != OutClassOutputIndex)
 			{
-				const bool bIsConnected = InputConnection.NodeID != FMetasoundNodeConnectionDescription::DisconnectedNodeID;
-				if (!bIsConnected)
-				{
-					continue;
-				}
-
-				const INode* FromNode = OutGraph.FindNode(InputConnection.NodeID);
-				if (nullptr == FromNode)
-				{
-					// TODO: log error.
-					continue;
-				}
-
-				OutGraph.AddDataEdge(*FromNode, *InputConnection.OutputName, *ToNode, *InputConnection.InputName);
+				return &InOwningGraph.Interface.Outputs[OutClassOutputIndex];
 			}
 		}
+		return nullptr;
 	}
 
-	void FFrontendGraphBuilder::AddInputDestinationsToGraph(const TArray<FMetasoundInputDescription>& InInputDescriptions, const TArray<FMetasoundNodeDescription>& InNodes, FFrontendGraph& OutGraph) const
+	const FMetasoundFrontendVertexLiteral* FFrontendGraphBuilder::FindInputLiteralForInputNode(const FMetasoundFrontendNode& InInputNode, const FMetasoundFrontendClass& InInputNodeClass, const FMetasoundFrontendClassInput& InOwningGraphClassInput) const
 	{
-		for (int32 InputIndex = 0; InputIndex < InInputDescriptions.Num(); InputIndex++)
+		// Default value priority is:
+		// 1. A value set directly on the node
+		// 2. A default value of the owning graph
+		// 3. A default value on the input node class.
+
+		const FMetasoundFrontendVertexLiteral* VertexLiteral = nullptr;
+
+		// Check for default value directly on node.
+		if (ensure(InInputNode.Interface.Inputs.Num() == 1))
 		{
-			const FMetasoundInputDescription& InputDescription = InInputDescriptions[InputIndex];
+			const FMetasoundFrontendVertex& InputVertex = InInputNode.Interface.Inputs[0];
 
-			auto IsInputWithSameName = [&](const FMetasoundNodeDescription& InNode) 
-			{ 
-				return (InNode.Name == InputDescription.Name) && (InNode.ObjectTypeOfNode == EMetasoundClassType::Input);
-			};
+			// Currently only support one point on input vertex.
+			if (ensure(InputVertex.PointIDs.Num() == 1))
+			{
+				VertexLiteral = InInputNode.InputLiterals.FindByPredicate(
+					[&](const FMetasoundFrontendVertexLiteral& InVertexLiteral)
+					{
+						return InVertexLiteral.PointID == InputVertex.PointIDs[0];
+					}
+				);
+			}
+		}
+
+		// Check for default value on owning graph.
+		if (nullptr == VertexLiteral)
+		{
+			if (ensure(InOwningGraphClassInput.PointIDs.Num() == 1))
+			{
+				int32 PointID = InOwningGraphClassInput.PointIDs[0];
+
+				VertexLiteral = InOwningGraphClassInput.Defaults.FindByPredicate(
+					[&](const FMetasoundFrontendVertexLiteral& InVertexLiteral)
+					{
+						return InVertexLiteral.PointID == PointID;
+					}
+				);
+			}
+		}
 
-			const FMetasoundNodeDescription* NodeDescription = InNodes.FindByPredicate(IsInputWithSameName);
+		// Check for default value on input node class
+		if (nullptr == VertexLiteral && ensure(InInputNodeClass.Interface.Inputs.Num() == 1))
+		{
+			const FMetasoundFrontendClassInput& InputNodeClassInput = InInputNodeClass.Interface.Inputs[0];
 
-			if (nullptr == NodeDescription)
+			if (ensure(InputNodeClassInput.PointIDs.Num() == 1))
 			{
-				// TODO: add error.
-				continue;
+				int32 PointID = InputNodeClassInput.PointIDs[0];
+
+				VertexLiteral = InputNodeClassInput.Defaults.FindByPredicate(
+					[&](const FMetasoundFrontendVertexLiteral& InVertexLiteral)
+					{
+						return InVertexLiteral.PointID == PointID;
+					}
+				);
 			}
+		}
+
+		return VertexLiteral;
+	}
 
-			OutGraph.AddInputNode(NodeDescription->UniqueID, InputIndex, InputDescription.Name, CreateInputNode(*NodeDescription, InputDescription));
+	// TODO: add errors here. Most will be a "PromptIfMissing"...
+	void FFrontendGraphBuilder::AddNodesToGraph(const FMetasoundFrontendGraphClass& InGraphClass, const TMap<int32, const FMetasoundFrontendClass*>& InClasses, FFrontendGraph& OutGraph) const
+	{
+		for (const FMetasoundFrontendNode& Node : InGraphClass.Graph.Nodes)
+		{
+			const FMetasoundFrontendClass* NodeClass = InClasses.FindRef(Node.ClassID);
+
+			if (ensure(nullptr != NodeClass))
+			{
+				switch (NodeClass->Metadata.Type)
+				{
+					case EMetasoundFrontendClassType::Input:
+						{
+							int32 InputIndex = INDEX_NONE;
+							const FMetasoundFrontendClassInput* ClassInput = FindClassInputForInputNode(InGraphClass, Node, InputIndex);
+
+							if ((nullptr != ClassInput) && (INDEX_NONE != InputIndex))
+							{
+
+								OutGraph.AddInputNode(Node.ID, InputIndex, ClassInput->Name, CreateInputNode(Node, *NodeClass, *ClassInput));
+							}
+							else
+							{
+								UE_LOG(LogMetasound, Error, TEXT("Failed to match input node [NodeID:%d, NodeName:%s] to owning graph [ClassID:%d] output."), Node.ID, *Node.Name, InGraphClass.ID);
+							}
+						}
+
+						break;
+
+					case EMetasoundFrontendClassType::Output:
+						{
+							int32 OutputIndex = INDEX_NONE;
+							const FMetasoundFrontendClassOutput* ClassOutput = FindClassOutputForOutputNode(InGraphClass, Node, OutputIndex);
+
+							if ((nullptr != ClassOutput) && (INDEX_NONE != OutputIndex))
+							{
+								OutGraph.AddOutputNode(Node.ID, OutputIndex, ClassOutput->Name, CreateOutputNode(Node, *NodeClass));
+							}
+							else
+							{
+								UE_LOG(LogMetasound, Error, TEXT("Failed to match output node [NodeID:%d, NodeName:%s] to owning graph [ClassID:%d] output."), Node.ID, *Node.Name, InGraphClass.ID);
+							}
+						}
+
+						break;
+
+					default:
+						OutGraph.AddNode(Node.ID, CreateExternalNode(Node, *NodeClass));
+				}
+			}
 		}
 	}
 
-	void FFrontendGraphBuilder::AddOutputSourcesToGraph(const TArray<FMetasoundOutputDescription>& InOutputDescriptions, const TArray<FMetasoundNodeDescription>& InNodes, FFrontendGraph& OutGraph) const
+	void FFrontendGraphBuilder::AddEdgesToGraph(const FMetasoundFrontendGraph& InFrontendGraph, FFrontendGraph& OutGraph) const
 	{
-		for (int32 OutputIndex = 0; OutputIndex < InOutputDescriptions.Num(); OutputIndex++)
+		// Pair of frontend node and core node. The frontend node can be one of
+		// several types.
+		struct FCoreNodeAndFrontendVertex
 		{
-			const FMetasoundOutputDescription& OutputDescription = InOutputDescriptions[OutputIndex];
+			const INode* Node = nullptr;
+			const FMetasoundFrontendVertex* Vertex = nullptr;
+		};
 
-			auto IsOutputWithSameName = [&](const FMetasoundNodeDescription& InNode) 
-			{ 
-				return (InNode.Name == OutputDescription.Name) && (InNode.ObjectTypeOfNode == EMetasoundClassType::Output);
-			};
+		// TODO: add support for array vertices.
+		typedef TTuple<int32, int32> FNodeIDPointID;
 
-			const FMetasoundNodeDescription* NodeDescription = InNodes.FindByPredicate(IsOutputWithSameName);
+		TMap<FNodeIDPointID, FCoreNodeAndFrontendVertex> NodeSourcesByID;
+		TMap<FNodeIDPointID, FCoreNodeAndFrontendVertex> NodeDestinationsByID;
 
-			if (nullptr == NodeDescription)
+		// Add nodes to NodeID/PointID map
+		for (const FMetasoundFrontendNode& Node : InFrontendGraph.Nodes)
+		{
+			const INode* CoreNode = OutGraph.FindNode(Node.ID);
+			if (nullptr == CoreNode)
 			{
-				// TODO: add error.
+				UE_LOG(LogMetasound, Display, TEXT("Could not find referenced node [ID:%d]"), Node.ID);
 				continue;
 			}
 
-			OutGraph.AddOutputNode(NodeDescription->UniqueID, OutputIndex, OutputDescription.Name, CreateOutputNode(*NodeDescription, OutputDescription));
-		}
-	}
+			for (const FMetasoundFrontendVertex& Vertex : Node.Interface.Inputs)
+			{
+				for (int32 PointID : Vertex.PointIDs)
+				{
+					NodeDestinationsByID.Add(FNodeIDPointID(Node.ID, PointID), FCoreNodeAndFrontendVertex({CoreNode, &Vertex}));
+				}
+			}
 
-	void FFrontendGraphBuilder::SplitNodesByType(const TArray<FMetasoundNodeDescription>& InNodes, TArray<FMetasoundNodeDescription>& OutExternalNodes, TArray<FMetasoundNodeDescription>& OutInputNodes, TArray<FMetasoundNodeDescription>& OutOutputNodes) const
-	{
-		for (const FMetasoundNodeDescription& Node : InNodes)
+			for (const FMetasoundFrontendVertex& Vertex : Node.Interface.Outputs)
+			{
+				for (int32 PointID : Vertex.PointIDs)
+				{
+					NodeSourcesByID.Add(FNodeIDPointID(Node.ID, PointID), FCoreNodeAndFrontendVertex({CoreNode, &Vertex}));
+				}
+			}
+		};
+		
+
+		for (const FMetasoundFrontendEdge& Edge : InFrontendGraph.Edges)
 		{
-			switch (Node.ObjectTypeOfNode)
+			const FNodeIDPointID DestinationKey(Edge.ToNodeID, Edge.ToPointID);
+			const FCoreNodeAndFrontendVertex* DestinationNodeAndVertex = NodeDestinationsByID.Find(DestinationKey);
+
+			if (nullptr == DestinationNodeAndVertex)
 			{
-				case EMetasoundClassType::External:
+				// TODO: bubble up error
+				UE_LOG(LogMetasound, Error, TEXT("Failed to add edge. Could not find destination [NodeID:%d, PointID:%d]"), Edge.ToNodeID, Edge.ToPointID);
+				continue;
+			}
 
-					OutExternalNodes.Add(Node);
-					break;
+			if (nullptr == DestinationNodeAndVertex->Node)
+			{
+				// TODO: bubble up error
+				UE_LOG(LogMetasound, Warning, TEXT("Skipping edge. Null destination node [ID:%d]"), Edge.ToNodeID);
+				continue;
+			}
+
+			const FNodeIDPointID SourceKey(Edge.FromNodeID, Edge.FromPointID);
+			const FCoreNodeAndFrontendVertex* SourceNodeAndVertex = NodeSourcesByID.Find(SourceKey);
+
+			if (nullptr == SourceNodeAndVertex)
+			{
+				UE_LOG(LogMetasound, Error, TEXT("Failed to add edge. Could not find source [NodeID:%d, PointID:%d]"), Edge.FromNodeID, Edge.FromPointID);
+				continue;
+			}
 
-				case EMetasoundClassType::Input:
+			if (nullptr == SourceNodeAndVertex->Node)
+			{
+				// TODO: bubble up error
+				UE_LOG(LogMetasound, Warning, TEXT("Skipping edge. Null source node [ID:%d]"), Edge.FromNodeID);
+				continue;
+			}
 
-					OutInputNodes.Add(Node);
-					break;
+			const INode* FromNode = SourceNodeAndVertex->Node;
+			const FVertexKey FromVertexKey = SourceNodeAndVertex->Vertex->Name;
 
-				case EMetasoundClassType::Output:
+			const INode* ToNode = DestinationNodeAndVertex->Node;
+			const FVertexKey ToVertexKey = DestinationNodeAndVertex->Vertex->Name;
 
-					OutOutputNodes.Add(Node);
-					break;
+			bool bSuccess = OutGraph.AddDataEdge(*FromNode, FromVertexKey,  *ToNode, ToVertexKey);
 
-				default:
-					checkNoEntry();
+			if (!bSuccess)
+			{
+				UE_LOG(LogMetasound, Error, TEXT("Failed to connect edge from [NodeID:%d, PointID:%d] to [NodeID:%d, PointID:%d]"), Edge.FromNodeID, Edge.FromPointID, Edge.ToNodeID, Edge.ToPointID);
 			}
 		}
 	}
 
 	/** Check that all dependencies are C++ class dependencies. */
-	bool FFrontendGraphBuilder::IsFlat(const FMetasoundDocument& InDocument) const
+	bool FFrontendGraphBuilder::IsFlat(const FMetasoundFrontendDocument& InDocument) const
 	{
-		return IsFlat(InDocument.RootClass, InDocument.Dependencies);
+		if (InDocument.Subgraphs.Num() > 0)
+		{
+			return false;
+		}
+
+		return IsFlat(InDocument.RootGraph, InDocument.Dependencies);
 	}
 
-	bool FFrontendGraphBuilder::IsFlat(const FMetasoundClassDescription& InRoot, const TArray<FMetasoundClassDescription>& InDependencies) const
+	bool FFrontendGraphBuilder::IsFlat(const FMetasoundFrontendGraphClass& InRoot, const TArray<FMetasoundFrontendClass>& InDependencies) const
 	{
 		// All dependencies are external dependencies in a flat graph
-		auto IsClassExternal = [&](const FMetasoundClassDescription& InDesc) { return InDesc.Metadata.NodeType == EMetasoundClassType::External; };
+		auto IsClassExternal = [&](const FMetasoundFrontendClass& InDesc) 
+		{ 
+			bool bIsExternal = (InDesc.Metadata.Type == EMetasoundFrontendClassType::External) ||
+				(InDesc.Metadata.Type == EMetasoundFrontendClassType::Input) ||
+				(InDesc.Metadata.Type == EMetasoundFrontendClassType::Output);
+			return bIsExternal;
+		};
 		const bool bIsEveryDependencyExternal = Algo::AllOf(InDependencies, IsClassExternal);
 
 		if (!bIsEveryDependencyExternal)
@@ -340,74 +521,54 @@ namespace Metasound
 			return false;
 		}
 
-		// None of the nodes are subgraphs in a flat graph
-		auto IsGraphNode = [&](const FMetasoundNodeDescription& InDesc) { return InDesc.ObjectTypeOfNode == EMetasoundClassType::MetasoundGraph; };
-		const bool bIsAnyNodeAGraph = Algo::AnyOf(InRoot.Graph.Nodes, IsGraphNode);
-
-		if (bIsAnyNodeAGraph)
-		{
-			return false;
-		}
-
-		// All the dependencies are available in a flat graph.
+		// All the dependencies are met 
 		TSet<int32> AvailableDependencies;
-		Algo::Transform(InDependencies, AvailableDependencies, [](const FMetasoundClassDescription& InDesc) { return InDesc.UniqueID; });
+		Algo::Transform(InDependencies, AvailableDependencies, [](const FMetasoundFrontendClass& InDesc) { return InDesc.ID; });
 
-		auto IsDependencyMet = [&](const FMetasoundNodeDescription& InNode) 
+		auto IsDependencyMet = [&](const FMetasoundFrontendNode& InNode) 
 		{ 
-			// Currently, input/output nodes do not have dependencies. This may change.
-			const bool bIsInputOrOutput = (InNode.ObjectTypeOfNode == EMetasoundClassType::Input) || (InNode.ObjectTypeOfNode == EMetasoundClassType::Output);
-			const bool bIsExternal = InNode.ObjectTypeOfNode == EMetasoundClassType::External;
-
-			return bIsInputOrOutput || (bIsExternal && AvailableDependencies.Contains(InNode.DependencyID));
+			return AvailableDependencies.Contains(InNode.ClassID);
 		};
 
 		const bool bIsEveryDependencyMet = Algo::AllOf(InRoot.Graph.Nodes, IsDependencyMet);
-		return bIsEveryDependencyMet;
-	}
 
-	
-	/* Metasound document should be inflated by now. */
-	TUniquePtr<FFrontendGraph> FFrontendGraphBuilder::CreateGraph(const FMetasoundDocument& InDocument) const
-	{
-		return CreateGraph(InDocument.RootClass, InDocument.Dependencies);
+		return bIsEveryDependencyMet;
 	}
 
-	TUniquePtr<FFrontendGraph> FFrontendGraphBuilder::CreateGraph(const FMetasoundClassDescription& InRoot, const TArray<FMetasoundClassDescription>& InDependencies) const
+	TUniquePtr<FFrontendGraph> FFrontendGraphBuilder::CreateGraph(const FMetasoundFrontendGraphClass& InGraph, const TArray<FMetasoundFrontendGraphClass>& InSubgraphs, const TArray<FMetasoundFrontendClass>& InDependencies) const
 	{
-		if (!IsFlat(InRoot, InDependencies))
+		if (!IsFlat(InGraph, InDependencies))
 		{
-			checkNoEntry();
+			// Likely this will change in the future and the builder will be able
+			// to build graphs with subgraphs..
+			UE_LOG(LogMetasound, Error, TEXT("Provided graph not flat. FFrontendGraphBuilder can only build flat graphs"));
 			return TUniquePtr<FFrontendGraph>(nullptr);
 		}
 
-		TUniquePtr<FFrontendGraph> Graph = MakeUnique<FFrontendGraph>(InRoot.Metadata.NodeName);
 
-		TMap<int32, const FMetasoundClassDescription*> NodeMap;
+		TMap<int32, const FMetasoundFrontendClass*> ClassMap;
 
-		for (const FMetasoundClassDescription& InDesc : InDependencies)
+		for (const FMetasoundFrontendClass& ExtClass : InDependencies)
 		{
-			NodeMap.Add(InDesc.UniqueID, &InDesc);
+			ClassMap.Add(ExtClass.ID, &ExtClass);
 		}
 
-		TArray<FMetasoundNodeDescription> ExternalNodes;
-		TArray<FMetasoundNodeDescription> InputNodes;
-		TArray<FMetasoundNodeDescription> OutputNodes;
-
-		SplitNodesByType(InRoot.Graph.Nodes, ExternalNodes, InputNodes, OutputNodes);
+		TUniquePtr<FFrontendGraph> MetasoundGraph = MakeUnique<FFrontendGraph>(InGraph.Metadata.Name.Name);
 
 		// TODO: will likely want to bubble up errors here for case where
 		// a datatype or node is not registered. 
-		AddExternalNodesToGraph(ExternalNodes, NodeMap, *Graph);
-
-		AddInputDestinationsToGraph(InRoot.Inputs, InputNodes, *Graph);
+		AddNodesToGraph(InGraph, ClassMap, *MetasoundGraph);
 
-		AddOutputSourcesToGraph(InRoot.Outputs, OutputNodes, *Graph);
+		AddEdgesToGraph(InGraph.Graph, *MetasoundGraph);
 
-		AddEdgesToGraph(InRoot.Graph, *Graph);
+		check(MetasoundGraph->OwnsAllReferencedNodes());
 
-		check(Graph->OwnsAllReferencedNodes());
-
-		return MoveTemp(Graph);
+		return MoveTemp(MetasoundGraph);
+	}
+	
+	/* Metasound document should be inflated by now. */
+	TUniquePtr<FFrontendGraph> FFrontendGraphBuilder::CreateGraph(const FMetasoundFrontendDocument& InDocument) const
+	{
+		return CreateGraph(InDocument.RootGraph, InDocument.Subgraphs, InDocument.Dependencies);
 	}
 }
diff --git a/Engine/Plugins/Runtime/Metasound/Source/MetasoundFrontend/Private/MetasoundFrontendGraph.h b/Engine/Plugins/Runtime/Metasound/Source/MetasoundFrontend/Private/MetasoundFrontendGraph.h
--- a/Engine/Plugins/Runtime/Metasound/Source/MetasoundFrontend/Private/MetasoundFrontendGraph.h
+++ b/Engine/Plugins/Runtime/Metasound/Source/MetasoundFrontend/Private/MetasoundFrontendGraph.h
@@ -5,7 +5,6 @@
 #include "CoreMinimal.h"
 
 #include "MetasoundFrontend.h"
-#include "MetasoundFrontendDataLayout.h"
 #include "MetasoundGraph.h"
 #include "MetasoundNodeInterface.h"
 
@@ -29,7 +28,7 @@ namespace Metasound
 
 			/** Add an input node to this graph.
 			 *
-			 * @param InNodeID - The NodeID related to the parent FMetasoundClassDescription.
+			 * @param InNodeID - The NodeID related to the parent FMetasoundFrontendClass.
 			 * @param InIndex - The positional index for the input.
 			 * @param InVertexKey - The key for the graph input vertex.
 			 * @param InNode - A unique pointer to an input node. 
@@ -38,7 +37,7 @@ namespace Metasound
 
 			/** Add an output node to this graph.
 			 *
-			 * @param InNodeID - The NodeID related to the parent FMetasoundClassDescription.
+			 * @param InNodeID - The NodeID related to the parent FMetasoundFrontendClass.
 			 * @param InIndex - The positional index for the output.
 			 * @param InVertexKey - The key for the graph output vertex.
 			 * @param InNode - A unique pointer to an output node. 
@@ -47,7 +46,7 @@ namespace Metasound
 
 			/** Store a node on this graph. 
 			 *
-			 * @param InNodeID - The NodeID related to the parent FMetasoundClassDescription.
+			 * @param InNodeID - The NodeID related to the parent FMetasoundFrontendClass.
 			 * @param InNode - A unique pointer to a node. 
 			 */
 			void AddNode(int32 InNodeID, TUniquePtr<INode> InNode);
@@ -93,7 +92,7 @@ namespace Metasound
 	};
 
 	/** FFrontendGraphBuilder builds a FFrontendGraph from a FMetasoundDoucment
-	 * or FMetasoundClassDescription.
+	 * or FMetasoundFrontendClass.
 	 */
 	class FFrontendGraphBuilder
 	{
@@ -105,40 +104,34 @@ namespace Metasound
 		 *
 		 * @return True if all dependencies are C++ classes. False otherwise.
 		 */
-		bool IsFlat(const FMetasoundDocument& InDocument) const;
+		bool IsFlat(const FMetasoundFrontendDocument& InDocument) const;
+
+		bool IsFlat(const FMetasoundFrontendGraphClass& InRoot, const TArray<FMetasoundFrontendClass>& InDependencies) const;
 
-		/** Check that all dependencies are C++ class dependencies. 
-		 * 
-		 * @param InDependencies - Array of dependencies to check.
-		 *
-		 * @return True if all dependencies are C++ classes. False otherwise.
-		 */
-		bool IsFlat(const FMetasoundClassDescription& InRoot, const TArray<FMetasoundClassDescription>& InDependencies) const;
-		
 		/* Metasound document should be in order to create this graph. */
-		TUniquePtr<FFrontendGraph> CreateGraph(const FMetasoundDocument& InDocument) const;
+		TUniquePtr<FFrontendGraph> CreateGraph(const FMetasoundFrontendDocument& InDocument) const;
+
+		/* Metasound document should be in order to create this graph. */
+		TUniquePtr<FFrontendGraph> CreateGraph(const FMetasoundFrontendGraphClass& InGraph, const TArray<FMetasoundFrontendGraphClass>& InSubgraphs, const TArray<FMetasoundFrontendClass>& InDependencies) const;
 
-		/** Creates a graph given the root class and all of it's dependencies. */
-		TUniquePtr<FFrontendGraph> CreateGraph(const FMetasoundClassDescription& InRootClass, const TArray<FMetasoundClassDescription>& InDependencies) const;
 
 	private:
 
-		TUniquePtr<INode> CreateInputNode(const FMetasoundNodeDescription& InNode, const FMetasoundInputDescription& InDescription) const;
+		const FMetasoundFrontendClassInput* FindClassInputForInputNode(const FMetasoundFrontendGraphClass& InOwningGraph, const FMetasoundFrontendNode& InInputNode, int32& OutClassInputIndex) const;
 
-		TUniquePtr<INode> CreateOutputNode(const FMetasoundNodeDescription& InNode, const FMetasoundOutputDescription& InDescription) const;
+		const FMetasoundFrontendClassOutput* FindClassOutputForOutputNode(const FMetasoundFrontendGraphClass& InOwningGraph, const FMetasoundFrontendNode& InOutputNode, int32& OutClassOutputIndex) const;
 
-		TUniquePtr<INode> CreateExternalNode(const FMetasoundNodeDescription& InNode, const FMetasoundClassDescription& InClass) const;
+		const FMetasoundFrontendVertexLiteral* FindInputLiteralForInputNode(const FMetasoundFrontendNode& InInputNode, const FMetasoundFrontendClass& InInputNodeClass, const FMetasoundFrontendClassInput& InOwningGraphClassInput) const;
 
-		void SplitNodesByType(const TArray<FMetasoundNodeDescription>& InNodes, TArray<FMetasoundNodeDescription>& OutExternalNodes, TArray<FMetasoundNodeDescription>& OutInputNodes, TArray<FMetasoundNodeDescription>& OutOutputNodes) const;
+		TUniquePtr<INode> CreateInputNode(const FMetasoundFrontendNode& InNode, const FMetasoundFrontendClass& InClass, const FMetasoundFrontendClassInput& InOwningGraphClassInput) const;
 
+		TUniquePtr<INode> CreateOutputNode(const FMetasoundFrontendNode& InNode, const FMetasoundFrontendClass& InExternal) const;
 
-		// TODO: add errors here. Most will be a "PromptIfMissing"...
-		void AddExternalNodesToGraph(const TArray<FMetasoundNodeDescription>& InNodes, const TMap<int32, const FMetasoundClassDescription*>& InClasses, FFrontendGraph& OutGraph) const;
+		TUniquePtr<INode> CreateExternalNode(const FMetasoundFrontendNode& InNode, const FMetasoundFrontendClass& InClass) const;
 
-		void AddEdgesToGraph(const FMetasoundGraphDescription& InGraphDescription, FFrontendGraph& OutGraph) const;
-
-		void AddInputDestinationsToGraph(const TArray<FMetasoundInputDescription>& InInputDescriptions, const TArray<FMetasoundNodeDescription>& InNodes, FFrontendGraph& OutGraph) const;
-		void AddOutputSourcesToGraph(const TArray<FMetasoundOutputDescription>& InOutputDescriptions, const TArray<FMetasoundNodeDescription>& InNodes, FFrontendGraph& OutGraph) const;
+		// TODO: add errors here. Most will be a "PromptIfMissing"...
+		void AddNodesToGraph(const FMetasoundFrontendGraphClass& InGraphClass, const TMap<int32, const FMetasoundFrontendClass*>& InClasses, FFrontendGraph& OutGraph) const;
 
+		void AddEdgesToGraph(const FMetasoundFrontendGraph& InGraphDescription, FFrontendGraph& OutGraph) const;
 	};
 }
diff --git a/Engine/Plugins/Runtime/Metasound/Source/MetasoundFrontend/Private/MetasoundFrontendInvalidController.cpp b/Engine/Plugins/Runtime/Metasound/Source/MetasoundFrontend/Private/MetasoundFrontendInvalidController.cpp
new file mode 100644
--- /dev/null
+++ b/Engine/Plugins/Runtime/Metasound/Source/MetasoundFrontend/Private/MetasoundFrontendInvalidController.cpp
@@ -0,0 +1,69 @@
+// Copyright Epic Games, Inc. All Rights Reserved.
+
+#include "MetasoundFrontendInvalidController.h"
+
+namespace Metasound
+{
+	namespace Frontend
+	{
+		TSharedRef<INodeController> FInvalidOutputController::GetOwningNode()
+		{
+			return FInvalidNodeController::GetInvalid();
+		}
+
+		TSharedRef<const INodeController> FInvalidOutputController::GetOwningNode() const
+		{
+			return FInvalidNodeController::GetInvalid();
+		}
+
+		TSharedRef<INodeController> FInvalidInputController::GetOwningNode()
+		{
+			return FInvalidNodeController::GetInvalid();
+		}
+
+		TSharedRef<const INodeController> FInvalidInputController::GetOwningNode() const
+		{
+			return FInvalidNodeController::GetInvalid();
+		}
+
+		TSharedRef<IGraphController> FInvalidNodeController::AsGraph()
+		{
+			return FInvalidGraphController::GetInvalid();
+		}
+
+		TSharedRef<const IGraphController> FInvalidNodeController::AsGraph() const
+		{
+			return FInvalidGraphController::GetInvalid();
+		}
+
+		TSharedRef<IGraphController> FInvalidNodeController::GetOwningGraph()
+		{
+			return FInvalidGraphController::GetInvalid();
+		}
+
+		TSharedRef<const IGraphController> FInvalidNodeController::GetOwningGraph() const
+		{
+			return FInvalidGraphController::GetInvalid();
+		}
+
+		TSharedRef<IDocumentController> FInvalidGraphController::GetOwningDocument()
+		{
+			return FInvalidDocumentController::GetInvalid();
+		}
+
+		TSharedRef<const IDocumentController> FInvalidGraphController::GetOwningDocument() const
+		{
+			return FInvalidDocumentController::GetInvalid();
+		}
+
+		TSharedRef<IGraphController> FInvalidDocumentController::GetRootGraph()
+		{
+			return FInvalidGraphController::GetInvalid();
+		}
+
+		TSharedRef<const IGraphController> FInvalidDocumentController::GetRootGraph() const
+		{
+			return FInvalidGraphController::GetInvalid();
+		}
+	}
+}
diff --git a/Engine/Plugins/Runtime/Metasound/Source/MetasoundFrontend/Private/MetasoundFrontendInvalidController.h b/Engine/Plugins/Runtime/Metasound/Source/MetasoundFrontend/Private/MetasoundFrontendInvalidController.h
new file mode 100644
--- /dev/null
+++ b/Engine/Plugins/Runtime/Metasound/Source/MetasoundFrontend/Private/MetasoundFrontendInvalidController.h
@@ -0,0 +1,326 @@
+// Copyright Epic Games, Inc. All Rights Reserved.
+
+#pragma once
+
+
+#include "CoreMinimal.h"
+#include "MetasoundBuilderInterface.h"
+#include "MetasoundFrontendBaseClasses.h"
+#include "MetasoundFrontendController.h"
+#include "MetasoundGraph.h"
+
+namespace Metasound
+{
+	namespace Frontend
+	{
+		namespace MetasoundFrontendInvalidControllerPrivate
+		{
+			// Helper to create a functional local static. Used when the interface requires
+			// a const reference returned.
+			template<typename Type>
+			const Type& GetInvalid()
+			{
+				static const Type InvalidValue;
+				return InvalidValue;
+			}
+		}
+
+		/** FInvalidOutputController is a controller which is always invalid. 
+		 *
+		 *  All methods return defaults and return error flags or invalid values.
+		 */
+		class METASOUNDFRONTEND_API FInvalidOutputController : public IOutputController
+		{
+		public:
+			FInvalidOutputController() = default;
+
+			virtual ~FInvalidOutputController() = default;
+
+			static TSharedRef<IOutputController> GetInvalid()
+			{
+				static TSharedRef<IOutputController> InvalidController = MakeShared<FInvalidOutputController>();
+				return MakeShared<FInvalidOutputController>();
+			}
+
+			bool IsValid() const override { return false; }
+			int32 GetID() const override { return Metasound::FrontendInvalidID; }
+			const FName& GetDataType() const override { return MetasoundFrontendInvalidControllerPrivate::GetInvalid<FName>(); }  
+			const FString& GetName() const override { return MetasoundFrontendInvalidControllerPrivate::GetInvalid<FString>(); }
+			const FText& GetDisplayName() const override { return FText::GetEmpty(); }
+			const FText& GetTooltip() const override { return FText::GetEmpty(); }
+
+			int32 GetOwningNodeID() const override { return Metasound::FrontendInvalidID; }
+			TSharedRef<INodeController> GetOwningNode() override;
+			TSharedRef<const INodeController> GetOwningNode() const override;
+
+
+			FConnectability CanConnectTo(const IInputController& InController) const override { return FConnectability(); }
+			bool Connect(IInputController& InController) override { return false; }
+			bool ConnectWithConverterNode(IInputController& InController, const FConverterNodeInfo& InNodeClassName) override { return false; }
+			bool Disconnect(IInputController& InController) override { return false; }
+
+		};
+
+		/** FInvalidInputController is a controller which is always invalid. 
+		 *
+		 *  All methods return defaults and return error flags or invalid values.
+		 */
+		class METASOUNDFRONTEND_API FInvalidInputController : public IInputController 
+		{
+		public:
+			FInvalidInputController() = default;
+			virtual ~FInvalidInputController() = default;
+
+			static TSharedRef<IInputController> GetInvalid()
+			{
+				static TSharedRef<IInputController> InvalidController = MakeShared<FInvalidInputController>();
+				return InvalidController;
+			}
+
+			bool IsValid() const override { return false; }
+			int32 GetID() const override { return Metasound::FrontendInvalidID; }
+			bool IsConnected() const override { return false; }
+			const FName& GetDataType() const override { return MetasoundFrontendInvalidControllerPrivate::GetInvalid<FName>(); }  
+			const FString& GetName() const override { return MetasoundFrontendInvalidControllerPrivate::GetInvalid<FString>(); }
+			const FText& GetDisplayName() const override { return FText::GetEmpty(); }
+			const FText& GetTooltip() const override { return FText::GetEmpty(); }
+
+			int32 GetOwningNodeID() const override { return Metasound::FrontendInvalidID; }
+			TSharedRef<INodeController> GetOwningNode() override;
+			TSharedRef<const INodeController> GetOwningNode() const override;
+
+			virtual TSharedRef<IOutputController> GetCurrentlyConnectedOutput() override { return FInvalidOutputController::GetInvalid(); }
+			virtual TSharedRef<const IOutputController> GetCurrentlyConnectedOutput() const override { return FInvalidOutputController::GetInvalid(); }
+
+			virtual FConnectability CanConnectTo(const IOutputController& InController) const override { return FConnectability(); }
+			virtual bool Connect(IOutputController& InController) override { return false; }
+			virtual bool ConnectWithConverterNode(IOutputController& InController, const FConverterNodeInfo& InNodeClassName) override { return false; }
+			virtual bool Disconnect(IOutputController& InController) override { return false; }
+			virtual bool Disconnect() override { return false; }
+		};
+
+		/** FInvalidNodeController is a controller which is always invalid. 
+		 *
+		 *  All methods return defaults and return error flags or invalid values.
+		 */
+		class METASOUNDFRONTEND_API FInvalidNodeController : public INodeController 
+		{
+
+		public:
+			FInvalidNodeController() = default;
+			virtual ~FInvalidNodeController() = default;
+
+			static TSharedRef<INodeController> GetInvalid()
+			{
+				static TSharedRef<INodeController> InvalidValue = MakeShared<FInvalidNodeController>();
+				return InvalidValue;
+			}
+
+			bool IsValid() const override { return false; }
+
+			TArray<TSharedRef<IInputController>> GetInputs() override { return TArray<TSharedRef<IInputController>>(); }
+			TArray<TSharedRef<IOutputController>> GetOutputs() override { return TArray<TSharedRef<IOutputController>>(); }
+			TArray<TSharedRef<const IInputController>> GetConstInputs() const override { return TArray<TSharedRef<const IInputController>>(); }
+			TArray<TSharedRef<const IOutputController>> GetConstOutputs() const override { return TArray<TSharedRef<const IOutputController>>(); }
+
+			TArray<FInputHandle> GetInputsWithVertexName(const FString& InName) override { return TArray<FInputHandle>(); }
+			TArray<FConstInputHandle> GetInputsWithVertexName(const FString& InName) const override { return TArray<FConstInputHandle>(); }
+			TArray<FOutputHandle> GetOutputsWithVertexName(const FString& InName) override { return TArray<FOutputHandle>(); }
+			TArray<FConstOutputHandle> GetOutputsWithVertexName(const FString& InName) const override { return TArray<FConstOutputHandle>(); }
+			TSharedRef<IInputController> GetInputWithID(int32 InPointID) override { return FInvalidInputController::GetInvalid(); }
+			TSharedRef<IOutputController> GetOutputWithID(int32 InPointID) override { return FInvalidOutputController::GetInvalid(); }
+			TSharedRef<const IInputController> GetInputWithID(int32 InPointID) const override { return FInvalidInputController::GetInvalid(); }
+			TSharedRef<const IOutputController> GetOutputWithID(int32 InPointID) const override { return FInvalidOutputController::GetInvalid(); }
+
+			bool CanAddInput(const FString& InVertexName) const override { return false; }
+			FInputHandle AddInput(const FString& InVertexName, const FMetasoundFrontendLiteral* InDefault) override { return FInvalidInputController::GetInvalid(); }
+			bool RemoveInput(int32 InPointID) override { return false; }
+
+			bool CanAddOutput(const FString& InVertexName) const override { return false; }
+			FInputHandle AddOutput(const FString& InVertexName, const FMetasoundFrontendLiteral* InDefault) override { return FInvalidInputController::GetInvalid(); }
+			bool RemoveOutput(int32 InPointID) override { return false; }
+
+			EMetasoundFrontendClassType GetClassType() const override { return EMetasoundFrontendClassType::Invalid; }
+			const FString& GetClassName() const override { return MetasoundFrontendInvalidControllerPrivate::GetInvalid<FString>(); }
+			FMetasoundFrontendVersionNumber GetClassVersionNumber() const override { return MetasoundFrontendInvalidControllerPrivate::GetInvalid<FMetasoundFrontendVersionNumber>(); }
+			const FText& GetClassDescription() const override { return MetasoundFrontendInvalidControllerPrivate::GetInvalid<FText>(); }
+
+
+			TSharedRef<IGraphController> AsGraph() override;
+			TSharedRef<const IGraphController> AsGraph() const override;
+
+			int32 GetID() const override { return Metasound::FrontendInvalidID; }
+			int32 GetClassID() const override { return Metasound::FrontendInvalidID; }
+
+			int32 GetOwningGraphClassID() const override { return Metasound::FrontendInvalidID; }
+			TSharedRef<IGraphController> GetOwningGraph() override;
+			TSharedRef<const IGraphController> GetOwningGraph() const override;
+
+			const FString& GetNodeName() const override { return MetasoundFrontendInvalidControllerPrivate::GetInvalid<FString>(); }
+		};
+
+		/** FInvalidGraphController is a controller which is always invalid. 
+		 *
+		 *  All methods return defaults and return error flags or invalid values.
+		 */
+		class METASOUNDFRONTEND_API FInvalidGraphController : public IGraphController 
+		{
+			public:
+			FInvalidGraphController() = default;
+			virtual ~FInvalidGraphController() = default;
+
+			static TSharedRef<IGraphController> GetInvalid()
+			{
+				static TSharedRef<IGraphController> InvalidGraph = MakeShared<FInvalidGraphController>();
+				return InvalidGraph;
+			}
+
+			bool IsValid() const override { return false; }
+
+			int32 GetNewPointID() const override { return Metasound::FrontendInvalidID; }
+			TArray<FString> GetInputVertexNames() const override { return TArray<FString>(); }
+			TArray<FString> GetOutputVertexNames() const override { return TArray<FString>(); }
+
+			TArray<TSharedRef<INodeController>> GetNodes() override { return TArray<TSharedRef<INodeController>>(); }
+			TArray<TSharedRef<const INodeController>> GetConstNodes() const override { return TArray<TSharedRef<const INodeController>>(); }
+
+			int32 GetClassID() const override { return Metasound::FrontendInvalidID; }
+			TSharedRef<const INodeController> GetNodeWithID(int32 InNodeID) const override { return FInvalidNodeController::GetInvalid(); }
+			TSharedRef<INodeController> GetNodeWithID(int32 InNodeID) override { return FInvalidNodeController::GetInvalid(); }
+
+			TArray<TSharedRef<INodeController>> GetOutputNodes() override { return TArray<TSharedRef<INodeController>>(); }
+			TArray<TSharedRef<INodeController>> GetInputNodes() override { return TArray<TSharedRef<INodeController>>(); }
+			TArray<TSharedRef<const INodeController>> GetConstOutputNodes() const override { return TArray<TSharedRef<const INodeController>>(); }
+			TArray<TSharedRef<const INodeController>> GetConstInputNodes() const override { return TArray<TSharedRef<const INodeController>>(); }
+
+			bool ContainsOutputNodeWithName(const FString& InName) const override { return false; }
+			bool ContainsInputNodeWithName(const FString& InName) const override { return false; }
+
+			TSharedRef<const INodeController> GetOutputNodeWithName(const FString& InName) const override { return FInvalidNodeController::GetInvalid(); }
+			TSharedRef<const INodeController> GetInputNodeWithName(const FString& InName) const override { return FInvalidNodeController::GetInvalid(); }
+			TSharedRef<INodeController> GetOutputNodeWithName(const FString& InName) override { return FInvalidNodeController::GetInvalid(); }
+			TSharedRef<INodeController> GetInputNodeWithName(const FString& InName) override { return FInvalidNodeController::GetInvalid(); }
+
+
+			FConstClassInputAccessPtr FindClassInputWithName(const FString& InName) const override { return FConstClassInputAccessPtr(); }
+			FConstClassOutputAccessPtr FindClassOutputWithName(const FString& InName) const override { return FConstClassOutputAccessPtr(); }
+
+
+
+			TSharedRef<INodeController> AddInputVertex(const FMetasoundFrontendClassInput& InDescription) override { return FInvalidNodeController::GetInvalid(); }
+			bool RemoveInputVertex(const FString& InputName) override { return false; }
+
+			TSharedRef<INodeController> AddOutputVertex(const FMetasoundFrontendClassOutput& InDescription) override { return FInvalidNodeController::GetInvalid(); }
+			bool RemoveOutputVertex(const FString& OutputName) override { return false; }
+
+			// This can be used to determine what kind of property editor we should use for the data type of a given input.
+			// Will return Invalid if the input couldn't be found, or if the input doesn't support any kind of literals.
+			ELiteralType GetPreferredLiteralTypeForInputVertex(const FString& InInputName) const override { return ELiteralType::Invalid; }
+
+			// For inputs whose preferred literal type is UObject or UObjectArray, this can be used to determine the UObject corresponding to that input's datatype.
+			UClass* GetSupportedClassForInputVertex(const FString& InInputName) override { return nullptr; }
+
+			TArray<int32> GetDefaultIDsForInputVertex(const FString& InInputName) const { return TArray<int32>(); }
+			TArray<int32> GetDefaultIDsForOutputVertex(const FString& InInputName) const { return TArray<int32>(); }
+
+			// These can be used to set the default value for a given input on this graph.
+			// @returns false if the input name couldn't be found, or if the literal type was incompatible with the Data Type of this input.
+			bool SetDefaultInputToLiteral(const FString& InInputName, int32 InPointID, bool bInValue) override { return false; }
+			bool SetDefaultInputToLiteral(const FString& InInputName, int32 InPointID, int32 InValue) override { return false; }
+			bool SetDefaultInputToLiteral(const FString& InInputName, int32 InPointID, float InValue) override { return false; }
+			bool SetDefaultInputToLiteral(const FString& InInputName, int32 InPointID, const FString& InValue) override { return false; }
+			bool SetDefaultInputToLiteral(const FString& InInputName, int32 InPointID, UObject* InValue) override { return false; }
+			bool SetDefaultInputToLiteral(const FString& InInputName, int32 InPointID, const TArray<UObject*>& InValue) override { return false; }
+
+
+			// Set the display name for the input with the given name
+			void SetInputDisplayName(FString InName, const FText& InDisplayName) override {}
+
+			// Set the display name for the output with the given name
+			void SetOutputDisplayName(FString InName, const FText& InDisplayName) override {}
+
+			// This can be used to clear the current literal for a given input.
+			// @returns false if the input name couldn't be found.
+			bool ClearLiteralForInput(const FString& InInputName, int32 InPointID) override { return false; }
+
+			TSharedRef<INodeController> AddNode(const FNodeClassInfo& InNodeClass) override { return FInvalidNodeController::GetInvalid(); }
+			TSharedRef<INodeController> AddNode(const FNodeRegistryKey& InNodeClass) override { return FInvalidNodeController::GetInvalid(); }
+
+			// Remove the node corresponding to this node handle.
+			// On success, invalidates the received node handle.
+			bool RemoveNode(INodeController& InNode) override { return false; }
+
+			// Returns the metadata for the current graph, including the name, description and author.
+			const FMetasoundFrontendClassMetadata& GetGraphMetadata() const override { return MetasoundFrontendInvalidControllerPrivate::GetInvalid<FMetasoundFrontendClassMetadata>(); }
+
+			// If the FNodeController given is itself a Metasound graph,
+			// and the FNodeController is a direct member of this FGraphController,
+			// If successful, this will invalidate the FNodeController and paste the graph for this node directly
+			// into the graph.
+			// If not successful, InNode will not be affected.
+			// @returns true on success, false on failure.
+			bool InflateNodeDirectlyIntoGraph(const INodeController& InNode) override { return false; }
+
+			TSharedRef<INodeController> CreateEmptySubgraph(const FMetasoundFrontendClassMetadata& InInfo) override { return FInvalidNodeController::GetInvalid(); }
+
+			TUniquePtr<IOperator> BuildOperator(const FOperatorSettings& InSettings, const FMetasoundEnvironment& InEnvironment, TArray<IOperatorBuilder::FBuildErrorPtr>& OutBuildErrors) const override
+			{
+				return TUniquePtr<IOperator>(nullptr);
+			}
+
+			TSharedRef<IDocumentController> GetOwningDocument() override;
+			TSharedRef<const IDocumentController> GetOwningDocument() const override;
+		};
+
+		/** FInvalidDocumentController is a controller which is always invalid. 
+		 *
+		 *  All methods return defaults and return error flags or invalid values.
+		 */
+		class METASOUNDFRONTEND_API FInvalidDocumentController : public IDocumentController 
+		{
+			public:
+				FInvalidDocumentController() = default;
+				virtual ~FInvalidDocumentController() = default;
+
+				static TSharedRef<IDocumentController> GetInvalid()
+				{
+					static TSharedRef<IDocumentController> InvalidDocument = MakeShared<FInvalidDocumentController>();
+					return InvalidDocument;
+				}
+
+				bool IsValid() const override { return false; }
+				bool IsRequiredInput(const FString& InInputName) const override { return false; }
+				bool IsRequiredOutput(const FString& InOutputName) const override { return false; }
+
+				TArray<FMetasoundFrontendClassVertex> GetRequiredInputs() const override { return TArray<FMetasoundFrontendClassVertex>(); }
+				TArray<FMetasoundFrontendClassVertex> GetRequiredOutputs() const override { return TArray<FMetasoundFrontendClassVertex>(); }
+
+				TArray<FMetasoundFrontendClass> GetDependencies() const override { return TArray<FMetasoundFrontendClass>(); }
+				TArray<FMetasoundFrontendGraphClass> GetSubgraphs() const override { return TArray<FMetasoundFrontendGraphClass>(); }
+				TArray<FMetasoundFrontendClass> GetClasses() const override { return TArray<FMetasoundFrontendClass>(); }
+
+				FConstClassAccessPtr FindDependencyWithID(int32 InClassID) const override { return FConstClassAccessPtr(); }
+				FConstGraphClassAccessPtr FindSubgraphWithID(int32 InClassID) const override { return FConstGraphClassAccessPtr(); }
+				FConstClassAccessPtr FindClassWithID(int32 InClassID) const override { return FConstClassAccessPtr(); }
+
+				FConstClassAccessPtr FindClass(const FNodeClassInfo& InNodeClass) const override { return FConstClassAccessPtr(); }
+				FConstClassAccessPtr FindOrAddClass(const FNodeClassInfo& InNodeClass) override { return FConstClassAccessPtr(); }
+				FConstClassAccessPtr FindClass(const FMetasoundFrontendClassMetadata& InMetadata) const override{ return FConstClassAccessPtr(); }
+				FConstClassAccessPtr FindOrAddClass(const FMetasoundFrontendClassMetadata& InMetadata) override{ return FConstClassAccessPtr(); }
+				void RemoveUnreferencedDependencies() override {}
+
+				TArray<FGraphHandle> GetSubgraphHandles() override { return TArray<FGraphHandle>(); }
+
+				TArray<FConstGraphHandle> GetSubgraphHandles() const override { return TArray<FConstGraphHandle>(); }
+
+				FGraphHandle GetSubgraphWithClassID(int32 InClassID) { return FInvalidGraphController::GetInvalid(); }
+
+				FConstGraphHandle GetSubgraphWithClassID(int32 InClassID) const { return FInvalidGraphController::GetInvalid(); }
+
+				TSharedRef<IGraphController> GetRootGraph() override;
+				TSharedRef<const IGraphController> GetRootGraph() const override;
+				bool ExportToJSONAsset(const FString& InAbsolutePath) const override { return false; }
+				FString ExportToJSON() const override { return FString(); }
+		};
+	}
+}
diff --git a/Engine/Plugins/Runtime/Metasound/Source/MetasoundFrontend/Private/MetasoundFrontendQuery.cpp b/Engine/Plugins/Runtime/Metasound/Source/MetasoundFrontend/Private/MetasoundFrontendQuery.cpp
--- a/Engine/Plugins/Runtime/Metasound/Source/MetasoundFrontend/Private/MetasoundFrontendQuery.cpp
+++ b/Engine/Plugins/Runtime/Metasound/Source/MetasoundFrontend/Private/MetasoundFrontendQuery.cpp
@@ -4,7 +4,7 @@
 
 #include "Algo/Sort.h"
 #include "CoreMinimal.h"
-#include "MetasoundFrontendDataLayout.h"
+#include "MetasoundFrontendDocument.h"
 
 namespace Metasound
 {
diff --git a/Engine/Plugins/Runtime/Metasound/Source/MetasoundFrontend/Private/MetasoundFrontendQuerySteps.cpp b/Engine/Plugins/Runtime/Metasound/Source/MetasoundFrontend/Private/MetasoundFrontendQuerySteps.cpp
--- a/Engine/Plugins/Runtime/Metasound/Source/MetasoundFrontend/Private/MetasoundFrontendQuerySteps.cpp
+++ b/Engine/Plugins/Runtime/Metasound/Source/MetasoundFrontend/Private/MetasoundFrontendQuerySteps.cpp
@@ -4,7 +4,7 @@
 
 #include "CoreMinimal.h"
 #include "MetasoundFrontend.h"
-#include "MetasoundFrontendDataLayout.h"
+#include "MetasoundFrontendDocument.h"
 
 namespace Metasound
 {
@@ -14,7 +14,7 @@ namespace Metasound
 
 		for (const Frontend::FNodeClassInfo& ClassInfo : ClassInfos)
 		{
-			OutEntries.Emplace(TInPlaceType<FMetasoundClassDescription>(), Frontend::GenerateClassDescription(ClassInfo));
+			OutEntries.Emplace(TInPlaceType<FMetasoundFrontendClass>(), Frontend::GenerateClassDescription(ClassInfo));
 		}
 	}
 
@@ -25,10 +25,10 @@ namespace Metasound
 
 	bool FFilterClassesByInputVertexDataType::Filter(const FFrontendQueryEntry& InEntry) const
 	{
-		check(InEntry.Value.IsType<FMetasoundClassDescription>());
+		check(InEntry.Value.IsType<FMetasoundFrontendClass>());
 
-		return InEntry.Value.Get<FMetasoundClassDescription>().Inputs.ContainsByPredicate(
-			[this](const FMetasoundInputDescription& InDesc)
+		return InEntry.Value.Get<FMetasoundFrontendClass>().Interface.Inputs.ContainsByPredicate(
+			[this](const FMetasoundFrontendClassInput& InDesc)
 			{
 				return InDesc.TypeName == InputVertexTypeName;
 			}
@@ -42,8 +42,8 @@ namespace Metasound
 
 	bool FFilterClassesByOutputVertexDataType::Filter(const FFrontendQueryEntry& InEntry) const
 	{
-		return InEntry.Value.Get<FMetasoundClassDescription>().Outputs.ContainsByPredicate(
-			[this](const FMetasoundOutputDescription& InDesc)
+		return InEntry.Value.Get<FMetasoundFrontendClass>().Interface.Outputs.ContainsByPredicate(
+			[this](const FMetasoundFrontendClassOutput& InDesc)
 			{
 				return InDesc.TypeName == OutputVertexTypeName;
 			}
diff --git a/Engine/Plugins/Runtime/Metasound/Source/MetasoundFrontend/Private/MetasoundFrontendRegistries.cpp b/Engine/Plugins/Runtime/Metasound/Source/MetasoundFrontend/Private/MetasoundFrontendRegistries.cpp
--- a/Engine/Plugins/Runtime/Metasound/Source/MetasoundFrontend/Private/MetasoundFrontendRegistries.cpp
+++ b/Engine/Plugins/Runtime/Metasound/Source/MetasoundFrontend/Private/MetasoundFrontendRegistries.cpp
@@ -1,7 +1,9 @@
 // Copyright Epic Games, Inc. All Rights Reserved.
 
 #include "MetasoundFrontendRegistries.h"
+
 #include "CoreMinimal.h"
+#include "MetasoundLog.h"
 #include "Misc/ScopeLock.h"
 #include "HAL/PlatformTime.h"
 
@@ -113,7 +115,7 @@ TUniquePtr<Metasound::INode> FMetasoundFrontendRegistryContainer::ConstructInput
 	}
 }
 
-TUniquePtr<Metasound::INode> FMetasoundFrontendRegistryContainer::ConstructOutputNode(const FName& InOutputType, const Metasound::FOutputNodeConstrutorParams& InParams)
+TUniquePtr<Metasound::INode> FMetasoundFrontendRegistryContainer::ConstructOutputNode(const FName& InOutputType, const Metasound::FOutputNodeConstructorParams& InParams)
 {
 	if (ensureAlwaysMsgf(DataTypeRegistry.Contains(InOutputType), TEXT("Couldn't find data type %s!"), *InOutputType.ToString()))
 	{
@@ -125,27 +127,27 @@ TUniquePtr<Metasound::INode> FMetasoundFrontendRegistryContainer::ConstructOutpu
 	}
 }
 
-Metasound::FDataTypeLiteralParam FMetasoundFrontendRegistryContainer::GenerateLiteralForUObject(const FName& InDataType, UObject* InObject)
+Metasound::FLiteral FMetasoundFrontendRegistryContainer::GenerateLiteralForUObject(const FName& InDataType, UObject* InObject)
 {
 	if (ensureAlwaysMsgf(DataTypeRegistry.Contains(InDataType), TEXT("Couldn't find data type %s!"), *InDataType.ToString()))
 	{
 		 Audio::IProxyDataPtr ProxyPtr = DataTypeRegistry[InDataType].Callbacks.CreateAudioProxy(InObject);
 		 if (ensureAlwaysMsgf(ProxyPtr.IsValid(), TEXT("UObject failed to create a valid proxy!")))
 		 {
-			 return Metasound::FDataTypeLiteralParam(MoveTemp(ProxyPtr));
+			 return Metasound::FLiteral(MoveTemp(ProxyPtr));
 		 }
 		 else
 		 {
-			 return Metasound::FDataTypeLiteralParam();
+			 return Metasound::FLiteral();
 		 }
 	}
 	else
 	{
-		return Metasound::FDataTypeLiteralParam();
+		return Metasound::FLiteral();
 	}
 }
 
-Metasound::FDataTypeLiteralParam FMetasoundFrontendRegistryContainer::GenerateLiteralForUObjectArray(const FName& InDataType, TArray<UObject*> InObjectArray)
+Metasound::FLiteral FMetasoundFrontendRegistryContainer::GenerateLiteralForUObjectArray(const FName& InDataType, TArray<UObject*> InObjectArray)
 {
 	if (ensureAlwaysMsgf(DataTypeRegistry.Contains(InDataType), TEXT("Couldn't find data type %s!"), *InDataType.ToString()))
 	{
@@ -161,11 +163,11 @@ Metasound::FDataTypeLiteralParam FMetasoundFrontendRegistryContainer::GenerateLi
 			}
 		}
 
-		return Metasound::FDataTypeLiteralParam(MoveTemp(ProxyArray));
+		return Metasound::FLiteral(MoveTemp(ProxyArray));
 	}
 	else
 	{
-		return Metasound::FDataTypeLiteralParam();
+		return Metasound::FLiteral();
 	}
 }
 
@@ -198,17 +200,17 @@ TArray<::Metasound::Frontend::FConverterNodeInfo> FMetasoundFrontendRegistryCont
 	}
 }
 
-Metasound::ELiteralArgType FMetasoundFrontendRegistryContainer::GetDesiredLiteralTypeForDataType(FName InDataType) const
+Metasound::ELiteralType FMetasoundFrontendRegistryContainer::GetDesiredLiteralTypeForDataType(FName InDataType) const
 {
 	if (!DataTypeRegistry.Contains(InDataType))
 	{
-		return Metasound::ELiteralArgType::Invalid;
+		return Metasound::ELiteralType::Invalid;
 	}
 
 	const FDataTypeRegistryElement& DataTypeInfo = DataTypeRegistry[InDataType];
 	
 	// If there's a designated preferred literal type for this datatype, use that.
-	if (DataTypeInfo.Info.PreferredLiteralType != Metasound::ELiteralArgType::None)
+	if (DataTypeInfo.Info.PreferredLiteralType != Metasound::ELiteralType::None)
 	{
 		return DataTypeInfo.Info.PreferredLiteralType;
 	}
@@ -216,30 +218,30 @@ Metasound::ELiteralArgType FMetasoundFrontendRegistryContainer::GetDesiredLitera
 	// Otherwise, we opt for the highest precision construction option available.
 	if (DataTypeInfo.Info.bIsStringParsable)
 	{
-		return Metasound::ELiteralArgType::String;
+		return Metasound::ELiteralType::String;
 	}
 	else if (DataTypeInfo.Info.bIsFloatParsable)
 	{
-		return Metasound::ELiteralArgType::Float;
+		return Metasound::ELiteralType::Float;
 	}
 	else if (DataTypeInfo.Info.bIsIntParsable)
 	{
-		return Metasound::ELiteralArgType::Integer;
+		return Metasound::ELiteralType::Integer;
 	}
 	else if (DataTypeInfo.Info.bIsBoolParsable)
 	{
-		return Metasound::ELiteralArgType::Boolean;
+		return Metasound::ELiteralType::Boolean;
 	}
 	else if (DataTypeInfo.Info.bIsDefaultParsable)
 	{
-		return Metasound::ELiteralArgType::None;
+		return Metasound::ELiteralType::None;
 	}
 	else
 	{
 		// if we ever hit this, something has gone terribly wrong with the REGISTER_METASOUND_DATATYPE macro.
 		// we should have failed to compile if any of these are false.
 		checkNoEntry();
-		return Metasound::ELiteralArgType::Invalid;
+		return Metasound::ELiteralType::Invalid;
 	}
 }
 
@@ -256,7 +258,7 @@ UClass* FMetasoundFrontendRegistryContainer::GetLiteralUClassForDataType(FName I
 	}
 }
 
-bool FMetasoundFrontendRegistryContainer::DoesDataTypeSupportLiteralType(FName InDataType, Metasound::ELiteralArgType InLiteralType) const
+bool FMetasoundFrontendRegistryContainer::DoesDataTypeSupportLiteralType(FName InDataType, Metasound::ELiteralType InLiteralType) const
 {
 	if (!DataTypeRegistry.Contains(InDataType))
 	{
@@ -268,35 +270,35 @@ bool FMetasoundFrontendRegistryContainer::DoesDataTypeSupportLiteralType(FName I
 	
 	switch (InLiteralType)
 	{
-		case Metasound::ELiteralArgType::Boolean:
+		case Metasound::ELiteralType::Boolean:
 		{
 			return DataTypeInfo.Info.bIsBoolParsable;
 		}
-		case Metasound::ELiteralArgType::Integer:
+		case Metasound::ELiteralType::Integer:
 		{
 			return DataTypeInfo.Info.bIsIntParsable;
 		}
-		case Metasound::ELiteralArgType::Float:
+		case Metasound::ELiteralType::Float:
 		{
 			return DataTypeInfo.Info.bIsFloatParsable;
 		}
-		case Metasound::ELiteralArgType::String:
+		case Metasound::ELiteralType::String:
 		{
 			return DataTypeInfo.Info.bIsStringParsable;
 		}
-		case Metasound::ELiteralArgType::UObjectProxy:
+		case Metasound::ELiteralType::UObjectProxy:
 		{
 			return DataTypeInfo.Info.bIsProxyParsable;
 		}
-		case Metasound::ELiteralArgType::UObjectProxyArray:
+		case Metasound::ELiteralType::UObjectProxyArray:
 		{
 			return DataTypeInfo.Info.bIsProxyArrayParsable;
 		}
-		case Metasound::ELiteralArgType::None:
+		case Metasound::ELiteralType::None:
 		{
 			return DataTypeInfo.Info.bIsDefaultParsable;
 		}
-		case Metasound::ELiteralArgType::Invalid:
+		case Metasound::ELiteralType::Invalid:
 		default:
 		{
 			return false;
@@ -304,7 +306,7 @@ bool FMetasoundFrontendRegistryContainer::DoesDataTypeSupportLiteralType(FName I
 	}
 }
 
-bool FMetasoundFrontendRegistryContainer::RegisterDataType(const ::Metasound::FDataTypeRegistryInfo& InDataInfo, ::Metasound::FDataTypeConstructorCallbacks&& InCallbacks)
+bool FMetasoundFrontendRegistryContainer::RegisterDataType(const ::Metasound::FDataTypeRegistryInfo& InDataInfo, const ::Metasound::FDataTypeConstructorCallbacks& InCallbacks)
 {
 	if (!ensureAlwaysMsgf(!DataTypeRegistry.Contains(InDataInfo.DataTypeName), TEXT("Name collision when trying to register Metasound Data Type %s! Make sure that you created a unique name for your data type, and that REGISTER_METASOUND_DATATYPE isn't called in a public header."), *InDataInfo.DataTypeName.ToString()))
 	{
@@ -313,14 +315,22 @@ bool FMetasoundFrontendRegistryContainer::RegisterDataType(const ::Metasound::FD
 	}
 	else
 	{
-		FDataTypeRegistryElement InElement = { MoveTemp(InCallbacks), InDataInfo };
-		DataTypeRegistry.Add(InDataInfo.DataTypeName, MoveTemp(InElement));
-		UE_LOG(LogTemp, Display, TEXT("Registered Metasound Datatype %s."), *InDataInfo.DataTypeName.ToString());
+		FDataTypeRegistryElement InElement = { InCallbacks, InDataInfo };
+
+		DataTypeRegistry.Add(InDataInfo.DataTypeName, InElement);
+
+		Metasound::Frontend::FNodeRegistryKey InputNodeRegistryKey = GetRegistryKey(InCallbacks.CreateFrontendInputClass().Metadata);
+		DataTypeNodeRegistry.Add(InputNodeRegistryKey, InElement);
+
+		Metasound::Frontend::FNodeRegistryKey OutputNodeRegistryKey = GetRegistryKey(InCallbacks.CreateFrontendOutputClass().Metadata);
+		DataTypeNodeRegistry.Add(OutputNodeRegistryKey, InElement);
+
+		UE_LOG(LogMetasound, Display, TEXT("Registered Metasound Datatype %s."), *InDataInfo.DataTypeName.ToString());
 		return true;
 	}
 }
 
-bool FMetasoundFrontendRegistryContainer::RegisterExternalNode(Metasound::FCreateMetasoundNodeFunction&& InCreateNode, Metasound::FCreateMetasoundClassDescriptionFunction&& InCreateDescription)
+bool FMetasoundFrontendRegistryContainer::RegisterExternalNode(Metasound::FCreateMetasoundNodeFunction&& InCreateNode, Metasound::FCreateMetasoundFrontendClassFunction&& InCreateDescription)
 {
 	FNodeRegistryElement RegistryElement = FNodeRegistryElement(MoveTemp(InCreateNode), MoveTemp(InCreateDescription));
 
@@ -341,11 +351,11 @@ bool FMetasoundFrontendRegistryContainer::RegisterExternalNode(Metasound::FCreat
 
 bool FMetasoundFrontendRegistryContainer::GetRegistryKey(const Metasound::Frontend::FNodeRegistryElement& InElement, Metasound::Frontend::FNodeRegistryKey& OutKey)
 {
-	if (InElement.CreateClassDescription)
+	if (InElement.CreateFrontendClass)
 	{
-		FMetasoundClassDescription Description = InElement.CreateClassDescription();
+		FMetasoundFrontendClass FrontendClass = InElement.CreateFrontendClass();
 
-		OutKey = MetasoundFrontendRegistryPrivate::GetRegistryKey(FName(Description.Metadata.NodeName), Description.Metadata.MajorVersion);
+		OutKey = GetRegistryKey(FrontendClass.Metadata);
 
 		return true;
 	}
@@ -358,6 +368,74 @@ Metasound::Frontend::FNodeRegistryKey FMetasoundFrontendRegistryContainer::GetRe
 	return MetasoundFrontendRegistryPrivate::GetRegistryKey(InNodeMetadata.ClassName, InNodeMetadata.MajorVersion);
 }
 
+Metasound::Frontend::FNodeRegistryKey FMetasoundFrontendRegistryContainer::GetRegistryKey(const FMetasoundFrontendClassMetadata& InNodeMetadata)
+{
+	return MetasoundFrontendRegistryPrivate::GetRegistryKey(FName(InNodeMetadata.Name.Name), InNodeMetadata.Version.Major);
+}
+
+bool FMetasoundFrontendRegistryContainer::GetFrontendClassFromRegistered(const FMetasoundFrontendClassMetadata& InMetadata, FMetasoundFrontendClass& OutClass)
+{
+	FMetasoundFrontendRegistryContainer* Registry = FMetasoundFrontendRegistryContainer::Get();
+
+	if (ensure(nullptr != Registry))
+	{
+		Metasound::Frontend::FNodeRegistryKey RegistryKey = GetRegistryKey(InMetadata);
+
+		if (InMetadata.Type == EMetasoundFrontendClassType::External)
+		{
+			if (Registry->ExternalNodeRegistry.Contains(RegistryKey))
+			{
+				OutClass = Registry->ExternalNodeRegistry[RegistryKey].CreateFrontendClass();
+				return true;
+			}
+		}
+		else if (InMetadata.Type == EMetasoundFrontendClassType::Input)
+		{
+			if (Registry->DataTypeNodeRegistry.Contains(RegistryKey))
+			{
+				OutClass = Registry->DataTypeNodeRegistry[RegistryKey].Callbacks.CreateFrontendInputClass();
+				return true;
+			}
+		}
+		else if (InMetadata.Type == EMetasoundFrontendClassType::Output)
+		{
+			if (Registry->DataTypeNodeRegistry.Contains(RegistryKey))
+			{
+				OutClass = Registry->DataTypeNodeRegistry[RegistryKey].Callbacks.CreateFrontendOutputClass();
+				return true;
+			}
+		}
+	}
+
+	return false;
+}
+
+bool FMetasoundFrontendRegistryContainer::GetInputNodeClassMetadataForDataType(const FName& InDataTypeName, FMetasoundFrontendClassMetadata& OutMetadata)
+{
+	if (FMetasoundFrontendRegistryContainer* Registry = FMetasoundFrontendRegistryContainer::Get())
+	{
+		if (Registry->DataTypeRegistry.Contains(InDataTypeName))
+		{
+			OutMetadata = Registry->DataTypeRegistry[InDataTypeName].Callbacks.CreateFrontendInputClass().Metadata;
+			return true;
+		}
+	}
+	return false;
+}
+
+bool FMetasoundFrontendRegistryContainer::GetOutputNodeClassMetadataForDataType(const FName& InDataTypeName, FMetasoundFrontendClassMetadata& OutMetadata)
+{
+	if (FMetasoundFrontendRegistryContainer* Registry = FMetasoundFrontendRegistryContainer::Get())
+	{
+		if (Registry->DataTypeRegistry.Contains(InDataTypeName))
+		{
+			OutMetadata = Registry->DataTypeRegistry[InDataTypeName].Callbacks.CreateFrontendOutputClass().Metadata;
+			return true;
+		}
+	}
+	return false;
+}
+
 bool FMetasoundFrontendRegistryContainer::RegisterConversionNode(const FConverterNodeRegistryKey& InNodeKey, const FConverterNodeInfo& InNodeInfo)
 {
 	if (!ConverterNodeRegistry.Contains(InNodeKey))
