Metasound Frontend Revamp

diff --git a/Engine/Plugins/Runtime/Metasound/Source/MetasoundFrontend/Private/MetasoundFrontendStandardController.cpp b/Engine/Plugins/Runtime/Metasound/Source/MetasoundFrontend/Private/MetasoundFrontendStandardController.cpp
new file mode 100644
--- /dev/null
+++ b/Engine/Plugins/Runtime/Metasound/Source/MetasoundFrontend/Private/MetasoundFrontendStandardController.cpp
@@ -0,0 +1,3641 @@
+// Copyright Epic Games, Inc. All Rights Reserved.
+
+#include "MetasoundFrontendStandardController.h"
+
+#include "Algo/ForEach.h"
+#include "Algo/NoneOf.h"
+#include "CoreMinimal.h"
+#include "HAL/FileManager.h"
+#include "MetasoundAccessPtr.h"
+#include "MetasoundFrontendController.h"
+#include "MetasoundFrontendDocument.h"
+#include "MetasoundFrontendGraph.h"
+#include "MetasoundJsonBackend.h"
+#include "MetasoundOperatorBuilder.h"
+#include "StructDeserializer.h"
+#include "StructSerializer.h"
+#include "Serialization/MemoryReader.h"
+
+namespace Metasound
+{
+	namespace Frontend
+	{
+		namespace FrontendControllerIntrinsics
+		{
+			// utility function for returning invalid values. If an invalid value type
+			// needs special construction, this template can be specialized. 
+			template<typename ValueType>
+			ValueType GetInvalidValue()
+			{
+				ValueType InvalidValue;
+				return InvalidValue;
+			}
+
+			// Utility function to get a value for an existing description pointer
+			// or return an invalid value const ref if the TAccessPtr is invalid.
+			template<typename ValueType, typename PtrType>
+			const ValueType& GetValueOrReturnInvalid(const TAccessPtr<PtrType>& InPointer, ValueType (PtrType::*InMember))
+			{
+				static const ValueType InvalidValue = GetInvalidValue<ValueType>();
+
+				if (InPointer.IsValid())
+				{
+					return (*InPointer).*InMember;
+				}
+
+				return InvalidValue;
+			}
+
+			// Invalid value specialization for int32
+			template<>
+			int32 GetInvalidValue<int32>() { return INDEX_NONE; }
+
+			// Invalid value specialization for EMetasoundFrontendClassType
+			template<>
+			EMetasoundFrontendClassType GetInvalidValue<EMetasoundFrontendClassType>() { return EMetasoundFrontendClassType::Invalid; }
+
+			// Invalid value specialization for FText
+			template<>
+			FText GetInvalidValue<FText>() { return FText::GetEmpty(); }
+
+			template<typename ValueType>
+			const ValueType& GetInvalidValueConstRef()
+			{
+				static const ValueType Value = GetInvalidValue<ValueType>();
+				return Value;
+			}
+
+			template<typename Type>
+			TAccessPtr<Type> MakeAutoAccessPtr(Type& InObj)
+			{
+				return MakeAccessPtr(InObj.AccessPoint, InObj);
+				//return MakeAccessPtr(InObj);
+			}
+
+
+			// TODO: All these access ptr helper routines should be merged into
+			// the TAccessPtr<> interface. The TAccessPtr<> class should be specialized 
+			// for the types used here, with the specialized TAccessPtr<> interface
+			// adding support for getting specific member access ptrs and doing
+			// specific queries. 
+			struct FGetSubgraphFromDocumentByID
+			{
+				FGetSubgraphFromDocumentByID(int32 InID)
+				: ID(InID)
+				{
+				}
+
+				FMetasoundFrontendGraphClass* operator() (FMetasoundFrontendDocument& InDoc) const
+				{
+					return InDoc.Subgraphs.FindByPredicate([&](const FMetasoundFrontendGraphClass& GraphClass) { return GraphClass.ID == ID; });
+				}
+
+				const FMetasoundFrontendGraphClass* operator() (const FMetasoundFrontendDocument& InDoc) const
+				{
+					return InDoc.Subgraphs.FindByPredicate([&](const FMetasoundFrontendGraphClass& GraphClass) { return GraphClass.ID == ID; });
+				}
+
+			private:
+				int32 ID;
+			};
+
+			struct FGetClassInputFromClassWithName
+			{
+				FGetClassInputFromClassWithName(const FString& InName)
+				: Name(InName)
+				{
+				}
+
+				FMetasoundFrontendClassInput* operator() (FMetasoundFrontendClass& InClass)
+				{
+					auto IsClassInputWithName = [&](const FMetasoundFrontendClassInput& ClassInput) { return ClassInput.Name == Name; };
+					return InClass.Interface.Inputs.FindByPredicate(IsClassInputWithName);
+				}
+
+				const FMetasoundFrontendClassInput* operator() (const FMetasoundFrontendClass& InClass)
+				{
+					auto IsClassInputWithName = [&](const FMetasoundFrontendClassInput& ClassInput) { return ClassInput.Name == Name; };
+					return InClass.Interface.Inputs.FindByPredicate(IsClassInputWithName);
+				}
+				private:
+				FString Name;
+			};
+
+			struct FGetClassOutputFromClassWithName
+			{
+				FGetClassOutputFromClassWithName(const FString& InName)
+				: Name(InName)
+				{
+				}
+
+				FMetasoundFrontendClassOutput* operator() (FMetasoundFrontendClass& InClass)
+				{
+					auto IsClassOutputWithName = [&](const FMetasoundFrontendClassOutput& ClassOutput) { return ClassOutput.Name == Name; };
+					return InClass.Interface.Outputs.FindByPredicate(IsClassOutputWithName);
+				}
+
+				const FMetasoundFrontendClassOutput* operator() (const FMetasoundFrontendClass& InClass)
+				{
+					auto IsClassOutputWithName = [&](const FMetasoundFrontendClassOutput& ClassOutput) { return ClassOutput.Name == Name; };
+					return InClass.Interface.Outputs.FindByPredicate(IsClassOutputWithName);
+				}
+				private:
+				FString Name;
+			};
+
+			auto FGetNodeInputFromNodeWithName(const FString& InName)
+			{
+				return [=](FMetasoundFrontendNode& InNode) -> FMetasoundFrontendVertex*
+				{
+					auto IsVertexWithName = [=](const FMetasoundFrontendVertex& Vertex) { return Vertex.Name == InName; };
+					return InNode.Interface.Inputs.FindByPredicate(IsVertexWithName);
+				};
+			}
+
+			auto FGetConstNodeInputFromNodeWithName(const FString& InName)
+			{
+				return [=](const FMetasoundFrontendNode& InNode) -> const FMetasoundFrontendVertex*
+				{
+					auto IsVertexWithName = [=](const FMetasoundFrontendVertex& Vertex) { return Vertex.Name == InName; };
+					return InNode.Interface.Inputs.FindByPredicate(IsVertexWithName);
+				};
+			}
+
+			auto FGetNodeOutputFromNodeWithName(const FString& InName)
+			{
+				return [=](FMetasoundFrontendNode& InNode) -> FMetasoundFrontendVertex*
+				{
+					auto IsVertexWithName = [=](const FMetasoundFrontendVertex& Vertex) { return Vertex.Name == InName; };
+					return InNode.Interface.Outputs.FindByPredicate(IsVertexWithName);
+				};
+			}
+
+			auto FGetConstNodeOutputFromNodeWithName(const FString& InName)
+			{
+				return [=](const FMetasoundFrontendNode& InNode) -> const FMetasoundFrontendVertex*
+				{
+					auto IsVertexWithName = [=](const FMetasoundFrontendVertex& Vertex) { return Vertex.Name == InName; };
+					return InNode.Interface.Outputs.FindByPredicate(IsVertexWithName);
+				};
+			}
+
+			auto FGetNodeInputFromNodeWithPoint(int32 InPointID)
+			{
+				return [=](FMetasoundFrontendNode& InNode) -> FMetasoundFrontendVertex*
+				{
+					auto IsVertexWithPoint = [=](const FMetasoundFrontendVertex& Vertex) { return Vertex.PointIDs.Contains(InPointID); };
+					return InNode.Interface.Inputs.FindByPredicate(IsVertexWithPoint);
+				};
+			}
+
+			auto FGetConstNodeInputFromNodeWithPoint(int32 InPointID)
+			{
+				return [=](const FMetasoundFrontendNode& InNode) -> const FMetasoundFrontendVertex*
+				{
+					auto IsVertexWithPoint = [=](const FMetasoundFrontendVertex& Vertex) { return Vertex.PointIDs.Contains(InPointID); };
+					return InNode.Interface.Inputs.FindByPredicate(IsVertexWithPoint);
+				};
+			}
+
+			auto FGetNodeOutputFromNodeWithPoint(int32 InPointID)
+			{
+				return [=](FMetasoundFrontendNode& InNode) -> FMetasoundFrontendVertex*
+				{
+					auto IsVertexWithPoint = [=](const FMetasoundFrontendVertex& Vertex) { return Vertex.PointIDs.Contains(InPointID); };
+					return InNode.Interface.Outputs.FindByPredicate(IsVertexWithPoint);
+				};
+			}
+
+			auto FGetConstNodeOutputFromNodeWithPoint(int32 InPointID)
+			{
+				return [=](const FMetasoundFrontendNode& InNode) -> const FMetasoundFrontendVertex*
+				{
+					auto IsVertexWithPoint = [=](const FMetasoundFrontendVertex& Vertex) { return Vertex.PointIDs.Contains(InPointID); };
+					return InNode.Interface.Outputs.FindByPredicate(IsVertexWithPoint);
+				};
+			}
+
+			auto FGetConstClassFromDocumentDependenciesWithID(int32 InClassID)
+			{
+				return [=](const FMetasoundFrontendDocument& InDocument) -> const FMetasoundFrontendClass*
+				{
+					auto IsClassWithID = [=](const FMetasoundFrontendClass& InClass) { return InClass.ID == InClassID; };
+					return InDocument.Dependencies.FindByPredicate(IsClassWithID);
+				};
+			}
+
+			auto FGetConstGraphClassFromDocumentWithID(int32 InClassID)
+			{
+				return [=](const FMetasoundFrontendDocument& InDocument) -> const FMetasoundFrontendGraphClass*
+				{
+					auto IsClassWithID = [=](const FMetasoundFrontendGraphClass& InClass) { return InClass.ID == InClassID; };
+					return InDocument.Subgraphs.FindByPredicate(IsClassWithID);
+				};
+			}
+
+			auto FGetConstClassFromDocumentWithClassInfo(const FNodeClassInfo& InInfo)
+			{
+				return [=](const FMetasoundFrontendDocument& InDocument) -> const FMetasoundFrontendClass*
+				{
+					auto IsClassWithClassInfo = [=](const FMetasoundFrontendClass& InClass) 
+					{ 
+						return FDocumentController::IsMatchingMetasoundClass(InInfo, InClass.Metadata);
+					};
+
+					const FMetasoundFrontendClass* MetasoundClass = InDocument.Dependencies.FindByPredicate(IsClassWithClassInfo);
+					if (nullptr == MetasoundClass)
+					{
+						MetasoundClass = InDocument.Subgraphs.FindByPredicate(IsClassWithClassInfo);
+					}
+					return MetasoundClass;
+				};
+			}
+
+			auto FGetConstClassFromDocumentWithMetadata(const FMetasoundFrontendClassMetadata& InMetadata)
+			{
+				return [=](const FMetasoundFrontendDocument& InDocument) -> const FMetasoundFrontendClass*
+				{
+					auto IsClassWithMetadata = [=](const FMetasoundFrontendClass& InClass) 
+					{ 
+						return FDocumentController::IsMatchingMetasoundClass(InMetadata, InClass.Metadata);
+					};
+
+					const FMetasoundFrontendClass* MetasoundClass = InDocument.Dependencies.FindByPredicate(IsClassWithMetadata);
+					if (nullptr == MetasoundClass)
+					{
+						MetasoundClass = InDocument.Subgraphs.FindByPredicate(IsClassWithMetadata);
+					}
+					return MetasoundClass;
+				};
+			}
+
+			auto FGetNodeFromGraphClassByNodeID(int32 InNodeID)
+			{
+				return [=](FMetasoundFrontendGraphClass& InGraphClass) -> FMetasoundFrontendNode*
+				{
+					auto IsNodeWithID = [&](const FMetasoundFrontendNode& InNode) { return InNode.ID == InNodeID; };
+					return InGraphClass.Graph.Nodes.FindByPredicate(IsNodeWithID);
+				};
+			}
+
+			auto FGetConstGraphFromGraphClass()
+			{
+				return [](const FMetasoundFrontendGraphClass& InGraphClass) -> const FMetasoundFrontendGraph* { return &InGraphClass.Graph; };
+			}
+			auto FGetGraphFromGraphClass()
+			{
+				return [](FMetasoundFrontendGraphClass& InGraphClass) -> FMetasoundFrontendGraph* { return &InGraphClass.Graph; };
+			}
+
+			auto FGetConstNodeFromGraphClassByNodeID(int32 InNodeID)
+			{
+				return [=](const FMetasoundFrontendGraphClass& InGraphClass) -> const FMetasoundFrontendNode*
+				{
+					auto IsNodeWithID = [&](const FMetasoundFrontendNode& InNode) { return InNode.ID == InNodeID; };
+					return InGraphClass.Graph.Nodes.FindByPredicate(IsNodeWithID);
+				};
+			}
+
+			auto FGetConstClassInputFromGraphClassWithNodeID(int32 InNodeID)
+			{
+				return [=](const FMetasoundFrontendGraphClass& InGraphClass) -> const FMetasoundFrontendClassInput*
+				{
+					return InGraphClass.Interface.Inputs.FindByPredicate([&](const FMetasoundFrontendClassInput& Desc) { return Desc.NodeID == InNodeID; });
+				};
+			}
+
+			auto FGetClassInputFromGraphClassWithNodeID(int32 InNodeID)
+			{
+				return [=](FMetasoundFrontendGraphClass& InGraphClass) -> FMetasoundFrontendClassInput*
+				{
+					return InGraphClass.Interface.Inputs.FindByPredicate([&](const FMetasoundFrontendClassInput& Desc) { return Desc.NodeID == InNodeID; });
+				};
+			}
+
+			auto FGetConstClassOutputFromGraphClassWithNodeID(int32 InNodeID)
+			{
+				return [=](const FMetasoundFrontendGraphClass& InGraphClass) -> const FMetasoundFrontendClassOutput*
+				{
+					return InGraphClass.Interface.Outputs.FindByPredicate([&](const FMetasoundFrontendClassOutput& Desc) { return Desc.NodeID == InNodeID; });
+				};
+			}
+
+			auto FGetClassOutputFromGraphClassWithNodeID(int32 InNodeID)
+			{
+				return [=](FMetasoundFrontendGraphClass& InGraphClass) -> FMetasoundFrontendClassOutput*
+				{
+					return InGraphClass.Interface.Outputs.FindByPredicate([&](const FMetasoundFrontendClassOutput& Desc) { return Desc.NodeID == InNodeID; });
+				};
+			}
+
+
+			/*
+			struct FGetNodeInputFromNodeWithName
+			{
+				FGetNodeInputFromNodeWithName(const FString& InName)
+				: Name(InName)
+				{
+				}
+
+				FMetasoundFrontendVertex* operator () (FMetasoundFrontendNode& InNode)
+				{
+					auto IsVertexWithName = [&](const FMetasoundFrontendVertex& Vertex) { return Vertex.Name == Name; };
+					return InNode.Interface.Inputs.FindByPredicate(IsVertexWithName);
+				}
+
+				const FMetasoundFrontendVertex* operator() (const FMetasoundFrontendNode& InNode)
+				{
+					auto IsVertexWithName = [&](const FMetasoundFrontendVertex& Vertex) { return Vertex.Name == Name; };
+					return InNode.Interface.Inputs.FindByPredicate(IsVertexWithName);
+				}
+			};
+			*/
+
+			int32 GetNewID(const TSet<int32> InExistingIDs)
+			{
+				// Assumption here is that we will never need more than ten thousand nodes,
+				// and four digits are easy enough to read/remember when looking at metasound graph documents.
+				int32 IDMax = 9999;
+
+				int32 NewID = Metasound::FrontendInvalidID;
+
+				while (InExistingIDs.Num() > (IDMax / 10))
+				{
+					IDMax *= 10;
+					IDMax += 9;
+
+					if (!ensure((IDMax > 0)))
+					{
+						// IDMax has rolled over to negative values. now is a good
+						// time to consider a new ID type.
+						return Metasound::FrontendInvalidID;
+					}
+				}
+
+				do 
+				{
+					NewID = FMath::RandRange(1, IDMax);
+				}
+				while (InExistingIDs.Contains(NewID));
+
+				return NewID;
+			}
+		}
+
+		//
+		// FBaseOutputController
+		//
+		FBaseOutputController::FBaseOutputController(const FBaseOutputController::FInitParams& InParams)
+		: ID(InParams.ID)
+		, NodeVertexPtr(InParams.NodeVertexPtr)
+		, OwningNode(InParams.OwningNode)
+		{
+		}
+
+		bool FBaseOutputController::IsValid() const
+		{
+			return OwningNode->IsValid() && NodeVertexPtr.IsValid();
+		}
+
+		int32 FBaseOutputController::GetID() const
+		{
+			return ID;
+		}
+
+		const FName& FBaseOutputController::GetDataType() const
+		{
+			if (NodeVertexPtr.IsValid())
+			{
+				return NodeVertexPtr->TypeName;
+			}
+			return FrontendControllerIntrinsics::GetInvalidValueConstRef<FName>();
+		}
+
+		const FString& FBaseOutputController::GetName() const
+		{
+			if (NodeVertexPtr.IsValid())
+			{
+				return NodeVertexPtr->Name;
+			}
+			return FrontendControllerIntrinsics::GetInvalidValueConstRef<FString>();
+		}
+
+		int32 FBaseOutputController::GetOwningNodeID() const
+		{
+			return OwningNode->GetID();
+		}
+
+		FNodeHandle FBaseOutputController::GetOwningNode()
+		{
+			return OwningNode;
+		}
+
+		FConstNodeHandle FBaseOutputController::GetOwningNode() const
+		{
+			return OwningNode;
+		}
+
+		FConnectability FBaseOutputController::CanConnectTo(const IInputController& InController) const
+		{
+			return InController.CanConnectTo(*this);
+		}
+
+		bool FBaseOutputController::Connect(IInputController& InController)
+		{
+			return InController.Connect(*this);
+		}
+
+		bool FBaseOutputController::ConnectWithConverterNode(IInputController& InController, const FConverterNodeInfo& InNodeClassName)
+		{
+			return InController.ConnectWithConverterNode(*this, InNodeClassName);
+		}
+
+		bool FBaseOutputController::Disconnect(IInputController& InController)
+		{
+			return InController.Disconnect(*this);
+		}
+
+		//
+		// FNodeOutputController
+		//
+		FNodeOutputController::FNodeOutputController(const FNodeOutputController::FInitParams& InParams)
+		: FBaseOutputController({InParams.ID, InParams.NodeVertexPtr, InParams.OwningNode})
+		, ClassOutputPtr(InParams.ClassOutputPtr)
+		{
+		}
+
+		bool FNodeOutputController::IsValid() const
+		{
+			return FBaseOutputController::IsValid() && ClassOutputPtr.IsValid();
+		}
+
+		const FText& FNodeOutputController::GetDisplayName() const
+		{
+			if (ClassOutputPtr.IsValid())
+			{
+				return ClassOutputPtr->Metadata.DisplayName;
+			}
+
+			return FrontendControllerIntrinsics::GetInvalidValueConstRef<FText>();
+		}
+
+		const FText& FNodeOutputController::GetTooltip() const
+		{
+			if (ClassOutputPtr.IsValid())
+			{
+				return ClassOutputPtr->Metadata.Description;
+			}
+
+			return FrontendControllerIntrinsics::GetInvalidValueConstRef<FText>();
+		}
+
+		//
+		// FInputNodeOutputController
+		// 
+		FInputNodeOutputController::FInputNodeOutputController(const FInputNodeOutputController::FInitParams& InParams)
+		: FBaseOutputController({InParams.ID, InParams.NodeVertexPtr, InParams.OwningNode})
+		, OwningGraphClassInputPtr(InParams.OwningGraphClassInputPtr)
+		{
+		}
+
+		bool FInputNodeOutputController::IsValid() const
+		{
+			return FBaseOutputController::IsValid() && OwningGraphClassInputPtr.IsValid();
+		}
+
+		const FText& FInputNodeOutputController::GetDisplayName() const
+		{
+			if (OwningGraphClassInputPtr.IsValid())
+			{
+				return OwningGraphClassInputPtr->Metadata.DisplayName;
+			}
+
+			return FrontendControllerIntrinsics::GetInvalidValueConstRef<FText>();
+		}
+
+		const FText& FInputNodeOutputController::GetTooltip() const
+		{
+			if (OwningGraphClassInputPtr.IsValid())
+			{
+				return OwningGraphClassInputPtr->Metadata.Description;
+			}
+
+			return FrontendControllerIntrinsics::GetInvalidValueConstRef<FText>();
+		}
+
+		//
+		// FBaseInputController
+		// 
+		FBaseInputController::FBaseInputController(const FBaseInputController::FInitParams& InParams)
+		: ID(InParams.ID)
+		, NodeVertexPtr(InParams.NodeVertexPtr)
+		, GraphPtr(InParams.GraphPtr)
+		, OwningNode(InParams.OwningNode)
+		{
+		}
+
+		bool FBaseInputController::IsValid() const 
+		{
+			return OwningNode->IsValid() && NodeVertexPtr.IsValid() &&  GraphPtr.IsValid();
+		}
+
+		int32 FBaseInputController::GetID() const
+		{
+			return ID;
+		}
+
+		const FName& FBaseInputController::GetDataType() const
+		{
+			if (NodeVertexPtr.IsValid())
+			{
+				return NodeVertexPtr->TypeName;
+			}
+			
+			return FrontendControllerIntrinsics::GetInvalidValueConstRef<FName>();
+		}
+
+		const FString& FBaseInputController::GetName() const
+		{
+			if (NodeVertexPtr.IsValid())
+			{
+				return NodeVertexPtr->Name;
+			}
+			
+			return FrontendControllerIntrinsics::GetInvalidValueConstRef<FString>();
+		}
+
+		bool FBaseInputController::IsConnected() const 
+		{
+			return (nullptr != FindEdge());
+		}
+
+		int32 FBaseInputController::GetOwningNodeID() const
+		{
+			return OwningNode->GetID();
+		}
+
+		FNodeHandle FBaseInputController::GetOwningNode()
+		{
+			return OwningNode;
+		}
+
+		FConstNodeHandle FBaseInputController::GetOwningNode() const
+		{
+			return OwningNode;
+		}
+
+		FOutputHandle FBaseInputController::GetCurrentlyConnectedOutput()
+		{
+			if (const FMetasoundFrontendEdge* Edge = FindEdge())
+			{
+				// Create output handle from output node.
+				FGraphHandle Graph = OwningNode->GetOwningGraph();
+				FNodeHandle OutputNode = Graph->GetNodeWithID(Edge->FromNodeID);
+				return OutputNode->GetOutputWithID(Edge->FromPointID);
+			}
+
+			return FInvalidOutputController::GetInvalid();
+		}
+
+		FConstOutputHandle FBaseInputController::GetCurrentlyConnectedOutput() const
+		{
+			if (const FMetasoundFrontendEdge* Edge = FindEdge())
+			{
+				// Create output handle from output node.
+				FConstGraphHandle Graph = OwningNode->GetOwningGraph();
+				FConstNodeHandle OutputNode = Graph->GetNodeWithID(Edge->FromNodeID);
+				return OutputNode->GetOutputWithID(Edge->FromPointID);
+			}
+
+			return FInvalidOutputController::GetInvalid();
+		}
+
+		FConnectability FBaseInputController::CanConnectTo(const IOutputController& InController) const
+		{
+			FConnectability OutConnectability;
+			OutConnectability.Connectable = FConnectability::EConnectable::No;
+
+			if (!(InController.IsValid() && IsValid()))
+			{
+				return OutConnectability;
+			}
+
+			if (InController.GetDataType() == GetDataType())
+			{
+				// If data types are equal, connection can happen.
+				OutConnectability.Connectable = FConnectability::EConnectable::Yes;
+				return OutConnectability;
+			}
+
+			// If data types are not equal, check for converter nodes which could
+			// convert data type.
+			OutConnectability.PossibleConverterNodeClasses = FMetasoundFrontendRegistryContainer::Get()->GetPossibleConverterNodes(InController.GetDataType(), GetDataType());
+
+			if (OutConnectability.PossibleConverterNodeClasses.Num() > 0)
+			{
+				OutConnectability.Connectable = FConnectability::EConnectable::YesWithConverterNode;
+				return OutConnectability;
+			}
+
+			return OutConnectability;
+		}
+
+		bool FBaseInputController::Connect(IOutputController& InController)
+		{
+			if (!IsValid() || !InController.IsValid())
+			{
+				return false;
+			}
+
+			if (ensureAlwaysMsgf(InController.GetDataType() == GetDataType(), TEXT("Cannot connect incompatible types.")))
+			{
+				// Overwrite an existing connection if it exists.
+				FMetasoundFrontendEdge* Edge = FindEdge();
+
+				if (!Edge)
+				{
+					Edge = &GraphPtr->Edges.AddDefaulted_GetRef();
+					Edge->ToNodeID = GetOwningNodeID();
+					Edge->ToPointID = GetID();
+				}
+
+				Edge->FromNodeID = InController.GetOwningNodeID();
+				Edge->FromPointID = InController.GetID();
+
+				return true;
+			}
+
+			return false;
+		}
+
+		bool FBaseInputController::ConnectWithConverterNode(IOutputController& InController, const FConverterNodeInfo& InConverterInfo)
+		{
+			FGraphHandle Graph = OwningNode->GetOwningGraph();
+
+			// Generate the converter node.
+			FNodeHandle ConverterNode = Graph->AddNode(InConverterInfo.NodeKey);
+
+			// TODO: may want to rename Point -> Pin. 
+			// TODO: fix naming confusion between "Pin" and "VertexKey". ConverterInfo should use vertexkey. 
+			TArray<FInputHandle> ConverterInputs = ConverterNode->GetInputsWithVertexName(InConverterInfo.PreferredConverterInputPin);
+			TArray<FOutputHandle> ConverterOutputs = ConverterNode->GetOutputsWithVertexName(InConverterInfo.PreferredConverterOutputPin);
+
+			if (ConverterInputs.Num() < 1)
+			{
+				UE_LOG(LogMetasound, Warning, TEXT("Converter node [Name: %s] does not support preferred input vertex [Vertex: %s]"), *InConverterInfo.NodeKey.NodeName.ToString(), *InConverterInfo.PreferredConverterInputPin);
+				return false;
+			}
+
+			if (ConverterOutputs.Num() < 1)
+			{
+				UE_LOG(LogMetasound, Warning, TEXT("Converter node [Name: %s] does not support preferred output vertex [Vertex: %s]"), *InConverterInfo.NodeKey.NodeName.ToString(), *InConverterInfo.PreferredConverterOutputPin);
+				return false;
+			}
+
+			FInputHandle ConverterInput = ConverterInputs[0];
+			FOutputHandle ConverterOutput = ConverterOutputs[0];
+
+			// Connect the output InController to the converter, than connect the converter to this input.
+			if (ConverterInput->Connect(InController) && Connect(*ConverterOutput))
+			{
+				return true;
+			}
+
+			return false;
+		}
+
+		bool FBaseInputController::Disconnect(IOutputController& InController) 
+		{
+			if (GraphPtr.IsValid())
+			{
+				int32 FromNodeID = InController.GetOwningNodeID();
+				int32 FromPointID = InController.GetID();
+				int32 ToNodeID = GetOwningNodeID();
+				int32 ToPointID = GetID();
+
+				auto IsMatchingEdge = [&](const FMetasoundFrontendEdge& Edge)
+				{
+					return (Edge.FromNodeID == FromNodeID) && (Edge.FromPointID == FromPointID) && (Edge.ToNodeID == ToNodeID) && (Edge.ToPointID == ToPointID);
+				};
+
+				int32 NumRemoved = GraphPtr->Edges.RemoveAllSwap(IsMatchingEdge);
+				return NumRemoved > 0;
+			}
+
+			return false;
+		}
+
+		bool FBaseInputController::Disconnect()
+		{
+			if (GraphPtr.IsValid())
+			{
+				const int32 NodeID = GetOwningNodeID();
+				const int32 PointID = GetID();
+
+				auto EdgeHasMatchingDestination = [&](const FMetasoundFrontendEdge& Edge)
+				{
+					return (Edge.ToNodeID == NodeID) && (Edge.ToPointID == PointID);
+				};
+
+				int32 NumRemoved = GraphPtr->Edges.RemoveAllSwap(EdgeHasMatchingDestination);
+				return NumRemoved > 0;
+			}
+
+			return false;
+		}
+
+
+		const FMetasoundFrontendEdge* FBaseInputController::FindEdge() const
+		{
+			if (GraphPtr.IsValid())
+			{
+				const int32 NodeID = GetOwningNodeID();
+				const int32 PointID = GetID();
+
+				auto EdgeHasMatchingDestination = [&](const FMetasoundFrontendEdge& Edge)
+				{
+					return (Edge.ToNodeID == NodeID) && (Edge.ToPointID == PointID);
+				};
+
+				return GraphPtr->Edges.FindByPredicate(EdgeHasMatchingDestination);
+			}
+
+			return nullptr;
+		}
+
+		FMetasoundFrontendEdge* FBaseInputController::FindEdge()
+		{
+			if (GraphPtr.IsValid())
+			{
+				const int32 NodeID = GetOwningNodeID();
+				const int32 PointID = GetID();
+
+				auto EdgeHasMatchingDestination = [&](const FMetasoundFrontendEdge& Edge)
+				{
+					return (Edge.ToNodeID == NodeID) && (Edge.ToPointID == PointID);
+				};
+
+				return GraphPtr->Edges.FindByPredicate(EdgeHasMatchingDestination);
+			}
+
+			return nullptr;
+		}
+
+		//
+		// FNodeInputController
+		//
+		FNodeInputController::FNodeInputController(const FNodeInputController::FInitParams& InParams)
+		: FBaseInputController({InParams.ID, InParams.NodeVertexPtr, InParams.GraphPtr, InParams.OwningNode})
+		, OwningGraphClassInputPtr(InParams.OwningGraphClassInputPtr)
+		{
+		}
+
+		bool FNodeInputController::IsValid() const
+		{
+			return FBaseInputController::IsValid() && OwningGraphClassInputPtr.IsValid();
+		}
+
+		const FText& FNodeInputController::GetDisplayName() const
+		{
+			if (OwningGraphClassInputPtr.IsValid())
+			{
+				return OwningGraphClassInputPtr->Metadata.DisplayName;
+			}
+			
+			return FrontendControllerIntrinsics::GetInvalidValueConstRef<FText>();
+		}
+
+		const FText& FNodeInputController::GetTooltip() const
+		{
+			if (OwningGraphClassInputPtr.IsValid())
+			{
+				return OwningGraphClassInputPtr->Metadata.Description;
+			}
+			
+			return FrontendControllerIntrinsics::GetInvalidValueConstRef<FText>();
+		}
+
+		//
+		// FOutputNodeInputController
+		//
+		FOutputNodeInputController::FOutputNodeInputController(const FOutputNodeInputController::FInitParams& InParams)
+		: FBaseInputController({InParams.ID, InParams.NodeVertexPtr, InParams.GraphPtr, InParams.OwningNode})
+		, OwningGraphClassOutputPtr(InParams.OwningGraphClassOutputPtr)
+		{
+		}
+
+		bool FOutputNodeInputController::IsValid() const
+		{
+			return FBaseInputController::IsValid() && OwningGraphClassOutputPtr.IsValid();
+		}
+
+		const FText& FOutputNodeInputController::GetDisplayName() const
+		{
+			if (OwningGraphClassOutputPtr.IsValid())
+			{
+				return OwningGraphClassOutputPtr->Metadata.DisplayName;
+			}
+			
+			return FrontendControllerIntrinsics::GetInvalidValueConstRef<FText>();
+		}
+
+		const FText& FOutputNodeInputController::GetTooltip() const
+		{
+			if (OwningGraphClassOutputPtr.IsValid())
+			{
+				return OwningGraphClassOutputPtr->Metadata.Description;
+			}
+			
+			return FrontendControllerIntrinsics::GetInvalidValueConstRef<FText>();
+		}
+
+
+		//
+		// FBaseNodeController
+		//
+		FBaseNodeController::FBaseNodeController(const FBaseNodeController::FInitParams& InParams)
+		: NodePtr(InParams.NodePtr)
+		, ClassPtr(InParams.ClassPtr)
+		, OwningGraph(InParams.OwningGraph)
+		{
+			if (NodePtr.IsValid() && ClassPtr.IsValid())
+			{
+				if (NodePtr->ClassID != ClassPtr->ID)
+				{
+					UE_LOG(LogMetasound, Warning, TEXT("Changing node's class id from [ClassID:%d] to [ClassID:%d]"), NodePtr->ClassID, ClassPtr->ID);
+					NodePtr->ClassID = ClassPtr->ID;
+				}
+			}
+		}
+
+		bool FBaseNodeController::IsValid() const
+		{
+			return OwningGraph->IsValid() && NodePtr.IsValid() && ClassPtr.IsValid();
+		}
+
+		int32 FBaseNodeController::GetOwningGraphClassID() const
+		{
+			return OwningGraph->GetClassID();
+		}
+
+		FGraphHandle FBaseNodeController::GetOwningGraph()
+		{
+			return OwningGraph;
+		}
+
+		FConstGraphHandle FBaseNodeController::GetOwningGraph() const
+		{
+			return OwningGraph;
+		}
+
+		int32 FBaseNodeController::GetID() const
+		{
+			if (NodePtr.IsValid())
+			{
+				return NodePtr->ID;
+			}
+			return Metasound::FrontendInvalidID;
+		}
+
+		int32 FBaseNodeController::GetClassID() const
+		{
+			if (ClassPtr.IsValid())
+			{
+				return ClassPtr->ID;
+			}
+			return Metasound::FrontendInvalidID;
+		}
+
+		FMetasoundFrontendVersionNumber FBaseNodeController::GetClassVersionNumber() const
+		{
+			if (ClassPtr.IsValid())
+			{
+				return ClassPtr->Metadata.Version;
+			}
+			return FMetasoundFrontendVersionNumber();
+		}
+
+		const FText& FBaseNodeController::GetClassDescription() const
+		{
+			if (ClassPtr.IsValid())
+			{
+				return ClassPtr->Metadata.Description;
+			}
+			return FrontendControllerIntrinsics::GetInvalidValueConstRef<FText>();
+		}
+
+		const FString& FBaseNodeController::GetNodeName() const
+		{
+			if (NodePtr.IsValid())
+			{
+				return NodePtr->Name;
+			}
+			return FrontendControllerIntrinsics::GetInvalidValueConstRef<FString>();
+		}
+
+		EMetasoundFrontendClassType FBaseNodeController::GetClassType() const
+		{
+			if (ClassPtr.IsValid())
+			{
+				return ClassPtr->Metadata.Type;
+			}
+			return EMetasoundFrontendClassType::Invalid;
+		}
+
+		const FString& FBaseNodeController::GetClassName() const
+		{
+			if (ClassPtr.IsValid())
+			{
+				return ClassPtr->Metadata.Name.Name;
+			}
+
+			return FrontendControllerIntrinsics::GetInvalidValueConstRef<FString>();
+		}
+
+		bool FBaseNodeController::CanAddInput(const FString& InVertexName) const
+		{
+			// TODO: not yet supported
+			return false;
+		}
+
+		FInputHandle FBaseNodeController::AddInput(const FString& InVertexName, const FMetasoundFrontendLiteral* InDefault)
+		{
+			checkNoEntry();
+			// TODO: not yet supported
+			return FInvalidInputController::GetInvalid();
+		}
+
+		bool FBaseNodeController::RemoveInput(int32 InPointID)
+		{
+			checkNoEntry();
+			// TODO: not yet supported
+			return false;
+		}
+
+		bool FBaseNodeController::CanAddOutput(const FString& InVertexName) const
+		{
+			// TODO: not yet supported
+			return false;
+		}
+
+		FInputHandle FBaseNodeController::AddOutput(const FString& InVertexName, const FMetasoundFrontendLiteral* InDefault)
+		{
+			checkNoEntry();
+			// TODO: not yet supported
+			return FInvalidInputController::GetInvalid();
+		}
+
+		bool FBaseNodeController::RemoveOutput(int32 InPointID)
+		{
+			checkNoEntry();
+			// TODO: not yet supported
+			return false;
+		}
+
+		TArray<FInputHandle> FBaseNodeController::GetInputs()
+		{
+			TArray<FInputHandle> Inputs;
+
+			FNodeHandle ThisNode = this->AsShared();
+
+			for (const FInputControllerParams& Params : GetInputControllerParams())
+			{
+				FInputHandle InputHandle = CreateInputController(Params.PointID, Params.NodeVertexPtr, Params.ClassInputPtr, ThisNode);
+				if (InputHandle->IsValid())
+				{
+					Inputs.Add(InputHandle);
+				}
+			}
+
+			return Inputs;
+		}
+
+		TArray<FOutputHandle> FBaseNodeController::GetOutputs()
+		{
+			TArray<FOutputHandle> Outputs;
+
+			FNodeHandle ThisNode = this->AsShared();
+
+			for (const FOutputControllerParams& Params : GetOutputControllerParams())
+			{
+				FOutputHandle OutputHandle = CreateOutputController(Params.PointID, Params.NodeVertexPtr, Params.ClassOutputPtr, ThisNode);
+				if (OutputHandle->IsValid())
+				{
+					Outputs.Add(OutputHandle);
+				}
+			}
+
+			return Outputs;
+		}
+
+		TArray<FConstInputHandle> FBaseNodeController::GetConstInputs() const
+		{
+			TArray<FConstInputHandle> Inputs;
+
+			// If I had a nickle for every time C++ backed me into a corner, I would be sitting
+			// on a tropical beach next to my mansion sipping strawberry daiquiris instead of 
+			// trying to code using this guileful language. The const cast is generally safe here
+			// because the FConstInputHandle only allows const access to the internal node controller. 
+			// Ostensibly, there could have been a INodeController and IConstNodeController
+			// which take different types in their constructor, but it starts to become
+			// difficult to maintain. So instead of adding 500 lines of nearly duplicate 
+			// code, a ConstCastSharedRef is used here. 
+			FNodeHandle ThisNode = ConstCastSharedRef<INodeController>(this->AsShared());
+
+			for (const FInputControllerParams& Params : GetInputControllerParams())
+			{
+				FConstInputHandle InputHandle = CreateInputController(Params.PointID, Params.NodeVertexPtr, Params.ClassInputPtr, ThisNode);
+				if (InputHandle->IsValid())
+				{
+					Inputs.Add(InputHandle);
+				}
+			}
+
+			return Inputs;
+		}
+
+		TArray<FConstOutputHandle> FBaseNodeController::GetConstOutputs() const
+		{
+			TArray<FConstOutputHandle> Outputs;
+
+			// See early use of ConstCastSharedRef in this class for discussion.
+			FNodeHandle ThisNode = ConstCastSharedRef<INodeController>(this->AsShared());
+
+			for (const FOutputControllerParams& Params : GetOutputControllerParams())
+			{
+				FConstOutputHandle OutputHandle = CreateOutputController(Params.PointID, Params.NodeVertexPtr, Params.ClassOutputPtr, ThisNode);
+				if (OutputHandle->IsValid())
+				{
+					Outputs.Add(OutputHandle);
+				}
+			}
+
+			return Outputs;
+		}
+
+		TArray<FInputHandle> FBaseNodeController::GetInputsWithVertexName(const FString& InName) 
+		{
+			TArray<FInputHandle> Inputs;
+
+			FNodeHandle ThisNode = this->AsShared();
+
+			for (const FInputControllerParams& Params : GetInputControllerParamsWithVertexName(InName))
+			{
+				FInputHandle InputHandle = CreateInputController(Params.PointID, Params.NodeVertexPtr, Params.ClassInputPtr, ThisNode);
+				if (InputHandle->IsValid())
+				{
+					Inputs.Add(InputHandle);
+				}
+			}
+
+			return Inputs;
+		}
+
+		TArray<FConstInputHandle> FBaseNodeController::GetInputsWithVertexName(const FString& InName) const 
+		{
+			TArray<FConstInputHandle> Inputs;
+
+			// See early use of ConstCastSharedRef in this class for discussion.
+			FNodeHandle ThisNode = ConstCastSharedRef<INodeController>(this->AsShared());
+
+			for (const FInputControllerParams& Params : GetInputControllerParamsWithVertexName(InName))
+			{
+				FInputHandle InputHandle = CreateInputController(Params.PointID, Params.NodeVertexPtr, Params.ClassInputPtr, ThisNode);
+				if (InputHandle->IsValid())
+				{
+					Inputs.Add(InputHandle);
+				}
+			}
+
+			return Inputs;
+		}
+
+		TArray<FOutputHandle> FBaseNodeController::GetOutputsWithVertexName(const FString& InName)
+		{
+			TArray<FOutputHandle> Outputs;
+
+			FNodeHandle ThisNode = this->AsShared();
+
+			for (const FOutputControllerParams& Params : GetOutputControllerParamsWithVertexName(InName))
+			{
+				FOutputHandle OutputHandle = CreateOutputController(Params.PointID, Params.NodeVertexPtr, Params.ClassOutputPtr, ThisNode);
+				if (OutputHandle->IsValid())
+				{
+					Outputs.Add(OutputHandle);
+				}
+			}
+
+			return Outputs;
+		}
+
+		TArray<FConstOutputHandle> FBaseNodeController::GetOutputsWithVertexName(const FString& InName) const
+		{
+			TArray<FConstOutputHandle> Outputs;
+
+			// See early use of ConstCastSharedRef in this class for discussion.
+			FNodeHandle ThisNode = ConstCastSharedRef<INodeController>(this->AsShared());
+
+			for (const FOutputControllerParams& Params : GetOutputControllerParamsWithVertexName(InName))
+			{
+				FConstOutputHandle OutputHandle = CreateOutputController(Params.PointID, Params.NodeVertexPtr, Params.ClassOutputPtr, ThisNode);
+				if (OutputHandle->IsValid())
+				{
+					Outputs.Add(OutputHandle);
+				}
+			}
+
+			return Outputs;
+		}
+
+		FInputHandle FBaseNodeController::GetInputWithID(int32 InPointID)
+		{
+			FInputControllerParams Params;
+
+			if (FindInputControllerParamsWithID(InPointID, Params))
+			{
+				FNodeHandle ThisNode = this->AsShared();
+				return CreateInputController(Params.PointID, Params.NodeVertexPtr, Params.ClassInputPtr, ThisNode);
+			}
+
+			return FInvalidInputController::GetInvalid();
+		}
+
+		FConstInputHandle FBaseNodeController::GetInputWithID(int32 InPointID) const
+		{
+			FInputControllerParams Params;
+
+			if (FindInputControllerParamsWithID(InPointID, Params))
+			{
+				// See early use of ConstCastSharedRef in this class for discussion.
+				FNodeHandle ThisNode = ConstCastSharedRef<INodeController>(this->AsShared());
+				return CreateInputController(Params.PointID, Params.NodeVertexPtr, Params.ClassInputPtr, ThisNode);
+			}
+
+			return FInvalidInputController::GetInvalid();
+		}
+
+		FOutputHandle FBaseNodeController::GetOutputWithID(int32 InPointID)
+		{
+			FOutputControllerParams Params;
+
+			if (FindOutputControllerParamsWithID(InPointID, Params))
+			{
+				FNodeHandle ThisNode = this->AsShared();
+				return CreateOutputController(Params.PointID, Params.NodeVertexPtr, Params.ClassOutputPtr, ThisNode);
+			}
+
+			return FInvalidOutputController::GetInvalid();
+		}
+
+		FConstOutputHandle FBaseNodeController::GetOutputWithID(int32 InPointID) const
+		{
+			FOutputControllerParams Params;
+
+			if (FindOutputControllerParamsWithID(InPointID, Params))
+			{
+				// See early use of ConstCastSharedRef in this class for discussion.
+				FNodeHandle ThisNode = ConstCastSharedRef<INodeController>(this->AsShared());
+				return CreateOutputController(Params.PointID, Params.NodeVertexPtr, Params.ClassOutputPtr, ThisNode);
+			}
+
+			return FInvalidOutputController::GetInvalid();
+		}
+
+
+		TArray<FBaseNodeController::FInputControllerParams> FBaseNodeController::GetInputControllerParams() const
+		{
+			using namespace FrontendControllerIntrinsics;
+
+			TArray<FBaseNodeController::FInputControllerParams> Inputs;
+
+			if (NodePtr.IsValid())
+			{
+				for (const FMetasoundFrontendVertex& NodeInputVertex : NodePtr->Interface.Inputs)
+				{
+
+					FConstVertexAccessPtr NodeVertexPtr = NodePtr.GetMemberAccessPtr<const FMetasoundFrontendVertex>(FGetNodeInputFromNodeWithName(NodeInputVertex.Name));
+
+					FConstClassInputAccessPtr ClassInputPtr = ClassPtr.GetMemberAccessPtr<const FMetasoundFrontendClassInput>(FGetClassInputFromClassWithName(NodeInputVertex.Name));
+
+					for (int32 PointID : NodeInputVertex.PointIDs)
+					{
+						Inputs.Add({PointID, NodeVertexPtr, ClassInputPtr});
+					}
+				}
+			}
+			/*
+
+
+
+
+
+
+			if (ClassPtr.IsValid() && NodePtr.IsValid())
+			{
+				for (const FMetasoundFrontendVertex& NodeInputVertex : NodePtr->Interface.Inputs)
+				{
+
+					FConstVertexAccessPtr NodeVertexPtr = FrontendControllerIntrinsics::MakeAutoAccessPtr(NodeInputVertex);
+
+
+					if (const FMetasoundFrontendClassInput* ClassInputVertex = FindClassInputWithName(NodeInputVertex.Name))
+					{
+
+
+FGetClassInputFromClassWithName
+
+						FConstClassInputAccessPtr ClassInputPtr = FrontendControllerIntrinsics::MakeAutoAccessPtr(*ClassInputVertex);
+
+						for (int32 PointID : NodeInputVertex.PointIDs)
+						{
+							Inputs.Add({PointID, NodeVertexPtr, ClassInputPtr});
+						}
+					}
+					else
+					{
+						UE_LOG(LogMetasound, Error, TEXT("Node vertex [NodeID:%d, InputName:%s] does not exist in class [ClassID:%d]"), NodePtr->ID, *NodeInputVertex.Name, ClassPtr->ID);
+					}
+				}
+			}
+			*/
+
+			return Inputs;
+		}
+
+		TArray<FBaseNodeController::FOutputControllerParams> FBaseNodeController::GetOutputControllerParams() const
+		{
+			using namespace FrontendControllerIntrinsics;
+
+			TArray<FBaseNodeController::FOutputControllerParams> Outputs;
+
+			for (const FMetasoundFrontendVertex& NodeOutputVertex : NodePtr->Interface.Outputs)
+			{
+				const FString& VertexName = NodeOutputVertex.Name;
+
+				FConstVertexAccessPtr NodeVertexPtr = NodePtr.GetMemberAccessPtr<const FMetasoundFrontendVertex>(FGetConstNodeOutputFromNodeWithName(VertexName));
+				FConstClassOutputAccessPtr ClassOutputPtr = ClassPtr.GetMemberAccessPtr<const FMetasoundFrontendClassOutput>(FGetClassOutputFromClassWithName(VertexName));
+
+				for (int32 PointID : NodeOutputVertex.PointIDs)
+				{
+					Outputs.Add({PointID, NodeVertexPtr, ClassOutputPtr});
+				}
+			}
+
+
+			/*
+			if (ClassPtr.IsValid() && NodePtr.IsValid())
+			{
+				for (const FMetasoundFrontendVertex& NodeOutputVertex : NodePtr->Interface.Outputs)
+				{
+					FConstVertexAccessPtr NodeVertexPtr = FrontendControllerIntrinsics::MakeAutoAccessPtr(NodeOutputVertex);
+
+					if (const FMetasoundFrontendClassOutput* ClassOutputVertex = FindClassOutputWithName(NodeOutputVertex.Name))
+					{
+						FConstClassOutputAccessPtr ClassOutputPtr = FrontendControllerIntrinsics::MakeAutoAccessPtr(*ClassOutputVertex);
+
+						for (int32 PointID : NodeOutputVertex.PointIDs)
+						{
+							Outputs.Add({PointID, NodeVertexPtr, ClassOutputPtr});
+						}
+					}
+					else
+					{
+						UE_LOG(LogMetasound, Error, TEXT("Node vertex [NodeID:%d, OutputName:%s] does not exist in class [ClassID:%d]"), NodePtr->ID, *NodeOutputVertex.Name, ClassPtr->ID);
+					}
+				}
+			}
+			*/
+
+			return Outputs;
+		}
+
+		TArray<FBaseNodeController::FInputControllerParams> FBaseNodeController::GetInputControllerParamsWithVertexName(const FString& InName) const
+		{
+			using namespace FrontendControllerIntrinsics;
+
+			TArray<FBaseNodeController::FInputControllerParams> Inputs;
+
+			FConstVertexAccessPtr NodeVertexPtr = NodePtr.GetMemberAccessPtr<const FMetasoundFrontendVertex>(FGetConstNodeInputFromNodeWithName(InName));
+
+			if (const FMetasoundFrontendVertex* Vertex = NodeVertexPtr.Get())
+			{
+				FConstClassInputAccessPtr ClassInputPtr = ClassPtr.GetMemberAccessPtr<const FMetasoundFrontendClassInput>(FGetClassInputFromClassWithName(InName));
+
+				for (int32 PointID : Vertex->PointIDs)
+				{
+					Inputs.Add({PointID, NodeVertexPtr, ClassInputPtr});
+				}
+			}
+
+
+			/*
+			if (ClassPtr.IsValid() && NodePtr.IsValid())
+			{
+				if (const FMetasoundFrontendVertex* NodeInputVertex = FindNodeInputWithName(InName))
+				{
+					FConstVertexAccessPtr NodeVertexPtr = FrontendControllerIntrinsics::MakeAutoAccessPtr(*NodeInputVertex);
+
+					if (const FMetasoundFrontendClassInput* ClassInputVertex = FindClassInputWithName(InName))
+					{
+						FConstClassInputAccessPtr ClassInputPtr = FrontendControllerIntrinsics::MakeAutoAccessPtr(*ClassInputVertex);
+
+						for (int32 PointID : NodeInputVertex->PointIDs)
+						{
+							Inputs.Add({PointID, NodeVertexPtr, ClassInputPtr});
+						}
+					}
+					else
+					{
+						UE_LOG(LogMetasound, Error, TEXT("Node vertex [NodeID:%d, InputName:%s] does not exist in class [ClassID:%d]"), NodePtr->ID, *NodeInputVertex->Name, ClassPtr->ID);
+					}
+				}
+			}
+			*/
+
+			return Inputs;
+		}
+
+		TArray<FBaseNodeController::FOutputControllerParams> FBaseNodeController::GetOutputControllerParamsWithVertexName(const FString& InName) const
+		{
+			using namespace FrontendControllerIntrinsics;
+
+			TArray<FBaseNodeController::FOutputControllerParams> Outputs;
+
+			FConstVertexAccessPtr NodeVertexPtr = NodePtr.GetMemberAccessPtr<const FMetasoundFrontendVertex>(FGetConstNodeOutputFromNodeWithName(InName));
+
+			if (const FMetasoundFrontendVertex* Vertex = NodeVertexPtr.Get())
+			{
+				FConstClassOutputAccessPtr ClassOutputPtr = ClassPtr.GetMemberAccessPtr<const FMetasoundFrontendClassOutput>(FGetClassOutputFromClassWithName(InName));
+
+				for (int32 PointID : Vertex->PointIDs)
+				{
+					Outputs.Add({PointID, NodeVertexPtr, ClassOutputPtr});
+				}
+			}
+			/*
+			if (ClassPtr.IsValid() && NodePtr.IsValid())
+			{
+				if (const FMetasoundFrontendVertex* NodeOutputVertex = FindNodeOutputWithName(InName))
+				{
+					FConstVertexAccessPtr NodeVertexPtr = FrontendControllerIntrinsics::MakeAutoAccessPtr(*NodeOutputVertex);
+
+					if (const FMetasoundFrontendClassOutput* ClassOutputVertex = FindClassOutputWithName(InName))
+					{
+						FConstClassOutputAccessPtr ClassOutputPtr = FrontendControllerIntrinsics::MakeAutoAccessPtr(*ClassOutputVertex);
+
+						for (int32 PointID : NodeOutputVertex->PointIDs)
+						{
+							Outputs.Add({PointID, NodeVertexPtr, ClassOutputPtr});
+						}
+					}
+					else
+					{
+						UE_LOG(LogMetasound, Error, TEXT("Node vertex [NodeID:%d, OutputName:%s] does not exist in class [ClassID:%d]"), NodePtr->ID, *NodeOutputVertex->Name, ClassPtr->ID);
+					}
+				}
+			}
+			*/
+
+			return Outputs;
+		}
+
+		bool FBaseNodeController::FindInputControllerParamsWithID(int32 InPointID, FInputControllerParams& OutParams) const
+		{
+			using namespace FrontendControllerIntrinsics;
+
+			FConstVertexAccessPtr NodeVertexPtr = NodePtr.GetMemberAccessPtr<const FMetasoundFrontendVertex>(FGetConstNodeInputFromNodeWithPoint(InPointID));
+
+			if (const FMetasoundFrontendVertex* Vertex = NodeVertexPtr.Get())
+			{
+				FConstClassInputAccessPtr ClassInputPtr = ClassPtr.GetMemberAccessPtr<const FMetasoundFrontendClassInput>(FGetClassInputFromClassWithName(Vertex->Name));
+
+				OutParams = FInputControllerParams{InPointID, NodeVertexPtr, ClassInputPtr};
+				return true;
+			}
+			/*
+			if (ClassPtr.IsValid() && NodePtr.IsValid())
+			{
+				for (const FMetasoundFrontendVertex& NodeInputVertex : NodePtr->Interface.Inputs)
+				{
+					if (NodeInputVertex.PointIDs.Contains(InPointID))
+					{
+						if (const FMetasoundFrontendClassInput* ClassInputVertex = FindClassInputWithName(NodeInputVertex.Name))
+						{
+							FConstVertexAccessPtr NodeVertexPtr = FrontendControllerIntrinsics::MakeAutoAccessPtr(NodeInputVertex);
+							FConstClassInputAccessPtr ClassInputPtr = FrontendControllerIntrinsics::MakeAutoAccessPtr(*ClassInputVertex);
+
+							OutParams = FInputControllerParams{InPointID, NodeVertexPtr, ClassInputPtr};
+							return true;
+						}
+						else
+						{
+							UE_LOG(LogMetasound, Error, TEXT("Node vertex [NodeID:%d, InputName:%s] does not exist in class [ClassID:%d]"), NodePtr->ID, *NodeInputVertex.Name, ClassPtr->ID);
+						}
+					}
+				}
+			}
+			*/
+
+			return false;
+		}
+
+		bool FBaseNodeController::FindOutputControllerParamsWithID(int32 InPointID, FOutputControllerParams& OutParams) const
+		{
+			using namespace FrontendControllerIntrinsics;
+
+			FConstVertexAccessPtr NodeVertexPtr = NodePtr.GetMemberAccessPtr<const FMetasoundFrontendVertex>(FGetConstNodeOutputFromNodeWithPoint(InPointID));
+
+			if (const FMetasoundFrontendVertex* Vertex = NodeVertexPtr.Get())
+			{
+				FConstClassOutputAccessPtr ClassOutputPtr = ClassPtr.GetMemberAccessPtr<const FMetasoundFrontendClassOutput>(FGetClassOutputFromClassWithName(Vertex->Name));
+
+				OutParams = FOutputControllerParams{InPointID, NodeVertexPtr, ClassOutputPtr};
+				return true;
+			}
+			/*
+			if (ClassPtr.IsValid() && NodePtr.IsValid())
+			{
+				for (const FMetasoundFrontendVertex& NodeOutputVertex : NodePtr->Interface.Outputs)
+				{
+					if (NodeOutputVertex.PointIDs.Contains(InPointID))
+					{
+						if (const FMetasoundFrontendClassOutput* ClassOutputVertex = FindClassOutputWithName(NodeOutputVertex.Name))
+						{
+							FConstVertexAccessPtr NodeVertexPtr = FrontendControllerIntrinsics::MakeAutoAccessPtr(NodeOutputVertex);
+							FConstClassOutputAccessPtr ClassOutputPtr = FrontendControllerIntrinsics::MakeAutoAccessPtr(*ClassOutputVertex);
+
+							OutParams = FOutputControllerParams{InPointID, NodeVertexPtr, ClassOutputPtr};
+							return true;
+						}
+						else
+						{
+							UE_LOG(LogMetasound, Error, TEXT("Node vertex [NodeID:%d, OutputName:%s] does not exist in class [ClassID:%d]"), NodePtr->ID, *NodeOutputVertex.Name, ClassPtr->ID);
+						}
+					}
+				}
+			}
+			*/
+
+			return false;
+		}
+
+		const FMetasoundFrontendClassInput* FBaseNodeController::FindClassInputWithName(const FString& InName) const
+		{
+			if (ClassPtr.IsValid())
+			{
+				auto IsClassInputWithName = [&](const FMetasoundFrontendClassInput& ClassInput) { return ClassInput.Name == InName; };
+
+				return ClassPtr->Interface.Inputs.FindByPredicate(IsClassInputWithName);
+			}
+
+			return nullptr;
+		}
+
+		const FMetasoundFrontendClassOutput* FBaseNodeController::FindClassOutputWithName(const FString& InName) const
+		{
+			if (ClassPtr.IsValid())
+			{
+				auto IsClassOutputWithName = [&](const FMetasoundFrontendClassOutput& ClassOutput) { return ClassOutput.Name == InName; };
+
+				return ClassPtr->Interface.Outputs.FindByPredicate(IsClassOutputWithName);
+			}
+
+			return nullptr;
+		}
+
+		const FMetasoundFrontendVertex* FBaseNodeController::FindNodeInputWithName(const FString& InName) const
+		{
+			if (NodePtr.IsValid())
+			{
+				auto IsNodeInputWithName = [&](const FMetasoundFrontendVertex& NodeInput) { return NodeInput.Name == InName; };
+
+				return NodePtr->Interface.Inputs.FindByPredicate(IsNodeInputWithName);
+			}
+
+			return nullptr;
+		}
+
+		const FMetasoundFrontendVertex* FBaseNodeController::FindNodeOutputWithName(const FString& InName) const
+		{
+			if (NodePtr.IsValid())
+			{
+				auto IsNodeOutputWithName = [&](const FMetasoundFrontendVertex& NodeOutput) { return NodeOutput.Name == InName; };
+
+				return NodePtr->Interface.Outputs.FindByPredicate(IsNodeOutputWithName);
+			}
+
+			return nullptr;
+		}
+
+		FGraphHandle FBaseNodeController::AsGraph()
+		{
+			// TODO: consider adding support for external graph owned in another document.
+			// Will require lookup support for external subgraphs..
+			
+			if (ClassPtr.IsValid())
+			{
+				return GetOwningGraph()->GetOwningDocument()->GetSubgraphWithClassID(ClassPtr->ID);
+			}
+
+			return FInvalidGraphController::GetInvalid();
+		}
+
+		FConstGraphHandle FBaseNodeController::AsGraph() const
+		{
+			// TODO: add support for graph owned in another asset.
+			// Will require lookup support for external subgraphs..
+			if (ClassPtr.IsValid())
+			{
+				return GetOwningGraph()->GetOwningDocument()->GetSubgraphWithClassID(ClassPtr->ID);
+			}
+
+			return FInvalidGraphController::GetInvalid();
+		}
+
+		// 
+		// FNodeController
+		//
+		FNodeController::FNodeController(EPrivateToken InToken, const FNodeController::FInitParams& InParams)
+		: FBaseNodeController({InParams.NodePtr, InParams.ClassPtr, InParams.OwningGraph})
+		, GraphPtr(InParams.GraphPtr)
+		{
+		}
+
+		FNodeHandle FNodeController::CreateNodeHandle(const FNodeController::FInitParams& InParams)
+		{
+			if (InParams.NodePtr.IsValid() && InParams.ClassPtr.IsValid())
+			{
+				// Cannot make a valid node handle if the node description and class description differ
+				if (InParams.NodePtr->ClassID == InParams.ClassPtr->ID)
+				{
+					return MakeShared<FNodeController>(EPrivateToken::Token, InParams);
+				}
+				else
+				{
+					UE_LOG(LogMetasound, Warning, TEXT("Frontend Node [NodeId:%d, ClassID: %d] is not of expected class class [ClassId:%d]"), InParams.NodePtr->ID, InParams.NodePtr->ClassID, InParams.ClassPtr->ID);
+				}
+			}
+			return FInvalidNodeController::GetInvalid();
+		}
+
+		FConstNodeHandle FNodeController::CreateConstNodeHandle(const FNodeController::FInitParams& InParams)
+		{
+			if (InParams.NodePtr.IsValid() && InParams.ClassPtr.IsValid())
+			{
+				// Cannot make a valid node handle if the node description and class description differ
+				if (InParams.NodePtr->ClassID == InParams.ClassPtr->ID)
+				{
+					return MakeShared<const FNodeController>(EPrivateToken::Token, InParams);
+				}
+				else
+				{
+					UE_LOG(LogMetasound, Warning, TEXT("Frontend Node [NodeId:%d, ClassID:%d] is not of expected class class [ClassId:%d]"), InParams.NodePtr->ID, InParams.NodePtr->ClassID, InParams.ClassPtr->ID);
+				}
+			}
+			return FInvalidNodeController::GetInvalid();
+		}
+
+		bool FNodeController::IsValid() const
+		{
+			return FBaseNodeController::IsValid() && GraphPtr.IsValid();
+		}
+
+		FInputHandle FNodeController::CreateInputController(int32 InPointID, FConstVertexAccessPtr InNodeVertexPtr, FConstClassInputAccessPtr InClassInputPtr, FNodeHandle InOwningNode) const
+		{
+			return MakeShared<FNodeInputController>(FNodeInputController::FInitParams{InPointID, InNodeVertexPtr, InClassInputPtr, GraphPtr, InOwningNode});
+		}
+
+		FOutputHandle FNodeController::CreateOutputController(int32 InPointID, FConstVertexAccessPtr InNodeVertexPtr, FConstClassOutputAccessPtr InClassOutputPtr, FNodeHandle InOwningNode) const
+		{
+			return MakeShared<FNodeOutputController>(FNodeOutputController::FInitParams{InPointID, InNodeVertexPtr, InClassOutputPtr, InOwningNode});
+		}
+
+		//
+		// FOutputNodeController
+		//
+		FOutputNodeController::FOutputNodeController(FOutputNodeController::EPrivateToken InToken, const FOutputNodeController::FInitParams& InParams)
+		: FBaseNodeController({InParams.NodePtr, InParams.ClassPtr, InParams.OwningGraph})
+		, GraphPtr(InParams.GraphPtr)
+		, OwningGraphClassOutputPtr(InParams.OwningGraphClassOutputPtr)
+		{
+		}
+
+		FNodeHandle FOutputNodeController::CreateOutputNodeHandle(const FOutputNodeController::FInitParams& InParams)
+		{
+			if (InParams.NodePtr.IsValid() && InParams.ClassPtr.IsValid())
+			{
+				if (EMetasoundFrontendClassType::Output == InParams.ClassPtr->Metadata.Type)
+				{
+					if (InParams.ClassPtr->ID == InParams.NodePtr->ClassID)
+					{
+						return MakeShared<FOutputNodeController>(EPrivateToken::Token, InParams);
+					}
+					else
+					{
+						UE_LOG(LogMetasound, Warning, TEXT("Frontend Node [NodeId:%d, ClassID: %d] is not of expected class class [ClassId:%d]"), InParams.NodePtr->ID, InParams.NodePtr->ClassID, InParams.ClassPtr->ID);
+					}
+				}
+				else
+				{
+					UE_LOG(LogMetasound, Warning, TEXT("Frontend Class of incorrect class type [ClassID:%d] while creating output node.. Must be EMetasoundFrontendClassType::Output."), InParams.ClassPtr->ID);
+				}
+			}
+
+			return FInvalidNodeController::GetInvalid();
+		}
+
+		FConstNodeHandle FOutputNodeController::CreateConstOutputNodeHandle(const FOutputNodeController::FInitParams& InParams)
+		{
+			if (InParams.NodePtr.IsValid() && InParams.ClassPtr.IsValid())
+			{
+				if (EMetasoundFrontendClassType::Output == InParams.ClassPtr->Metadata.Type)
+				{
+					if (InParams.ClassPtr->ID == InParams.NodePtr->ClassID)
+					{
+						return MakeShared<const FOutputNodeController>(EPrivateToken::Token, InParams);
+					}
+					else
+					{
+						UE_LOG(LogMetasound, Warning, TEXT("Frontend Node [NodeId:%d, ClassID: %d] is not of expected class class [ClassId:%d]"), InParams.NodePtr->ID, InParams.NodePtr->ClassID, InParams.ClassPtr->ID);
+					}
+				}
+				else
+				{
+					UE_LOG(LogMetasound, Warning, TEXT("Frontend Class of incorrect class type [ClassID:%d] while creating output node. Must be EMetasoundFrontendClassType::Output."), InParams.ClassPtr->ID);
+				}
+			}
+
+			return FInvalidNodeController::GetInvalid();
+		}
+
+
+		bool FOutputNodeController::IsValid() const
+		{
+			return OwningGraph->IsValid() && OwningGraphClassOutputPtr.IsValid() && GraphPtr.IsValid();
+		}
+
+		FInputHandle FOutputNodeController::CreateInputController(int32 InPointID, FConstVertexAccessPtr InNodeVertexPtr, FConstClassInputAccessPtr InClassInputPtr, FNodeHandle InOwningNode) const
+		{
+			return MakeShared<FOutputNodeInputController>(FOutputNodeInputController::FInitParams{InPointID, InNodeVertexPtr, OwningGraphClassOutputPtr, GraphPtr, InOwningNode});
+		}
+
+		FOutputHandle FOutputNodeController::CreateOutputController(int32 InPointID, FConstVertexAccessPtr InNodeVertexPtr, FConstClassOutputAccessPtr InClassOutputPtr, FNodeHandle InOwningNode) const
+		{
+			return FInvalidOutputController::GetInvalid();
+		}
+
+		//
+		// FInputNodeController
+		//
+		FInputNodeController::FInputNodeController(EPrivateToken InToken, const FInputNodeController::FInitParams& InParams)
+		: FBaseNodeController({InParams.NodePtr, InParams.ClassPtr, InParams.OwningGraph})
+		, OwningGraphClassInputPtr(InParams.OwningGraphClassInputPtr)
+		, GraphPtr(InParams.GraphPtr)
+		{
+		}
+
+		FNodeHandle FInputNodeController::CreateInputNodeHandle(const FInputNodeController::FInitParams& InParams)
+		{
+			if (InParams.NodePtr.IsValid() && InParams.ClassPtr.IsValid())
+			{
+				if (EMetasoundFrontendClassType::Input == InParams.ClassPtr->Metadata.Type)
+				{
+					if (InParams.ClassPtr->ID == InParams.NodePtr->ClassID)
+					{
+						return MakeShared<FInputNodeController>(EPrivateToken::Token, InParams);
+					}
+					else
+					{
+						UE_LOG(LogMetasound, Warning, TEXT("Frontend Node [NodeId:%d, ClassID: %d] is not of expected class class [ClassId:%d]"), InParams.NodePtr->ID, InParams.NodePtr->ClassID, InParams.ClassPtr->ID);
+					}
+				}
+				else
+				{
+					UE_LOG(LogMetasound, Warning, TEXT("Frontend Class of incorrect class type [ClassID:%d] while creating input node. Must be EMetasoundFrontendClassType::Input."), InParams.ClassPtr->ID);
+				}
+			}
+
+			return FInvalidNodeController::GetInvalid();
+		}
+
+		FConstNodeHandle FInputNodeController::CreateConstInputNodeHandle(const FInputNodeController::FInitParams& InParams)
+		{
+			if (InParams.NodePtr.IsValid() && InParams.ClassPtr.IsValid())
+			{
+				if (EMetasoundFrontendClassType::Input == InParams.ClassPtr->Metadata.Type)
+				{
+					if (InParams.ClassPtr->ID == InParams.NodePtr->ClassID)
+					{
+						return MakeShared<const FInputNodeController>(EPrivateToken::Token, InParams);
+					}
+					else
+					{
+						UE_LOG(LogMetasound, Warning, TEXT("Frontend Node [NodeId:%d, ClassID: %d] is not of expected class class [ClassId:%d]"), InParams.NodePtr->ID, InParams.NodePtr->ClassID, InParams.ClassPtr->ID);
+					}
+				}
+				else
+				{
+					UE_LOG(LogMetasound, Warning, TEXT("Frontend Class of incorrect class type [ClassID:%d] while creating input node. Must be EMetasoundFrontendClassType::Input."), InParams.ClassPtr->ID);
+				}
+			}
+
+			return FInvalidNodeController::GetInvalid();
+		}
+
+
+		bool FInputNodeController::IsValid() const
+		{
+			return OwningGraph->IsValid() && OwningGraphClassInputPtr.IsValid() && GraphPtr.IsValid();
+		}
+
+		FInputHandle FInputNodeController::CreateInputController(int32 InPointID, FConstVertexAccessPtr InNodeVertexPtr, FConstClassInputAccessPtr InClassInputPtr, FNodeHandle InOwningNode) const
+		{
+			return FInvalidInputController::GetInvalid();
+		}
+
+		FOutputHandle FInputNodeController::CreateOutputController(int32 InPointID, FConstVertexAccessPtr InNodeVertexPtr, FConstClassOutputAccessPtr InClassOutputPtr, FNodeHandle InOwningNode) const
+		{
+			return MakeShared<FInputNodeOutputController>(FInputNodeOutputController::FInitParams{InPointID, InNodeVertexPtr, OwningGraphClassInputPtr, InOwningNode});
+		}
+
+
+		// 
+		// FGraphController
+		//
+		FGraphController::FGraphController(EPrivateToken InToken, const FGraphController::FInitParams& InParams)
+		: GraphClassPtr(InParams.GraphClassPtr)
+		, OwningDocument(InParams.OwningDocument)
+		{
+		}
+
+		FGraphHandle FGraphController::CreateGraphHandle(const FGraphController::FInitParams& InParams)
+		{
+			if (InParams.GraphClassPtr.IsValid())
+			{
+				if (InParams.GraphClassPtr->Metadata.Type == EMetasoundFrontendClassType::Graph)
+				{
+					return MakeShared<FGraphController>(EPrivateToken::Token, InParams);
+				}
+				else
+				{
+					UE_LOG(LogMetasound, Warning, TEXT("Failed to make graph controller [ClassId:%d]. Class must be EMeatsoundFrontendClassType::Graph."), InParams.GraphClassPtr->ID)
+				}
+			}
+			return FInvalidGraphController::GetInvalid();
+		}
+
+		FConstGraphHandle FGraphController::CreateConstGraphHandle(const FGraphController::FInitParams& InParams)
+		{
+			if (InParams.GraphClassPtr.IsValid())
+			{
+				if (InParams.GraphClassPtr->Metadata.Type == EMetasoundFrontendClassType::Graph)
+				{
+					return MakeShared<FGraphController>(EPrivateToken::Token, InParams);
+				}
+				else
+				{
+					UE_LOG(LogMetasound, Warning, TEXT("Failed to make graph controller [ClassId:%d]. Class must be EMeatsoundFrontendClassType::Graph."), InParams.GraphClassPtr->ID)
+				}
+			}
+			return FInvalidGraphController::GetInvalid();
+		}
+
+		bool FGraphController::IsValid() const
+		{
+			return GraphClassPtr.IsValid() && OwningDocument->IsValid();
+		}
+
+		int32 FGraphController::GetClassID() const
+		{
+			if (GraphClassPtr.IsValid())
+			{
+				return GraphClassPtr->ID;
+			}
+
+			return Metasound::FrontendInvalidID;
+		}
+
+		int32 FGraphController::GetNewPointID() const
+		{
+			if (const FMetasoundFrontendGraphClass* ThisClass = GraphClassPtr.Get())
+			{
+				TSet<int32> UsedIDs;
+
+				for (const FMetasoundFrontendClassInput& Input : ThisClass->Interface.Inputs)
+				{
+					UsedIDs.Append(Input.PointIDs);
+				}
+
+				for (const FMetasoundFrontendClassOutput& Output : ThisClass->Interface.Outputs)
+				{
+					UsedIDs.Append(Output.PointIDs);
+				}
+
+				return FrontendControllerIntrinsics::GetNewID(UsedIDs);
+			}
+
+			return Metasound::FrontendInvalidID;
+		}
+
+		TArray<FString> FGraphController::GetInputVertexNames() const 
+		{
+			TArray<FString> Names;
+
+			if (GraphClassPtr.IsValid())
+			{
+				for (const FMetasoundFrontendClassInput& Input : GraphClassPtr->Interface.Inputs)
+				{
+					Names.Add(Input.Name);
+				}
+			}
+
+			return Names;
+		}
+
+		TArray<FString> FGraphController::GetOutputVertexNames() const
+		{
+			TArray<FString> Names;
+
+			if (GraphClassPtr.IsValid())
+			{
+				for (const FMetasoundFrontendClassOutput& Output : GraphClassPtr->Interface.Outputs)
+				{
+					Names.Add(Output.Name);
+				}
+			}
+
+			return Names;
+		}
+
+		/*
+		const FMetasoundFrontendClassInput* FGraphController::FindClassInputWithName(const FString& InName) const
+		{
+			if (GraphClassPtr.IsValid())
+			{
+				auto IsInputWithName = [&](const FMetasoundFrontendClassInput& Input) { return Input.Name == InName; };
+				return GraphClassPtr->Interface.Inputs.FindByPredicate(IsInputWithName);
+			}
+			return nullptr;
+		}
+
+		const FMetasoundFrontendClassOutput* FGraphController::FindClassOutputWithName(const FString& InName) const
+		{
+			if (GraphClassPtr.IsValid())
+			{
+				auto IsOutputWithName = [&](const FMetasoundFrontendClassOutput& Output) { return Output.Name == InName; };
+				return GraphClassPtr->Interface.Outputs.FindByPredicate(IsOutputWithName);
+			}
+			return nullptr;
+		}
+		*/
+		FConstClassInputAccessPtr FGraphController::FindClassInputWithName(const FString& InName) const
+		{
+			using namespace FrontendControllerIntrinsics;
+			return GraphClassPtr.GetMemberAccessPtr<FMetasoundFrontendClassInput>(FGetClassInputFromClassWithName(InName));
+		}
+
+		FConstClassOutputAccessPtr FGraphController::FindClassOutputWithName(const FString& InName) const
+		{
+			using namespace FrontendControllerIntrinsics;
+			return GraphClassPtr.GetMemberAccessPtr<FMetasoundFrontendClassOutput>(FGetClassOutputFromClassWithName(InName));
+		}
+
+		TArray<int32> FGraphController::GetDefaultIDsForInputVertex(const FString& InInputName) const
+		{
+			if (const FMetasoundFrontendClassInput* Input = FindClassInputWithName(InInputName).Get())
+			{
+				return Input->PointIDs;
+			}
+			return TArray<int32>();
+		}
+
+		TArray<int32> FGraphController::GetDefaultIDsForOutputVertex(const FString& InOutputName) const
+		{
+			if (const FMetasoundFrontendClassOutput* Output = FindClassOutputWithName(InOutputName).Get())
+			{
+				return Output->PointIDs;
+			}
+			return TArray<int32>();
+		}
+
+		TArray<FNodeHandle> FGraphController::GetNodes()
+		{
+			return GetNodeHandles(GetNodesAndClasses());
+		}
+
+		TArray<FConstNodeHandle> FGraphController::GetConstNodes() const
+		{
+			return GetNodeHandles(GetNodesAndClasses());
+		}
+
+		FConstNodeHandle FGraphController::GetNodeWithID(int32 InNodeID) const
+		{
+			auto IsNodeWithSameID = [&](const FMetasoundFrontendClass& NodeClas, const FMetasoundFrontendNode& Node) 
+			{ 
+				return Node.ID == InNodeID; 
+			};
+			return GetNodeByPredicate(IsNodeWithSameID);
+		}
+
+		FNodeHandle FGraphController::GetNodeWithID(int32 InNodeID)
+		{
+			auto IsNodeWithSameID = [&](const FMetasoundFrontendClass& NodeClas, const FMetasoundFrontendNode& Node) 
+			{ 
+				return Node.ID == InNodeID; 
+			};
+			return GetNodeByPredicate(IsNodeWithSameID);
+		}
+
+		TArray<FNodeHandle> FGraphController::GetOutputNodes()
+		{
+			auto IsOutputNode = [](const FMetasoundFrontendClass& NodeClass, const FMetasoundFrontendNode& Node) 
+			{
+				return NodeClass.Metadata.Type == EMetasoundFrontendClassType::Output; 
+			};
+			return GetNodesByPredicate(IsOutputNode);
+		}
+
+		TArray<FNodeHandle> FGraphController::GetInputNodes()
+		{
+			auto IsInputNode = [](const FMetasoundFrontendClass& NodeClass, const FMetasoundFrontendNode& Node) 
+			{
+				return NodeClass.Metadata.Type == EMetasoundFrontendClassType::Input; 
+			};
+			return GetNodesByPredicate(IsInputNode);
+		}
+
+		TArray<FConstNodeHandle> FGraphController::GetConstOutputNodes() const
+		{
+			auto IsOutputNode = [](const FMetasoundFrontendClass& NodeClass, const FMetasoundFrontendNode& Node) 
+			{
+				return NodeClass.Metadata.Type == EMetasoundFrontendClassType::Output; 
+			};
+			return GetNodesByPredicate(IsOutputNode);
+		}
+
+		TArray<FConstNodeHandle> FGraphController::GetConstInputNodes() const
+		{
+			auto IsInputNode = [](const FMetasoundFrontendClass& NodeClass, const FMetasoundFrontendNode& Node) 
+			{
+				return NodeClass.Metadata.Type == EMetasoundFrontendClassType::Input; 
+			};
+			return GetNodesByPredicate(IsInputNode);
+		}
+
+		bool FGraphController::ContainsOutputNodeWithName(const FString& InName) const
+		{
+			if (GraphClassPtr.IsValid())
+			{
+				auto IsOutputNodeWithSameName = [&](const FMetasoundFrontendClass& NodeClass, const FMetasoundFrontendNode& Node) 
+				{ 
+					return (NodeClass.Metadata.Type == EMetasoundFrontendClassType::Output) && (Node.Name == InName); 
+				};
+				return ContainsNodesAndClassesByPredicate(IsOutputNodeWithSameName);
+			}
+			return false;
+		}
+
+		bool FGraphController::ContainsInputNodeWithName(const FString& InName) const
+		{
+			if (GraphClassPtr.IsValid())
+			{
+				auto IsInputNodeWithSameName = [&](const FMetasoundFrontendClass& NodeClass, const FMetasoundFrontendNode& Node) 
+				{ 
+					return (NodeClass.Metadata.Type == EMetasoundFrontendClassType::Input) && (Node.Name == InName); 
+				};
+				return ContainsNodesAndClassesByPredicate(IsInputNodeWithSameName);
+			}
+			return false;
+		}
+
+		FConstNodeHandle FGraphController::GetOutputNodeWithName(const FString& InName) const
+		{
+			auto IsOutputNodeWithSameName = [&](const FMetasoundFrontendClass& NodeClass, const FMetasoundFrontendNode& Node) 
+			{ 
+				return (NodeClass.Metadata.Type == EMetasoundFrontendClassType::Output) && (Node.Name == InName); 
+			};
+			return GetNodeByPredicate(IsOutputNodeWithSameName);
+		}
+
+		FConstNodeHandle FGraphController::GetInputNodeWithName(const FString& InName) const
+		{
+			auto IsInputNodeWithSameName = [&](const FMetasoundFrontendClass& NodeClass, const FMetasoundFrontendNode& Node) 
+			{ 
+				return (NodeClass.Metadata.Type == EMetasoundFrontendClassType::Input) && (Node.Name == InName); 
+			};
+			return GetNodeByPredicate(IsInputNodeWithSameName);
+		}
+
+		FNodeHandle FGraphController::GetOutputNodeWithName(const FString& InName)
+		{
+			auto IsOutputNodeWithSameName = [&](const FMetasoundFrontendClass& NodeClass, const FMetasoundFrontendNode& Node) 
+			{ 
+				return (NodeClass.Metadata.Type == EMetasoundFrontendClassType::Output) && (Node.Name == InName); 
+			};
+			return GetNodeByPredicate(IsOutputNodeWithSameName);
+		}
+
+		FNodeHandle FGraphController::GetInputNodeWithName(const FString& InName)
+		{
+			auto IsInputNodeWithSameName = [&](const FMetasoundFrontendClass& NodeClass, const FMetasoundFrontendNode& Node) 
+			{ 
+				return (NodeClass.Metadata.Type == EMetasoundFrontendClassType::Input) && (Node.Name == InName); 
+			};
+			return GetNodeByPredicate(IsInputNodeWithSameName);
+		}
+
+		FNodeHandle FGraphController::AddInputVertex(const FMetasoundFrontendClassInput& InClassInput)
+		{
+			using namespace FrontendControllerIntrinsics;
+
+			if (FMetasoundFrontendGraphClass* GraphClass = GraphClassPtr.Get())
+			{
+				auto IsInputWithSameName = [&](const FMetasoundFrontendClassInput& ExistingDesc) { return ExistingDesc.Name == InClassInput.Name; };
+				if (Algo::NoneOf(GraphClass->Interface.Inputs, IsInputWithSameName))
+				{
+					FMetasoundFrontendClassMetadata ClassMetadata;
+					if (FMetasoundFrontendRegistryContainer::GetInputNodeClassMetadataForDataType(InClassInput.TypeName, ClassMetadata))
+					{
+						if (FConstClassAccessPtr InputClass = OwningDocument->FindOrAddClass(ClassMetadata))
+						{
+							FMetasoundFrontendNode& Node = GraphClass->Graph.Nodes.Emplace_GetRef(*InputClass);
+
+							Node.ID = NewNodeID();
+
+							// TODO: have something that checks if input node has valid interface.
+							auto IsVertexWithTypeName = [&](FMetasoundFrontendVertex& Vertex) { return Vertex.TypeName == InClassInput.TypeName; };
+							if (FMetasoundFrontendVertex* InputVertex = Node.Interface.Inputs.FindByPredicate(IsVertexWithTypeName))
+							{ 
+								InputVertex->Name = InClassInput.Name;
+							}
+							else
+							{
+								UE_LOG(LogMetasound, Error, TEXT("Input node [TypeName:%s] does not contain input vertex with type [TypeName:%s]"), *InClassInput.TypeName.ToString(), *InClassInput.TypeName.ToString());
+							}
+
+							if (FMetasoundFrontendVertex* OutputVertex = Node.Interface.Outputs.FindByPredicate(IsVertexWithTypeName))
+							{ 
+								OutputVertex->Name = InClassInput.Name;
+							}
+							else
+							{
+								UE_LOG(LogMetasound, Error, TEXT("Input node [TypeName:%s] does not contain input vertex with type [TypeName:%s]"), *InClassInput.TypeName.ToString(), *InClassInput.TypeName.ToString());
+							}
+
+							FMetasoundFrontendClassInput& NewInput = GraphClassPtr->Interface.Inputs.Add_GetRef(InClassInput);
+							NewInput.NodeID = Node.ID;
+
+							FNodeAccessPtr NodePtr = GraphClassPtr.GetMemberAccessPtr<FMetasoundFrontendNode>(FGetNodeFromGraphClassByNodeID(Node.ID));
+							return GetNodeHandle(FGraphController::FNodeAndClass{NodePtr, InputClass});
+						}
+					}
+					else 
+					{
+						UE_LOG(LogMetasound, Display, TEXT("Failed to add input. No input node registered for data type [TypeName:%s]"), *InClassInput.TypeName.ToString());
+					}
+				}
+				else
+				{
+					UE_LOG(LogMetasound, Display, TEXT("Failed to add input. Input with same name \"%s\" exists in class [ClassID:%d]"), *InClassInput.Name, GraphClassPtr->ID);
+				}
+			}
+			return FInvalidNodeController::GetInvalid();
+		}
+
+		bool FGraphController::RemoveInputVertex(const FString& InName)
+		{
+			if (GraphClassPtr.IsValid())
+			{
+				auto IsInputNodeWithSameName = [&](const FMetasoundFrontendClass& InClass, const FMetasoundFrontendNode& InNode) 
+				{ 
+					return (InClass.Metadata.Type == EMetasoundFrontendClassType::Input) && (InNode.Name == InName); 
+				};
+
+				for (const FNodeAndClass& NodeAndClass : GetNodesAndClassesByPredicate(IsInputNodeWithSameName))
+				{
+					return RemoveInput(*NodeAndClass.Node);
+				}
+			}
+
+			return false;
+		}
+
+		FNodeHandle FGraphController::AddOutputVertex(const FMetasoundFrontendClassOutput& InClassOutput)
+		{
+			using namespace FrontendControllerIntrinsics;
+
+			if (FMetasoundFrontendGraphClass* GraphClass = GraphClassPtr.Get())
+			{
+				auto IsOutputWithSameName = [&](const FMetasoundFrontendClassOutput& ExistingDesc) { return ExistingDesc.Name == InClassOutput.Name; };
+				if (Algo::NoneOf(GraphClass->Interface.Outputs, IsOutputWithSameName))
+				{
+					FMetasoundFrontendClassMetadata ClassMetadata;
+					if (FMetasoundFrontendRegistryContainer::GetOutputNodeClassMetadataForDataType(InClassOutput.TypeName, ClassMetadata))
+					{
+						//if (const FMetasoundFrontendClass* OutputClass = OwningDocument->FindOrAddClass(ClassMetadata))
+						if (FConstClassAccessPtr OutputClass = OwningDocument->FindOrAddClass(ClassMetadata))
+						{
+							FMetasoundFrontendNode& Node = GraphClass->Graph.Nodes.Emplace_GetRef(*OutputClass);
+
+							Node.ID = NewNodeID();
+
+							// TODO: have something that checks if input node has valid interface.
+							auto IsVertexWithTypeName = [&](FMetasoundFrontendVertex& Vertex) { return Vertex.TypeName == InClassOutput.TypeName; };
+							if (FMetasoundFrontendVertex* InputVertex = Node.Interface.Inputs.FindByPredicate(IsVertexWithTypeName))
+							{ 
+								InputVertex->Name = InClassOutput.Name;
+							}
+							else
+							{
+								UE_LOG(LogMetasound, Error, TEXT("Output node [TypeName:%s] does not contain input vertex with type [TypeName:%s]"), *InClassOutput.TypeName.ToString(), *InClassOutput.TypeName.ToString());
+							}
+
+							if (FMetasoundFrontendVertex* OutputVertex = Node.Interface.Outputs.FindByPredicate(IsVertexWithTypeName))
+							{ 
+								OutputVertex->Name = InClassOutput.Name;
+							}
+							else
+							{
+								UE_LOG(LogMetasound, Error, TEXT("Output node [TypeName:%s] does not contain input vertex with type [TypeName:%s]"), *InClassOutput.TypeName.ToString(), *InClassOutput.TypeName.ToString());
+							}
+
+							FMetasoundFrontendClassOutput& NewOutput = GraphClassPtr->Interface.Outputs.Add_GetRef(InClassOutput);
+							NewOutput.NodeID = Node.ID;
+
+							FNodeAccessPtr NodePtr = GraphClassPtr.GetMemberAccessPtr<FMetasoundFrontendNode>(FGetNodeFromGraphClassByNodeID(Node.ID));
+							return GetNodeHandle(FGraphController::FNodeAndClass{NodePtr, OutputClass});
+						}
+					}
+					else 
+					{
+						UE_LOG(LogMetasound, Display, TEXT("Failed to add output. No output node registered for data type [TypeName:%s]"), *InClassOutput.TypeName.ToString());
+					}
+				}
+				else
+				{
+					UE_LOG(LogMetasound, Display, TEXT("Failed to add output. Output with same name \"%s\" exists in class [ClassID:%d]"), *InClassOutput.Name, GraphClassPtr->ID);
+				}
+			}
+
+			return FInvalidNodeController::GetInvalid();
+		}
+
+		bool FGraphController::RemoveOutputVertex(const FString& InName)
+		{
+			if (GraphClassPtr.IsValid())
+			{
+				auto IsOutputNodeWithSameName = [&](const FMetasoundFrontendClass& InClass, const FMetasoundFrontendNode& InNode) 
+				{ 
+					return (InClass.Metadata.Type == EMetasoundFrontendClassType::Output) && (InNode.Name == InName); 
+				};
+
+				for (const FNodeAndClass& NodeAndClass : GetNodesAndClassesByPredicate(IsOutputNodeWithSameName))
+				{
+					return RemoveOutput(*NodeAndClass.Node);
+				}
+			}
+
+			return false;
+		}
+
+		// This can be used to determine what kind of property editor we should use for the data type of a given input.
+		// Will return Invalid if the input couldn't be found, or if the input doesn't support any kind of literals.
+		ELiteralType FGraphController::GetPreferredLiteralTypeForInputVertex(const FString& InInputName) const
+		{
+			if (const FMetasoundFrontendClassInput* Desc = FindInputDescriptionWithName(InInputName))
+			{
+				return FMetasoundFrontendRegistryContainer::Get()->GetDesiredLiteralTypeForDataType(Desc->TypeName);
+			}
+			return ELiteralType::Invalid;
+		}
+
+		// For inputs whose preferred literal type is UObject or UObjectArray, this can be used to determine the UObject corresponding to that input's datatype.
+		UClass* FGraphController::GetSupportedClassForInputVertex(const FString& InInputName)
+		{
+			if (const FMetasoundFrontendClassInput* Desc = FindInputDescriptionWithName(InInputName))
+			{
+				return FMetasoundFrontendRegistryContainer::Get()->GetLiteralUClassForDataType(Desc->TypeName);
+			}
+			return nullptr;
+		}
+
+
+		// These can be used to set the default value for a given input on this graph.
+		// @returns false if the input name couldn't be found, or if the literal type was incompatible with the Data Type of this input.
+		bool FGraphController::SetDefaultInputToLiteral(const FString& InInputName, int32 InPointID, bool bInValue)
+		{
+			return SetDefaultInputToLiteralInternal(InInputName, InPointID, bInValue);
+		}
+
+		bool FGraphController::SetDefaultInputToLiteral(const FString& InInputName, int32 InPointID, int32 InValue)
+		{
+			return SetDefaultInputToLiteralInternal(InInputName, InPointID, InValue);
+		}
+
+		bool FGraphController::SetDefaultInputToLiteral(const FString& InInputName, int32 InPointID, float InValue)
+		{
+			return SetDefaultInputToLiteralInternal(InInputName, InPointID, InValue);
+		}
+
+		bool FGraphController::SetDefaultInputToLiteral(const FString& InInputName, int32 InPointID, const FString& InValue)
+		{
+			return SetDefaultInputToLiteralInternal(InInputName, InPointID, InValue);
+		}
+
+		bool FGraphController::SetDefaultInputToLiteral(const FString& InInputName, int32 InPointID, UObject* InValue)
+		{
+			return SetDefaultInputToLiteralInternal(InInputName, InPointID, InValue);
+		}
+
+		bool FGraphController::SetDefaultInputToLiteral(const FString& InInputName, int32 InPointID, const TArray<UObject*>& InValue)
+		{
+			return SetDefaultInputToLiteralInternal(InInputName, InPointID, InValue);
+		}
+
+		// Set the display name for the input with the given name
+		void FGraphController::SetInputDisplayName(FString InName, const FText& InDisplayName)
+		{
+			if (FMetasoundFrontendClassInput* Desc = FindInputDescriptionWithName(InName))
+			{
+				Desc->Metadata.DisplayName = InDisplayName;
+			}
+		}
+
+		// Set the display name for the output with the given name
+		void FGraphController::SetOutputDisplayName(FString InName, const FText& InDisplayName)
+		{
+			if (FMetasoundFrontendClassOutput* Desc = FindOutputDescriptionWithName(InName))
+			{
+				Desc->Metadata.DisplayName = InDisplayName;
+			}
+		}
+
+		FMetasoundFrontendClassMetadata FGraphController::CreateInputClassMetadata(const FMetasoundFrontendClassInput& InClassInput)
+		{
+			FMetasoundFrontendClassMetadata ClassMetadata;
+			FMetasoundFrontendRegistryContainer::GetInputNodeClassMetadataForDataType(InClassInput.TypeName, ClassMetadata);
+			return ClassMetadata;
+		}
+
+		FMetasoundFrontendClassMetadata FGraphController::CreateOutputClassMetadata(const FMetasoundFrontendClassOutput& InClassOutput)
+		{
+			FMetasoundFrontendClassMetadata ClassMetadata;
+			FMetasoundFrontendRegistryContainer::GetOutputNodeClassMetadataForDataType(InClassOutput.TypeName, ClassMetadata);
+			return ClassMetadata;
+		}
+
+
+		// This can be used to clear the current literal for a given input.
+		// @returns false if the input name couldn't be found.
+		bool FGraphController::ClearLiteralForInput(const FString& InInputName, int32 InPointID)
+		{
+			if (FMetasoundFrontendClassInput* Desc = FindInputDescriptionWithName(InInputName))
+			{
+				auto IsLiteralWithSamePointID = [&](const FMetasoundFrontendVertexLiteral& InVertexLiteral) 
+				{ 
+					return InVertexLiteral.PointID == InPointID; 
+				};
+
+				if (FMetasoundFrontendVertexLiteral* VertexLiteral = Desc->Defaults.FindByPredicate(IsLiteralWithSamePointID))
+				{
+					VertexLiteral->Value.Clear();
+					return true;
+				}
+			}
+
+			return false;
+		}
+
+		FNodeHandle FGraphController::AddNode(const FNodeClassInfo& InNodeClass)
+		{
+			if (IsValid())
+			{
+				//if (const FMetasoundFrontendClass* DependencyDescription = OwningDocument->FindOrAddClass(InNodeClass))
+				if (FConstClassAccessPtr DependencyDescription = OwningDocument->FindOrAddClass(InNodeClass))
+				{
+					return AddNode(DependencyDescription);
+				}
+			}
+
+			return FInvalidNodeController::GetInvalid();
+		}
+
+		FNodeHandle FGraphController::AddNode(const FNodeRegistryKey& InNodeClass)
+		{
+			// Construct a FNodeClassInfo from this lookup key.
+			FNodeClassInfo ClassInfo;
+			ClassInfo.LookupKey = InNodeClass;
+			ClassInfo.NodeName = InNodeClass.NodeName.ToString();
+			ClassInfo.NodeType = EMetasoundFrontendClassType::External;
+
+			return AddNode(ClassInfo);
+		}
+
+		// Remove the node corresponding to this node handle.
+		// On success, invalidates the received node handle.
+		bool FGraphController::RemoveNode(INodeController& InNode)
+		{
+			int32 NodeID = InNode.GetID();
+			const EMetasoundFrontendClassType NodeType = InNode.GetClassType();
+
+			auto IsNodeWithSameID = [&](const FMetasoundFrontendNode& InDesc) { return InDesc.ID == NodeID; };
+			
+			if (const FMetasoundFrontendNode* Desc = GraphClassPtr->Graph.Nodes.FindByPredicate(IsNodeWithSameID))
+			{
+				if (EMetasoundFrontendClassType::Input == NodeType)
+				{
+					return RemoveInput(*Desc);
+				}
+				else if (EMetasoundFrontendClassType::Output == NodeType)
+				{
+					return RemoveOutput(*Desc);
+				}
+				else
+				{
+					return RemoveNode(*Desc);
+				}
+			}
+
+			return false;
+		}
+
+		// Returns the metadata for the current graph, including the name, description and author.
+		const FMetasoundFrontendClassMetadata& FGraphController::GetGraphMetadata() const
+		{
+			if (GraphClassPtr.IsValid())
+			{
+				return GraphClassPtr->Metadata;
+			}
+			return FrontendControllerIntrinsics::GetInvalidValueConstRef<FMetasoundFrontendClassMetadata>();
+		}
+
+		bool FGraphController::InflateNodeDirectlyIntoGraph(const INodeController& InNode)
+		{
+			// TODO: implement
+			checkNoEntry();
+
+			return false;
+		}
+
+		FNodeHandle FGraphController::CreateEmptySubgraph(const FMetasoundFrontendClassMetadata& InInfo)
+		{
+			if (IsValid())
+			{
+				if (InInfo.Type == EMetasoundFrontendClassType::Graph)
+				{
+					if (const FMetasoundFrontendClass* ExistingDependency = OwningDocument->FindClass(InInfo).Get())
+					{
+						UE_LOG(LogMetasound, Error, TEXT("Cannot add new subgraph. Metasound class already exists with matching metadata Name: \"%s\", Version %d.%d"), *(ExistingDependency->Metadata.Name.GetFullName()), ExistingDependency->Metadata.Version.Major, ExistingDependency->Metadata.Version.Minor);
+					}
+					//else if (const FMetasoundFrontendClass* DependencyDescription = OwningDocument->FindOrAddClass(InInfo))
+					else if (FConstClassAccessPtr DependencyDescription = OwningDocument->FindOrAddClass(InInfo))
+					{
+						return AddNode(DependencyDescription);
+					}
+				}
+				else
+				{
+					UE_LOG(LogMetasound, Warning, TEXT("Incompatible Metasound NodeType encountered when attempting to create an empty subgraph.  NodeType must equal EMetasoundFrontendClassType::Graph"));
+				}
+			}
+			
+			return FInvalidNodeController::GetInvalid();
+		}
+
+		TUniquePtr<IOperator> FGraphController::BuildOperator(const FOperatorSettings& InSettings, const FMetasoundEnvironment& InEnvironment, TArray<IOperatorBuilder::FBuildErrorPtr>& OutBuildErrors) const
+		{
+			if (!IsValid())
+			{
+				return TUniquePtr<IOperator>(nullptr);
+			}
+
+			// TODO: Implement inflation step here.
+
+			// TODO: bubble up errors. 
+			TArray<FMetasoundFrontendGraphClass> Subgraphs = OwningDocument->GetSubgraphs();
+			TArray<FMetasoundFrontendClass> Dependencies = OwningDocument->GetDependencies();
+
+			FFrontendGraphBuilder GraphBuilder;
+			TUniquePtr<FFrontendGraph> Graph = GraphBuilder.CreateGraph(*GraphClassPtr, Subgraphs, Dependencies);
+
+			if (!Graph.IsValid())
+			{
+				return TUniquePtr<IOperator>(nullptr);
+			}
+
+			// Step 5: Invoke Operator Builder
+			FOperatorBuilder OperatorBuilder(FOperatorBuilderSettings::GetDefaultSettings());
+
+			return OperatorBuilder.BuildGraphOperator(*Graph, InSettings, InEnvironment, OutBuildErrors);
+		}
+
+		FDocumentHandle FGraphController::GetOwningDocument()
+		{
+			return OwningDocument;
+		}
+
+		FConstDocumentHandle FGraphController::GetOwningDocument() const
+		{
+			return OwningDocument;
+		}
+
+		FNodeHandle FGraphController::AddNode(FConstClassAccessPtr InExistingDependency)
+		{
+			using namespace FrontendControllerIntrinsics;
+
+			if (FMetasoundFrontendGraphClass* GraphClass = GraphClassPtr.Get())
+			{
+				if (const FMetasoundFrontendClass* NodeClass = InExistingDependency.Get())
+				{
+					FMetasoundFrontendNode& Node = GraphClass->Graph.Nodes.Emplace_GetRef(*InExistingDependency);
+
+					Node.ID = NewNodeID();
+
+					FNodeAccessPtr NodePtr = GraphClassPtr.GetMemberAccessPtr<FMetasoundFrontendNode>(FGetNodeFromGraphClassByNodeID(Node.ID));
+
+					//FGraphAccessPtr GraphPtr = GraphClassPtr.GetMemberAccessPtr<FMetasoundFrontendGraph>(FGetGraphFromGraphClass());
+
+					//return FNodeController::CreateNodeHandle(FNodeController::FInitParams{NodePtr, InExistingDependency, GraphPtr, this->AsShared()});
+					return GetNodeHandle(FGraphController::FNodeAndClass{NodePtr, InExistingDependency});
+				}
+			}
+			/*
+			if (GraphClassPtr.IsValid())
+			{
+				FMetasoundFrontendNode& NewNodeDescription = GraphClassPtr->Graph.Nodes.AddDefaulted_GetRef();
+
+				NewNodeDescription.Name = InExistingDependency.Metadata.Name.Name;
+				NewNodeDescription.ID = NewNodeID();
+				NewNodeDescription.ClassID = InExistingDependency.ID;
+
+				FNodeAccessPtr NodePtr = FrontendControllerIntrinsics::MakeAutoAccessPtr(NewNodeDescription);
+				FConstClassAccessPtr ClassPtr = FrontendControllerIntrinsics::MakeAutoAccessPtr(InExistingDependency);
+				FGraphAccessPtr GraphPtr = FrontendControllerIntrinsics::MakeAutoAccessPtr(GraphClassPtr->Graph);
+
+				return FNodeController::CreateNodeHandle(FNodeController::FInitParams{NodePtr, ClassPtr, GraphPtr, this->AsShared()});
+			}
+			*/
+
+			return FInvalidNodeController::GetInvalid();
+		}
+
+		bool FGraphController::RemoveNode(const FMetasoundFrontendNode& InDesc)
+		{
+			if (GraphClassPtr.IsValid())
+			{
+				auto IsConnectionFromThisNode = [&](const FMetasoundFrontendEdge& ConDesc) { return ConDesc.FromNodeID == InDesc.ID; };
+
+				// Remove any reference connections
+				GraphClassPtr->Graph.Edges.RemoveAll(IsConnectionFromThisNode);
+
+				auto IsNodeWithID = [&](const FMetasoundFrontendNode& Desc) { return InDesc.ID == Desc.ID; };
+
+				int32 NumRemoved = GraphClassPtr->Graph.Nodes.RemoveAll(IsNodeWithID);
+
+				OwningDocument->RemoveUnreferencedDependencies();
+
+				return (NumRemoved > 0);
+			}
+			return false;
+		}
+
+		bool FGraphController::RemoveInput(const FMetasoundFrontendNode& InNode)
+		{
+			if (GraphClassPtr.IsValid())
+			{
+				auto IsInputWithSameNodeID = [&](const FMetasoundFrontendClassInput& ClassInput) { return ClassInput.NodeID == InNode.ID; };
+
+				int32 NumInputsRemoved = GraphClassPtr->Interface.Inputs.RemoveAll(IsInputWithSameNodeID);
+
+				bool bDidRemoveNode = RemoveNode(InNode);
+
+				return (NumInputsRemoved > 0) || bDidRemoveNode;
+			}
+
+			return false;
+		}
+
+		bool FGraphController::RemoveOutput(const FMetasoundFrontendNode& InNode)
+		{
+			if (GraphClassPtr.IsValid())
+			{
+				auto IsOutputWithSameNodeID = [&](const FMetasoundFrontendClassOutput& ClassOutput) { return ClassOutput.NodeID == InNode.ID; };
+
+				int32 NumOutputsRemoved = GraphClassPtr->Interface.Outputs.RemoveAll(IsOutputWithSameNodeID);
+
+				bool bDidRemoveNode = RemoveNode(InNode);
+
+				return (NumOutputsRemoved > 0) || bDidRemoveNode;
+			}
+
+			return false;
+		}
+
+		int32 FGraphController::NewNodeID() const
+		{
+			if (GraphClassPtr.IsValid())
+			{
+				TSet<int32> UsedIDs;
+
+				for (const FMetasoundFrontendNode& Node : GraphClassPtr->Graph.Nodes)
+				{
+					UsedIDs.Add(Node.ID);
+				}
+
+				return FrontendControllerIntrinsics::GetNewID(UsedIDs);
+			}
+
+			return Metasound::FrontendInvalidID;
+		}
+
+		bool FGraphController::ContainsNodesAndClassesByPredicate(TFunctionRef<bool (const FMetasoundFrontendClass&, const FMetasoundFrontendNode&)> InPredicate) const
+		{
+			if (GraphClassPtr.IsValid())
+			{
+				for (FMetasoundFrontendNode& Node : GraphClassPtr->Graph.Nodes)
+				{
+					const FMetasoundFrontendClass* NodeClass = OwningDocument->FindClassWithID(Node.ClassID).Get();
+					if (nullptr != NodeClass)
+					{
+						if (InPredicate(*NodeClass, Node))
+						{
+							return true;
+						}
+					}
+					else
+					{
+						UE_LOG(LogMetasound, Warning, TEXT("Failed to find class for node [NodeID:%d, ClassID:%d]"), Node.ID, Node.ClassID);
+					}
+				}
+			}
+
+			return false;
+		}
+
+		TArray<FGraphController::FNodeAndClass> FGraphController::GetNodesAndClasses()
+		{
+			using namespace FrontendControllerIntrinsics;
+			TArray<FNodeAndClass> NodesAndClasses;
+
+			if (GraphClassPtr.IsValid())
+			{
+				for (FMetasoundFrontendNode& Node : GraphClassPtr->Graph.Nodes)
+				{
+					FNodeAccessPtr NodePtr = GraphClassPtr.GetMemberAccessPtr<FMetasoundFrontendNode>(FGetNodeFromGraphClassByNodeID(Node.ID));
+					//const FMetasoundFrontendClass* NodeClass = OwningDocument->FindClassWithID(Node.ClassID);
+					FConstClassAccessPtr NodeClassPtr = OwningDocument->FindClassWithID(Node.ClassID);
+					if (NodeClassPtr && NodePtr)
+					{
+						NodesAndClasses.Add(FGraphController::FNodeAndClass{ NodePtr, NodeClassPtr });
+					}
+					else
+					{
+						UE_LOG(LogMetasound, Warning, TEXT("Failed to find class for node [NodeID:%d, ClassID:%d]"), Node.ID, Node.ClassID);
+					}
+				}
+			}
+
+			return NodesAndClasses;
+		}
+
+		TArray<FGraphController::FConstNodeAndClass> FGraphController::GetNodesAndClasses() const
+		{
+			using namespace FrontendControllerIntrinsics;
+			TArray<FConstNodeAndClass> NodesAndClasses;
+
+			if (GraphClassPtr.IsValid())
+			{
+				for (const FMetasoundFrontendNode& Node : GraphClassPtr->Graph.Nodes)
+				{
+					FConstNodeAccessPtr NodePtr = GraphClassPtr.GetMemberAccessPtr<const FMetasoundFrontendNode>(FGetConstNodeFromGraphClassByNodeID(Node.ID));
+					//const FMetasoundFrontendClass* NodeClass = OwningDocument->FindClassWithID(Node.ClassID);
+					FConstClassAccessPtr NodeClassPtr = OwningDocument->FindClassWithID(Node.ClassID);
+					if (NodeClassPtr && NodePtr)
+					{
+						NodesAndClasses.Add(FGraphController::FConstNodeAndClass{ NodePtr, NodeClassPtr });
+					}
+					/*
+					const FMetasoundFrontendClass* NodeClass = OwningDocument->FindClassWithID(Node.ClassID);
+					if (nullptr != NodeClass)
+					{
+						NodesAndClasses.Emplace(
+							FGraphController::FConstNodeAndClass
+							{
+								FrontendControllerIntrinsics::MakeAutoAccessPtr(Node),
+								FrontendControllerIntrinsics::MakeAutoAccessPtr(*NodeClass)
+							}
+						);
+					}
+					*/
+					else
+					{
+						UE_LOG(LogMetasound, Warning, TEXT("Failed to find class for node [NodeID:%d, ClassID:%d]"), Node.ID, Node.ClassID);
+					}
+				}
+			}
+
+			return NodesAndClasses;
+		}
+
+		TArray<FGraphController::FNodeAndClass> FGraphController::GetNodesAndClassesByPredicate(TFunctionRef<bool (const FMetasoundFrontendClass&, const FMetasoundFrontendNode&)> InPredicate)
+		{
+			using namespace FrontendControllerIntrinsics;
+
+			TArray<FNodeAndClass> NodesAndClasses;
+
+			if (GraphClassPtr.IsValid())
+			{
+				for (FMetasoundFrontendNode& Node : GraphClassPtr->Graph.Nodes)
+				{
+					FConstClassAccessPtr NodeClassPtr = OwningDocument->FindClassWithID(Node.ClassID);
+					if (const FMetasoundFrontendClass* NodeClass = NodeClassPtr.Get())
+					{
+						if (InPredicate(*NodeClass, Node))
+						{
+							FNodeAccessPtr NodePtr = GraphClassPtr.GetMemberAccessPtr<FMetasoundFrontendNode>(FGetNodeFromGraphClassByNodeID(Node.ID));
+							NodesAndClasses.Add(FGraphController::FNodeAndClass{ NodePtr, NodeClassPtr });
+						}
+					}
+					else
+					{
+						UE_LOG(LogMetasound, Warning, TEXT("Failed to find class for node [NodeID:%d, ClassID:%d]"), Node.ID, Node.ClassID);
+					}
+				}
+			}
+
+			return NodesAndClasses;
+		}
+
+		TArray<FGraphController::FConstNodeAndClass> FGraphController::GetNodesAndClassesByPredicate(TFunctionRef<bool (const FMetasoundFrontendClass&, const FMetasoundFrontendNode&)> InPredicate) const
+		{
+			using namespace FrontendControllerIntrinsics;
+
+			TArray<FConstNodeAndClass> NodesAndClasses;
+
+			if (GraphClassPtr.IsValid())
+			{
+				for (const FMetasoundFrontendNode& Node : GraphClassPtr->Graph.Nodes)
+				{
+					FConstClassAccessPtr NodeClassPtr = OwningDocument->FindClassWithID(Node.ClassID);
+					if (const FMetasoundFrontendClass* NodeClass = NodeClassPtr.Get())
+					{
+						if (InPredicate(*NodeClass, Node))
+						{
+							FConstNodeAccessPtr NodePtr = GraphClassPtr.GetMemberAccessPtr<const FMetasoundFrontendNode>(FGetConstNodeFromGraphClassByNodeID(Node.ID));
+							NodesAndClasses.Add(FGraphController::FConstNodeAndClass{ NodePtr, NodeClassPtr });
+						}
+					}
+					else
+					{
+						UE_LOG(LogMetasound, Warning, TEXT("Failed to find class for node [NodeID:%d, ClassID:%d]"), Node.ID, Node.ClassID);
+					}
+				}
+			}
+
+			return NodesAndClasses;
+		}
+
+		FNodeHandle FGraphController::GetNodeByPredicate(TFunctionRef<bool (const FMetasoundFrontendClass&, const FMetasoundFrontendNode&)> InPredicate)
+		{
+			if (GraphClassPtr.IsValid())
+			{
+				TArray<FNodeAndClass> NodeAndClass = GetNodesAndClassesByPredicate(InPredicate);
+
+				if (NodeAndClass.Num() > 0)
+				{
+					return GetNodeHandle(NodeAndClass[0]);
+				}
+			}
+
+			return FInvalidNodeController::GetInvalid();
+		}
+
+		FConstNodeHandle FGraphController::GetNodeByPredicate(TFunctionRef<bool (const FMetasoundFrontendClass&, const FMetasoundFrontendNode&)> InPredicate) const
+		{
+			if (GraphClassPtr.IsValid())
+			{
+				TArray<FConstNodeAndClass> NodeAndClass = GetNodesAndClassesByPredicate(InPredicate);
+
+				if (NodeAndClass.Num() > 0)
+				{
+					return GetNodeHandle(NodeAndClass[0]);
+				}
+			}
+
+			return FInvalidNodeController::GetInvalid();
+		}
+
+		TArray<FNodeHandle> FGraphController::GetNodesByPredicate(TFunctionRef<bool (const FMetasoundFrontendClass&, const FMetasoundFrontendNode&)> InFilterFunc)
+		{
+			if (GraphClassPtr.IsValid())
+			{
+				return GetNodeHandles(GetNodesAndClassesByPredicate(InFilterFunc));
+			}
+			return TArray<FNodeHandle>();
+		}
+
+		TArray<FConstNodeHandle> FGraphController::GetNodesByPredicate(TFunctionRef<bool (const FMetasoundFrontendClass&, const FMetasoundFrontendNode&)> InFilterFunc) const
+		{
+			if (GraphClassPtr.IsValid())
+			{
+				return GetNodeHandles(GetNodesAndClassesByPredicate(InFilterFunc));
+			}
+			return TArray<FConstNodeHandle>();
+		}
+
+		TArray<FNodeHandle> FGraphController::GetNodeHandles(TArrayView<const FGraphController::FNodeAndClass> InNodesAndClasses)
+		{
+			TArray<FNodeHandle> Nodes;
+
+			for (const FNodeAndClass& NodeAndClass: InNodesAndClasses)
+			{
+				FNodeHandle NodeController = GetNodeHandle(NodeAndClass);
+				if (NodeController->IsValid())
+				{
+					Nodes.Add(NodeController);
+				}
+			}
+
+			return Nodes;
+		}
+
+		TArray<FConstNodeHandle> FGraphController::GetNodeHandles(TArrayView<const FGraphController::FConstNodeAndClass> InNodesAndClasses) const
+		{
+			TArray<FConstNodeHandle> Nodes;
+
+			for (const FConstNodeAndClass& NodeAndClass : InNodesAndClasses)
+			{
+				FConstNodeHandle NodeController = GetNodeHandle(NodeAndClass);
+				if (NodeController->IsValid())
+				{
+					Nodes.Add(NodeController);
+				}
+			}
+
+			return Nodes;
+		}
+
+		FNodeHandle FGraphController::GetNodeHandle(const FGraphController::FNodeAndClass& InNodeAndClass)
+		{
+			using namespace FrontendControllerIntrinsics;
+
+			if (InNodeAndClass.IsValid() && GraphClassPtr.IsValid())
+			{
+				FGraphHandle OwningGraph = this->AsShared();
+				FGraphAccessPtr GraphPtr = GraphClassPtr.GetMemberAccessPtr<FMetasoundFrontendGraph>(FGetGraphFromGraphClass());
+				//FGraphAccessPtr GraphPtr = FrontendControllerIntrinsics::MakeAutoAccessPtr(GraphClassPtr->Graph);
+
+				switch (InNodeAndClass.Class->Metadata.Type)
+				{
+					case EMetasoundFrontendClassType::Input:
+						if (FClassInputAccessPtr OwningGraphClassInputPtr = FindInputDescriptionWithNodeID(InNodeAndClass.Node->ID))
+						{
+							FInputNodeController::FInitParams InitParams
+							{
+								InNodeAndClass.Node,
+								InNodeAndClass.Class,
+								OwningGraphClassInputPtr,
+								GraphPtr,
+								OwningGraph
+							};
+							return FInputNodeController::CreateInputNodeHandle(InitParams);
+						}
+						break;
+
+					case EMetasoundFrontendClassType::Output:
+						if (FClassOutputAccessPtr OwningGraphClassOutputPtr = FindOutputDescriptionWithNodeID(InNodeAndClass.Node->ID))
+						{
+							FOutputNodeController::FInitParams InitParams
+							{
+								InNodeAndClass.Node,
+								InNodeAndClass.Class,
+								OwningGraphClassOutputPtr,
+								GraphPtr,
+								OwningGraph
+							};
+							return FOutputNodeController::CreateOutputNodeHandle(InitParams);
+						}
+						break;
+
+					case EMetasoundFrontendClassType::External:
+					case EMetasoundFrontendClassType::Graph:
+						{
+							FNodeController::FInitParams InitParams
+							{
+								InNodeAndClass.Node,
+								InNodeAndClass.Class,
+								GraphPtr,
+								OwningGraph
+							};
+							return FNodeController::CreateNodeHandle(InitParams);
+						}
+						break;
+
+					default:
+						checkNoEntry();
+				}
+			}
+
+			return FInvalidNodeController::GetInvalid();
+		}
+
+		FConstNodeHandle FGraphController::GetNodeHandle(const FGraphController::FConstNodeAndClass& InNodeAndClass) const
+		{
+			using namespace FrontendControllerIntrinsics;
+
+			if (InNodeAndClass.IsValid() && GraphClassPtr.IsValid())
+			{
+				FConstGraphHandle OwningGraph = this->AsShared();
+				//FConstGraphAccessPtr GraphPtr = FrontendControllerIntrinsics::MakeAutoAccessPtr(GraphClassPtr->Graph);
+				FConstGraphAccessPtr GraphPtr = GraphClassPtr.GetMemberAccessPtr<const FMetasoundFrontendGraph>(FGetConstGraphFromGraphClass());
+				switch (InNodeAndClass.Class->Metadata.Type)
+				{
+					case EMetasoundFrontendClassType::Input:
+						if (FConstClassInputAccessPtr OwningGraphClassInputPtr = FindInputDescriptionWithNodeID(InNodeAndClass.Node->ID))
+						{
+							FInputNodeController::FInitParams InitParams
+							{
+								ConstCastAccessPtr<FMetasoundFrontendNode>(InNodeAndClass.Node),
+								InNodeAndClass.Class,
+								ConstCastAccessPtr<FMetasoundFrontendClassInput>(OwningGraphClassInputPtr),
+								ConstCastAccessPtr<FMetasoundFrontendGraph>(GraphPtr),
+								ConstCastSharedRef<IGraphController>(OwningGraph)
+							};
+							return FInputNodeController::CreateConstInputNodeHandle(InitParams);
+						}
+						break;
+
+					case EMetasoundFrontendClassType::Output:
+						if (FConstClassOutputAccessPtr OwningGraphClassOutputPtr = FindOutputDescriptionWithNodeID(InNodeAndClass.Node->ID))
+						{
+							FOutputNodeController::FInitParams InitParams
+							{
+								ConstCastAccessPtr<FMetasoundFrontendNode>(InNodeAndClass.Node),
+								InNodeAndClass.Class,
+								ConstCastAccessPtr<FMetasoundFrontendClassOutput>(OwningGraphClassOutputPtr),
+								ConstCastAccessPtr<FMetasoundFrontendGraph>(GraphPtr),
+								ConstCastSharedRef<IGraphController>(OwningGraph)
+							};
+							return FOutputNodeController::CreateConstOutputNodeHandle(InitParams);
+						}
+						break;
+
+					case EMetasoundFrontendClassType::External:
+					case EMetasoundFrontendClassType::Graph:
+						{
+							FNodeController::FInitParams InitParams
+							{
+								ConstCastAccessPtr<FMetasoundFrontendNode>(InNodeAndClass.Node),
+								InNodeAndClass.Class,
+								ConstCastAccessPtr<FMetasoundFrontendGraph>(GraphPtr),
+								ConstCastSharedRef<IGraphController>(OwningGraph)
+							};
+							return FNodeController::CreateConstNodeHandle(InitParams);
+						}
+						break;
+
+					default:
+						checkNoEntry();
+				}
+			}
+
+			return FInvalidNodeController::GetInvalid();
+		}
+
+
+		FMetasoundFrontendClassInput* FGraphController::FindInputDescriptionWithName(const FString& InName)
+		{
+			if (GraphClassPtr.IsValid())
+			{
+				return GraphClassPtr->Interface.Inputs.FindByPredicate([&](const FMetasoundFrontendClassInput& Desc) { return Desc.Name == InName; });
+			}
+			return nullptr;
+		}
+
+		const FMetasoundFrontendClassInput* FGraphController::FindInputDescriptionWithName(const FString& InName) const
+		{
+			if (GraphClassPtr.IsValid())
+			{
+				return GraphClassPtr->Interface.Inputs.FindByPredicate([&](const FMetasoundFrontendClassInput& Desc) { return Desc.Name == InName; });
+			}
+			return nullptr;
+		}
+
+		FMetasoundFrontendClassOutput* FGraphController::FindOutputDescriptionWithName(const FString& InName)
+		{
+			if (GraphClassPtr.IsValid())
+			{
+				return GraphClassPtr->Interface.Outputs.FindByPredicate([&](const FMetasoundFrontendClassOutput& Desc) { return Desc.Name == InName; });
+			}
+			return nullptr;
+		}
+
+		const FMetasoundFrontendClassOutput* FGraphController::FindOutputDescriptionWithName(const FString& InName) const
+		{
+			if (GraphClassPtr.IsValid())
+			{
+				return GraphClassPtr->Interface.Outputs.FindByPredicate([&](const FMetasoundFrontendClassOutput& Desc) { return Desc.Name == InName; });
+			}
+			return nullptr;
+		}
+
+		FClassInputAccessPtr FGraphController::FindInputDescriptionWithNodeID(int32 InNodeID)
+		{
+			using namespace FrontendControllerIntrinsics;
+
+			return GraphClassPtr.GetMemberAccessPtr<FMetasoundFrontendClassInput>(FGetClassInputFromGraphClassWithNodeID(InNodeID));
+			/*
+			if (GraphClassPtr.IsValid())
+			{
+				return GraphClassPtr->Interface.Inputs.FindByPredicate([&](const FMetasoundFrontendClassInput& Desc) { return Desc.NodeID == InNodeID; });
+			}
+			return nullptr;
+			*/
+		}
+
+		FConstClassInputAccessPtr FGraphController::FindInputDescriptionWithNodeID(int32 InNodeID) const
+		{
+			using namespace FrontendControllerIntrinsics;
+
+			return GraphClassPtr.GetMemberAccessPtr<const FMetasoundFrontendClassInput>(FGetConstClassInputFromGraphClassWithNodeID(InNodeID));
+			/*
+			if (GraphClassPtr.IsValid())
+			{
+				return GraphClassPtr->Interface.Inputs.FindByPredicate([&](const FMetasoundFrontendClassInput& Desc) { return Desc.NodeID == InNodeID; });
+			}
+			return nullptr;
+			*/
+		}
+
+		FClassOutputAccessPtr FGraphController::FindOutputDescriptionWithNodeID(int32 InNodeID)
+		{
+			using namespace FrontendControllerIntrinsics;
+
+			return GraphClassPtr.GetMemberAccessPtr<FMetasoundFrontendClassOutput>(FGetClassOutputFromGraphClassWithNodeID(InNodeID));
+			/*
+			if (GraphClassPtr.IsValid())
+			{
+				return GraphClassPtr->Interface.Outputs.FindByPredicate([&](const FMetasoundFrontendClassOutput& Desc) { return Desc.NodeID == InNodeID; });
+			}
+			return nullptr;
+			*/
+		}
+
+		FConstClassOutputAccessPtr FGraphController::FindOutputDescriptionWithNodeID(int32 InNodeID) const
+		{
+			using namespace FrontendControllerIntrinsics;
+
+			return GraphClassPtr.GetMemberAccessPtr<const FMetasoundFrontendClassOutput>(FGetConstClassOutputFromGraphClassWithNodeID(InNodeID));
+			/*
+			if (GraphClassPtr.IsValid())
+			{
+				return GraphClassPtr->Interface.Outputs.FindByPredicate([&](const FMetasoundFrontendClassOutput& Desc) { return Desc.NodeID == InNodeID; });
+			}
+			return nullptr;
+			*/
+		}
+
+		/*
+		FMetasoundFrontendNode* FGraphController::FindNodeByID(int32 InNodeID)
+		{
+			if (GraphClassPtr.IsValid())
+			{
+				auto IsNodeWithID = [&](const FMetasoundFrontendNode& InNode) { return InNode.ID == InNodeID; };
+				return GraphClassPtr->Graph.Nodes.FindByPredicate(IsNodeWithID);
+			}
+			return nullptr;
+		}
+
+		const FMetasoundFrontendNode* FGraphController::FindNodeByID(int32 InNodeID) const
+		{
+			if (GraphClassPtr.IsValid())
+			{
+				auto IsNodeWithID = [&](const FMetasoundFrontendNode& InNode) { return InNode.ID == InNodeID; };
+				return GraphClassPtr->Graph.Nodes.FindByPredicate(IsNodeWithID);
+			}
+			return nullptr;
+		}
+		*/
+
+		//
+		// FDocumentController
+		//
+		FDocumentController::FDocumentController(FDocumentAccessPtr InDocumentPtr)
+		:	DocumentPtr(InDocumentPtr)
+		{
+		}
+
+		bool FDocumentController::IsValid() const
+		{
+			return DocumentPtr.IsValid();
+		}
+
+		bool FDocumentController::IsRequiredInput(const FString& InInputName) const
+		{
+			if (DocumentPtr.IsValid())
+			{ 
+				auto IsInputWithSameName = [&](const FMetasoundFrontendClassVertex& Desc) { return Desc.Name == InInputName; };
+
+				return DocumentPtr->Archetype.Interface.Inputs.ContainsByPredicate(IsInputWithSameName);
+			}
+			return false;
+		}
+
+		bool FDocumentController::IsRequiredOutput(const FString& InOutputName) const
+		{
+			if (DocumentPtr.IsValid())
+			{ 
+				auto IsOutputWithSameName = [&](const FMetasoundFrontendClassVertex& Desc) { return Desc.Name == InOutputName; };
+
+				return DocumentPtr->Archetype.Interface.Outputs.ContainsByPredicate(IsOutputWithSameName);
+			}
+			return false;
+		}
+
+		TArray<FMetasoundFrontendClassVertex> FDocumentController::GetRequiredInputs() const
+		{
+			if (DocumentPtr.IsValid())
+			{
+				return DocumentPtr->Archetype.Interface.Inputs;
+			}
+
+			return TArray<FMetasoundFrontendClassVertex>();
+		}
+
+		TArray<FMetasoundFrontendClassVertex> FDocumentController::GetRequiredOutputs() const
+		{
+			if (DocumentPtr.IsValid())
+			{
+				return DocumentPtr->Archetype.Interface.Outputs;
+			}
+
+			return TArray<FMetasoundFrontendClassVertex>();
+		}
+
+		TArray<FMetasoundFrontendClass> FDocumentController::GetDependencies() const
+		{
+			if (DocumentPtr.IsValid())
+			{
+				return DocumentPtr->Dependencies;
+			}
+			return TArray<FMetasoundFrontendClass>();
+		}
+
+		TArray<FMetasoundFrontendGraphClass> FDocumentController::GetSubgraphs() const
+		{
+			if (DocumentPtr.IsValid())
+			{
+				return DocumentPtr->Subgraphs;
+			}
+			return TArray<FMetasoundFrontendGraphClass>();
+		}
+
+		TArray<FMetasoundFrontendClass> FDocumentController::GetClasses() const 
+		{
+			TArray<FMetasoundFrontendClass> Classes = GetDependencies();
+			Classes.Append(GetSubgraphs());
+			return Classes;
+		}
+
+		FConstClassAccessPtr FDocumentController::FindDependencyWithID(int32 InClassID) const 
+		{
+			using namespace FrontendControllerIntrinsics;
+			return DocumentPtr.GetMemberAccessPtr<const FMetasoundFrontendClass>(FGetConstClassFromDocumentDependenciesWithID(InClassID));
+
+			/*
+			if (DocumentPtr.IsValid())
+			{
+				auto IsMatchingID = [&](const FMetasoundFrontendClass& InClass) { return InClass.ID == InClassID; };
+				return DocumentPtr->Dependencies.FindByPredicate(IsMatchingID);
+			}
+			return nullptr;
+			*/
+		}
+
+		FConstGraphClassAccessPtr FDocumentController::FindSubgraphWithID(int32 InClassID) const
+		{
+			using namespace FrontendControllerIntrinsics;
+			return DocumentPtr.GetMemberAccessPtr<const FMetasoundFrontendGraphClass>(FGetConstGraphClassFromDocumentWithID(InClassID));
+			/*
+			if (DocumentPtr.IsValid())
+			{
+				auto IsMatchingID = [&](const FMetasoundFrontendGraphClass& InClass) { return InClass.ID == InClassID; };
+				return DocumentPtr->Subgraphs.FindByPredicate(IsMatchingID);
+			}
+			return nullptr;
+			*/
+		}
+
+		FConstClassAccessPtr FDocumentController::FindClassWithID(int32 InClassID) const
+		{
+			FConstClassAccessPtr MetasoundClass = FindDependencyWithID(InClassID);
+
+			if (!MetasoundClass.IsValid())
+			{
+				MetasoundClass = FindSubgraphWithID(InClassID);
+			}
+
+			return MetasoundClass;
+		}
+
+		FConstClassAccessPtr FDocumentController::FindClass(const FNodeClassInfo& InNodeClass) const
+		{
+			using namespace FrontendControllerIntrinsics;
+
+			return DocumentPtr.GetMemberAccessPtr<const FMetasoundFrontendClass>(FGetConstClassFromDocumentWithClassInfo(InNodeClass));
+			/*
+			if (DocumentPtr.IsValid())
+			{
+				auto IsMatchingClass = [&](const FMetasoundFrontendClass& Desc)
+				{
+					return IsMatchingMetasoundClass(InNodeClass, Desc.Metadata);
+				};
+
+				FConstClassAccessPtr MetasoundClass = DocumentPtr->Dependencies.FindByPredicate(IsMatchingClass);
+
+				if (nullptr == MetasoundClass)
+				{
+					MetasoundClass = DocumentPtr->Subgraphs.FindByPredicate(IsMatchingClass);
+				}
+
+				return MetasoundClass;
+			}
+
+			return nullptr;
+			*/
+		}
+
+		FConstClassAccessPtr FDocumentController::FindOrAddClass(const FNodeClassInfo& InNodeClass)
+		{
+			using namespace FrontendControllerIntrinsics;
+			
+			FConstClassAccessPtr ClassPtr = FindClass(InNodeClass);
+
+			if ((!ClassPtr.IsValid()) && DocumentPtr.IsValid())
+			{
+				// FNodeClassInfo does not contain enough info add a subgraph.
+				check(EMetasoundFrontendClassType::Graph != InNodeClass.NodeType);
+
+				FMetasoundFrontendClass NewClass = GenerateClassDescription(InNodeClass);
+				NewClass.ID = NewClassID();
+
+				DocumentPtr->Dependencies.Add(NewClass);
+
+				ClassPtr = FindClass(InNodeClass);
+			}
+
+			return ClassPtr;
+		}
+
+		FConstClassAccessPtr FDocumentController::FindClass(const FMetasoundFrontendClassMetadata& InMetadata) const
+		{
+			using namespace FrontendControllerIntrinsics;
+
+			return DocumentPtr.GetMemberAccessPtr<const FMetasoundFrontendClass>(FGetConstClassFromDocumentWithMetadata(InMetadata));
+			/*
+			if (DocumentPtr.IsValid())
+			{
+				auto IsMatchingClass = [&](const FMetasoundFrontendClass& Desc)
+				{
+					return IsMatchingMetasoundClass(InMetadata, Desc.Metadata);
+				};
+
+				FConstClassAccessPtr MetasoundClass = DocumentPtr->Dependencies.FindByPredicate(IsMatchingClass);
+
+				if (nullptr == MetasoundClass)
+				{
+					MetasoundClass = DocumentPtr->Subgraphs.FindByPredicate(IsMatchingClass);
+				}
+			}
+			return nullptr;
+			*/
+		}
+
+		FConstClassAccessPtr FDocumentController::FindOrAddClass(const FMetasoundFrontendClassMetadata& InMetadata)
+		{
+			FConstClassAccessPtr ClassPtr = FindClass(InMetadata);
+
+			if ((!ClassPtr.IsValid()) && DocumentPtr.IsValid())
+			{
+				if ((EMetasoundFrontendClassType::External == InMetadata.Type) || (EMetasoundFrontendClassType::Input == InMetadata.Type) || (EMetasoundFrontendClassType::Output == InMetadata.Type))
+				{
+					FMetasoundFrontendClass NewClass;
+					if (FMetasoundFrontendRegistryContainer::GetFrontendClassFromRegistered(InMetadata, NewClass))
+					{
+						NewClass.ID = NewClassID();
+						DocumentPtr->Dependencies.Add(NewClass);
+					}
+					else
+					{
+						UE_LOG(LogMetasound, Error, TEXT("Cannot add external dependency. No Metasound class found with matching metadata Name: \"%s\", Version %d.%d. Suggested solution \"%s\" by %s."), *InMetadata.Name.GetFullName(), InMetadata.Version.Major, InMetadata.Version.Minor, *InMetadata.PromptIfMissing.ToString(), *InMetadata.Author.ToString());
+					}
+				} 
+				else if (EMetasoundFrontendClassType::Graph == InMetadata.Type)
+				{
+					FMetasoundFrontendGraphClass NewClass;
+					NewClass.ID = NewClassID();
+					NewClass.Metadata = InMetadata;
+
+					DocumentPtr->Subgraphs.Add(NewClass);
+				}
+				else
+				{
+					UE_LOG(LogMetasound, Error, TEXT("Unsupported metasound class type for node: \"%s\", Version %d.%d."), *InMetadata.Name.GetFullName(), InMetadata.Version.Major, InMetadata.Version.Minor);
+					checkNoEntry();
+				}
+
+				ClassPtr = FindClass(InMetadata);
+			}
+
+			return ClassPtr;
+		}
+
+		/*
+		const FMetasoundFrontendClass* FDocumentController::FindDependencyWithID(int32 InClassID) const 
+		{
+			if (DocumentPtr.IsValid())
+			{
+				auto IsMatchingID = [&](const FMetasoundFrontendClass& InClass) { return InClass.ID == InClassID; };
+				return DocumentPtr->Dependencies.FindByPredicate(IsMatchingID);
+			}
+			return nullptr;
+		}
+
+		const FMetasoundFrontendGraphClass* FDocumentController::FindSubgraphWithID(int32 InClassID) const
+		{
+			if (DocumentPtr.IsValid())
+			{
+				auto IsMatchingID = [&](const FMetasoundFrontendGraphClass& InClass) { return InClass.ID == InClassID; };
+				return DocumentPtr->Subgraphs.FindByPredicate(IsMatchingID);
+			}
+			return nullptr;
+		}
+
+		const FMetasoundFrontendClass* FDocumentController::FindClassWithID(int32 InClassID) const
+		{
+			const FMetasoundFrontendClass* MetasoundClass = FindDependencyWithID(InClassID);
+
+			if (nullptr == MetasoundClass)
+			{
+				MetasoundClass = FindSubgraphWithID(InClassID);
+			}
+
+			return MetasoundClass;
+		}
+
+		const FMetasoundFrontendClass* FDocumentController::FindClass(const FNodeClassInfo& InNodeClass) const
+		{
+			if (DocumentPtr.IsValid())
+			{
+				auto IsMatchingClass = [&](const FMetasoundFrontendClass& Desc)
+				{
+					return IsMatchingMetasoundClass(InNodeClass, Desc.Metadata);
+				};
+
+				const FMetasoundFrontendClass* MetasoundClass = DocumentPtr->Dependencies.FindByPredicate(IsMatchingClass);
+
+				if (nullptr == MetasoundClass)
+				{
+					MetasoundClass = DocumentPtr->Subgraphs.FindByPredicate(IsMatchingClass);
+				}
+
+				return MetasoundClass;
+			}
+
+			return nullptr;
+		}
+
+		const FMetasoundFrontendClass* FDocumentController::FindOrAddClass(const FNodeClassInfo& InNodeClass)
+		{
+			if (DocumentPtr.IsValid())
+			{
+
+				if (const FMetasoundFrontendClass* MetasoundClass = FindClass(InNodeClass))
+				{
+					return MetasoundClass;
+				}
+
+				// FNodeClassInfo does not contain enough info add a subgraph.
+				check(EMetasoundFrontendClassType::Graph != InNodeClass.NodeType);
+
+				FMetasoundFrontendClass NewClass = GenerateClassDescription(InNodeClass);
+				NewClass.ID = NewClassID();
+
+				return &(DocumentPtr->Dependencies.Add_GetRef(NewClass));
+			}
+
+			return nullptr;
+		}
+
+		const FMetasoundFrontendClass* FDocumentController::FindClass(const FMetasoundFrontendClassMetadata& InMetadata) const
+		{
+			if (DocumentPtr.IsValid())
+			{
+				auto IsMatchingClass = [&](const FMetasoundFrontendClass& Desc)
+				{
+					return IsMatchingMetasoundClass(InMetadata, Desc.Metadata);
+				};
+
+				const FMetasoundFrontendClass* MetasoundClass = DocumentPtr->Dependencies.FindByPredicate(IsMatchingClass);
+
+				if (nullptr == MetasoundClass)
+				{
+					MetasoundClass = DocumentPtr->Subgraphs.FindByPredicate(IsMatchingClass);
+				}
+			}
+			return nullptr;
+		}
+
+		const FMetasoundFrontendClass* FDocumentController::FindOrAddClass(const FMetasoundFrontendClassMetadata& InMetadata)
+		{
+			if (DocumentPtr.IsValid())
+			{
+				if (const FMetasoundFrontendClass* MetasoundClass = FindClass(InMetadata))
+				{
+					return MetasoundClass;
+				}
+
+				if ((EMetasoundFrontendClassType::External == InMetadata.Type) || (EMetasoundFrontendClassType::Input == InMetadata.Type) || (EMetasoundFrontendClassType::Output == InMetadata.Type))
+				{
+					FMetasoundFrontendClass NewClass;
+					if (FMetasoundFrontendRegistryContainer::GetFrontendClassFromRegistered(InMetadata, NewClass))
+					{
+						NewClass.ID = NewClassID();
+						return &(DocumentPtr->Dependencies.Add_GetRef(NewClass));
+					}
+					else
+					{
+						UE_LOG(LogMetasound, Error, TEXT("Cannot add external dependency. No Metasound class found with matching metadata Name: \"%s\", Version %d.%d. Suggested solution \"%s\" by %s."), *InMetadata.Name.GetFullName(), InMetadata.Version.Major, InMetadata.Version.Minor, *InMetadata.PromptIfMissing.ToString(), *InMetadata.Author.ToString());
+					}
+				} 
+				else if (EMetasoundFrontendClassType::Graph == InMetadata.Type)
+				{
+					FMetasoundFrontendGraphClass NewClass;
+					NewClass.ID = NewClassID();
+					NewClass.Metadata = InMetadata;
+
+					return &(DocumentPtr->Subgraphs.Add_GetRef(NewClass));
+				}
+				else
+				{
+					UE_LOG(LogMetasound, Error, TEXT("Unsupported metasound class type for node: \"%s\", Version %d.%d."), *InMetadata.Name.GetFullName(), InMetadata.Version.Major, InMetadata.Version.Minor);
+					checkNoEntry();
+				}
+			}
+
+			return nullptr;
+		}
+		*/
+
+		void FDocumentController::RemoveUnreferencedDependencies()
+		{
+			if (DocumentPtr.IsValid())
+			{
+				int32 NumDependenciesRemovedThisItr = 0;
+
+				// Repeatedly remove unreferenced dependencies until there are 
+				// no unreferenced dependencies left. 
+				do
+				{
+					TSet<int32> ReferencedDependencyIDs;
+					auto AddNodeClassIDToSet = [&](const FMetasoundFrontendNode& Node)
+					{
+						ReferencedDependencyIDs.Add(Node.ClassID);
+					};
+
+					auto AddGraphNodeClassIDsToSet = [&](const FMetasoundFrontendGraphClass& GraphClass)
+					{
+						Algo::ForEach(GraphClass.Graph.Nodes, AddNodeClassIDToSet);
+					};
+					
+					// Referenced dependencies in root class
+					Algo::ForEach(DocumentPtr->RootGraph.Graph.Nodes, AddNodeClassIDToSet);
+					// Referenced dependencies in 
+					Algo::ForEach(DocumentPtr->Subgraphs, AddGraphNodeClassIDsToSet);
+
+					auto IsDependencyUnreferenced = [&](const FMetasoundFrontendClass& ClassDependency)
+					{
+						return !ReferencedDependencyIDs.Contains(ClassDependency.ID);
+					};
+
+					NumDependenciesRemovedThisItr = DocumentPtr->Dependencies.RemoveAllSwap(IsDependencyUnreferenced);
+				} 
+				while(NumDependenciesRemovedThisItr > 0);
+			}
+		}
+
+		FGraphHandle FDocumentController::GetRootGraph()
+		{
+			if (DocumentPtr.IsValid())
+			{
+				//FGraphClassAccessPtr GraphClass = FrontendControllerIntrinsics::MakeAutoAccessPtr(DocumentPtr->RootGraph);
+				FGraphClassAccessPtr GraphClass = DocumentPtr.GetMemberAccessPtr<FMetasoundFrontendGraphClass>([](FMetasoundFrontendDocument& Doc) { return &Doc.RootGraph; });
+				return FGraphController::CreateGraphHandle(FGraphController::FInitParams{GraphClass, this->AsShared()});
+			}
+
+			return FInvalidGraphController::GetInvalid();
+		}
+
+		FConstGraphHandle FDocumentController::GetRootGraph() const
+		{
+			if (DocumentPtr.IsValid())
+			{
+				//FGraphClassAccessPtr GraphClass = FrontendControllerIntrinsics::MakeAutoAccessPtr(DocumentPtr->RootGraph);
+				FConstGraphClassAccessPtr GraphClass = DocumentPtr.GetMemberAccessPtr<const FMetasoundFrontendGraphClass>([](const FMetasoundFrontendDocument& Doc) { return &Doc.RootGraph; });
+				return FGraphController::CreateConstGraphHandle(FGraphController::FInitParams
+					{
+						ConstCastAccessPtr<FMetasoundFrontendGraphClass>(GraphClass),
+						ConstCastSharedRef<IDocumentController>(this->AsShared())
+					});
+			}
+			return FInvalidGraphController::GetInvalid();
+		}
+
+		TArray<FGraphHandle> FDocumentController::GetSubgraphHandles() 
+		{
+			TArray<FGraphHandle> Subgraphs;
+
+			if (DocumentPtr.IsValid())
+			{
+				for (FMetasoundFrontendGraphClass& GraphClass : DocumentPtr->Subgraphs)
+				{
+					/*
+					FGraphClassAccessPtr GraphClassPtr = FrontendControllerIntrinsics::MakeAutoAccessPtr(GraphClass);
+					FGraphController::CreateGraphHandle(FGraphController::FInitParams{GraphClassPtr, this->AsShared()}));
+					*/
+					Subgraphs.Add(GetSubgraphWithClassID(GraphClass.ID));
+				}
+			}
+
+			return Subgraphs;
+		}
+
+		TArray<FConstGraphHandle> FDocumentController::GetSubgraphHandles() const 
+		{
+			TArray<FConstGraphHandle> Subgraphs;
+
+			if (DocumentPtr.IsValid())
+			{
+				for (const FMetasoundFrontendGraphClass& GraphClass : DocumentPtr->Subgraphs)
+				{
+					/*
+					FConstGraphClassAccessPtr GraphClassPtr = FrontendControllerIntrinsics::MakeAutoAccessPtr(GraphClass);
+					Subgraphs.Add(FGraphController::CreateConstGraphHandle(
+						FGraphController::FInitParams
+						{
+							ConstCastAccessPtr<FMetasoundFrontendGraphClass>(GraphClassPtr), 
+							ConstCastSharedRef<IDocumentController>(this->AsShared())
+						}
+					));
+					*/
+					Subgraphs.Add(GetSubgraphWithClassID(GraphClass.ID));
+				}
+			}
+
+			return Subgraphs;
+		}
+
+		FGraphHandle FDocumentController::GetSubgraphWithClassID(int32 InClassID)
+		{
+			FGraphClassAccessPtr GraphClassPtr = DocumentPtr.GetMemberAccessPtr<FMetasoundFrontendGraphClass>(FrontendControllerIntrinsics::FGetSubgraphFromDocumentByID(InClassID));
+			return FGraphController::CreateGraphHandle(FGraphController::FInitParams{GraphClassPtr, this->AsShared()});
+			/*
+			if (DocumentPtr.IsValid())
+			{
+
+
+					[InClassID](FMetasoundFrontendDocument& Doc)
+					{
+						return Doc.Subgraphs.FindByPredicate(
+							[&](const FMetasoundFrontendGraphClass& InGraphClass) 
+							{ 
+								return InGraphClass.ID == InClassID; 
+							}
+						);
+					}
+				);
+				*/
+
+				/*
+
+
+				auto IsMatchingID = [&](const FMetasoundFrontendGraphClass& InGraphClass) { return InGraphClass.ID == InClassID; };
+				if (FMetasoundFrontendGraphClass* GraphClass = DocumentPtr->Subgraphs.FindByPredicate(IsMatchingID))
+				{
+					//FGraphClassAccessPtr GraphClassPtr = FrontendControllerIntrinsics::MakeAutoAccessPtr(*GraphClass);
+					return FGraphController::CreateGraphHandle(FGraphController::FInitParams{GraphClassPtr, this->AsShared()});
+				}
+			}
+			return FInvalidGraphController::GetInvalid();
+				*/
+		}
+
+		FConstGraphHandle FDocumentController::GetSubgraphWithClassID(int32 InClassID) const
+		{
+			FConstGraphClassAccessPtr GraphClassPtr = DocumentPtr.GetMemberAccessPtr<const FMetasoundFrontendGraphClass>(FrontendControllerIntrinsics::FGetSubgraphFromDocumentByID(InClassID));
+
+			return FGraphController::CreateConstGraphHandle(FGraphController::FInitParams{ConstCastAccessPtr<FMetasoundFrontendGraphClass>(GraphClassPtr), ConstCastSharedRef<IDocumentController>(this->AsShared())});
+
+			/*
+			if (DocumentPtr.IsValid())
+			{
+				auto IsMatchingID = [&](const FMetasoundFrontendGraphClass& InGraphClass) { return InGraphClass.ID == InClassID; };
+
+				if (const FMetasoundFrontendGraphClass* GraphClass = DocumentPtr->Subgraphs.FindByPredicate(IsMatchingID))
+				{
+					FConstGraphClassAccessPtr GraphClassPtr = FrontendControllerIntrinsics::MakeAutoAccessPtr(*GraphClass);
+					return FGraphController::CreateConstGraphHandle(FGraphController::FInitParams{ConstCastAccessPtr<FMetasoundFrontendGraphClass>(GraphClassPtr), ConstCastSharedRef<IDocumentController>(this->AsShared())});
+				}
+			}
+			return FInvalidGraphController::GetInvalid();
+			*/
+		}
+
+		bool FDocumentController::ExportToJSONAsset(const FString& InAbsolutePath) const
+		{
+			if (DocumentPtr.IsValid())
+			{
+				if (TUniquePtr<FArchive> FileWriter = TUniquePtr<FArchive>(IFileManager::Get().CreateFileWriter(*InAbsolutePath)))
+				{
+					TJsonStructSerializerBackend<DefaultCharType> Backend(*FileWriter, EStructSerializerBackendFlags::Default);
+					FStructSerializer::Serialize<FMetasoundFrontendDocument>(*DocumentPtr, Backend);
+			
+					FileWriter->Close();
+
+					return true;
+				}
+				else
+				{
+					UE_LOG(LogMetasound, Error, TEXT("Failed to export Metasound json asset. Could not write to path \"%s\"."), *InAbsolutePath);
+				}
+			}
+
+			return false;
+		}
+
+
+		FString FDocumentController::ExportToJSON() const
+		{
+			TArray<uint8> WriterBuffer;
+			FMemoryWriter MemWriter(WriterBuffer);
+
+			Metasound::TJsonStructSerializerBackend<Metasound::DefaultCharType> Backend(MemWriter, EStructSerializerBackendFlags::Default);
+			FStructSerializer::Serialize<FMetasoundFrontendDocument>(*DocumentPtr, Backend);
+
+			MemWriter.Close();
+
+			// null terminator
+			WriterBuffer.AddZeroed(sizeof(ANSICHAR));
+
+			FString Output;
+			Output.AppendChars(reinterpret_cast<ANSICHAR*>(WriterBuffer.GetData()), WriterBuffer.Num() / sizeof(ANSICHAR));
+
+			return Output;
+		}
+		
+		bool FDocumentController::IsMatchingMetasoundClass(const FMetasoundFrontendClassMetadata& InMetadataA, const FMetasoundFrontendClassMetadata& InMetadataB) 
+		{
+			if (InMetadataA.Type == InMetadataB.Type)
+			{
+				if (InMetadataA.Name == InMetadataB.Name)
+				{
+					return FRegistry::GetRegistryKey(InMetadataA) == FRegistry::GetRegistryKey(InMetadataB);
+				}
+			}
+			return false;
+		}
+
+		bool FDocumentController::IsMatchingMetasoundClass(const FNodeClassInfo& InNodeClass, const FMetasoundFrontendClassMetadata& InMetadata) 
+		{
+			if (InNodeClass.NodeType == InMetadata.Type)
+			{
+				if (InNodeClass.NodeName == InMetadata.Name.Name)
+				{
+					return InNodeClass.LookupKey == FRegistry::GetRegistryKey(InMetadata);
+				}
+			}
+			return false;
+		}
+
+		int32 FDocumentController::NewClassID() const 
+		{
+			if (DocumentPtr.IsValid())
+			{
+				TSet<int32> UsedIDs;
+				for (const FMetasoundFrontendClass& InDesc : DocumentPtr->Dependencies)
+				{
+					UsedIDs.Add(InDesc.ID);
+				}
+				for (const FMetasoundFrontendClass& InDesc : DocumentPtr->Subgraphs)
+				{
+					UsedIDs.Add(InDesc.ID);
+				}
+
+				return FrontendControllerIntrinsics::GetNewID(UsedIDs);
+			}
+
+			return Metasound::FrontendInvalidID;
+		}
+	}
+}
+
