Metasound Frontend Revamp

diff --git a/Engine/Plugins/Runtime/Metasound/Source/MetasoundFrontend/Private/MetasoundFrontendStandardController.h b/Engine/Plugins/Runtime/Metasound/Source/MetasoundFrontend/Private/MetasoundFrontendStandardController.h
new file mode 100644
--- /dev/null
+++ b/Engine/Plugins/Runtime/Metasound/Source/MetasoundFrontend/Private/MetasoundFrontendStandardController.h
@@ -0,0 +1,846 @@
+// Copyright Epic Games, Inc. All Rights Reserved.
+
+#pragma once
+
+#include "CoreMinimal.h"
+#include "MetasoundBuilderInterface.h"
+#include "MetasoundFrontendBaseClasses.h"
+#include "MetasoundFrontendController.h"
+#include "MetasoundFrontendDocument.h"
+#include "MetasoundFrontendInvalidController.h"
+#include "MetasoundLog.h"
+
+namespace Metasound
+{
+	namespace Frontend
+	{
+		/*
+		using FVertexAccessPtr = TAccessPtr<FMetasoundFrontendVertex>;
+		using FConstVertexAccessPtr = TAccessPtr<const FMetasoundFrontendVertex>;
+		using FClassOutputAccessPtr = TAccessPtr<FMetasoundFrontendClassOutput>;
+		using FConstClassOutputAccessPtr = TAccessPtr<const FMetasoundFrontendClassOutput>;
+		using FClassInputAccessPtr = TAccessPtr<FMetasoundFrontendClassInput>;
+		using FConstClassInputAccessPtr = TAccessPtr<const FMetasoundFrontendClassInput>;
+		using FNodeAccessPtr = TAccessPtr<FMetasoundFrontendNode>;
+		using FConstNodeAccessPtr = TAccessPtr<const FMetasoundFrontendNode>;
+		using FGraphAccessPtr = TAccessPtr<FMetasoundFrontendGraph>;
+		using FConstGraphAccessPtr = TAccessPtr<const FMetasoundFrontendGraph>;
+		using FClassAccessPtr = TAccessPtr<FMetasoundFrontendClass>;
+		using FConstClassAccessPtr = TAccessPtr<const FMetasoundFrontendClass>;
+		using FGraphClassAccessPtr = TAccessPtr<FMetasoundFrontendGraphClass>;
+		using FConstGraphClassAccessPtr = TAccessPtr<const FMetasoundFrontendGraphClass>;
+		using FDocumentAccessPtr = TAccessPtr<FMetasoundFrontendDocument>;
+		using FConstDocumentAccessPtr = TAccessPtr<const FMetasoundFrontendDocument>;
+		*/
+
+		/** FBaseOutputController provides common functionality for multiple derived
+		 * output controllers.
+		 */
+		class METASOUNDFRONTEND_API FBaseOutputController: public IOutputController
+		{
+		public:
+
+			struct FInitParams
+			{
+				int32 ID;
+
+				FConstVertexAccessPtr NodeVertexPtr;
+
+				/* Node handle which owns this output. */
+				FNodeHandle OwningNode;
+			};
+
+			/** Construct the output controller base.  */
+			FBaseOutputController(const FInitParams& InParams);
+
+			virtual ~FBaseOutputController() = default;
+
+			bool IsValid() const override;
+
+			int32 GetID() const override;
+			const FName& GetDataType() const override;
+			const FString& GetName() const override;
+
+			// Return info on containing node. 
+			int32 GetOwningNodeID() const override;
+			FNodeHandle GetOwningNode() override;
+			FConstNodeHandle GetOwningNode() const override;
+
+			// Connection logic.
+			FConnectability CanConnectTo(const IInputController& InController) const override;
+			bool Connect(IInputController& InController) override;
+			bool ConnectWithConverterNode(IInputController& InController, const FConverterNodeInfo& InNodeClassName) override;
+			bool Disconnect(IInputController& InController) override;
+
+		protected:
+
+			int32 ID;
+			FConstVertexAccessPtr NodeVertexPtr;	
+			FNodeHandle OwningNode;
+		};
+
+		/** FNodeOutputController represents a single output of a single node. */
+		class METASOUNDFRONTEND_API FNodeOutputController : public FBaseOutputController
+		{
+		public:
+			struct FInitParams
+			{
+				int32 ID;
+
+				FConstVertexAccessPtr NodeVertexPtr;
+			
+				FConstClassOutputAccessPtr ClassOutputPtr;
+
+				/* Node handle which owns this output. */
+				FNodeHandle OwningNode;
+			};
+
+			FNodeOutputController(const FInitParams& InInitParams);
+
+			virtual ~FNodeOutputController() = default;
+
+			bool IsValid() const override;
+
+			// Output metadata
+			const FText& GetDisplayName() const override;
+			const FText& GetTooltip() const override;
+
+		private:
+
+			FConstClassOutputAccessPtr ClassOutputPtr;
+		};
+
+		/** FInputNodeOutputController represents the output vertex of an input 
+		 * node. 
+		 *
+		 * FInputNodeOutputController is largely to represent inputs coming into 
+		 * graph. 
+		 */
+		class METASOUNDFRONTEND_API FInputNodeOutputController : public FBaseOutputController
+		{
+		public:
+			struct FInitParams
+			{
+				int32 ID;
+
+				FConstVertexAccessPtr NodeVertexPtr;
+			
+				FConstClassInputAccessPtr OwningGraphClassInputPtr;
+
+				/* Node handle which owns this output. */
+				FNodeHandle OwningNode;
+			};
+
+			/** Constructs the output controller. */
+			FInputNodeOutputController(const FInitParams& InParams);
+
+			virtual ~FInputNodeOutputController() = default;
+
+			bool IsValid() const override;
+
+			// Output metadata
+			const FText& GetDisplayName() const override;
+			const FText& GetTooltip() const override;
+
+		private:
+
+			FConstClassInputAccessPtr OwningGraphClassInputPtr;
+		};
+
+		
+		/** FBaseInputController provides common functionality for multiple derived
+		 * input controllers.
+		 */
+		class METASOUNDFRONTEND_API FBaseInputController : public IInputController 
+		{
+		public:
+
+			struct FInitParams
+			{
+				int32 ID; 
+				FConstVertexAccessPtr NodeVertexPtr;
+				FGraphAccessPtr GraphPtr; 
+				FNodeHandle OwningNode;
+			};
+
+			/** Construct the input controller base. */
+			FBaseInputController(const FInitParams& InParams);
+
+			virtual ~FBaseInputController() = default;
+
+			bool IsValid() const override;
+
+			int32 GetID() const override;
+			const FName& GetDataType() const override;
+			const FString& GetName() const override;
+
+			// Owning node info
+			int32 GetOwningNodeID() const override;
+			FNodeHandle GetOwningNode() override;
+			FConstNodeHandle GetOwningNode() const override;
+
+			// Connection info
+			bool IsConnected() const override;
+			FOutputHandle GetCurrentlyConnectedOutput() override;
+			FConstOutputHandle GetCurrentlyConnectedOutput() const override;
+
+			FConnectability CanConnectTo(const IOutputController& InController) const override;
+			bool Connect(IOutputController& InController) override;
+
+			// Connection controls.
+			bool ConnectWithConverterNode(IOutputController& InController, const FConverterNodeInfo& InNodeClassName) override;
+
+			bool Disconnect(IOutputController& InController) override;
+			bool Disconnect() override;
+
+		protected:
+
+			const FMetasoundFrontendEdge* FindEdge() const;
+			FMetasoundFrontendEdge* FindEdge();
+
+		private:
+
+			int32 ID;
+			FConstVertexAccessPtr NodeVertexPtr;
+			FGraphAccessPtr GraphPtr;
+			FNodeHandle OwningNode;
+		};
+
+		/** FNodeInputController represents a single input of a single node. */
+		class METASOUNDFRONTEND_API FNodeInputController : public FBaseInputController 
+		{
+		public:
+
+			struct FInitParams
+			{
+				int32 ID; 
+				FConstVertexAccessPtr NodeVertexPtr;
+				FConstClassInputAccessPtr OwningGraphClassInputPtr;
+				FGraphAccessPtr GraphPtr; 
+				FNodeHandle OwningNode;
+			};
+
+			/** Constructs the input controller.  */
+			FNodeInputController(const FInitParams& InParams);
+
+			bool IsValid() const override;
+
+			// Input metadata
+			const FText& GetDisplayName() const override;
+			const FText& GetTooltip() const override;
+
+		private:
+
+			FConstClassInputAccessPtr OwningGraphClassInputPtr;
+		};
+
+		/** FOutputNodeInputController represents the input vertex of an output 
+		 * node. 
+		 *
+		 * FOutputNodeInputController is largely to represent outputs exposed from
+		 * a graph. 
+		 */
+		class METASOUNDFRONTEND_API FOutputNodeInputController : public FBaseInputController 
+		{
+		public:
+			struct FInitParams
+			{
+				int32 ID; 
+				FConstVertexAccessPtr NodeVertexPtr;
+				FConstClassOutputAccessPtr OwningGraphClassOutputPtr;
+				FGraphAccessPtr GraphPtr; 
+				FNodeHandle OwningNode;
+			};
+
+			/** Constructs the input controller. */
+			FOutputNodeInputController(const FInitParams& InParams);
+
+			bool IsValid() const override;
+
+			// Input metadata
+			const FText& GetDisplayName() const override;
+			const FText& GetTooltip() const override;
+
+		private:
+
+			FConstClassOutputAccessPtr OwningGraphClassOutputPtr;
+		};
+
+		/** FBaseNodeController provides common functionality for multiple derived
+		 * node controllers.
+		 */
+		class METASOUNDFRONTEND_API FBaseNodeController : public INodeController
+		{
+		public:
+
+			struct FInitParams
+			{
+				FNodeAccessPtr NodePtr;
+				FConstClassAccessPtr ClassPtr;
+				FGraphHandle OwningGraph;
+			};
+
+			/** Construct a base node controller. */
+			FBaseNodeController(const FInitParams& InParams);
+
+			bool IsValid() const override;
+
+			// Owning graph info
+			int32 GetOwningGraphClassID() const override;
+			FGraphHandle GetOwningGraph() override;
+			FConstGraphHandle GetOwningGraph() const override;
+
+			// Info about this node.
+			int32 GetID() const override;
+			int32 GetClassID() const override;
+			FMetasoundFrontendVersionNumber GetClassVersionNumber() const override;
+			const FText& GetClassDescription() const override;
+
+			const FString& GetNodeName() const override;
+			EMetasoundFrontendClassType GetClassType() const override;
+
+			bool CanAddInput(const FString& InVertexName) const override;
+			FInputHandle AddInput(const FString& InVertexName, const FMetasoundFrontendLiteral* InDefault) override;
+			bool RemoveInput(int32 InPointID) override;
+
+			bool CanAddOutput(const FString& InVertexName) const override;
+			FInputHandle AddOutput(const FString& InVertexName, const FMetasoundFrontendLiteral* InDefault) override;
+			bool RemoveOutput(int32 InPointID) override;
+
+			const FString& GetClassName() const override;
+
+			/** Returns all node inputs. */
+			TArray<FInputHandle> GetInputs() override;
+
+			/** Returns all node inputs. */
+			TArray<FConstInputHandle> GetConstInputs() const override;
+
+			TArray<FInputHandle> GetInputsWithVertexName(const FString& InName) override;
+			TArray<FConstInputHandle> GetInputsWithVertexName(const FString& InName) const override;
+
+			/** Returns all node outputs. */
+			TArray<FOutputHandle> GetOutputs() override;
+
+			/** Returns all node outputs. */
+			TArray<FConstOutputHandle> GetConstOutputs() const override;
+
+			TArray<FOutputHandle> GetOutputsWithVertexName(const FString& InName) override;
+			TArray<FConstOutputHandle> GetOutputsWithVertexName(const FString& InName) const override;
+
+			/** Returns an input with the given id. 
+			 *
+			 * If the input does not exist, an invalid handle is returned. 
+			 */
+			FInputHandle GetInputWithID(int32 InPointID) override;
+
+			/** Returns an input with the given name. 
+			 *
+			 * If the input does not exist, an invalid handle is returned. 
+			 */
+			FConstInputHandle GetInputWithID(int32 InPointID) const override;
+
+			/** Returns an output with the given name. 
+			 *
+			 * If the output does not exist, an invalid handle is returned. 
+			 */
+			FOutputHandle GetOutputWithID(int32 InPointID) override;
+
+			/** Returns an output with the given name. 
+			 *
+			 * If the output does not exist, an invalid handle is returned. 
+			 */
+			FConstOutputHandle GetOutputWithID(int32 InPointID) const override;
+
+			FGraphHandle AsGraph() override;
+			FConstGraphHandle AsGraph() const override;
+
+		protected:
+
+			FNodeAccessPtr NodePtr;
+			FConstClassAccessPtr ClassPtr;
+			FGraphHandle OwningGraph;
+
+		private:
+
+			struct FInputControllerParams
+			{
+				int32 PointID;
+				FConstVertexAccessPtr NodeVertexPtr;
+				FConstClassInputAccessPtr ClassInputPtr;
+			};
+
+			struct FOutputControllerParams
+			{
+				int32 PointID;
+				FConstVertexAccessPtr NodeVertexPtr;
+				FConstClassOutputAccessPtr ClassOutputPtr;
+			};
+
+			TArray<FInputControllerParams> GetInputControllerParams() const;
+			TArray<FOutputControllerParams> GetOutputControllerParams() const;
+
+			TArray<FInputControllerParams> GetInputControllerParamsWithVertexName(const FString& InName) const;
+			TArray<FOutputControllerParams> GetOutputControllerParamsWithVertexName(const FString& InName) const;
+
+			bool FindInputControllerParamsWithID(int32 InPointID, FInputControllerParams& OutParams) const;
+			bool FindOutputControllerParamsWithID(int32 InPointID, FOutputControllerParams& OutParams) const;
+
+			const FMetasoundFrontendClassInput* FindClassInputWithName(const FString& InName) const;
+			const FMetasoundFrontendClassOutput* FindClassOutputWithName(const FString& InName) const;
+
+			const FMetasoundFrontendVertex* FindNodeInputWithName(const FString& InName) const;
+			const FMetasoundFrontendVertex* FindNodeOutputWithName(const FString& InName) const;
+
+
+
+			virtual FInputHandle CreateInputController(int32 InPointID, FConstVertexAccessPtr InNodeVertexPtr, FConstClassInputAccessPtr InClassInputPtr, FNodeHandle InOwningNode) const = 0;
+
+			virtual FOutputHandle CreateOutputController(int32 InPointID, FConstVertexAccessPtr InNodeVertexPtr, FConstClassOutputAccessPtr InClassOutputPtr, FNodeHandle InOwningNode) const = 0;
+		};
+
+		/** FNodeController represents a external or subgraph node. */
+		class METASOUNDFRONTEND_API FNodeController: public FBaseNodeController
+		{
+			// Private token only allows members or friends to call constructor.
+			enum EPrivateToken { Token };
+
+		public:
+			struct FInitParams
+			{
+				FNodeAccessPtr NodePtr;
+				FConstClassAccessPtr ClassPtr;
+				FGraphAccessPtr GraphPtr;
+				FGraphHandle OwningGraph;
+			};
+
+			// Constructor takes a private token so it can only be instantiated by
+			// using the static creation functions. This protects against some
+			// error conditions which would result in a zombie object. The creation
+			// methods can detect the error conditions and return an invalid controller
+			// on error
+			FNodeController(EPrivateToken InToken, const FInitParams& InParams);
+
+			/** Create a node handle for a external or subgraph node. 
+			 *
+			 * @return A Node handle. On error, an invalid node handle is returned. 
+			 */
+			static FNodeHandle CreateNodeHandle(const FInitParams& InParams);
+
+			/** Create a node handle for a external or subgraph node. 
+			 *
+			 * @return A Node handle. On error, an invalid node handle is returned. 
+			 */
+			static FConstNodeHandle CreateConstNodeHandle(const FInitParams& InParams);
+
+			virtual ~FNodeController() = default;
+
+			bool IsValid() const override;
+
+		private:
+
+			FInputHandle CreateInputController(int32 InPointID, FConstVertexAccessPtr InNodeVertexPtr, FConstClassInputAccessPtr InClassInputPtr, FNodeHandle InOwningNode) const override;
+
+			FOutputHandle CreateOutputController(int32 InPointID, FConstVertexAccessPtr InNodeVertexPtr, FConstClassOutputAccessPtr InClassOutputPtr, FNodeHandle InOwningNode) const override;
+
+			FGraphAccessPtr GraphPtr;
+
+		};
+
+		/** FOutputNodeController represents an output node. */
+		class METASOUNDFRONTEND_API FOutputNodeController: public FBaseNodeController
+		{
+			// Private token only allows members or friends to call constructor.
+			enum EPrivateToken { Token };
+
+		public:
+			struct FInitParams
+			{
+				FNodeAccessPtr NodePtr;
+				FConstClassAccessPtr ClassPtr;
+				FConstClassOutputAccessPtr OwningGraphClassOutputPtr; 
+				FGraphAccessPtr GraphPtr;
+				FGraphHandle OwningGraph;
+			};
+
+			// Constructor takes a private token so it can only be instantiated by
+			// using the static creation functions. This protects against some
+			// error conditions which would result in a zombie object. The creation
+			// methods can detect the error conditions and return an invalid controller
+			// on error
+			FOutputNodeController(EPrivateToken InToken, const FInitParams& InParams);
+
+			/** Create a node handle for an output node. 
+			 *
+			 * @return A Node handle. On error, an invalid node handle is returned. 
+			 */
+			static FNodeHandle CreateOutputNodeHandle(const FInitParams& InParams);
+
+			/** Create a node handle for an output node. 
+			 *
+			 * @return A Node handle. On error, an invalid node handle is returned. 
+			 */
+			static FConstNodeHandle CreateConstOutputNodeHandle(const FInitParams& InParams);
+
+			virtual ~FOutputNodeController() = default;
+
+			bool IsValid() const override;
+
+		private:
+			FInputHandle CreateInputController(int32 InPointID, FConstVertexAccessPtr InNodeVertexPtr, FConstClassInputAccessPtr InClassInputPtr, FNodeHandle InOwningNode) const override;
+
+			FOutputHandle CreateOutputController(int32 InPointID, FConstVertexAccessPtr InNodeVertexPtr, FConstClassOutputAccessPtr InClassOutputPtr, FNodeHandle InOwningNode) const override;
+
+			FGraphAccessPtr GraphPtr;
+			FConstClassOutputAccessPtr OwningGraphClassOutputPtr; 
+		};
+
+		/** FInputNodeController represents an input node. */
+		class METASOUNDFRONTEND_API FInputNodeController: public FBaseNodeController
+		{
+			// Private token only allows members or friends to call constructor.
+			enum EPrivateToken { Token };
+
+		public:
+			struct FInitParams
+			{
+				FNodeAccessPtr NodePtr;
+				FConstClassAccessPtr ClassPtr;
+				FConstClassInputAccessPtr OwningGraphClassInputPtr; 
+				FGraphAccessPtr GraphPtr;
+				FGraphHandle OwningGraph;
+			};
+
+			// Constructor takes a private token so it can only be instantiated by
+			// using the static creation functions. This protects against some
+			// error conditions which would result in a zombie object. The creation
+			// methods can detect the error conditions and return an invalid controller
+			// on error
+			FInputNodeController(EPrivateToken InToken, const FInitParams& InParams);
+
+			/** Create a node handle for an input node. 
+			 *
+			 * @return A Node handle. On error, an invalid node handle is returned. 
+			 */
+			static FNodeHandle CreateInputNodeHandle(const FInitParams& InParams);
+
+			/** Create a node handle for an input node. 
+			 *
+			 * @return A Node handle. On error, an invalid node handle is returned. 
+			 */
+			static FConstNodeHandle CreateConstInputNodeHandle(const FInitParams& InParams);
+
+			virtual ~FInputNodeController() = default;
+
+			bool IsValid() const override;
+
+
+		private:
+			FInputHandle CreateInputController(int32 InPointID, FConstVertexAccessPtr InNodeVertexPtr, FConstClassInputAccessPtr InClassInputPtr, FNodeHandle InOwningNode) const override;
+
+			FOutputHandle CreateOutputController(int32 InPointID, FConstVertexAccessPtr InNodeVertexPtr, FConstClassOutputAccessPtr InClassOutputPtr, FNodeHandle InOwningNode) const override;
+
+			FConstClassInputAccessPtr OwningGraphClassInputPtr;
+			FGraphAccessPtr GraphPtr;
+		};
+
+		/** FGraphController represents a Metasound graph class. */
+		class METASOUNDFRONTEND_API FGraphController : public IGraphController
+		{
+			// Util for setting template to literal based on template argument
+			template<typename ArgType>
+			bool SetDefaultInputToLiteralInternal(const FString& InInputName, int32 InPointID, ArgType InValue)
+			{
+				if (FMetasoundFrontendClassInput* Desc = FindInputDescriptionWithName(InInputName))
+				{
+					auto IsLiteralWithSamePointID = [&](const FMetasoundFrontendVertexLiteral& InVertexLiteral) 
+					{ 
+						return InVertexLiteral.PointID == InPointID; 
+					};
+
+					FMetasoundFrontendVertexLiteral* VertexLiteral = Desc->Defaults.FindByPredicate(IsLiteralWithSamePointID);
+					if (nullptr == VertexLiteral)
+					{
+						VertexLiteral = &Desc->Defaults.AddDefaulted_GetRef();
+						VertexLiteral->PointID = InPointID;
+					}
+
+					if (ensure(FMetasoundFrontendRegistryContainer::Get()->DoesDataTypeSupportLiteralType<ArgType>(Desc->TypeName)))
+					{
+						VertexLiteral->Value.Set(InValue);
+						return true;
+					}
+				}
+
+				return false;
+			}
+
+			// Private token only allows members or friends to call constructor.
+			enum EPrivateToken { Token };
+
+		public:
+
+			struct FInitParams
+			{
+				FGraphClassAccessPtr GraphClassPtr;
+				FDocumentHandle OwningDocument;
+			};
+
+			// Constructor takes a private token so it can only be instantiated by
+			// using the static creation functions. This protects against some
+			// error conditions which would result in a zombie object. The creation
+			// methods can detect the error conditions and return an invalid controller
+			// on error
+			FGraphController(EPrivateToken InToken, const FInitParams& InParams);
+
+			/** Create a graph handle. 
+			 *
+			 * @return A graph handle. On error, an invalid handle is returned. 
+			 */
+			static FGraphHandle CreateGraphHandle(const FInitParams& InParams);
+
+			/** Create a graph handle. 
+			 *
+			 * @return A graph handle. On error, an invalid handle is returned. 
+			 */
+			static FConstGraphHandle CreateConstGraphHandle(const FInitParams& InParams);
+
+			virtual ~FGraphController() = default;
+
+			bool IsValid() const override;
+
+			int32 GetClassID() const override;
+
+			int32 GetNewPointID() const override;
+
+			TArray<FString> GetInputVertexNames() const override;
+			TArray<FString> GetOutputVertexNames() const override;
+			FConstClassInputAccessPtr FindClassInputWithName(const FString& InName) const override;
+			FConstClassOutputAccessPtr FindClassOutputWithName(const FString& InName) const override;
+			TArray<int32> GetDefaultIDsForInputVertex(const FString& InInputName) const override;
+			TArray<int32> GetDefaultIDsForOutputVertex(const FString& InOutputName) const override;
+
+			TArray<FNodeHandle> GetNodes() override;
+			TArray<FConstNodeHandle> GetConstNodes() const override;
+
+			FConstNodeHandle GetNodeWithID(int32 InNodeID) const override;
+			FNodeHandle GetNodeWithID(int32 InNodeID) override;
+
+			TArray<FNodeHandle> GetInputNodes() override;
+			TArray<FConstNodeHandle> GetConstInputNodes() const override;
+
+			TArray<FNodeHandle> GetOutputNodes() override;
+			TArray<FConstNodeHandle> GetConstOutputNodes() const override;
+
+			bool ContainsInputNodeWithName(const FString& InName) const override;
+			bool ContainsOutputNodeWithName(const FString& InName) const override;
+
+			FConstNodeHandle GetInputNodeWithName(const FString& InName) const override;
+			FConstNodeHandle GetOutputNodeWithName(const FString& InName) const override;
+
+			FNodeHandle GetInputNodeWithName(const FString& InName) override;
+			FNodeHandle GetOutputNodeWithName(const FString& InName) override;
+
+			FNodeHandle AddInputVertex(const FMetasoundFrontendClassInput& InDescription) override;
+			bool RemoveInputVertex(const FString& InName) override;
+
+			FNodeHandle AddOutputVertex(const FMetasoundFrontendClassOutput& InDescription) override;
+			bool RemoveOutputVertex(const FString& InName) override;
+
+			// This can be used to determine what kind of property editor we should use for the data type of a given input.
+			// Will return Invalid if the input couldn't be found, or if the input doesn't support any kind of literals.
+			ELiteralType GetPreferredLiteralTypeForInputVertex(const FString& InInputName) const override;
+
+			// For inputs whose preferred literal type is UObject or UObjectArray, this can be used to determine the UObject corresponding to that input's datatype.
+			UClass* GetSupportedClassForInputVertex(const FString& InInputName) override;
+
+			// These can be used to set the default value for a given input on this graph.
+			// @returns false if the input name couldn't be found, or if the literal type was incompatible with the Data Type of this input.
+			bool SetDefaultInputToLiteral(const FString& InInputName, int32 InPointID, bool bInValue) override;
+			bool SetDefaultInputToLiteral(const FString& InInputName, int32 InPointID, int32 InValue) override;
+			bool SetDefaultInputToLiteral(const FString& InInputName, int32 InPointID, float InValue) override;
+			bool SetDefaultInputToLiteral(const FString& InInputName, int32 InPointID, const FString& InValue) override;
+			bool SetDefaultInputToLiteral(const FString& InInputName, int32 InPointID, UObject* InValue) override;
+			bool SetDefaultInputToLiteral(const FString& InInputName, int32 InPointID, const TArray<UObject*>& InValue) override;
+
+			// Set the display name for the input with the given name
+			void SetInputDisplayName(FString InName, const FText& InDisplayName) override;
+
+			// Set the display name for the output with the given name
+			void SetOutputDisplayName(FString InName, const FText& InDisplayName) override;
+
+			// This can be used to clear the current literal for a given input.
+			// @returns false if the input name couldn't be found.
+			bool ClearLiteralForInput(const FString& InInputName, int32 InPointID) override;
+
+			FNodeHandle AddNode(const FNodeClassInfo& InNodeClass) override;
+			FNodeHandle AddNode(const FNodeRegistryKey& InNodeClass) override;
+
+			// Remove the node corresponding to this node handle.
+			// On success, invalidates the received node handle.
+			bool RemoveNode(INodeController& InNode) override;
+
+			// Returns the metadata for the current graph, including the name, description and author.
+			const FMetasoundFrontendClassMetadata& GetGraphMetadata() const override;
+
+			bool InflateNodeDirectlyIntoGraph(const INodeController& InNode) override;
+
+			FNodeHandle CreateEmptySubgraph(const FMetasoundFrontendClassMetadata& InInfo) override;
+
+			TUniquePtr<IOperator> BuildOperator(const FOperatorSettings& InSettings, const FMetasoundEnvironment& InEnvironment, TArray<IOperatorBuilder::FBuildErrorPtr>& OutBuildErrors) const override;
+
+			FDocumentHandle GetOwningDocument() override;
+			FConstDocumentHandle GetOwningDocument() const override;
+
+		private:
+
+			// Add/remove nodes
+			//FNodeHandle AddNode(const FMetasoundFrontendClass& InExistingDependency);
+			FNodeHandle AddNode(FConstClassAccessPtr InExistingDependency);
+			bool RemoveNode(const FMetasoundFrontendNode& InDesc);
+
+			// Remove inputs
+			bool RemoveInput(const FMetasoundFrontendNode& InNode);
+			bool RemoveOutput(const FMetasoundFrontendNode& InNode);
+
+			int32 NewNodeID() const;
+
+			FNodeHandle GetNodeByPredicate(TFunctionRef<bool (const FMetasoundFrontendClass&, const FMetasoundFrontendNode&)> InPredicate);
+			FConstNodeHandle GetNodeByPredicate(TFunctionRef<bool (const FMetasoundFrontendClass&, const FMetasoundFrontendNode&)> InPredicate) const;
+
+			TArray<FNodeHandle> GetNodesByPredicate(TFunctionRef<bool (const FMetasoundFrontendClass&, const FMetasoundFrontendNode&)> InFilterFunc);
+			TArray<FConstNodeHandle> GetNodesByPredicate(TFunctionRef<bool (const FMetasoundFrontendClass&, const FMetasoundFrontendNode&)> InFilterFunc) const;
+
+
+			struct FNodeAndClass
+			{
+				FNodeAccessPtr Node;
+				FConstClassAccessPtr Class;
+
+				bool IsValid() const { return Node.IsValid() && Class.IsValid(); }
+			};
+
+			struct FConstNodeAndClass
+			{
+				FConstNodeAccessPtr Node;
+				FConstClassAccessPtr Class;
+
+				bool IsValid() const { return Node.IsValid() && Class.IsValid(); }
+			};
+
+			bool ContainsNodesAndClassesByPredicate(TFunctionRef<bool (const FMetasoundFrontendClass&, const FMetasoundFrontendNode&)> InPredicate) const;
+			TArray<FNodeAndClass> GetNodesAndClasses();
+			TArray<FConstNodeAndClass> GetNodesAndClasses() const;
+
+			TArray<FNodeAndClass> GetNodesAndClassesByPredicate(TFunctionRef<bool (const FMetasoundFrontendClass&, const FMetasoundFrontendNode&)> InPredicate);
+			TArray<FConstNodeAndClass> GetNodesAndClassesByPredicate(TFunctionRef<bool (const FMetasoundFrontendClass&, const FMetasoundFrontendNode&)> InPredicate) const;
+
+
+			TArray<FNodeHandle> GetNodeHandles(TArrayView<const FNodeAndClass> InNodesAndClasses);
+			TArray<FConstNodeHandle> GetNodeHandles(TArrayView<const FConstNodeAndClass> InNodesAndClasses) const;
+
+			FNodeHandle GetNodeHandle(const FNodeAndClass& InNodeAndClass);
+			FConstNodeHandle GetNodeHandle(const FConstNodeAndClass& InNodeAndClass) const;
+
+			FMetasoundFrontendClassInput* FindInputDescriptionWithName(const FString& InName);
+			const FMetasoundFrontendClassInput* FindInputDescriptionWithName(const FString& InName) const;
+
+			FMetasoundFrontendClassOutput* FindOutputDescriptionWithName(const FString& InName);
+			const FMetasoundFrontendClassOutput* FindOutputDescriptionWithName(const FString& InName) const;
+
+			FClassInputAccessPtr FindInputDescriptionWithNodeID(int32 InNodeID);
+			FConstClassInputAccessPtr FindInputDescriptionWithNodeID(int32 InNodeID) const;
+
+			FClassOutputAccessPtr FindOutputDescriptionWithNodeID(int32 InNodeID);
+			FConstClassOutputAccessPtr FindOutputDescriptionWithNodeID(int32 InNodeID) const;
+
+			/*
+			FMetasoundFrontendNode* FindNodeByID(int32 InNodeID);
+			const FMetasoundFrontendNode* FindNodeByID(int32 InNodeID) const;
+			*/
+
+			FMetasoundFrontendClassMetadata CreateInputClassMetadata(const FMetasoundFrontendClassInput& InClassInput);
+			FMetasoundFrontendClassMetadata CreateOutputClassMetadata(const FMetasoundFrontendClassOutput& InClassOutput);
+
+			FGraphClassAccessPtr GraphClassPtr;
+			FDocumentHandle OwningDocument; 
+		};
+
+		/** FDocumentController represents an entire Metasound document. */
+		class METASOUNDFRONTEND_API FDocumentController : public IDocumentController 
+		{
+			using FRegistry = FMetasoundFrontendRegistryContainer;
+
+		public:
+			/** Construct a FDocumentController.
+			 *
+			 * @param InDocument - Document to be manipulated.
+			 */
+			FDocumentController(FDocumentAccessPtr InDocumentPtr);
+
+			/** Create a FDocumentController.
+			 *
+			 * @param InDocument - Document to be manipulated.
+			 *
+			 * @return A document handle. 
+			 */
+			static FDocumentHandle CreateDocumentHandle(FDocumentAccessPtr InDocument)
+			{
+				return MakeShared<FDocumentController>(InDocument);
+			}
+
+			virtual ~FDocumentController() = default;
+
+			bool IsValid() const override;
+
+			bool IsRequiredInput(const FString& InInputName) const override;
+			bool IsRequiredOutput(const FString& InOutputName) const override;
+
+			TArray<FMetasoundFrontendClassVertex> GetRequiredInputs() const override;
+			TArray<FMetasoundFrontendClassVertex> GetRequiredOutputs() const override;
+
+			TArray<FMetasoundFrontendClass> GetDependencies() const override;
+			TArray<FMetasoundFrontendGraphClass> GetSubgraphs() const override;
+			TArray<FMetasoundFrontendClass> GetClasses() const override;
+
+			FConstClassAccessPtr FindDependencyWithID(int32 InClassID) const override;
+			FConstGraphClassAccessPtr FindSubgraphWithID(int32 InClassID) const override;
+			FConstClassAccessPtr FindClassWithID(int32 InClassID) const override;
+
+			FConstClassAccessPtr FindClass(const FNodeClassInfo& InNodeClass) const override;
+			FConstClassAccessPtr FindClass(const FMetasoundFrontendClassMetadata& InMetadata) const override;
+
+			FConstClassAccessPtr FindOrAddClass(const FNodeClassInfo& InNodeClass) override;
+			FConstClassAccessPtr FindOrAddClass(const FMetasoundFrontendClassMetadata& InMetadata) override;
+
+			void RemoveUnreferencedDependencies() override;
+
+			FGraphHandle GetRootGraph() override;
+			FConstGraphHandle GetRootGraph() const override;
+
+			/** Returns an array of all subgraphs for this document. */
+			TArray<FGraphHandle> GetSubgraphHandles() override;
+
+			/** Returns an array of all subgraphs for this document. */
+			TArray<FConstGraphHandle> GetSubgraphHandles() const override;
+
+			/** Returns a graphs in the document with the given class ID.*/
+			FGraphHandle GetSubgraphWithClassID(int32 InClassID) override;
+
+			/** Returns a graphs in the document with the given class ID.*/
+			FConstGraphHandle GetSubgraphWithClassID(int32 InClassID) const override;
+
+			bool ExportToJSONAsset(const FString& InAbsolutePath) const override;
+			FString ExportToJSON() const override;
+
+			static bool IsMatchingMetasoundClass(const FMetasoundFrontendClassMetadata& InMetadataA, const FMetasoundFrontendClassMetadata& InMetadataB);
+			static bool IsMatchingMetasoundClass(const FNodeClassInfo& InNodeClass, const FMetasoundFrontendClassMetadata& InMetadata);
+		private:
+
+			int32 NewClassID() const;
+
+			FDocumentAccessPtr DocumentPtr;
+		};
+	}
+}
+
+
diff --git a/Engine/Plugins/Runtime/Metasound/Source/MetasoundFrontend/Private/MetasoundFrontendTransform.cpp b/Engine/Plugins/Runtime/Metasound/Source/MetasoundFrontend/Private/MetasoundFrontendTransform.cpp
new file mode 100644
--- /dev/null
+++ b/Engine/Plugins/Runtime/Metasound/Source/MetasoundFrontend/Private/MetasoundFrontendTransform.cpp
@@ -0,0 +1,69 @@
+// Copyright Epic Games, Inc. All Rights Reserved.
+
+#include "MetasoundFrontendTransform.h"
+
+namespace Metasound
+{
+	namespace Frontend
+	{
+		bool FMatchRootGraphToArchetype::Transform(FDocumentHandle InDocument) const
+		{
+			if (!InDocument->IsValid())
+			{
+				return false;
+			}
+
+			bool bDidEdit = false;
+			FGraphHandle Graph = InDocument->GetRootGraph();
+
+			TArray<FMetasoundFrontendClassVertex> RequiredInputs = InDocument->GetRequiredInputs();
+			TArray<FMetasoundFrontendClassVertex> RequiredOutputs = InDocument->GetRequiredOutputs();
+
+			// Go through each input and add or swap if something is missing.
+			for (const FMetasoundFrontendClassVertex& RequiredInput : RequiredInputs)
+			{
+				const FMetasoundFrontendClassInput* ClassInput = Graph->FindClassInputWithName(RequiredInput.Name).Get();
+
+				bool bIsMatching = false;
+				if (nullptr != ClassInput)
+				{
+					bIsMatching = FMetasoundFrontendClassVertex::IsFunctionalEquivalent(RequiredInput, *ClassInput);
+				}
+
+				if (!bIsMatching)
+				{
+					bDidEdit = true;
+					if (nullptr != ClassInput)
+					{
+						Graph->RemoveInputVertex(RequiredInput.Name);
+					}
+					Graph->AddInputVertex(RequiredInput);
+				}
+			}
+
+			// Go through each output and add or swap if something is missing.
+			for (const FMetasoundFrontendClassVertex& RequiredOutput : RequiredOutputs)
+			{
+				const FMetasoundFrontendClassOutput* ClassOutput = Graph->FindClassOutputWithName(RequiredOutput.Name).Get();
+
+				bool bIsMatching = false;
+				if (nullptr != ClassOutput)
+				{
+					bIsMatching = FMetasoundFrontendClassVertex::IsFunctionalEquivalent(RequiredOutput, *ClassOutput);
+				}
+
+				if (!bIsMatching)
+				{
+					bDidEdit = true;
+					if (nullptr != ClassOutput)
+					{
+						Graph->RemoveOutputVertex(RequiredOutput.Name);
+					}
+					Graph->AddOutputVertex(RequiredOutput);
+				}
+			}
+
+			return bDidEdit;
+		}
+	}
+}
diff --git a/Engine/Plugins/Runtime/Metasound/Source/MetasoundFrontend/Private/MetasoundRouter.cpp b/Engine/Plugins/Runtime/Metasound/Source/MetasoundFrontend/Private/MetasoundRouter.cpp
--- a/Engine/Plugins/Runtime/Metasound/Source/MetasoundFrontend/Private/MetasoundRouter.cpp
+++ b/Engine/Plugins/Runtime/Metasound/Source/MetasoundFrontend/Private/MetasoundRouter.cpp
@@ -21,7 +21,7 @@ static FAutoConsoleCommand GPushFloatCommand(
 			FName ChannelName = FName(*Args[0]);
 			float ValueToPush = TCString<TCHAR>::Atof(*Args[1]);
 
-			Metasound::FDataTypeLiteralParam LiteralParam(ValueToPush);
+			Metasound::FLiteral LiteralParam(ValueToPush);
 
 			if (!Metasound::FDataTransmissionCenter::Get().PushLiteral(ChannelName, LiteralParam))
 			{
@@ -47,7 +47,7 @@ static FAutoConsoleCommand GPushBoolCommand(
 
 			bool ValueToPush = ValueAsInt != 0;
 
-			Metasound::FDataTypeLiteralParam LiteralParam(ValueToPush);
+			Metasound::FLiteral LiteralParam(ValueToPush);
 
 			if (!Metasound::FDataTransmissionCenter::Get().PushLiteral(ChannelName, LiteralParam))
 			{
@@ -71,7 +71,7 @@ static FAutoConsoleCommand GPushIntCommand(
 			FName ChannelName = FName(*Args[0]);
 			int32 ValueToPush = TCString<TCHAR>::Atoi(*Args[1]);
 
-			Metasound::FDataTypeLiteralParam LiteralParam(ValueToPush);
+			Metasound::FLiteral LiteralParam(ValueToPush);
 
 			if (!Metasound::FDataTransmissionCenter::Get().PushLiteral(ChannelName, LiteralParam))
 			{
@@ -94,7 +94,7 @@ static FAutoConsoleCommand GPushStringCommand(
 
 			FName ChannelName = FName(*Args[0]);
 
-			Metasound::FDataTypeLiteralParam LiteralParam(Args[1]);
+			Metasound::FLiteral LiteralParam(Args[1]);
 
 			if (!Metasound::FDataTransmissionCenter::Get().PushLiteral(ChannelName, LiteralParam))
 			{
diff --git a/Engine/Plugins/Runtime/Metasound/Source/MetasoundFrontend/Public/MetasoundAccessPtr.h b/Engine/Plugins/Runtime/Metasound/Source/MetasoundFrontend/Public/MetasoundAccessPtr.h
--- a/Engine/Plugins/Runtime/Metasound/Source/MetasoundFrontend/Public/MetasoundAccessPtr.h
+++ b/Engine/Plugins/Runtime/Metasound/Source/MetasoundFrontend/Public/MetasoundAccessPtr.h
@@ -3,11 +3,16 @@
 #pragma once
 
 #include "CoreMinimal.h"
+#include "Templates/UnrealTemplate.h"
 
+#if !(UE_BUILD_SHIPPING || UE_BUILD_TEST)
+#define METASOUND_FRONTEND_ACCESSPTR_DEBUG_INFO 1
+#endif
 namespace Metasound
 {
 	namespace Frontend
 	{
+
 		class FAccessPoint;
 
 		struct FAccessToken {};
@@ -17,8 +22,8 @@ namespace Metasound
 		 * TAccessPtr is used to determine whether an object has been destructed or not.
 		 * It is useful when an object cannot be wrapped in a TSharedPtr. A TAccessPtr
 		 * works similar to a TWeakPtr, but it cannot pin the object. Object pointers 
-		 * held within a TAccessPtr should only be accessed on a thread where the objects
-		 * cannot be destructed.
+		 * held within a TAccessPtr must only be accessed on the thread where the objects
+		 * gets destructed to avoid having the object destructed while in use.
 		 *
 		 * If the TAccessPtr's underlying object is accessed when the pointer is invalid,
 		 * a fallback object will be returned. 
@@ -28,19 +33,16 @@ namespace Metasound
 		template<typename Type>
 		class TAccessPtr
 		{
+			enum EConstCast { Tag };
+
 		public:
-			static Type FallbackObject;
+			using FTokenType = FAccessToken;
 
-			/** Create an invalid access pointer. */
-			static TAccessPtr<Type> CreateInvalid()
-			{
-				TSharedPtr<FAccessToken> InvalidToken;
-				return TAccessPtr<Type>(InvalidToken, FallbackObject);
-			}
+			static Type FallbackObject;
 
 			/** Returns true if object is valid. Should only be called in same
 			 * thread which created the token. */
-			bool IsValid() const { return Token.IsValid(); }
+			bool IsValid() const { return (Get() != nullptr); }
 
 			FORCEINLINE explicit operator bool() const
 			{
@@ -49,34 +51,92 @@ namespace Metasound
 
 			Type* Get() const
 			{
-				if (IsValid())
-				{
-					return Object;
-				}
-				return nullptr;
+#if METASOUND_FRONTEND_ACCESSPTR_DEBUG_INFO 
+				CachedObjectPtr = GetObject();
+#endif
+				return GetObject();
 			}
 
 			Type& operator*() const
 			{
-				if (ensure(IsValid()))
+				if (Type* Object = Get())
 				{
 					return *Object;
 				}
 
+				checkNoEntry();
 				return FallbackObject;
 			}
 
-			Type& operator->() const
+			Type* operator->() const
 			{
-				if (ensure(IsValid()))
+				if (Type* Object = Get())
 				{
-					return *Object;
+					return Object;
 				}
 
-				return FallbackObject;
+				return &FallbackObject;
+			}
+
+			template<typename MemberType>
+			TAccessPtr<MemberType> GetMemberAccessPtr(TFunction<MemberType*(Type&)> InGetMember) const
+			{
+				TFunction<MemberType*()> GetMemberFromObject = [=]() -> MemberType*
+				{
+					if (Type* Object = Get())
+					{
+						return InGetMember(*Object);
+					}
+					return static_cast<MemberType*>(nullptr);
+				};
+
+				return TAccessPtr<MemberType>(GetMemberFromObject);
 			}
 
-			TAccessPtr() = delete;
+
+			TAccessPtr()
+			//: Object(nullptr)
+			: GetObject([]() { return static_cast<Type*>(nullptr); })
+			{
+#if METASOUND_FRONTEND_ACCESSPTR_DEBUG_INFO 
+				Get();
+#endif
+			}
+
+
+			template<typename OtherType>
+			TAccessPtr(const TAccessPtr<OtherType>& InOther, EConstCast InTag)
+			//: Token(InOther.Token)
+			//, Object(const_cast<Type*>(InOther.Object))
+			{
+				TFunction<OtherType*()> OtherGetObject = InOther.GetObject;
+				GetObject = [=]() -> Type*
+				{
+					return const_cast<Type*>(OtherGetObject());
+				};
+#if METASOUND_FRONTEND_ACCESSPTR_DEBUG_INFO 
+				Get();
+#endif
+			} 
+
+			template <
+				typename OtherType,
+				typename = decltype(ImplicitConv<Type*>((OtherType*)nullptr))
+			>
+			TAccessPtr(const TAccessPtr<OtherType>& InOther)
+			//: Token(InOther.Token)
+			//, Object(InOther.Object)
+			{
+				TFunction<OtherType*()> OtherGetObject = InOther.GetObject;
+
+				GetObject = [=]() -> Type*
+				{
+					return static_cast<Type*>(OtherGetObject());
+				};
+#if METASOUND_FRONTEND_ACCESSPTR_DEBUG_INFO 
+				Get();
+#endif
+			}
 
 			TAccessPtr(const TAccessPtr<Type>& InOther) = default;
 			TAccessPtr<Type>& operator=(const TAccessPtr<Type>& InOther) = default;
@@ -88,13 +148,41 @@ namespace Metasound
 			template<typename RelatedType>
 			friend TAccessPtr<RelatedType> MakeAccessPtr(const FAccessPoint& InAccessPoint, RelatedType& InRef);
 
-			TWeakPtr<FAccessToken> Token;
-			Type* Object;
+			template<typename ToType, typename FromType> 
+			friend TAccessPtr<ToType> ConstCastAccessPtr(const TAccessPtr<FromType>& InAccessPtr);
+
+			template<typename OtherType>
+			friend class TAccessPtr;
 
-			TAccessPtr(TWeakPtr<FAccessToken> AccessToken, Type& InRef)
-			:	Token(AccessToken)
-			,	Object(&InRef)
+#if METASOUND_FRONTEND_ACCESSPTR_DEBUG_INFO 
+			mutable Type* CachedObjectPtr = nullptr;
+#endif
+			TFunction<Type*()> GetObject;
+
+			TAccessPtr(TFunction<Type*()> InGetObject)
+			: GetObject(InGetObject)
 			{
+#if METASOUND_FRONTEND_ACCESSPTR_DEBUG_INFO 
+				Get();
+#endif
+			}
+
+			TAccessPtr(TWeakPtr<FTokenType> AccessToken, Type& InRef)
+			{
+				Type* RefPtr = &InRef;
+
+				GetObject = [=]() -> Type*
+				{
+					Type* Object = nullptr;
+					if (AccessToken.IsValid())
+					{
+						Object = RefPtr;
+					}
+					return Object;
+				};
+#if METASOUND_FRONTEND_ACCESSPTR_DEBUG_INFO 
+				Get();
+#endif
 			}
 		};
 
@@ -107,22 +195,33 @@ namespace Metasound
 		class FAccessPoint
 		{
 		public:
+			using FTokenType = FAccessToken;
 
 			FAccessPoint() 
 			{
-				Token = MakeShared<FAccessToken>();
+				Token = MakeShared<FTokenType>();
+			}
+
+			FAccessPoint(const FAccessPoint&) 
+			{
+				// Do not copy token from other access point on copy.
+				Token = MakeShared<FTokenType>();
+			}
+
+			// Do not copy token from other access point on assignment
+			FAccessPoint& operator=(const FAccessPoint&) 
+			{
+				return *this;
 			}
 
 		private:
 			template<typename Type>
 			friend TAccessPtr<Type> MakeAccessPtr(const FAccessPoint& InAccessPoint, Type& InRef);
 
-			FAccessPoint(const FAccessPoint&) = delete;
-			FAccessPoint& operator=(const FAccessPoint&) = delete;
 			FAccessPoint(FAccessPoint&&) = delete;
 			FAccessPoint& operator=(FAccessPoint&&) = delete;
 
-			TSharedPtr<FAccessToken> Token;
+			TSharedPtr<FTokenType> Token;
 		};
 
 		template<typename Type>
@@ -130,5 +229,11 @@ namespace Metasound
 		{
 			return TAccessPtr<Type>(InAccessPoint.Token, InRef);
 		}
+
+		template<typename ToType, typename FromType> 
+		TAccessPtr<ToType> ConstCastAccessPtr(const TAccessPtr<FromType>& InAccessPtr)
+		{
+			return TAccessPtr<ToType>(InAccessPtr, TAccessPtr<ToType>::EConstCast::Tag);
+		}
 	}
 }
diff --git a/Engine/Plugins/Runtime/Metasound/Source/MetasoundFrontend/Public/MetasoundArchetype.h b/Engine/Plugins/Runtime/Metasound/Source/MetasoundFrontend/Public/MetasoundArchetype.h
--- a/Engine/Plugins/Runtime/Metasound/Source/MetasoundFrontend/Public/MetasoundArchetype.h
+++ b/Engine/Plugins/Runtime/Metasound/Source/MetasoundFrontend/Public/MetasoundArchetype.h
@@ -2,23 +2,23 @@
 #pragma once
 
 #include "CoreMinimal.h"
-#include "MetasoundFrontendDataLayout.h"
+#include "MetasoundFrontendDocument.h"
 
 namespace Metasound
 {
 	namespace Frontend
 	{
-		METASOUNDFRONTEND_API bool IsSubsetOfArchetype(const FMetasoundArchetype& InSubsetArchetype, const FMetasoundArchetype& InSupersetArchetype);
+		METASOUNDFRONTEND_API bool IsSubsetOfArchetype(const FMetasoundFrontendArchetype& InSubsetArchetype, const FMetasoundFrontendArchetype& InSupersetArchetype);
 
-		METASOUNDFRONTEND_API bool IsEquivalentArchetype(const FMetasoundArchetype& InInputArchetype, const FMetasoundArchetype& InTargetArchetype);
+		METASOUNDFRONTEND_API bool IsEquivalentArchetype(const FMetasoundFrontendArchetype& InInputArchetype, const FMetasoundFrontendArchetype& InTargetArchetype);
 
-		METASOUNDFRONTEND_API bool IsEqualArchetype(const FMetasoundArchetype& InInputArchetype, const FMetasoundArchetype& InTargetArchetype);
+		METASOUNDFRONTEND_API bool IsEqualArchetype(const FMetasoundFrontendArchetype& InInputArchetype, const FMetasoundFrontendArchetype& InTargetArchetype);
 
-		METASOUNDFRONTEND_API int32 InputOutputDifferenceCount(const FMetasoundArchetype& InArchetypeA, const FMetasoundArchetype& InArchetypeB);
+		METASOUNDFRONTEND_API int32 InputOutputDifferenceCount(const FMetasoundFrontendArchetype& InArchetypeA, const FMetasoundFrontendArchetype& InArchetypeB);
 
-		METASOUNDFRONTEND_API void GatherRequiredEnvironmentVariables(const FMetasoundDocument& InDocument, TArray<FMetasoundEnvironmentVariableDescription>& OutEnvironmentVariables);
+		METASOUNDFRONTEND_API void GatherRequiredEnvironmentVariables(const FMetasoundFrontendDocument& InDocument, TArray<FMetasoundFrontendEnvironmentVariable>& OutEnvironmentVariables);
 
-		METASOUNDFRONTEND_API const FMetasoundArchetype* FindMostSimilarArchetypeSupportingEnvironment(const FMetasoundDocument& InDocument, const TArray<FMetasoundArchetype>& InCandidateArchetypes);
+		METASOUNDFRONTEND_API const FMetasoundFrontendArchetype* FindMostSimilarArchetypeSupportingEnvironment(const FMetasoundFrontendDocument& InDocument, const TArray<FMetasoundFrontendArchetype>& InCandidateArchetypes);
 	}
 }
 
diff --git a/Engine/Plugins/Runtime/Metasound/Source/MetasoundFrontend/Public/MetasoundAssetBase.h b/Engine/Plugins/Runtime/Metasound/Source/MetasoundFrontend/Public/MetasoundAssetBase.h
--- a/Engine/Plugins/Runtime/Metasound/Source/MetasoundFrontend/Public/MetasoundAssetBase.h
+++ b/Engine/Plugins/Runtime/Metasound/Source/MetasoundFrontend/Public/MetasoundAssetBase.h
@@ -4,16 +4,16 @@
 #include "CoreMinimal.h"
 #include "MetasoundAccessPtr.h"
 #include "MetasoundFrontend.h"
-#include "MetasoundFrontendDataLayout.h"
 #include "MetasoundFrontendBaseClasses.h"
+#include "MetasoundFrontendController.h"
 #include "MetasoundGraph.h"
 #include "UObject/WeakObjectPtrTemplates.h"
 
 class UEdGraph;
 
 /** FMetasoundAssetBase is intended to be a mix-in subclass for UObjects which utilize
- * Metasound assets.  It provides consistent access to FMetasoundDocuments, control
- * over the FMetasoundArchetype of the FMetasoundDocument.  It also enables the UObject
+ * Metasound assets.  It provides consistent access to FMetasoundFrontendDocuments, control
+ * over the FMetasoundFrontendArchetype of the FMetasoundFrontendDocument.  It also enables the UObject
  * to be utilized by a host of other engine tools built to support Metasounds.
  *
  * Example Usage:
@@ -22,9 +22,7 @@ class UEdGraph;
  * {
  *
  *     	UPROPERTY()
- *     	FMetasoundDocument MetasoundDocument;
- *
- *	 	Metasound::Frontend::FAccessPoint MetasoundSourceAccessPoint;
+ *     	FMetasoundFrontendDocument MetasoundDocument;
  *
  * #if WITH_EDITORONLY_DATA
  * 		UPROPERTY()
@@ -47,7 +45,7 @@ public:
 
 	// Construct an FMetasoundAssetBase with a default Metasound Archetype.
 	// @param InDefaultArchetype - Default archetype for a Metasound Document in this class.
-	FMetasoundAssetBase(const FMetasoundArchetype& InDefaultArchetype);
+	FMetasoundAssetBase(const FMetasoundFrontendArchetype& InDefaultArchetype);
 
 	virtual ~FMetasoundAssetBase() = default;
 
@@ -68,28 +66,28 @@ public:
 #endif // WITH_EDITORONLY_DATA
 
 	// Sets/overwrites the root class metadata
-	virtual void SetMetadata(FMetasoundClassMetadata& InMetadata);
+	virtual void SetMetadata(FMetasoundFrontendClassMetadata& InMetadata);
 
 	// Returns  a description of the required inputs and outputs for this metasound UClass.
-	const FMetasoundArchetype& GetArchetype() const;
+	const FMetasoundFrontendArchetype& GetArchetype() const;
 
 	// Sets/overwrites the document archetype.
 	// @param InArchetype - The desired archetype .
 	//
 	// @return True on success. False if archetype is not supported by this object. 
-	bool SetArchetype(const FMetasoundArchetype& InArchetype);
+	bool SetArchetype(const FMetasoundFrontendArchetype& InArchetype);
 
 	// Returns an array of archetypes preferred for this class.
-	virtual const TArray<FMetasoundArchetype>& GetPreferredArchetypes() const = 0;
+	virtual const TArray<FMetasoundFrontendArchetype>& GetPreferredArchetypes() const = 0;
 
 	// Returns true if the archetype is supported by this object.
-	virtual bool IsArchetypeSupported(const FMetasoundArchetype& InArchetype) const;
+	virtual bool IsArchetypeSupported(const FMetasoundFrontendArchetype& InArchetype) const;
 
 	// Returns the preferred archetype for the given document.
-	virtual const FMetasoundArchetype& GetPreferredArchetype(const FMetasoundDocument& InDocument) const;
+	virtual const FMetasoundFrontendArchetype& GetPreferredArchetype(const FMetasoundFrontendDocument& InDocument) const;
 
 	// Returns the root class metadata
-	FMetasoundClassMetadata GetMetadata();
+	FMetasoundFrontendClassMetadata GetMetadata();
 
 	// Imports data from a JSON string directly
 	bool ImportFromJSON(const FString& InJSON);
@@ -97,48 +95,38 @@ public:
 	// Imports the asset from a JSON file at provided path
 	bool ImportFromJSONAsset(const FString& InAbsolutePath);
 
+	// Returns handle for the root metasound graph of this asset.
+	Metasound::Frontend::FDocumentHandle GetDocumentHandle();
+
 	// Returns handle for the root metasound graph of this asset.
 	Metasound::Frontend::FGraphHandle GetRootGraphHandle();
 
-	// Returns all handles for subgraphs referenced
-	TArray<Metasound::Frontend::FGraphHandle> GetAllSubgraphHandles();
 
 	// Overwrites the existing document. If the document's archetype is not supported,
 	// the FMetasoundAssetBase be while queried for a new one using `GetPreferredArchetype`. If `bForceUpdateArchetype`
 	// is true, `GetPreferredArchetype` will be used whether or not the provided document's archetype
 	// is supported. 
-	void SetDocument(const FMetasoundDocument& InDocument, bool bForceUpdateArchetype=false);
+	void SetDocument(const FMetasoundFrontendDocument& InDocument, bool bForceUpdateArchetype=false);
 
 	
-	FMetasoundDocument& GetDocumentChecked();
-	const FMetasoundDocument& GetDocumentChecked() const;
+	FMetasoundFrontendDocument& GetDocumentChecked();
+	const FMetasoundFrontendDocument& GetDocumentChecked() const;
 
 	// This must be called on UObject::PostLoad, as well as in this asset's UFactory, to fix up the root document based on the most recent version of the archetype.
 	void ConformDocumentToArchetype();
 
 protected:
-	// Returns private token allowing implementing asset class to use graph/node handle system
-	static Metasound::Frontend::FHandleInitParams::EPrivateToken GetPrivateToken()
-	{
-		return Metasound::Frontend::FHandleInitParams::PrivateToken;
-	}
 
 	// Returns an access pointer to the document.
-	virtual Metasound::Frontend::TAccessPtr<FMetasoundDocument> GetDocument() = 0;
+	virtual Metasound::Frontend::TAccessPtr<FMetasoundFrontendDocument> GetDocument() = 0;
 
 	// Returns an access pointer to the document.
-	virtual Metasound::Frontend::TAccessPtr<const FMetasoundDocument> GetDocument() const = 0;
+	virtual Metasound::Frontend::TAccessPtr<const FMetasoundFrontendDocument> GetDocument() const = 0;
 
 	// Returns the owning asset responsible for transactions applied to metasound
 	virtual UObject* GetOwningAsset() const = 0;
 
-	// Returns a access point that can be used to build a TDescriptionPtr
-	// for direct editing of the FMetasoundClassDescription tree.
-	// For advance use only, and requires knowledge of Metasound::Frontend::FDescPath syntax.
-	// For most use cases, use GetRootGraphHandle() instead.
-	Metasound::Frontend::FDescriptionAccessPoint GetGraphAccessPoint();
-
 private:
 
-	FMetasoundArchetype Archetype;
+	FMetasoundFrontendArchetype Archetype;
 };
diff --git a/Engine/Plugins/Runtime/Metasound/Source/MetasoundFrontend/Public/MetasoundDataTypeRegistrationMacro.h b/Engine/Plugins/Runtime/Metasound/Source/MetasoundFrontend/Public/MetasoundDataTypeRegistrationMacro.h
--- a/Engine/Plugins/Runtime/Metasound/Source/MetasoundFrontend/Public/MetasoundDataTypeRegistrationMacro.h
+++ b/Engine/Plugins/Runtime/Metasound/Source/MetasoundFrontend/Public/MetasoundDataTypeRegistrationMacro.h
@@ -116,7 +116,7 @@ namespace Metasound
 		AttemptToRegisterConverter<TDataType, FString>();
 	}
 
-	template<typename TDataType, ELiteralArgType PreferredArgType = ELiteralArgType::None, typename UClassToUse = UObject>
+	template<typename TDataType, ELiteralType PreferredArgType = ELiteralType::None, typename UClassToUse = UObject>
 	bool RegisterDataTypeWithFrontend()
 	{
 		// if we reenter this code (because DECLARE_METASOUND_DATA_REFERENCE_TYPES was called twice with the same type),
@@ -136,10 +136,26 @@ namespace Metasound
 			return TUniquePtr<INode>(new TInputNode<TDataType>(InParams.InNodeName, InParams.InVertexName, MoveTemp(InParams.InitParam)));
 		};
 
-		FCreateOutputNodeFunction OutputNodeConstructor = [](const FOutputNodeConstrutorParams& InParams) -> TUniquePtr<INode>
+		FCreateMetasoundFrontendClassFunction CreateFrontendInputClass = []() -> FMetasoundFrontendClass
+		{
+			// Create class info using prototype node
+			// TODO: register input nodes with static class info.
+  			static TInputNode<TDataType> Prototype(TEXT(""), TEXT(""), FLiteral());
+			return Metasound::Frontend::GenerateClassDescription(Prototype.GetMetadata(), EMetasoundFrontendClassType::Input);
+		};
+
+		FCreateOutputNodeFunction OutputNodeConstructor = [](const FOutputNodeConstructorParams& InParams) -> TUniquePtr<INode>
 		{
 			return TUniquePtr<INode>(new TOutputNode<TDataType>(InParams.InNodeName, InParams.InVertexName));
 		};
+
+		FCreateMetasoundFrontendClassFunction CreateFrontendOutputClass = []() -> FMetasoundFrontendClass
+		{
+			// Create class info using prototype node
+			// TODO: register input nodes with static class info.
+			static TOutputNode<TDataType> Prototype(TEXT(""), TEXT(""));
+			return Metasound::Frontend::GenerateClassDescription(Prototype.GetMetadata(), EMetasoundFrontendClassType::Output);
+		};
 		
 		// By default, this function should not be used, unless the preferred arg type is UObjectProxy or UObjectProxyArray, and UClassToUse should be specified.
 		FCreateAudioProxyFunction ProxyGenerator = [](UObject* InObject) -> Audio::IProxyDataPtr
@@ -176,7 +192,14 @@ namespace Metasound
 		static const FName DataTypeName = GetMetasoundDataTypeName<TDataType>();
 
 		// Pack all of our various constructor lambdas to a single struct.
-		FDataTypeConstructorCallbacks Callbacks = { MoveTemp(InputNodeConstructor), MoveTemp(OutputNodeConstructor), MoveTemp(ProxyGenerator) };
+		FDataTypeConstructorCallbacks Callbacks = 
+		{ 
+			MoveTemp(InputNodeConstructor), 
+			MoveTemp(CreateFrontendInputClass),
+			MoveTemp(OutputNodeConstructor), 
+			MoveTemp(CreateFrontendOutputClass),
+			MoveTemp(ProxyGenerator) 
+		};
 
 
 		FDataTypeRegistryInfo RegistryInfo;
@@ -217,11 +240,11 @@ namespace Metasound
 // This should be used to expose a datatype as a potential input or output for a metasound graph.
 // The first argument to the macro is the class to expose.
 // the second argument is the display name of that type in the Metasound editor.
-// Optionally, a Metasound::ELiteralArgType can be passed in to designate a preferred literal type-
-// For example, if Metasound::ELiteralArgType::Float is passed in, we will default to using a float parameter to create this datatype.
+// Optionally, a Metasound::ELiteralType can be passed in to designate a preferred literal type-
+// For example, if Metasound::ELiteralType::Float is passed in, we will default to using a float parameter to create this datatype.
 // If no argument is passed in, we will infer a literal type to use.
 // If 
-// Metasound::ELiteralArgType::Invalid can be used to enforce that we don't provide space for a literal, in which case you should have a default constructor or a constructor that takes [const FOperatorSettings&] implemented.
+// Metasound::ELiteralType::Invalid can be used to enforce that we don't provide space for a literal, in which case you should have a default constructor or a constructor that takes [const FOperatorSettings&] implemented.
 // If you pass in a preferred arg type, please make sure that the passed in datatype has a matching constructor, since we won't check this until runtime.
 
 #define CANNOT_REGISTER_METASOUND_DATA_TYPE_ASSERT_STRING(DataType) \
diff --git a/Engine/Plugins/Runtime/Metasound/Source/MetasoundFrontend/Public/MetasoundFrontend.h b/Engine/Plugins/Runtime/Metasound/Source/MetasoundFrontend/Public/MetasoundFrontend.h
--- a/Engine/Plugins/Runtime/Metasound/Source/MetasoundFrontend/Public/MetasoundFrontend.h
+++ b/Engine/Plugins/Runtime/Metasound/Source/MetasoundFrontend/Public/MetasoundFrontend.h
@@ -4,8 +4,8 @@
 
 #include "CoreMinimal.h"
 #include "MetasoundGraph.h"
-#include "MetasoundFrontendDataLayout.h"
 #include "MetasoundFrontendBaseClasses.h"
+#include "MetasoundFrontendDocument.h"
 #include "MetasoundBuilderInterface.h"
 #include "UObject/WeakObjectPtrTemplates.h"
 
@@ -14,28 +14,8 @@ class FMetasoundAssetBase;
 
 namespace Metasound
 {	
-	/**
-	 * These classes are used for editing Metasounds.
-	 * In general, the workflow for editing a Metasound graph will be:
-	 * 1) Load or create a metasound asset.
-	 * 2) Call UMetasound::GetGraphHandle() to get a handle to the graph for that Metasound asset.
-	 * 
-	 * All Metasound documents are saved as a FMetasoundClassDescription, which itself can own FMetasoundClassDescriptions in a tree-like hierarchy.
-	 * Typically the workflow for creating a Metasound subgraph will be
-	 * 1) Get a Metasound::Frontend::FGraphHandle (typically through the two steps described in the first paragraph)
-	 * 2) Build a FMetasoundClassMetadata struct with whatever name/author/description you want to give the subgraph,
-	 * 3) call Metasound::Frontend::FGraphHandle::CreateEmptySubgraphNode, which will return that subgraph as it's own FGraphHandle,
-	 *    as well as the FNodeHandle for that subgraph in the current graph.
-	 *
-	 * If you want to bounce that subgraph to it's own saved Metasound Asset, call FGraphHandle::SaveToNewMetasoundAsset on that subgraph's FGraphHandle.
-	 *
-	 * General note- these apis are NOT thread safe. 
-	 * Make sure that any FGraphHandle, FNodeHandle, FInputHandle and FOutputHandle that access similar data are called on the same thread.
-	 */
 	namespace Frontend
 	{
-		// Forward Declarations
-		class FInputHandle;
 
 		// Struct with the basics of a node class' information,
 		// used to look up that node from our node browser functions,
@@ -46,7 +26,7 @@ namespace Metasound
 			FString NodeName;
 
 			// The type for this node.
-			EMetasoundClassType NodeType;
+			EMetasoundFrontendClassType NodeType = EMetasoundFrontendClassType::Invalid;
 
 			// The lookup key used for the internal node registry.
 			FNodeRegistryKey LookupKey;
@@ -61,25 +41,25 @@ namespace Metasound
 		 *
 		 * @return Metadata for node.
 		 */
-		METASOUNDFRONTEND_API FMetasoundClassMetadata GenerateClassMetadata(const FNodeClassInfo& InInfo);
+		METASOUNDFRONTEND_API FMetasoundFrontendClassMetadata GenerateClassMetadata(const FNodeClassInfo& InInfo);
 
-		/** Generates a new FMetasoundClassDescription from Node Metadata 
+		/** Generates a new FMetasoundFrontendClass from Node Metadata 
 		 *
 		 * @param InNodeMetadata - Metadata describing an external node.
 		 *
 		 * @return Class description for external node.
 		 */
-		METASOUNDFRONTEND_API FMetasoundClassDescription GenerateClassDescription(const FNodeInfo& InNodeMetadata);
+		METASOUNDFRONTEND_API FMetasoundFrontendClass GenerateClassDescription(const FNodeInfo& InNodeMetadata, EMetasoundFrontendClassType ClassType=EMetasoundFrontendClassType::External);
 
-		/** Generates a new FMetasoundClassDescription from node lookup info.
+		/** Generates a new FMetasoundFrontendClass from node lookup info.
 		 *
 		 * @param InInfo - Class info for a already registered external node.
 		 *
 		 * @return Class description for external node.
 		 */
-		METASOUNDFRONTEND_API FMetasoundClassDescription GenerateClassDescription(const FNodeClassInfo& InInfo);
+		METASOUNDFRONTEND_API FMetasoundFrontendClass GenerateClassDescription(const FNodeClassInfo& InInfo);
 
-		/** Generates a new FMetasoundClassDescription from Node init data
+		/** Generates a new FMetasoundFrontendClass from Node init data
 		 *
 		 * @tparam NodeType - Type of node to instantiate.
 		 * @param InNodeInitData - Data used to call constructor of node.
@@ -87,7 +67,7 @@ namespace Metasound
 		 * @return Class description for external node.
 		 */
 		template<typename NodeType>
-		FMetasoundClassDescription GenerateClassDescription(const FNodeInitData& InNodeInitData)
+		FMetasoundFrontendClass GenerateClassDescription(const FNodeInitData& InNodeInitData)
 		{
 			TUniquePtr<INode> Node = MakeUnique<NodeType>(InNodeInitData);
 
@@ -96,17 +76,17 @@ namespace Metasound
 				return GenerateClassDescription(Node->GetMetadata());
 			}
 
-			return FMetasoundClassDescription();
+			return FMetasoundFrontendClass();
 		}
 
-		/** Generates a new FMetasoundClassDescription from a NodeType
+		/** Generates a new FMetasoundFrontendClass from a NodeType
 		 *
 		 * @tparam NodeType - Type of node.
 		 *
 		 * @return Class description for external node.
 		 */
 		template<typename NodeType>
-		FMetasoundClassDescription GenerateClassDescription()
+		FMetasoundFrontendClass GenerateClassDescription()
 		{
 			FNodeInitData InitData;
 			InitData.InstanceName = FString(TEXT("GenerateClassDescriptionForNode"));
@@ -114,8 +94,8 @@ namespace Metasound
 			return GenerateClassDescription<NodeType>(InitData);
 		}
 
-		// Convenience function to convert Metasound::ELiteralArgType to EMetasoundLiteralType, the UEnum used for metasound documents.
-		METASOUNDFRONTEND_API EMetasoundLiteralType GetMetasoundLiteralType(Metasound::ELiteralArgType InLiteralType);
+		// Convenience function to convert Metasound::ELiteralType to EMetasoundFrontendLiteralType, the UEnum used for metasound documents.
+		METASOUNDFRONTEND_API EMetasoundFrontendLiteralType GetMetasoundLiteralType(Metasound::ELiteralType InLiteralType);
 
 		template<typename DataType>
 		FName GetDataTypeName()
@@ -133,11 +113,11 @@ namespace Metasound
 
 		// Takes a JSON string and deserializes it into a Metasound document struct.
 		// @returns false if the file couldn't be found or parsed into a document.
-		METASOUNDFRONTEND_API bool ImportJSONToMetasound(const FString& InJSON, FMetasoundDocument& OutMetasoundDocument);
+		METASOUNDFRONTEND_API bool ImportJSONToMetasound(const FString& InJSON, FMetasoundFrontendDocument& OutMetasoundDocument);
 
 		// Opens a json document at the given absolute path and deserializes it into a Metasound document struct.
 		// @returns false if the file couldn't be found or parsed into a document.
-		METASOUNDFRONTEND_API bool ImportJSONAssetToMetasound(const FString& InPath, FMetasoundDocument& OutMetasoundDocument);
+		METASOUNDFRONTEND_API bool ImportJSONAssetToMetasound(const FString& InPath, FMetasoundFrontendDocument& OutMetasoundDocument);
 
 		// Struct that indicates whether an input and an output can be connected,
 		// and whether an intermediate node is necessary to connect the two.
@@ -150,7 +130,7 @@ namespace Metasound
 				YesWithConverterNode
 			};
 
-			EConnectable Connectable;
+			EConnectable Connectable = EConnectable::No;
 
 			// If Connectable is EConnectable::YesWithConverterNode,
 			// this will be a populated list of nodes we can use 
@@ -158,354 +138,5 @@ namespace Metasound
 			TArray<FConverterNodeInfo> PossibleConverterNodeClasses;
 		};
 
-		struct METASOUNDFRONTEND_API FHandleInitParams
-		{
-			FDescriptionAccessPoint InAccessPoint;
-
-			// Path in the document to the element we're getting a handle to.
-			// for the FGraphHandle, this is a path to the graph within the graph's Class description.
-			// for the FNodeHandle, FInputHandle and FOutputHandle, this is a path to the individual node within a Graph description.
-			const FDescPath& InPath;
-
-			// Name for the vertex to get a handle for.
-			//const FString& InName;
-
-			// ID for node or class.
-			int32 NodeID;
-			int32 DependencyID;
-
-			// The asset that owns the FMetasoundDocument this handle is associated with.
-			TWeakObjectPtr<UObject> InOwningAsset;
-
-		private:
-			// This is used to grant specific classes the ability to create handles.
-			enum EPrivateToken { Token };
-			static const EPrivateToken PrivateToken;
-
-			friend class FOutputHandle;
-			friend class FInputHandle;
-			friend class FNodeHandle;
-			friend class FGraphHandle;
-			friend class ::FMetasoundAssetBase;
-		};
-
-		class METASOUNDFRONTEND_API FOutputHandle : protected ITransactable
-		{
-		public:
-			FOutputHandle() = delete;
-
-			enum EInputNodeTag { InputNode };
-			enum EDefaultTag { Default };
-
-			explicit FOutputHandle(FHandleInitParams::EPrivateToken InToken, const FHandleInitParams& InParams, const FString& InOutputName, EDefaultTag InTag);
-
-			// Constructor used for the outgoing connection from an input node.
-			explicit FOutputHandle(FHandleInitParams::EPrivateToken InToken, const FHandleInitParams& InParams, const FString& InOutputName, EInputNodeTag InTag);
-
-			~FOutputHandle() = default;
-
-			static FOutputHandle InvalidHandle();
-
-			bool IsValid() const;
-
-			// @todo: with the current spec, we'd have to scan all nodes in the graph for this output to solve this.
-			// Is this worth it, or should we just require folks to use FInputHandle::IsConnected()?
-			//bool IsConnected() const;
-
-			FName GetOutputType() const;
-			FString GetOutputName() const;
-			FText GetOutputTooltip() const;
-			int32 GetOwningNodeID() const;
-
-			// @todo: with the current spec, we'd have to scan all nodes in the graph for this output to solve this.
-			// Is this worth it, or should we just require folks to use FInputHandle::GetCurrentlyConnectedOutput()?
-			//FInputHandle GetCurrentlyConnectedInput();
-
-			FConnectability CanConnectTo(const FInputHandle& InHandle) const;
-			bool Connect(FInputHandle& InHandle);
-			bool ConnectWithConverterNode(FInputHandle& InHandle, const FConverterNodeInfo& InNodeClassName);
-			bool Disconnect(FInputHandle& InHandle);
-
-		private:
-			TDescriptionPtr<FMetasoundNodeDescription> NodePtr;
-			TDescriptionPtr<FMetasoundClassDescription> NodeClass;
-
-			// owning output description ptr for the node's class.
-			// This can be none if the node that owns this output is an input node itself.
-			TDescriptionPtr<FMetasoundOutputDescription> OutputPtr;
-
-			// Optional description pointer used when this output connection is owned by an input node.
-			TDescriptionPtr<FMetasoundInputDescription> InputNodePtr;
-		};
-
-		class METASOUNDFRONTEND_API FInputHandle : protected ITransactable
-		{
-		public:
-			FInputHandle() = delete;
-			~FInputHandle() = default;
-
-			enum EOutputNodeTag { OutputNode };
-			enum EDefaultTag { Default };
-
-			// Sole constructor for FInputHandle. Can only be used by friend classes for FHandleInitParams.
-			explicit FInputHandle(FHandleInitParams::EPrivateToken PrivateToken, const FHandleInitParams& InParams, const FString& InputName, EDefaultTag InTag);
-
-			// constructor used exclusively by output nodes.
-			explicit FInputHandle(FHandleInitParams::EPrivateToken PrivateToken, const FHandleInitParams& InParams, const FString& InputName, EOutputNodeTag InTag);
-
-			static FInputHandle InvalidHandle();
-
-			bool IsValid() const;
-			bool IsConnected() const;
-
-			FName GetInputType() const;
-			FString GetInputName() const;
-			FText GetInputTooltip() const;
-			int32 GetOwningNodeID() const;
-
-			FOutputHandle GetCurrentlyConnectedOutput() const;
-
-			FConnectability CanConnectTo(const FOutputHandle& InHandle) const;
-			bool Connect(FOutputHandle& InHandle);
-			bool ConnectWithConverterNode(FOutputHandle& InHandle, const FConverterNodeInfo& InNodeClassName);
-			bool Disconnect(FOutputHandle& InHandle);
-			bool Disconnect();
-
-		private:
-			TDescriptionPtr<FMetasoundNodeDescription> NodePtr;
-
-			// owning node class for the node that owns this input connection.
-			// This can be none if the node that owns this output is an output node itself.
-			TDescriptionPtr<FMetasoundClassDescription> NodeClass;
-
-			// owning input description ptr for the node's class.
-			// This can be none if the node that owns this output is an output node itself.
-			TDescriptionPtr<FMetasoundInputDescription> InputPtr;
-
-			// Optional description pointer used when this input connection is owned by an output node.
-			TDescriptionPtr<FMetasoundOutputDescription> OutputNodePtr;
-
-			FString InputName;
-
-			const FMetasoundNodeConnectionDescription* GetConnectionDescription() const;
-			FMetasoundNodeConnectionDescription* GetConnectionDescription();
-		};
-
-		class FGraphHandle;
-
-		// Opaque handle to a single node on a graph.
-		// Can retrieve metadata about that node,
-		// get FNodeHandles for other nodes it is connected to, and validate and modify connections.
-		class METASOUNDFRONTEND_API FNodeHandle : protected ITransactable
-		{
-		public:
-			FNodeHandle() = delete;
-			~FNodeHandle() = default;
-
-			// Sole constructor for FNodeHandle. Can only be used by friends of FHandleInitParams.
-			explicit FNodeHandle(FHandleInitParams::EPrivateToken PrivateToken, const FHandleInitParams& InParams, EMetasoundClassType InNodeClassType);
-
-			static FNodeHandle InvalidHandle();
-
-			bool IsValid() const;
-
-			TArray<FInputHandle> GetAllInputs();
-			TArray<FOutputHandle> GetAllOutputs();
-
-			FInputHandle GetInputWithName(const FString& InName);
-			FOutputHandle GetOutputWithName(const FString& InName);
-
-			FNodeClassInfo GetClassInfo() const;
-
-			EMetasoundClassType GetNodeType() const;
-			const FString& GetNodeClassName() const;
-
-			// If this node is itself a Metasound,
-			// use this to return the contained graph for the metasound.
-			// Otherwise it will return an invalid FGraphHandle.
-			void GetContainedGraph(FGraphHandle& OutGraph);
-
-			int32 GetNodeID() const;
-			static int32 GetNodeID(const FDescPath& InNodePath);
-
-			const FString& GetNodeName() const;
-
-		private:
-
-			static TDescriptionPtr<FMetasoundClassDescription> GetNodeClassDescriptionForNodeHandle(const FHandleInitParams& InitParams, EMetasoundClassType InNodeClassType);
-
-			TDescriptionPtr<FMetasoundNodeDescription> NodePtr;
-			TDescriptionPtr<FMetasoundClassDescription> NodeClass;
-
-			// whether this node is an input or output to it's owning graph,
-			// an externally implemented node, or a metasound graph itself.
-			EMetasoundClassType NodeClassType;
-		};
-
-		class METASOUNDFRONTEND_API FGraphHandle : protected ITransactable
-		{
-		public:
-			FGraphHandle() = delete;
-			~FGraphHandle() = default;
-
-			// Sole constructor for FGraphHandle. Can only be used by friends of FHandleInitParams.
-			FGraphHandle(FHandleInitParams::EPrivateToken PrivateToken, const FHandleInitParams& InParams);
-
-			static FGraphHandle GetHandle(UObject* InOwner, FMetasoundDocument& InRootMetasoundDocument, const Metasound::Frontend::FDescriptionAccessPoint& InAccessPoint)
-			{
-				using namespace Path;
-
-				FDescPath PathToGraph = FDescPath()[EFromDocument::ToRootClass][EFromClass::ToGraph];
-				FHandleInitParams InitParams = { InAccessPoint, PathToGraph, INDEX_NONE, InRootMetasoundDocument.RootClass.UniqueID, MakeWeakObjectPtr(InOwner) };
-				return FGraphHandle(FHandleInitParams::PrivateToken, InitParams);
-			}
-
-			static FGraphHandle InvalidHandle();
-
-			void CopyGraph(FGraphHandle& InOther);
-
-			bool IsValid() const;
-
-			bool IsRequiredInput(const FString& InOutputName) const;
-			bool IsRequiredOutput(const FString& InOutputName) const;
-
-			const TArray<FMetasoundInputDescription>& GetRequiredInputs() const;
-			const TArray<FMetasoundOutputDescription>& GetRequiredOutputs() const;
-
-			const FMetasoundEditorData& GetEditorData() const;
-			void SetEditorData(const FMetasoundEditorData& InEditorData);
-
-			TArray<FNodeHandle> GetAllNodes();
-			FNodeHandle GetNodeWithId(int32 InNodeId) const;
-			TArray<FNodeHandle> GetOutputNodes();
-			TArray<FNodeHandle> GetInputNodes();
-
-			bool ContainsOutputNodeWithName(const FString& InName) const;
-			bool ContainsInputNodeWithName(const FString& InName) const;
-
-			FNodeHandle GetOutputNodeWithName(const FString& InName);
-			FNodeHandle GetInputNodeWithName(const FString& InName);
-
-			FNodeHandle AddNewInput(const FMetasoundInputDescription& InDescription);
-			bool RemoveInput(const FString& InputName);
-
-			FNodeHandle AddNewOutput(const FMetasoundOutputDescription& InDescription);
-			bool RemoveOutput(const FString& OutputName);
-
-			// This can be used to determine what kind of property editor we should use for the data type of a given input.
-			// Will return Invalid if the input couldn't be found, or if the input doesn't support any kind of literals.
-			ELiteralArgType GetPreferredLiteralTypeForInput(const FString& InInputName);
-
-			// For inputs whose preferred literal type is UObject or UObjectArray, this can be used to determine the UObject corresponding to that input's datatype.
-			UClass* GetSupportedClassForInput(const FString& InInputName);
-
-
-			// These can be used to set the default value for a given input on this graph.
-			// @returns false if the input name couldn't be found, or if the literal type was incompatible with the Data Type of this input.
-			bool SetInputToLiteral(const FString& InInputName, bool bInValue);
-			bool SetInputToLiteral(const FString& InInputName, int32 InValue);
-			bool SetInputToLiteral(const FString& InInputName, float InValue);
-			bool SetInputToLiteral(const FString& InInputName, const FString& InValue);
-			bool SetInputToLiteral(const FString& InInputName, UObject* InValue);
-			bool SetInputToLiteral(const FString& InInputName, TArray<UObject*> InValue);
-
-			// Returns the display name for the input with the given name
-			const FText& GetInputDisplayName(FString InputName) const;
-
-			// Returns the display name for the output with the given name
-			const FText& GetOutputDisplayName(FString OutputName) const;
-
-			// Returns the tool tip for the input with the given name
-			const FText& GetInputToolTip(FString InputName) const;
-
-			// Returns the tool tip for the output with the given name
-			const FText& GetOutputToolTip(FString OutputName) const;
-
-			// Set the display name for the input with the given name
-			void SetInputDisplayName(FString InName, const FText& InDisplayName);
-
-			// Set the display name for the output with the given name
-			void SetOutputDisplayName(FString InName, const FText& InDisplayName);
-
-			// This can be used to clear the current literal for a given input.
-			// @returns false if the input name couldn't be found.
-			bool ClearLiteralForInput(const FString& InInputName);
-
-			FNodeHandle AddNewNode(const FNodeClassInfo& InNodeClass);
-			FNodeHandle AddNewNode(const FNodeRegistryKey& InNodeClass);
-
-			// Remove the node corresponding to this node handle.
-			// On success, invalidates the received node handle.
-			bool RemoveNode(const FNodeHandle& InNode);
-
-			// Returns the metadata for the current graph, including the name, description and author.
-			FMetasoundClassMetadata GetGraphMetadata();
-
-			// Exports this graph to a JSON string.
-			// @returns JSON string.
-			FString ExportToJSON() const;
-
-			// Exports this graph to a JSON at the given path.
-			// @returns true on success.
-			bool ExportToJSONAsset(const FString& InAbsolutePath) const;
-
-			// If the FNodeHandle given is itself a Metasound graph,
-			// and the FNodeHandle is a direct member of this FGraphHandle,
-			// If successful, this will invalidate the FNodeHandle and paste the graph for this node directly
-			// into the graph.
-			// If not successful, InNode will not be affected.
-			// @returns true on success, false on failure.
-			bool InflateNodeDirectlyIntoGraph(FNodeHandle& InNode);
-
-			// When called, creates a new node in this graph that can be edited directly with the FGraphHandle
-			// that gets returned.
-			// @param InInfo the name, author, etc that should be used for this graph.
-			// @returns a tuple of the FGraphHandle for this subgraph (which can be independently manipulated),
-			//          as well as the FNodeHandle for the subgraph's node in this graph.
-			//          Both of these will be invalid on failure.
-			TTuple<FGraphHandle, FNodeHandle> CreateEmptySubgraphNode(const FMetasoundClassMetadata& InInfo);
-
-			// This invokes the Metasound Builder to synchronously compile a 
-			// metasound operator, which can then be used for playback.
-			// @returns nullptr on failure.
-			// @todo: make an API for doing this async.
-			TUniquePtr<IOperator> BuildOperator(const FOperatorSettings& InSettings, const FMetasoundEnvironment& InEnvironment, TArray<IOperatorBuilder::FBuildErrorPtr>& OutBuildErrors) const;
-
-			// This function will ensure that the document's root class has the inputs and outputs required by the archetype.
-			void FixDocumentToMatchArchetype();
-
-		private:
-
-			// The graph struct itself.
-			TDescriptionPtr<FMetasoundGraphDescription> GraphPtr;
-
-			// The class description for the graph corresponding to this.
-			TDescriptionPtr<FMetasoundClassDescription> GraphsClassDeclaration;
-
-			// the outermost document for all dependencies for the document this graph lives in.
-			TDescriptionPtr<FMetasoundDocument> OwningDocument;
-
-			// Remove the node corresponding to this node handle.
-			// On success, invalidates the received node handle.
-			// Can remove inputs and outputs, but does not remove
-			// from the input/output arrays.
-			bool RemoveNodeInternal(const FNodeHandle& InNode);
-
-			// Scans all existing node ids to guarantee a new unique ID.
-			int32 FindNewUniqueNodeId();
-
-			// Scans all existing dependency IDs in the root dependency list to guarantee a new unique ID when adding a dependency.
-			int32 FindNewUniqueDependencyId();
-
-			FMetasoundLiteralDescription* GetLiteralDescriptionForInput(const FString& InInputName, FName& OutDataType) const;
-
-			bool GetDataTypeForInput(const FString& InInputName, FName& OutDataType);
-		};
-
-		// If there's a Metasound node referenced by a class in it's dependencies list,
-		// but the node' graph isn't implemented in that dependency list,
-		// this can be used to attempt to find which asset that node is implemented in,
-		// and return a graph handle for that node.
-		FGraphHandle GetGraphHandleForClass(const FMetasoundClassDescription& InClass);
 	} // namespace Frontend
 } // namespace Metasound
