Metasound Frontend Revamp

diff --git a/Engine/Plugins/Runtime/Metasound/Source/MetasoundFrontend/Public/MetasoundFrontendBaseClasses.h b/Engine/Plugins/Runtime/Metasound/Source/MetasoundFrontend/Public/MetasoundFrontendBaseClasses.h
--- a/Engine/Plugins/Runtime/Metasound/Source/MetasoundFrontend/Public/MetasoundFrontendBaseClasses.h
+++ b/Engine/Plugins/Runtime/Metasound/Source/MetasoundFrontend/Public/MetasoundFrontendBaseClasses.h
@@ -5,7 +5,7 @@
 #include "CoreMinimal.h"
 #include "Misc/TVariant.h"
 #include "MetasoundAccessPtr.h"
-#include "MetasoundFrontendDataLayout.h"
+#include "MetasoundFrontendDocument.h"
 #include "MetasoundFrontendRegistries.h"
 
 namespace Metasound
@@ -21,38 +21,29 @@ namespace Metasound
 		// Convenience functions to create an INode corresponding to a specific input or output for a metasound graph.
 		// @returns nullptr if the type given wasn't found.
 		METASOUNDFRONTEND_API TUniquePtr<INode> ConstructInputNode(const FName& InInputType, FInputNodeConstructorParams&& InParams);
-		METASOUNDFRONTEND_API TUniquePtr<INode> ConstructOutputNode(const FName& InOutputType, const FOutputNodeConstrutorParams& InParams);
+		METASOUNDFRONTEND_API TUniquePtr<INode> ConstructOutputNode(const FName& InOutputType, const FOutputNodeConstructorParams& InParams);
 
 		// Convenience functions to create an INodeB corresponding to a specific externally declared node type.
 		// InNodeType and InNodeHash can be retrieved from the FNodeClassInfo generated from the node registry queries in the metasound frontend (GetAllAvailableNodeClasses, GetAllNodeClassesInNamespace, etc.)
 		// @returns nullptr if the type given wasn't found.
 		METASOUNDFRONTEND_API TUniquePtr<INode> ConstructExternalNode(const FName& InNodeType, uint32 InNodeHash, const FNodeInitData& InInitData);
 
-		// Utility functions for setting serialized literals.
-		METASOUNDFRONTEND_API void SetLiteralDescription(FMetasoundLiteralDescription& OutDescription, bool InValue);
-		METASOUNDFRONTEND_API void SetLiteralDescription(FMetasoundLiteralDescription& OutDescription, int32 InValue);
-		METASOUNDFRONTEND_API void SetLiteralDescription(FMetasoundLiteralDescription& OutDescription, float InValue);
-		METASOUNDFRONTEND_API void SetLiteralDescription(FMetasoundLiteralDescription& OutDescription, const FString& InValue);
-		METASOUNDFRONTEND_API void SetLiteralDescription(FMetasoundLiteralDescription& OutDescription, UObject* InValue);
-		METASOUNDFRONTEND_API void SetLiteralDescription(FMetasoundLiteralDescription& OutDescription, const TArray<UObject*>& InValue);
-		
-		METASOUNDFRONTEND_API void ClearLiteralDescription(FMetasoundLiteralDescription& OutDescription);
 
 		// Utility functions for building a ::Metasound::FDataInitParam corresponding to a literal.
 
 		// Return the literal description parsed into a init param. 
 		// @Returns an invalid init param if the data type couldn't be found, or if the literal type was incompatible with the data type.
-		METASOUNDFRONTEND_API FDataTypeLiteralParam GetLiteralParamForDataType(FName InDataType, const FMetasoundLiteralDescription& InDescription);
+		METASOUNDFRONTEND_API FLiteral GetLiteralParamForDataType(FName InDataType, const FMetasoundFrontendLiteral& InDescription);
 
-		METASOUNDFRONTEND_API bool DoesDataTypeSupportLiteralType(FName InDataType, EMetasoundLiteralType InLiteralType);
-		METASOUNDFRONTEND_API bool DoesDataTypeSupportLiteralType(FName InDataType, ELiteralArgType InLiteralType);
+		METASOUNDFRONTEND_API bool DoesDataTypeSupportLiteralType(FName InDataType, EMetasoundFrontendLiteralType InLiteralType);
+		METASOUNDFRONTEND_API bool DoesDataTypeSupportLiteralType(FName InDataType, ELiteralType InLiteralType);
 		
 		// Returns a literal param without any type checking.
-		METASOUNDFRONTEND_API FDataTypeLiteralParam GetLiteralParam(const FMetasoundLiteralDescription& InDescription);
+		METASOUNDFRONTEND_API FLiteral GetLiteralParam(const FMetasoundFrontendLiteral& InDescription);
 
 		// Returns the defaulted version of a literal param for the given data type.
 		// @Returns an invalid init param if the data type couldn't be found.
-		METASOUNDFRONTEND_API FDataTypeLiteralParam GetDefaultParamForDataType(FName InDataType);
+		METASOUNDFRONTEND_API FLiteral GetDefaultParamForDataType(FName InDataType);
 
 		// Utility base class for classes that want to support naive multi-level undo/redo.
 		class METASOUNDFRONTEND_API ITransactable : public TSharedFromThis<ITransactable>
@@ -139,566 +130,5 @@ namespace Metasound
 			// If this transactable is operating on a UAsset, this can be filled to mark it's corresponding package as dirty.
 			TWeakObjectPtr<UObject> OwningAsset;
 		};
-
-		namespace Path
-		{
-			// This series of enums is used to denote 
-			// the path to a specific element in the class description.
-
-			// This enum is every possible Description struct we have.
-			enum class EDescType : uint8
-			{
-				Document, // FMetasoundDocument
-				Class, // FMetasoundClassDescription
-				Metadata, // FMetasoundClassMetadata
-				DocDependencies, // TArray<FMetasoundClassDescription>
-				ClassDependencies,
-				Graph, // 
-				Nodes,
-				Node,
-				Inputs,
-				Input,
-				Outputs,
-				Output,
-				Invalid
-			};
-
-
-			// This series of enums is used to delineate which description types are direct children or other descripiton types.
-			enum class EFromDocument : uint32
-			{
-				ToRootClass,
-				ToDependencies
-			};
-
-			
-			enum class EFromClass : uint32
-			{
-				ToMetadata,
-				ToInputs,
-				ToOutputs,
-				ToDependencies,
-				ToGraph,
-				ToInvalid
-			};
-
-			enum class EFromGraph : uint32
-			{
-				ToNodes,
-				Invalid
-			};
-
-			struct FElement
-			{
-				EDescType CurrentDescType;
-				
-				// This is only used if EDescType is Inputs or Outputs
-				FString LookupName;
-
-				// This is only used if EDescType == Node or Dependencies
-				int32 LookupID;
-
-				FElement()
-					: CurrentDescType(EDescType::Invalid)
-					, LookupID(FMetasoundNodeDescription::InvalidID)
-				{}
-			};
-		}
-
-		// This struct represents a selector for a specific element in a FMetasoundClassDescription.
-		// Typically these will be used with FDescriptionAccessPoint::GetElementAtPath.
-		struct FDescPath
-		{
-			TArray<Path::FElement> Path;
-
-			FDescPath()
-			{
-				Path::FElement RootElement;
-				RootElement.CurrentDescType = Path::EDescType::Document;
-				Path.Add(RootElement);
-			}
-
-			int32 Num() const
-			{
-				return Path.Num();
-			}
-
-			bool IsValid() const
-			{
-				return Path.Num() > 0 &&
-					Path[0].CurrentDescType == Path::EDescType::Document;
-			}
-
-			/** 
-			 * Bracket operator used for setting up a path into a metasound class description.
-			 * These can be chained, for example:
-			 * FDescPath PathToNode = RootClass[Path::EFromClass::ToGraph][Path::EFromGraph::ToNodes][12];
-			*/
-			FDescPath operator[](Path::EFromDocument InElement) const
-			{
-				if (!ensureAlwaysMsgf(Path.Num() > 0, TEXT("Tried to append to a path that had no root.")))
-				{
-					return FDescPath();
-				}
-
-				// Copy the old path, and based on the last element in that path, append a new child element.
-				FDescPath NewPath = *this;
-
-				const Path::FElement& LastElementInPath = Path.Last();
-
-				if(!ensureAlwaysMsgf(LastElementInPath.CurrentDescType == Path::EDescType::Document, TEXT("Tried to build an invalid path.")))
-				{
-					return FDescPath();
-				}
-
-				Path::FElement NewElement = Path::FElement();
-
-				switch (InElement)
-				{
-					case Path::EFromDocument::ToRootClass:
-					{
-						NewElement.CurrentDescType = Path::EDescType::Class;
-						break;
-					}
-					case Path::EFromDocument::ToDependencies:
-					{
-						NewElement.CurrentDescType = Path::EDescType::DocDependencies;
-						break;
-					}
-				}
-
-				NewPath.Path.Add(NewElement);
-				return NewPath;
-			}
-
-			FDescPath operator[](Path::EFromClass InElement) const
-			{
-				if (!ensureAlwaysMsgf(Path.Num() > 0, TEXT("Tried to append to a path that had no root.")))
-				{
-					return FDescPath();
-				}
-
-				// Copy the old path, and based on the last element in that path, append a new child element.
-				FDescPath NewPath = *this;
-
-				const Path::FElement& LastElementInPath = Path.Last();
-
-				if (!ensureAlwaysMsgf(LastElementInPath.CurrentDescType == Path::EDescType::Class, TEXT("Tried to build an invalid path.")))
-				{
-					return FDescPath();
-				}
-
-				Path::FElement NewElement = Path::FElement();
-
-				switch (InElement)
-				{
-					case Path::EFromClass::ToDependencies:
-					{
-						NewElement.CurrentDescType = Path::EDescType::ClassDependencies;
-						break;
-					}
-					case Path::EFromClass::ToGraph:
-					{
-						NewElement.CurrentDescType = Path::EDescType::Graph;
-						break;
-					}
-					case Path::EFromClass::ToInputs:
-					{
-						NewElement.CurrentDescType = Path::EDescType::Inputs;
-						break;
-					}
-					case Path::EFromClass::ToOutputs:
-					{
-						NewElement.CurrentDescType = Path::EDescType::Outputs;
-						break;
-					}
-					case Path::EFromClass::ToInvalid:
-					default:
-					{
-						ensureAlwaysMsgf(false, TEXT("Invalid path set up. Returning an empty path."));
-						return FDescPath();
-						break;
-					}
-				}
-
-				NewPath.Path.Add(NewElement);
-				return NewPath;
-			}
-
-			FDescPath operator[](Path::EFromGraph InElement) const
-			{
-				if (!ensureAlwaysMsgf(Path.Num() > 0, TEXT("Tried to append to a path that had no root.")))
-				{
-					return FDescPath();
-				}
-
-				// Copy the old path, and based on the last element in that path, append a new child element.
-				FDescPath NewPath = *this;
-
-				const Path::FElement& LastElementInPath = Path.Last();
-
-				if (!ensureAlwaysMsgf(LastElementInPath.CurrentDescType == Path::EDescType::Graph, TEXT("Tried to build an invalid path.")))
-				{
-					return FDescPath();
-				}
-
-				Path::FElement NewElement = Path::FElement();
-
-				switch (InElement)
-				{
-					case Path::EFromGraph::ToNodes:
-					{
-						NewElement.CurrentDescType = Path::EDescType::Nodes;
-						break;
-					}
-					case Path::EFromGraph::Invalid:
-					default:
-					{
-						ensureAlwaysMsgf(false, TEXT("Invalid path set up. Returning an empty path."));
-						return FDescPath();
-						break;
-					}
-				}
-
-				NewPath.Path.Add(NewElement);
-				return NewPath;
-			}
-
-			FDescPath operator[](int32 InElement) const
-			{
-				if (!ensureAlwaysMsgf(Path.Num() > 0, TEXT("Tried to append to a path that had no root.")))
-				{
-					return FDescPath();
-				}
-
-				// Copy the old path, and based on the last element in that path, append a new child element.
-				FDescPath NewPath = *this;
-
-				const Path::FElement& LastElementInPath = Path.Last();
-
-				if (!ensureAlwaysMsgf(LastElementInPath.CurrentDescType == Path::EDescType::Nodes || LastElementInPath.CurrentDescType == Path::EDescType::DocDependencies, TEXT("Tried to build an invalid path.")))
-				{
-					return FDescPath();
-				}
-
-				Path::FElement NewElement = Path::FElement();
-				NewElement.LookupID = InElement;
-
-				if (LastElementInPath.CurrentDescType == Path::EDescType::Nodes)
-				{
-					NewElement.CurrentDescType = Path::EDescType::Node;
-				}
-				else if (LastElementInPath.CurrentDescType == Path::EDescType::ClassDependencies)
-				{
-					NewElement.CurrentDescType = Path::EDescType::DocDependencies;
-				}
-				else if (LastElementInPath.CurrentDescType == Path::EDescType::DocDependencies)
-				{
-					NewElement.CurrentDescType = Path::EDescType::Class;
-				}
-
-				NewPath.Path.Add(NewElement);
-				return NewPath;
-			}
-
-			FDescPath operator [](const TCHAR* InElementName) const
-			{
-				if (!ensureMsgf(Path.Num() > 0, TEXT("Tried to append to a path that had no root.")))
-				{
-					return FDescPath();
-				}
-
-				// Copy the old path, and based on the last element in that path, append a new child element.
-				FDescPath NewPath = *this;
-
-				const Path::FElement& LastElementInPath = Path.Last();
-				Path::FElement NewElement = Path::FElement();
-
-				switch (LastElementInPath.CurrentDescType)
-				{
-					case Path::EDescType::DocDependencies:
-					{
-						ensureAlwaysMsgf(false, TEXT("Invalid path set up. DocDependencies for are accessed by using a int32."));
-						return FDescPath();
-						break;
-					}
-					case Path::EDescType::Inputs:
-					{
-						NewElement.CurrentDescType = Path::EDescType::Input;
-						NewElement.LookupName = FString(InElementName);
-						break;
-					}
-					case Path::EDescType::Outputs:
-					{
-						NewElement.CurrentDescType = Path::EDescType::Output;
-						NewElement.LookupName = FString(InElementName);
-						break;
-					}
-					/** Class descriptions and graph descriptions are navigated via enums. */
-					case Path::EDescType::Class:
-					case Path::EDescType::Graph:
-					{
-						ensureAlwaysMsgf(false, TEXT("Invalid path set up. Classes and Graphs are accessed by using the DescriptionPath::EClass and DesciptionPath::EGraph enums, respectively."));
-						return FDescPath();
-						break;
-					}
-					/** Dependencies for individual classes are navigated to via a int32. */
-					case Path::EDescType::ClassDependencies:
-					{
-						ensureAlwaysMsgf(false, TEXT("Invalid path set up. Dependencies for individual classes are accessed by using a int32."));
-						return FDescPath();
-						break;;
-					}
-					/** Node list descriptions are navigated via a int32. */
-					case Path::EDescType::Nodes:
-					{
-						ensureAlwaysMsgf(false, TEXT("Invalid path set up. Node lists are accessed by using a int32."));
-						return FDescPath();
-						break;
-					}
-					/** Individual inputs, outputs and nodes, as well as metadata, do not have any child elements. */
-					case Path::EDescType::Input:
-					case Path::EDescType::Output:
-					case Path::EDescType::Node:
-					case Path::EDescType::Metadata:
-					default:
-					{
-						ensureAlwaysMsgf(false, TEXT("Invalid path. Tried to add more pathing off of a description type that has no children."));
-						return FDescPath();
-						break;
-					}
-				}
-
-				NewPath.Path.Add(NewElement);
-				return NewPath;
-			}
-
-			// This can be used to go up multiple levels in a path. Effectively removes the last UnwindCount elements from the path.
-			FDescPath& operator<< (uint32 UnwindCount)
-			{
-				if (!ensureAlwaysMsgf(((uint32)Path.Num()) > UnwindCount, TEXT("Tried to unwind past the root of an FDescPath!")))
-				{
-					return *this;
-				}
-
-				while (UnwindCount-- > 0)
-				{
-					Path.RemoveAt(Path.Num() - 1);
-				}
-
-				return *this;
-			}
-		};
-
-		// Series of utility functions for description paths.
-		namespace Path
-		{
-			// When given a path for a node, finds that node's declaration in the Dependencies list.
-			FDescPath GetPathToClassForNode(FDescPath InPathForNode, FString& InNodeName);
-			
-			// Returns a path for whatever class description encapsulates the given graph.
-			FDescPath GetOwningClassDescription(FDescPath InPathForGraph);
-			
-			FDescPath GetDependencyPath(int32 InDependencyID);
-
-			// Given the path to an input node, returns the path to the corresponding  input description.
-			FDescPath GetInputDescriptionPath(FDescPath InPathForInputNode, const FString& InputName);
-
-			FDescPath GetOutputDescriptionPath(FDescPath InPathForOutputNode, const FString& OutputName);
-
-			FDescPath GetOuterGraphPath(FDescPath InPath);
-
-			// Generates a human readable string for a path into a Metasound description. Used for debugging.
-			FString GetPrintableString(FDescPath InPath);
-		}
-
-		using FClassDependencyIDs = TArray<int32>;
-
-		/**
-		 * Utility class that lives alongside a FMetasoundClassDescription to allow it's individual elements to be safely read and edited.
-		 * Note: Could be moved to a private header.
-		 *       It's also only necessary because we can't wrap a FMetasoundClassDescription in a shared ptr and also use it as a UPROPERTY.
-		 */
-		class METASOUNDFRONTEND_API FDescriptionAccessPoint
-		{
-			using FMetasoundDescriptionPtr = TVariant<
-				FMetasoundDocument*,
-				FMetasoundClassDescription*,
-				FMetasoundNodeDescription*,
-				FMetasoundGraphDescription*,
-				FMetasoundInputDescription*,
-				FMetasoundOutputDescription*,
-				FMetasoundClassMetadata*,
-				FClassDependencyIDs*,
-				void*>
-				;
-		public:
-			FDescriptionAccessPoint() = delete;
-			FDescriptionAccessPoint(TAccessPtr<FMetasoundDocument> InRootDocument);
-
-			static FDescriptionAccessPoint CreateInvalid()
-			{
-				return FDescriptionAccessPoint(TAccessPtr<FMetasoundDocument>::CreateInvalid());
-			}
-
-			// Returns the root class description associated with this FDescriptionAccessPoint.
-			FMetasoundDocument& GetRootChecked() const;
-
-			bool IsValid() const
-			{
-				return RootDocumentPtr.IsValid();
-			}
-
-			// Get a descriptor element at a specific path from the root of a metasound.
-			// @see FDescPath for more details.
-			// @returns a pointer to a descriptor struct, or an invalid ptr if the path is invalid.
-			template <typename MetasoundDescriptionType>
-			MetasoundDescriptionType* GetElementAtPath(const FDescPath& InPathFromRoot) const
-			{
-				FMetasoundDescriptionPtr Ptr;
-				if (ensure(GetDescriptionPtrFromPath(InPathFromRoot, Ptr)))
-				{
-					return Ptr.Get<MetasoundDescriptionType*>();
-				}
-
-				return nullptr;
-			}
-
-		private:
-			// The root class description for all description structs
-			// accessed via this class.
-			TAccessPtr<FMetasoundDocument> RootDocumentPtr;
-
-
-			bool GetDescriptionPtrFromPath(const FDescPath& InPathFromRoot, FMetasoundDescriptionPtr& OutPtr) const;
-
-			struct FDescriptionUnwindStep
-			{
-				FMetasoundDescriptionPtr DescriptionStructPtr;
-				Path::EDescType Type;
-
-				// TODO: make some simple constructors for this so we don't have
-				// to manually fill out info each time. 
-
-				static FDescriptionUnwindStep CreateInvalid()
-				{
-					FDescriptionUnwindStep InvalidStep;
-					InvalidStep.DescriptionStructPtr.Set<void*>(nullptr);
-					InvalidStep.Type = Path::EDescType::Invalid;
-					return InvalidStep;
-				}
-			};
-
-			FDescriptionUnwindStep GoToNextFromDocument(FMetasoundDocument& InDocument, FDescPath& InPath, const Path::FElement& InNext) const;
-			FDescriptionUnwindStep GoToNextFromClass(FMetasoundClassDescription& InClass, FDescPath& InPath, const Path::FElement& InNext) const;
-
-			// This function pops off elements of InPath to go to the next level down in a Metasound description.
-			FDescriptionUnwindStep GoToNext(FDescPath& InPath, FDescriptionUnwindStep InElement) const;
-		};
-
-		// Utility class used by the frontend handle classes to access the underlying Description data they are modifying.
-		// Similar to a TSoftObjectPtr: 
-		// this class caches a FDescPath to a specific element in a Metasound Class Description, and a weak pointer to a FDescriptionAccessPoint
-		// to use it with.
-		template <typename MetasoundDescriptionType>
-		class TDescriptionPtr
-		{
-			static constexpr bool bIsValidMetasoundDescriptionType =
-				TIsSame< MetasoundDescriptionType, FMetasoundDocument>::Value ||
-				TIsSame< MetasoundDescriptionType, FMetasoundGraphDescription>::Value ||
-				TIsSame< MetasoundDescriptionType, FMetasoundNodeDescription>::Value ||
-				TIsSame< MetasoundDescriptionType, FMetasoundInputDescription>::Value ||
-				TIsSame< MetasoundDescriptionType, FMetasoundOutputDescription>::Value ||
-				TIsSame< MetasoundDescriptionType, FMetasoundClassMetadata>::Value ||
-				TIsSame< MetasoundDescriptionType, FMetasoundClassDescription>::Value;
-
-			static_assert(bIsValidMetasoundDescriptionType, R"(Tried to use a Metasound::Frontend::FDescriptionAccessor with an invalid type. Supported types are: 
-			FMetasoundDocument			
-			FMetasoundGraphDescription, 
-			FMetasoundNodeDescription,
-			FMetasoundInputDescription,
-			FMetasoundOutputDescription,
-			FMetasoundClassMetadata,
-			FMetasoundClassDescription)");
-
-			FDescriptionAccessPoint AccessPoint;
-			FDescPath PathFromRoot;
-			bool bAccessFailure = false;
-
-		public:
-			TDescriptionPtr() = delete;
-
-			TDescriptionPtr(FDescriptionAccessPoint InAccessPoint, const FDescPath& InPathFromRoot)
-				: AccessPoint(InAccessPoint)
-				, PathFromRoot(InPathFromRoot)
-			{
-				// Test accessor
-				bAccessFailure = Get() == nullptr;
-			}
-
-			static TDescriptionPtr<MetasoundDescriptionType> CreateInvalid()
-			{
-				return TDescriptionPtr<MetasoundDescriptionType>(FDescriptionAccessPoint::CreateInvalid(), FDescPath());
-			}
-
-			TDescriptionPtr(TDescriptionPtr&& InDescriptionPtr) = default;
-			TDescriptionPtr(const TDescriptionPtr& InDescriptionPtr) = default;
-
-			TDescriptionPtr& operator =(TDescriptionPtr&& InDescriptionPtr) = default;
-			TDescriptionPtr& operator =(const TDescriptionPtr & InDescriptionPtr) = default;
-
-			// @returns false if the description this accessor was referencing has been destroyed,
-			// or if this was created with bad arguments.
-			bool IsValid() const
-			{
-				return AccessPoint.IsValid() && PathFromRoot.IsValid() && !bAccessFailure;
-			}
-
-			FDescPath GetPath() const
-			{
-				return PathFromRoot;
-			}
-
-			const FDescriptionAccessPoint& GetAccessPoint() const
-			{
-				return AccessPoint;
-			}
-
-			// returns a pointer or reference to the description data struct.
-			// Note that this is NOT thread safe- this should be called on the same thread as whomever owns the lifecycle of the underlying MetasoundDescriptionType.
-			MetasoundDescriptionType* Get() const
-			{
-				if (AccessPoint.IsValid())
-				{
-					return AccessPoint.GetElementAtPath<MetasoundDescriptionType>(PathFromRoot);
-				}
-				else
-				{
-					return nullptr;
-				}
-			}
-
-			MetasoundDescriptionType& GetChecked() const
-			{
-				MetasoundDescriptionType* ElementPtr = Get();
-				check(ElementPtr);
-				return *ElementPtr;
-			}
-
-			MetasoundDescriptionType& operator*() const
-			{
-				return GetChecked();
-			}
-
-			MetasoundDescriptionType* operator->() const
-			{
-				return Get();
-			}
-
-			FORCEINLINE explicit operator bool() const
-			{
-				return IsValid();
-			}
-		};
 	}
 }
diff --git a/Engine/Plugins/Runtime/Metasound/Source/MetasoundFrontend/Public/MetasoundFrontendController.h b/Engine/Plugins/Runtime/Metasound/Source/MetasoundFrontend/Public/MetasoundFrontendController.h
new file mode 100644
--- /dev/null
+++ b/Engine/Plugins/Runtime/Metasound/Source/MetasoundFrontend/Public/MetasoundFrontendController.h
@@ -0,0 +1,679 @@
+// Copyright Epic Games, Inc. All Rights Reserved.
+
+#pragma once
+
+#include "CoreMinimal.h"
+#include "MetasoundAccessPtr.h"
+#include "MetasoundBuilderInterface.h"
+#include "MetasoundFrontend.h"
+#include "MetasoundFrontendBaseClasses.h"
+#include "MetasoundFrontendDocument.h"
+#include "MetasoundGraph.h"
+
+/* Metasound Controllers and Handles provide a object oriented interface for  manipulating Metasound Documents. 
+ *
+ * Each controller interface is associated with a single Metasound entity such as 
+ * Document, Graph, Node, Input or Output. Access to these entities generally begins
+ * with the DocumentController which provides access to Graphs. Graphs provide access
+ * to Nodes and Nodes provide access to Inputs and Outputs. 
+ *
+ * A "Handle" is simply a TSharedRef<> of a controller.
+ *
+ * In general, the workflow for editing a Metasound graph will be:
+ * 1) Load or create a metasound asset.
+ * 2) Call UMetasound::GetDocumentHandle() to get a handle to the document for that asset.
+ * 
+ * Typically the workflow for creating a Metasound subgraph will be
+ * 1) Get a Metasound::Frontend::FGraphHandle (typically through the two steps described in the first paragraph)
+ * 2) Build a FMetasoundFrontendClassMetadata struct with whatever name/author/description you want to give the subgraph,
+ * 3) call Metasound::Frontend::FGraphHandle::CreateEmptySubgraphNode, which will return that subgraph as it as a FNodeHandle for that subgraph in the current graph.
+ * 4) class Metasound::Frontend::FNodeHandle::AsGraph which provides access to edit the subgraphs internal structure as well as externally facing inputs and outputs.
+ *
+ * General note- these apis are NOT thread safe. 
+ * Make sure that any FDocumentHandle, FGraphHandle, FNodeHandle, FInputHandle and FOutputHandle that access similar data are called on the same thread.
+ */
+namespace Metasound
+{
+	namespace Frontend
+	{
+		// Forward declare 
+		class IInputController;
+		class IOutputController;
+		class INodeController;
+		class IGraphController;
+		class IDocumentController;
+
+		// Metasound Frontend Handles are all TSharedRefs of various Metasound Frontend Controllers.
+		typedef TSharedRef<IInputController> FInputHandle;
+		typedef TSharedRef<const IInputController> FConstInputHandle;
+		typedef TSharedRef<IOutputController> FOutputHandle;
+		typedef TSharedRef<const IOutputController> FConstOutputHandle;
+		typedef TSharedRef<INodeController> FNodeHandle;
+		typedef TSharedRef<const INodeController> FConstNodeHandle;
+		typedef TSharedRef<IGraphController> FGraphHandle;
+		typedef TSharedRef<const IGraphController> FConstGraphHandle;
+		typedef TSharedRef<IDocumentController> FDocumentHandle;
+		typedef TSharedRef<const IDocumentController> FConstDocumentHandle;
+
+		using FVertexAccessPtr = TAccessPtr<FMetasoundFrontendVertex>;
+		using FConstVertexAccessPtr = TAccessPtr<const FMetasoundFrontendVertex>;
+		using FClassOutputAccessPtr = TAccessPtr<FMetasoundFrontendClassOutput>;
+		using FConstClassOutputAccessPtr = TAccessPtr<const FMetasoundFrontendClassOutput>;
+		using FClassInputAccessPtr = TAccessPtr<FMetasoundFrontendClassInput>;
+		using FConstClassInputAccessPtr = TAccessPtr<const FMetasoundFrontendClassInput>;
+		using FNodeAccessPtr = TAccessPtr<FMetasoundFrontendNode>;
+		using FConstNodeAccessPtr = TAccessPtr<const FMetasoundFrontendNode>;
+		using FGraphAccessPtr = TAccessPtr<FMetasoundFrontendGraph>;
+		using FConstGraphAccessPtr = TAccessPtr<const FMetasoundFrontendGraph>;
+		using FClassAccessPtr = TAccessPtr<FMetasoundFrontendClass>;
+		using FConstClassAccessPtr = TAccessPtr<const FMetasoundFrontendClass>;
+		using FGraphClassAccessPtr = TAccessPtr<FMetasoundFrontendGraphClass>;
+		using FConstGraphClassAccessPtr = TAccessPtr<const FMetasoundFrontendGraphClass>;
+		using FDocumentAccessPtr = TAccessPtr<FMetasoundFrontendDocument>;
+		using FConstDocumentAccessPtr = TAccessPtr<const FMetasoundFrontendDocument>;
+		
+		/* An IOutputController provides methods for querying and manipulating a metasound output vertex. */
+		class METASOUNDFRONTEND_API IOutputController : public TSharedFromThis<IOutputController>
+		{
+		public:
+			static FOutputHandle GetInvalidHandle();
+
+			IOutputController() = default;
+			virtual ~IOutputController() = default;
+
+			/** Returns true if the controller is in a valid state. */
+			virtual bool IsValid() const = 0;
+
+			virtual int32 GetID() const = 0;
+			
+			/** Returns the data type name associated with this output. */
+			virtual const FName& GetDataType() const = 0;
+			
+			/** Returns the name associated with this output. */
+			virtual const FString& GetName() const = 0;
+			
+			/** Returns the human readable name associated with this output. */
+			virtual const FText& GetDisplayName() const = 0;
+			
+			/** Returns the tooltip associated with this output. */
+			virtual const FText& GetTooltip() const = 0;
+
+			/** Returns the ID of the node which owns this output. */
+			virtual int32 GetOwningNodeID() const = 0;
+			
+			/** Returns a FNodeHandle to the node which owns this output. */
+			virtual FNodeHandle GetOwningNode() = 0;
+			
+			/** Returns a FConstNodeHandle to the node which owns this output. */
+			virtual FConstNodeHandle GetOwningNode() const = 0;
+
+
+			// TODO: with the current spec, we'd have to scan all nodes in the graph for this output to solve this.
+			// Is this worth it, or should we just require folks to use FInputController::GetCurrentlyConnectedOutput()?
+			//IInputController* GetCurrentlyConnectedInput();
+			//const IInputController* GetCurrentlyConnectedInput() const;
+			//void Disconnect();
+
+			/** Returns information describing connectability between this output and the supplied input. */
+			virtual FConnectability CanConnectTo(const IInputController& InController) const = 0;
+			
+			/** Connects this output and the supplied input. 
+			 * @return True on success, false on failure. 
+			 */
+			virtual bool Connect(IInputController& InController) = 0;
+			
+			/** Connects this output to the input using a converter node.
+			 * @return True on success, false on failure. 
+			 */
+			virtual bool ConnectWithConverterNode(IInputController& InController, const FConverterNodeInfo& InNodeClassName) = 0;
+			
+			/** Disconnects this output from the input. 
+			 * @return True on success, false on failure. 
+			 */
+			virtual bool Disconnect(IInputController& InController) = 0;
+
+		};
+
+		/* An IInputController provides methods for querying and manipulating a metasound input vertex. */
+		class METASOUNDFRONTEND_API IInputController : public TSharedFromThis<IInputController>
+		{
+		public:
+			static FInputHandle GetInvalidHandle();
+
+			IInputController() = default;
+			virtual ~IInputController() = default;
+
+			/** Returns true if the controller is in a valid state. */
+			virtual bool IsValid() const = 0;
+
+			virtual int32 GetID() const = 0;
+
+			/** Return true if the input is connect to an output. */
+			virtual bool IsConnected() const = 0;
+
+			/** Returns the data type name associated with this input. */
+			virtual const FName& GetDataType() const = 0;
+
+			/** Returns the data type name associated with this input. */
+			virtual const FString& GetName() const = 0;
+
+			/** Returns the data type name associated with this input. */
+			virtual const FText& GetDisplayName() const = 0;
+
+			/** Returns the data type name associated with this input. */
+			virtual const FText& GetTooltip() const = 0;
+
+			/** Returns the ID of the node which owns this output. */
+			virtual int32 GetOwningNodeID() const = 0;
+			
+			/** Returns a FNodeHandle to the node which owns this output. */
+			virtual FNodeHandle GetOwningNode() = 0;
+			
+			/** Returns a FConstNodeHandle to the node which owns this output. */
+			virtual FConstNodeHandle GetOwningNode() const = 0;
+			
+			/** Returns the currently connected output. If this input is not
+			 * connected, the returned handle will be invalid. */
+			virtual FOutputHandle GetCurrentlyConnectedOutput() = 0;
+
+			/** Returns the currently connected output. If this input is not
+			 * connected, the returned handle will be invalid. */
+			virtual FConstOutputHandle GetCurrentlyConnectedOutput() const = 0;
+
+			/** Returns information describing connectability between this input and the supplied output. */
+			virtual FConnectability CanConnectTo(const IOutputController& InController) const = 0;
+
+			/** Connects this input and the supplied output. 
+			 * @return True on success, false on failure. 
+			 */
+			virtual bool Connect(IOutputController& InController) = 0;
+
+			/** Connects this input to the output using a converter node.
+			 * @return True on success, false on failure. 
+			 */
+			virtual bool ConnectWithConverterNode(IOutputController& InController, const FConverterNodeInfo& InNodeClassName) = 0;
+
+			/** Disconnects this input from the given output. 
+			 * @return True on success, false on failure. 
+			 */
+			virtual bool Disconnect(IOutputController& InController) = 0;
+
+			/** Disconnects this input from any connected output. 
+			 * @return True on success, false on failure. 
+			 */
+			virtual bool Disconnect() = 0;
+		};
+
+		/* An INodeController provides methods for querying and manipulating a Metasound node. */
+		class METASOUNDFRONTEND_API INodeController : public TSharedFromThis<INodeController>
+		{
+
+		public:
+			static FNodeHandle GetInvalidHandle();
+
+			INodeController() = default;
+			virtual ~INodeController() = default;
+
+			/** Returns true if the controller is in a valid state. */
+			virtual bool IsValid() const = 0;
+
+			/** Returns all node inputs. */
+			virtual TArray<FInputHandle> GetInputs() = 0;
+
+			/** Returns all node inputs. */
+			virtual TArray<FConstInputHandle> GetConstInputs() const = 0;
+
+			virtual TArray<FInputHandle> GetInputsWithVertexName(const FString& InName) = 0;
+			virtual TArray<FConstInputHandle> GetInputsWithVertexName(const FString& InName) const = 0;
+
+			/** Returns all node outputs. */
+			virtual TArray<FOutputHandle> GetOutputs() = 0;
+
+			/** Returns all node outputs. */
+			virtual TArray<FConstOutputHandle> GetConstOutputs() const = 0;
+
+			virtual TArray<FOutputHandle> GetOutputsWithVertexName(const FString& InName) = 0;
+			virtual TArray<FConstOutputHandle> GetOutputsWithVertexName(const FString& InName) const = 0;
+
+			/** Returns an input with the given id. 
+			 *
+			 * If the input does not exist, an invalid handle is returned. 
+			 */
+			virtual FInputHandle GetInputWithID(int32 InPointID) = 0;
+
+			/** Returns an input with the given name. 
+			 *
+			 * If the input does not exist, an invalid handle is returned. 
+			 */
+			virtual FConstInputHandle GetInputWithID(int32 InPointID) const = 0;
+
+			/** Returns an output with the given name. 
+			 *
+			 * If the output does not exist, an invalid handle is returned. 
+			 */
+			virtual FOutputHandle GetOutputWithID(int32 InPointID) = 0;
+
+			/** Returns an output with the given name. 
+			 *
+			 * If the output does not exist, an invalid handle is returned. 
+			 */
+			virtual FConstOutputHandle GetOutputWithID(int32 InPointID) const = 0;
+
+			virtual bool CanAddInput(const FString& InVertexName) const = 0;
+			virtual FInputHandle AddInput(const FString& InVertexName, const FMetasoundFrontendLiteral* InDefault) = 0;
+			virtual bool RemoveInput(int32 InPointID) = 0;
+
+			// TODO: PointID -> PinID? PointID? SlotID? RefID? 
+			virtual bool CanAddOutput(const FString& InVertexName) const = 0;
+			virtual FInputHandle AddOutput(const FString& InVertexName, const FMetasoundFrontendLiteral* InDefault) = 0;
+			virtual bool RemoveOutput(int32 InPointID) = 0;
+
+			/** Returns the class type for this node. */
+			virtual EMetasoundFrontendClassType GetClassType() const = 0;
+
+			/** Returns the name of the node class. */
+			virtual const FString& GetClassName() const = 0;
+			virtual FMetasoundFrontendVersionNumber GetClassVersionNumber() const = 0;
+			virtual const FText& GetClassDescription() const = 0;
+
+			/** If the node is also a graph, this returns a graph handle.
+			 * If the node is not also a graph, it will return an invalid handle.
+			 */
+			virtual FGraphHandle AsGraph() = 0;
+
+			/** If the node is also a graph, this returns a graph handle.
+			 * If the node is not also a graph, it will return an invalid handle.
+			 */
+			virtual FConstGraphHandle AsGraph() const = 0;
+
+			/** Returns the name of this node. */
+			virtual const FString& GetNodeName() const = 0;
+
+			/** Returns the ID associated with this node. */
+			virtual int32 GetID() const = 0;
+
+			/** Returns the ID associated with the node class. */
+			virtual int32 GetClassID() const = 0;
+
+			/** Returns the class ID of the metasound class which owns this node. */
+			virtual int32 GetOwningGraphClassID() const = 0;
+
+			/** Returns the graph which owns this node. */
+			virtual FGraphHandle GetOwningGraph() = 0;
+
+			/** Returns the graph which owns this node. */
+			virtual FConstGraphHandle GetOwningGraph() const = 0;
+		};
+
+		/* An IGraphController provides methods for querying and manipulating a Metasound graph. */
+		class METASOUNDFRONTEND_API IGraphController : public TSharedFromThis<IGraphController>
+		{
+			public:
+			static FGraphHandle GetInvalidHandle();
+
+			IGraphController() = default;
+			virtual ~IGraphController() = default;
+
+			/** Returns true if the controller is in a valid state. */
+			virtual bool IsValid() const = 0;
+
+			virtual int32 GetNewPointID() const = 0;
+
+			virtual TArray<FString> GetInputVertexNames() const = 0;
+			virtual TArray<FString> GetOutputVertexNames() const = 0;
+
+			/** Returns all nodes in the graph. */
+			virtual TArray<FNodeHandle> GetNodes() = 0;
+
+			/** Returns a node by NodeID. If the node does not exist, an invalid handle is returned. */
+			virtual FConstNodeHandle GetNodeWithID(int32 InNodeID) const = 0;
+
+			/** Returns all nodes in the graph. */
+			virtual TArray<FConstNodeHandle> GetConstNodes() const = 0;
+
+			/** Returns a node by NodeID. If the node does not exist, an invalid handle is returned. */
+			virtual FNodeHandle GetNodeWithID(int32 InNodeID) = 0;
+
+			/** Returns all output nodes in the graph. */
+			virtual TArray<FNodeHandle> GetOutputNodes() = 0;
+
+			/** Returns all output nodes in the graph. */
+			virtual TArray<FConstNodeHandle> GetConstOutputNodes() const = 0;
+
+			/** Returns all input nodes in the graph. */
+			virtual TArray<FNodeHandle> GetInputNodes() = 0;
+
+			/** Returns all input nodes in the graph. */
+			virtual TArray<FConstNodeHandle> GetConstInputNodes() const = 0;
+
+			/** Returns true if an output node with the given NodeName exists.
+			 *
+			 * @param InName - Name of node.
+			 * @return True if the node exists, false otherwise. 
+			 */
+			virtual bool ContainsOutputNodeWithName(const FString& InName) const = 0;
+
+			/** Returns true if an input node with the given NodeName exists.
+			 *
+			 * @param InName - Name of node.
+			 * @return True if the node exists, false otherwise. 
+			 */
+			virtual bool ContainsInputNodeWithName(const FString& InName) const = 0;
+
+			/** Returns a handle to an existing output node for the given graph output name.
+			 * If no output exists with the given name, an invalid node handle is returned. 
+			 *
+			 * @param InName - Name of graph output.. 
+			 * @return The node handle for the output node. If the output does not exist, an invalid handle is returned.
+			 */
+			virtual FNodeHandle GetOutputNodeWithName(const FString& InName) = 0;
+
+			/** Returns a handle to an existing output node for the given graph output name.
+			 * If no output exists with the given name, an invalid node handle is returned. 
+			 *
+			 * @param InName - Name of graph output.
+			 * @return The node handle for the output node. If the node does not exist, an invalid handle is returned.
+			 */
+			virtual FConstNodeHandle GetOutputNodeWithName(const FString& InName) const = 0;
+
+			/** Returns a handle to an existing input node for the given graph input name.
+			 * If no input exists with the given name, an invalid node handle is returned. 
+			 *
+			 * @param InName - Name of graph input. 
+			 * @return The node handle for the input node. If the node does not exist, an invalid handle is returned.
+			 */
+			virtual FNodeHandle GetInputNodeWithName(const FString& InName) = 0;
+
+			/** Returns a handle to an existing input node for the given graph input name.
+			 * If no input exists with the given name, an invalid node handle is returned. 
+			 *
+			 * @param InName - Name of graph input. 
+			 * @return The node handle for the input node. If the node does not exist, an invalid handle is returned.
+			 */
+			virtual FConstNodeHandle GetInputNodeWithName(const FString& InName) const = 0;
+
+			virtual FConstClassInputAccessPtr FindClassInputWithName(const FString& InName) const = 0;
+			virtual FConstClassOutputAccessPtr FindClassOutputWithName(const FString& InName) const = 0;
+
+
+			/** Add a new input node using the input description. 
+			 *
+			 * @param InDescription - Description for input of graph.
+			 * @return On success, a valid input node handle. On failure, an invalid node handle.
+			 */
+			virtual FNodeHandle AddInputVertex(const FMetasoundFrontendClassInput& InDescription) = 0;
+
+			/** Remove the input with the given name. Returns true if successfully removed, false otherwise. */
+			virtual bool RemoveInputVertex(const FString& InputName) = 0;
+
+			/** Add a new output node using the output description. 
+			 *
+			 * @param InDescription - Description for output of graph.
+			 * @return On success, a valid output node handle. On failure, an invalid node handle.
+			 */
+			virtual FNodeHandle AddOutputVertex(const FMetasoundFrontendClassOutput& InDescription) = 0;
+
+			/** Remove the output with the given name. Returns true if successfully removed, false otherwise. */
+			virtual bool RemoveOutputVertex(const FString& OutputName) = 0;
+
+			/** Returns the preferred literal argument type for a given input.
+			 * Returns ELiteralType::Invalid if the input couldn't be found, 
+			 * or if the input doesn't support any kind of literals.
+			 *
+			 * @param InInputName - Name of graph input.
+			 * @return The preferred literal argument type. 
+			 */
+			virtual ELiteralType GetPreferredLiteralTypeForInputVertex(const FString& InInputName) const = 0;
+
+			/** Return the UObject class corresponding an input. Meaningful for inputs whose preferred 
+			 * literal type is UObject or UObjectArray.
+			 *
+			 * @param InInputName - Name of graph input.
+			 *
+			 * @return The UClass* for the literal argument input. nullptr on error or if UObject argument is not supported.
+			 */
+			virtual UClass* GetSupportedClassForInputVertex(const FString& InInputName) = 0;
+
+			virtual TArray<int32> GetDefaultIDsForInputVertex(const FString& InInputName) const = 0;
+			virtual TArray<int32> GetDefaultIDsForOutputVertex(const FString& InOutputName) const = 0;
+
+			/** Set the default value for the graph input.
+			 *
+			 * @param InInputName - Name of the graph input.
+			 * @param InValue - Value to give to the graph input. 
+			 *
+			 * @return True on success. False if the input does not exist or if the literal type was incompatible with the input.
+			 */
+			virtual bool SetDefaultInputToLiteral(const FString& InInputName, int32 InPointID, bool InValue) = 0;
+
+			/** Set the default value for the graph input.
+			 *
+			 * @param InInputName - Name of the graph input.
+			 * @param InValue - Value to give to the graph input. 
+			 *
+			 * @return True on success. False if the input does not exist or if the literal type was incompatible with the input.
+			 */
+			virtual bool SetDefaultInputToLiteral(const FString& InInputName, int32 InPointID, int32 InValue) = 0;
+
+			/** Set the default value for the graph input.
+			 *
+			 * @param InInputName - Name of the graph input.
+			 * @param InValue - Value to give to the graph input. 
+			 *
+			 * @return True on success. False if the input does not exist or if the literal type was incompatible with the input.
+			 */
+			virtual bool SetDefaultInputToLiteral(const FString& InInputName, int32 InPointID, float InValue) = 0;
+
+			/** Set the default value for the graph input.
+			 *
+			 * @param InInputName - Name of the graph input.
+			 * @param InValue - Value to give to the graph input. 
+			 *
+			 * @return True on success. False if the input does not exist or if the literal type was incompatible with the input.
+			 */
+			virtual bool SetDefaultInputToLiteral(const FString& InInputName, int32 InPointID, const FString& InValue) = 0;
+
+			/** Set the default value for the graph input.
+			 *
+			 * @param InInputName - Name of the graph input.
+			 * @param InValue - Value to give to the graph input. 
+			 *
+			 * @return True on success. False if the input does not exist or if the literal type was incompatible with the input.
+			 */
+			virtual bool SetDefaultInputToLiteral(const FString& InInputName, int32 InPointID, UObject* InValue) = 0;
+
+			/** Set the default value for the graph input.
+			 *
+			 * @param InInputName - Name of the graph input.
+			 * @param InValue - Value to give to the graph input. 
+			 *
+			 * @return True on success. False if the input does not exist or if the literal type was incompatible with the input.
+			 */
+			virtual bool SetDefaultInputToLiteral(const FString& InInputName, int32 InPointID, const TArray<UObject*>& InValue) = 0;
+
+
+			/** Set the display name for the input with the given name. */
+			virtual void SetInputDisplayName(FString InName, const FText& InDisplayName) = 0;
+
+			/** Set the display name for the output with the given name. */
+			virtual void SetOutputDisplayName(FString InName, const FText& InDisplayName) = 0;
+
+			/** Clear the current literal for a given input.
+			 *
+			 * @return True on success, false on failure.
+			 */
+			virtual bool ClearLiteralForInput(const FString& InInputName, int32 InPointID) = 0;
+
+			/** Add a new node to this graph.
+			 *
+			 * @param InNodeClass - Info for node class.
+			 * 
+			 * @return Node handle for class. On error, an invalid handle is returned. 
+			 */
+			virtual FNodeHandle AddNode(const FNodeClassInfo& InNodeClass) = 0;
+
+			/** Add a new node to this graph.
+			 *
+			 * @param InNodeClass - Info for node class.
+			 * 
+			 * @return Node handle for class. On error, an invalid handle is returned. 
+			 */
+			virtual FNodeHandle AddNode(const FNodeRegistryKey& InNodeClass) = 0;
+
+			/** Remove the node corresponding to this node handle.
+			 *
+			 * @return True on success, false on failure. 
+			 */
+			virtual bool RemoveNode(INodeController& InNode) = 0;
+
+			/** Returns the ClassID associated with this graph. */
+			virtual int32 GetClassID() const = 0;
+
+			/** Return the metadata for the current graph. */
+			virtual const FMetasoundFrontendClassMetadata& GetGraphMetadata() const = 0;
+
+			/** Inflates a subgraph of this graph into the this graph.
+			 *
+			 * If the INodeController given is itself a Metasound graph,
+			 * and the INodeController is a direct member of this IGraphController,
+			 * this will invalidate the FNodeController and paste the graph for this 
+			 * node directly into this graph.
+			 *
+			 * If not successful, InNode will not be affected.
+			 *
+			 * @returns True on success, false on failure.
+			 */
+			virtual bool InflateNodeDirectlyIntoGraph(const INodeController& InNode) = 0;
+
+			/** Creates and inserts a new subgraph into this graph using the given metadata.
+			 * By calling AsGraph() on the returned node handle, callers can modify
+			 * the new subgraph.
+			 *
+			 * @param InInfo - Metadata for the subgraph.
+			 *
+			 * @return Handle to the subgraph node. On error, the handle is invalid.
+			 */
+			virtual FNodeHandle CreateEmptySubgraph(const FMetasoundFrontendClassMetadata& InInfo) = 0;
+
+			/** Creates a runtime operator for the given graph.
+			 *
+			 * @param InSettings - Settings to use when creating operators.
+			 * @param InEnvironment - Environment variables available during creation.
+			 * @param OutBuildErrors - An array to populate with errors encountered during the build process.
+			 *
+			 * @return On success, a valid pointer to a Metasound operator. An invalid pointer on failure. 
+			 */
+			virtual TUniquePtr<IOperator> BuildOperator(const FOperatorSettings& InSettings, const FMetasoundEnvironment& InEnvironment, TArray<IOperatorBuilder::FBuildErrorPtr>& OutBuildErrors) const = 0;
+
+			/** Returns a handle to the document owning this graph. */
+			virtual FDocumentHandle GetOwningDocument() = 0;
+
+			/** Returns a handle to the document owning this graph. */
+			virtual FConstDocumentHandle GetOwningDocument() const = 0;
+		};
+
+		/* An IDocumentController provides methods for querying and manipulating a Metasound document. */
+		class METASOUNDFRONTEND_API IDocumentController : public TSharedFromThis<IDocumentController>
+		{
+		public:
+			static FDocumentHandle GetInvalidHandle();
+
+			/** Create a document from FMetasoundFrontendDocument description pointer. */
+			static FDocumentHandle CreateDocumentHandle(TAccessPtr<FMetasoundFrontendDocument> InDocument);
+
+			IDocumentController() = default;
+			virtual ~IDocumentController() = default;
+
+			/** Returns true if the controller is in a valid state. */
+			virtual bool IsValid() const = 0;
+
+			/** Returns true if an input with the provided name is required to satisfy the document archetype. */
+			virtual bool IsRequiredInput(const FString& InInputName) const = 0;
+
+			/** Returns true if an input with the provided name is required to satisfy the document archetype. */
+			virtual bool IsRequiredOutput(const FString& InOutputName) const = 0;
+
+			/** Returns an array of inputs which describe the required inputs needed to satisfy the document archetype. */
+			virtual TArray<FMetasoundFrontendClassVertex> GetRequiredInputs() const = 0;
+
+			/** Returns an array of outputs which describe the required outputs needed to satisfy the document archetype. */
+			virtual TArray<FMetasoundFrontendClassVertex> GetRequiredOutputs() const = 0;
+
+			// TODO: add info on environment variables. 
+			// TODO: consider find/add subgraph
+			// TODO: perhaps functions returning access pointers could be removed from main interface and only exist in FDocumentController.
+			
+			/** Returns an array of all class dependencies for this document. */
+			virtual TArray<FMetasoundFrontendClass> GetDependencies() const = 0;
+			virtual TArray<FMetasoundFrontendGraphClass> GetSubgraphs() const = 0;
+			virtual TArray<FMetasoundFrontendClass> GetClasses() const = 0;
+
+			virtual FConstClassAccessPtr FindDependencyWithID(int32 InClassID) const = 0;
+			virtual FConstGraphClassAccessPtr FindSubgraphWithID(int32 InClassID) const = 0;
+			virtual FConstClassAccessPtr FindClassWithID(int32 InClassID) const = 0;
+
+
+			/** Returns an existing Metasound class description corresponding to 
+			 * a dependency which matches the provided class information.
+			 *
+			 * @return A pointer to the found object, or nullptr if it could not be found.
+			 */
+			virtual FConstClassAccessPtr FindClass(const FNodeClassInfo& InNodeClass) const = 0; // TODO: swap for just using metadata in future. 
+
+			/** Attempts to find an existing Metasound class description corresponding
+			 * to a dependency which matches the provided class information. If the
+			 * class is not found in the current dependencies, it is added to the 
+			 * dependencies.
+			 *
+			 * @return A pointer to the object, or nullptr on error.
+			 */
+			virtual FConstClassAccessPtr FindOrAddClass(const FNodeClassInfo& InNodeClass) = 0; // TODO: swap for just using metadata in future. 
+
+			/** Returns an existing Metasound class description corresponding to 
+			 * a dependency which matches the provided class information.
+			 *
+			 * @return A pointer to the found object, or nullptr if it could not be found.
+			 */
+			virtual FConstClassAccessPtr FindClass(const FMetasoundFrontendClassMetadata& InMetadata) const = 0;
+
+			/** Attempts to find an existing Metasound class description corresponding
+			 * to a dependency which matches the provided class information. If the
+			 * class is not found in the current dependencies, it is added to the 
+			 * dependencies.
+			 *
+			 * @return A pointer to the object, or nullptr on error.
+			 */
+			virtual FConstClassAccessPtr FindOrAddClass(const FMetasoundFrontendClassMetadata& InMetadata) = 0;
+
+			/** Removes all dependencies which are no longer referenced by any graphs within this document.  */
+			virtual void RemoveUnreferencedDependencies() = 0;
+
+			/** Returns an array of all subgraphs for this document. */
+			virtual TArray<FGraphHandle> GetSubgraphHandles() = 0;
+
+			/** Returns an array of all subgraphs for this document. */
+			virtual TArray<FConstGraphHandle> GetSubgraphHandles() const = 0;
+
+			/** Returns a graphs in the document with the given class ID.*/
+			virtual FGraphHandle GetSubgraphWithClassID(int32 InClassID) = 0;
+
+			/** Returns a graphs in the document with the given class ID.*/
+			virtual FConstGraphHandle GetSubgraphWithClassID(int32 InClassID) const = 0;
+
+			/** Returns the root graph of this document. */
+			virtual FGraphHandle GetRootGraph() = 0;
+
+			/** Returns the root graph of this document. */
+			virtual FConstGraphHandle GetRootGraph() const = 0;
+
+			/** Exports the document to a json file at the provided path.
+			 *
+			 * @return True on success, false on failure. 
+			 */
+			virtual bool ExportToJSONAsset(const FString& InAbsolutePath) const = 0;
+
+			/** Exports the document to a json formatted string. */
+			virtual FString ExportToJSON() const = 0;
+		};
+	}
+}
diff --git a/Engine/Plugins/Runtime/Metasound/Source/MetasoundFrontend/Public/MetasoundFrontendDocument.h b/Engine/Plugins/Runtime/Metasound/Source/MetasoundFrontend/Public/MetasoundFrontendDocument.h
new file mode 100644
--- /dev/null
+++ b/Engine/Plugins/Runtime/Metasound/Source/MetasoundFrontend/Public/MetasoundFrontendDocument.h
@@ -0,0 +1,712 @@
+// Copyright Epic Games, Inc. All Rights Reserved.
+
+#pragma once
+
+#include "CoreMinimal.h"
+#include "Internationalization/Text.h"
+#include "MetasoundAccessPtr.h"
+
+#include "MetasoundFrontendDocument.generated.h"
+
+namespace Metasound
+{
+	static constexpr int32 FrontendInvalidID = INDEX_NONE;
+}
+
+// Forward declare
+struct FMetasoundFrontendClass;
+struct FMetasoundFrontendClassInterface;
+
+UENUM()
+enum class EMetasoundFrontendClassType : uint8
+{
+	// The Metasound class is defined externally, in compiled code or in another document.
+	External,
+
+	// The Metasound class is a graph within the containing document.
+	Graph,
+
+	// The Metasound class is an input into a graph in the containing document.
+	Input,
+
+	// The Metasound class is an input into a graph in the containing document.
+	Output,
+
+	Invalid UMETA(Hidden)
+};
+
+// General purpose version number for Metasound Frontend objects.
+USTRUCT()
+struct FMetasoundFrontendVersionNumber
+{
+	GENERATED_BODY()
+
+	// Major version number.
+	UPROPERTY(EditAnywhere, Category = General)
+	int32 Major = 1;
+
+	// Minor version number.
+	UPROPERTY(EditAnywhere, Category = General)
+	int32 Minor = 0;
+};
+
+// General purpose version info for Metasound Frontend objects.
+USTRUCT()
+struct FMetasoundFrontendVersion
+{
+	GENERATED_BODY()
+
+	// Name of version.
+	UPROPERTY()
+	FName Name;
+
+	// Version number.
+	UPROPERTY()
+	FMetasoundFrontendVersionNumber Number;
+};
+
+// The type of a given literal for an input value.
+UENUM()
+enum class EMetasoundFrontendLiteralType : uint8
+{
+	None,
+	Bool,
+	Float,
+	Integer,
+	String,
+	UObject,
+	UObjectArray,
+	Invalid UMETA(Hidden)
+};
+
+// Represents the serialized version of variant literal types. 
+USTRUCT()
+struct METASOUNDFRONTEND_API FMetasoundFrontendLiteral
+{
+	GENERATED_BODY()
+
+	// The set type of this literal.
+	UPROPERTY(EditAnywhere, Category = Customized)
+	EMetasoundFrontendLiteralType Type = EMetasoundFrontendLiteralType::None;
+
+	UPROPERTY(EditAnywhere, Category = Customized)
+	bool AsBool = false;
+
+	UPROPERTY(EditAnywhere, Category = Customized)
+	int32 AsInteger = 0;
+
+	UPROPERTY(EditAnywhere, Category = Customized)
+	float AsFloat = 0.f;
+
+	UPROPERTY(EditAnywhere, Category = Customized)
+	FString AsString = TEXT("");
+
+	UPROPERTY(EditAnywhere, Category = Customized)
+	UObject* AsUObject = nullptr;
+
+	UPROPERTY(EditAnywhere, Category = Customized)
+	TArray<UObject*> AsUObjectArray;
+
+	void Set(bool InValue);
+	void Set(int32 InValue);
+	void Set(float InValue);
+	void Set(const FString& InValue);
+	void Set(UObject* InValue);
+	void Set(const TArray<UObject*>& InValue);
+	void Clear();
+};
+
+// A FMetasoundFrontendVertex provides an named connection point of a node.
+USTRUCT() 
+struct METASOUNDFRONTEND_API FMetasoundFrontendVertex
+{
+	GENERATED_BODY()
+
+	virtual ~FMetasoundFrontendVertex() = default;
+
+	// Name of the vertex. Unique amongst other vertexes on the same interface.
+	UPROPERTY(VisibleAnywhere, Category = CustomView)
+	FString Name;
+
+	// Data type name of the vertex.
+	UPROPERTY(EditAnywhere, Category = Parameters)
+	FName TypeName;
+
+	// IDs of connection points supported by the vertex.
+	UPROPERTY()	
+	TArray<int32> PointIDs;
+
+	// Returns true if vertexes have equal name, type and number of IDs. 
+	static bool IsFunctionalEquivalent(const FMetasoundFrontendVertex& InLHS, const FMetasoundFrontendVertex& InRHS);
+};
+
+// Contains a default value for a single vertex ID
+USTRUCT() 
+struct FMetasoundFrontendVertexLiteral
+{
+	GENERATED_BODY()
+
+	// ID of vertex.
+	UPROPERTY(EditAnywhere, Category = Parameters)
+	int32 PointID = Metasound::FrontendInvalidID;
+
+	// Value to use when constructing input. 
+	UPROPERTY(EditAnywhere, Category = Parameters)
+	FMetasoundFrontendLiteral Value;
+};
+
+
+USTRUCT()
+struct METASOUNDFRONTEND_API FMetasoundFrontendNodeInterface
+{
+	GENERATED_BODY()
+
+	FMetasoundFrontendNodeInterface() = default;
+
+	// Create a node interface which satisfies an existing class interface.
+	FMetasoundFrontendNodeInterface(const FMetasoundFrontendClassInterface& InClassInterface);
+
+	// Input vertices to node.
+	UPROPERTY()
+	TArray<FMetasoundFrontendVertex> Inputs;
+
+	// Output vertices to node.
+	UPROPERTY()
+	TArray<FMetasoundFrontendVertex> Outputs;
+
+	// Environment variables of node.
+	UPROPERTY()
+	TArray<FMetasoundFrontendVertex> Environment;
+};
+
+
+// An FMetasoundFrontendNode represents a single instance of a FMetasoundFrontendClass
+USTRUCT()
+struct METASOUNDFRONTEND_API FMetasoundFrontendNode 
+{
+	GENERATED_BODY()
+
+	FMetasoundFrontendNode() = default;
+
+	// Construct node to satisfy class. 
+	FMetasoundFrontendNode(const FMetasoundFrontendClass& InClass);
+
+	// Unique ID of this node.
+	UPROPERTY()
+	int32 ID = Metasound::FrontendInvalidID;
+
+	// ID of FMetasoundFrontendClass corresponding to this node.
+	UPROPERTY()
+	int32 ClassID = Metasound::FrontendInvalidID;
+
+	// Name of node instance. 
+	UPROPERTY()
+	FString Name;
+
+	// Interface of node instance.
+	UPROPERTY()
+	FMetasoundFrontendNodeInterface Interface;
+
+	// Default values for node inputs.
+	UPROPERTY()
+	TArray<FMetasoundFrontendVertexLiteral> InputLiterals;
+
+	// TODO: Add node init data if it is used. It may be dropped as a concept.
+};
+
+
+// Represents a single connection from one point to another.
+USTRUCT()
+struct FMetasoundFrontendEdge
+{
+	GENERATED_BODY()
+
+	// ID of source node.
+	UPROPERTY()
+	int32 FromNodeID = Metasound::FrontendInvalidID;
+
+	// ID of source point on source node.
+	UPROPERTY()
+	int32 FromPointID = Metasound::FrontendInvalidID;
+
+	// ID of destination node.
+	UPROPERTY()
+	int32 ToNodeID = Metasound::FrontendInvalidID;
+
+	// ID of destination point on destination node.
+	UPROPERTY()
+	int32 ToPointID = Metasound::FrontendInvalidID;
+};
+
+// Display style for an edge.
+UENUM()
+enum class EMetasoundFrontendStyleEdgeDisplay : uint8
+{
+	Default,
+	Inherited,
+	Hidden
+};
+
+// Styling for edges
+USTRUCT()
+struct FMetasoundFrontendStyleEdge
+{
+	GENERATED_BODY()
+
+	UPROPERTY()
+	EMetasoundFrontendStyleEdgeDisplay Display;
+};
+
+
+// Styling for a class of edges dependent upon edge data type.
+USTRUCT()
+struct FMetasoundFrontendStyleEdgeClass
+{
+	GENERATED_BODY()
+
+	// Datatype of edge to apply style to
+	UPROPERTY()
+	FName TypeName;
+
+	// Style information for edge.
+	UPROPERTY()
+	FMetasoundFrontendStyleEdge Style;
+};
+
+// Styling for a class
+USTRUCT() 
+struct FMetasoundFrontendGraphStyle 
+{
+	GENERATED_BODY()
+
+	// Edge styles for graph.
+	UPROPERTY()
+	TArray<FMetasoundFrontendStyleEdgeClass> EdgeStyles;
+};
+
+USTRUCT()
+struct FMetasoundFrontendGraph
+{
+	GENERATED_BODY()
+
+	// Node contained in graph
+	UPROPERTY()
+	TArray<FMetasoundFrontendNode> Nodes;
+
+	// Connections between points on nodes.
+	UPROPERTY()
+	TArray<FMetasoundFrontendEdge> Edges;
+
+	// Style of graph display.
+	UPROPERTY()
+	FMetasoundFrontendGraphStyle Style;
+};
+
+// The type of the vertex.
+UENUM()
+enum class EMetasoundFrontendVertexType : uint8
+{
+	Point, //< Vertex represents a single value.
+	//TODO: Array   //< Vertex represents a variable sized array of values.
+};
+
+
+// Defines the behavior of the vertex.
+USTRUCT()
+struct METASOUNDFRONTEND_API FMetasoundFrontendVertexBehavior
+{
+	GENERATED_BODY()
+
+	// TODO: should all these be editable? 
+	UPROPERTY()
+	EMetasoundFrontendVertexType Type = EMetasoundFrontendVertexType::Point;
+
+	// Minimum number of connection points. Only used if the vertex type is EMetasoundFrontendVertexType::Array
+	UPROPERTY()
+	int32 ArrayMin = 1;
+
+	// Maximum number of connection points. Only used if the vertex type is EMetasoundFrontendVertexType::Array
+	UPROPERTY()
+	int32 ArrayMax = 1;
+
+	static bool IsFunctionalEquivalent(const FMetasoundFrontendVertexBehavior& InLHS, const FMetasoundFrontendVertexBehavior& InRHS);
+};
+
+// Metadata associated with a vertex.
+USTRUCT() 
+struct FMetasoundFrontendVertexMetadata
+{
+	GENERATED_BODY()
+		
+	// Display name for a vertex
+	UPROPERTY(EditAnywhere, Category = Parameters, meta = (DisplayName = "Name"))
+	FText DisplayName;
+
+	// Description of the vertex.
+	UPROPERTY(EditAnywhere, Category = Parameters)
+	FText Description;
+
+	// Keywords associated with the vertex
+	UPROPERTY()
+	TArray<FString> Keywords;
+
+	// Vertexes of the same group are generally placed together. 
+	UPROPERTY()
+	FString Group;
+
+	// If true, vertex is shown for advanced display.
+	UPROPERTY()
+	bool bIsAdvancedDisplay = false;
+};
+
+USTRUCT()
+struct FMetasoundFrontendEnvironmentVariableMetadata
+{
+	GENERATED_BODY()
+
+	// Display name for a environment variable
+	UPROPERTY()
+	FText DisplayName;
+
+	// Description of the environment variable
+	UPROPERTY()
+	FText Description;
+};
+
+
+
+USTRUCT() 
+struct METASOUNDFRONTEND_API FMetasoundFrontendClassVertex : public FMetasoundFrontendVertex
+{
+	GENERATED_BODY()
+
+	virtual ~FMetasoundFrontendClassVertex() = default;
+
+	UPROPERTY()
+	int32 NodeID = Metasound::FrontendInvalidID;
+
+	// Metadata associated with input.
+	UPROPERTY(EditAnywhere, Category = CustomView)
+	FMetasoundFrontendVertexMetadata Metadata;
+
+	// Behavior of input vertex.
+	UPROPERTY()
+	FMetasoundFrontendVertexBehavior Behavior;
+
+	static bool IsFunctionalEquivalent(const FMetasoundFrontendClassVertex& InLHS, const FMetasoundFrontendClassVertex& InRHS);
+};
+
+// Contains info for input vertex of a Metasound class.
+USTRUCT() 
+struct METASOUNDFRONTEND_API FMetasoundFrontendClassInput : public FMetasoundFrontendClassVertex
+{
+	GENERATED_BODY()
+
+	FMetasoundFrontendClassInput() = default;
+
+	FMetasoundFrontendClassInput(const FMetasoundFrontendClassVertex& InOther);
+
+	virtual ~FMetasoundFrontendClassInput() = default;
+
+	// Default values for vertex IDs in this input.
+	UPROPERTY(EditAnywhere, Category = Parameters)
+	TArray<FMetasoundFrontendVertexLiteral> Defaults;
+};
+
+USTRUCT()
+struct METASOUNDFRONTEND_API FMetasoundFrontendClassOutput : public FMetasoundFrontendClassVertex
+{
+	GENERATED_BODY()
+
+	FMetasoundFrontendClassOutput() = default;
+
+	FMetasoundFrontendClassOutput(const FMetasoundFrontendClassVertex& InOther)
+	:	FMetasoundFrontendClassVertex(InOther)
+	{
+	}
+
+	virtual ~FMetasoundFrontendClassOutput() = default;
+};
+
+USTRUCT() 
+struct METASOUNDFRONTEND_API FMetasoundFrontendEnvironmentVariable
+{
+	GENERATED_BODY()
+
+	virtual ~FMetasoundFrontendEnvironmentVariable() = default;
+
+	// Name of environment variable.
+	UPROPERTY()
+	FString Name;
+
+	// Type of environment variable.
+	UPROPERTY()
+	FName TypeName;
+
+	// Metadata of environment variable.
+	UPROPERTY()
+	FMetasoundFrontendEnvironmentVariableMetadata Metadata;
+};
+
+USTRUCT() 
+struct METASOUNDFRONTEND_API FMetasoundFrontendClassEnvironmentVariable : public FMetasoundFrontendEnvironmentVariable
+{
+	GENERATED_BODY()
+
+	virtual ~FMetasoundFrontendClassEnvironmentVariable() = default;
+
+	// True if the environment variable is needed in order to instantiate a node instance of the class. 
+	UPROPERTY()
+	bool bIsRequired = true;
+};
+
+// Layout mode for an interface.
+UENUM()
+enum class EMetasoundFrontendStyleInterfaceLayoutMode : uint8
+{
+	Default,
+	Inherited
+};
+
+
+// Style info of an interface.
+USTRUCT()
+struct FMetasoundFrontendInterfaceStyle
+{
+	GENERATED_BODY()
+
+	// Interface layoud mode
+	UPROPERTY()
+	EMetasoundFrontendStyleInterfaceLayoutMode LayoutMode = EMetasoundFrontendStyleInterfaceLayoutMode::Inherited;
+
+	// Default vertex order
+	UPROPERTY()
+	TArray<FString> DefaultOrder;
+};
+
+
+USTRUCT()
+struct FMetasoundFrontendClassInterface
+{
+	GENERATED_BODY()
+
+	// Style info for inputs.
+	UPROPERTY()
+	FMetasoundFrontendInterfaceStyle InputStyle;
+
+	// Style info for outputs.
+	UPROPERTY()
+	FMetasoundFrontendInterfaceStyle OutputStyle;
+
+	// Description of class inputs.
+	UPROPERTY(EditAnywhere, Category = CustomView)
+	TArray<FMetasoundFrontendClassInput> Inputs;
+
+	// Description of class outputs.
+	UPROPERTY(EditAnywhere, Category = CustomView)
+	TArray<FMetasoundFrontendClassOutput> Outputs;
+
+	// Description of class environment variables.
+	UPROPERTY(EditAnywhere, Category = CustomView)
+	TArray<FMetasoundFrontendClassEnvironmentVariable> Environment;
+};
+
+
+// Name of a Metasound class
+USTRUCT()
+struct METASOUNDFRONTEND_API FMetasoundFrontendClassName
+{
+	GENERATED_BODY()
+
+	// Namespace of class.
+	UPROPERTY()
+	FString Namespace;
+
+	// Name of class.
+	UPROPERTY()
+	FString Name;
+
+	// Variant of class. The Variant is used to describe an equivalent class which performs the same operation but on differing types.
+	UPROPERTY()
+	FString Variant;
+
+	// Returns a full name of the class.
+	FString GetFullName() const;
+
+	METASOUNDFRONTEND_API friend bool operator==(const FMetasoundFrontendClassName& InLHS, const FMetasoundFrontendClassName& InRHS);
+};
+
+
+USTRUCT()
+struct FMetasoundFrontendClassMetadata
+{
+	GENERATED_BODY()
+
+	UPROPERTY()
+	FMetasoundFrontendClassName Name;
+
+	UPROPERTY(EditAnywhere, Category = General)
+	FMetasoundFrontendVersionNumber Version;
+
+	UPROPERTY(VisibleAnywhere, Category = General)
+	EMetasoundFrontendClassType Type = EMetasoundFrontendClassType::Invalid;
+
+	UPROPERTY(EditAnywhere, Category = General, meta = (DisplayName = "Description"))
+	FText Description;
+
+	UPROPERTY()
+	FText PromptIfMissing;
+
+	UPROPERTY()
+	FText Author;
+
+	UPROPERTY()
+	TArray<FName> Keywords;
+};
+
+UENUM()
+enum class EMetasoundFrontendStyleNodeDisplay : uint8
+{
+	Default,
+	Inherited,
+	Minimized
+};
+
+USTRUCT() 
+struct FMetasoundFrontendClassStyle
+{
+	GENERATED_BODY()
+
+	UPROPERTY()
+	EMetasoundFrontendStyleNodeDisplay NodeDisplay = EMetasoundFrontendStyleNodeDisplay::Inherited;
+};
+
+
+
+USTRUCT()
+struct FMetasoundFrontendEditorData
+{
+	GENERATED_BODY()
+
+	UPROPERTY()
+	FMetasoundFrontendVersion Version;
+
+	UPROPERTY()
+	TArray<uint8> Data;
+};
+
+USTRUCT()
+struct METASOUNDFRONTEND_API FMetasoundFrontendClass
+{
+	GENERATED_BODY()
+
+	//Metasound::Frontend::TAccessPoint<FMetasoundFrontendClass> AccessPoint;
+	//Metasound::Frontend::FAccessPoint AccessPoint;
+
+	virtual ~FMetasoundFrontendClass() = default;
+
+	UPROPERTY()
+	int32 ID = Metasound::FrontendInvalidID;
+
+	UPROPERTY(EditAnywhere, Category = CustomView)
+	FMetasoundFrontendClassMetadata Metadata;
+
+	UPROPERTY(EditAnywhere, Category = CustomView)
+	FMetasoundFrontendClassInterface Interface;
+
+	UPROPERTY()
+	FMetasoundFrontendEditorData EditorData;
+
+	UPROPERTY()
+	FMetasoundFrontendClassStyle Style;
+};
+
+USTRUCT()
+struct METASOUNDFRONTEND_API FMetasoundFrontendGraphClass : public FMetasoundFrontendClass
+{
+	GENERATED_BODY()
+
+	FMetasoundFrontendGraphClass();
+
+	virtual ~FMetasoundFrontendGraphClass() = default;
+
+	UPROPERTY()
+	FMetasoundFrontendGraph Graph;
+
+};
+
+USTRUCT()
+struct FMetasoundFrontendDocumentMetadata
+{
+	GENERATED_BODY()
+
+	UPROPERTY()
+	FMetasoundFrontendVersion Format;
+};
+
+
+USTRUCT()
+struct FMetasoundFrontendArchetypeInterface
+{
+	GENERATED_BODY()
+
+	UPROPERTY()
+	TArray<FMetasoundFrontendClassVertex> Inputs;
+
+	UPROPERTY()
+	TArray<FMetasoundFrontendClassVertex> Outputs;
+
+	UPROPERTY()
+	TArray<FMetasoundFrontendEnvironmentVariable> Environment;
+};
+
+// This is used to describe the required inputs and outputs for a metasound, and is used to make sure we can use a metasound graph for specific applications.
+// For example, a UMetasoundSource needs to generate audio, so its RequiredOutputs will contain "MainAudioOutput"
+USTRUCT()
+struct FMetasoundFrontendArchetype
+{
+	GENERATED_BODY()
+
+	// Name of the archetype we're using.
+	UPROPERTY()
+	FName Name;
+
+	UPROPERTY()
+	FMetasoundFrontendVersionNumber Version;
+
+	UPROPERTY()
+	FMetasoundFrontendArchetypeInterface Interface;
+};
+
+
+USTRUCT()
+struct METASOUNDFRONTEND_API FMetasoundFrontendDocument
+{
+	GENERATED_BODY()
+
+	//Metasound::Frontend::TAccessPoint<FMetasoundFrontendDocument> AccessPoint;
+	Metasound::Frontend::FAccessPoint AccessPoint;
+
+	FMetasoundFrontendDocument();
+
+	UPROPERTY(EditAnywhere, Category = Metadata)
+	FMetasoundFrontendDocumentMetadata Metadata;
+
+	UPROPERTY(EditAnywhere, Category = CustomView)
+	FMetasoundFrontendGraphClass RootGraph;
+
+	UPROPERTY()
+	TArray<FMetasoundFrontendGraphClass> Subgraphs;
+
+
+	UPROPERTY()
+	FMetasoundFrontendEditorData EditorData;
+
+	UPROPERTY()
+	FMetasoundFrontendArchetype Archetype;
+
+	UPROPERTY()
+	TArray<FMetasoundFrontendClass> Dependencies;
+};
+
+
diff --git a/Engine/Plugins/Runtime/Metasound/Source/MetasoundFrontend/Public/MetasoundFrontendQuery.h b/Engine/Plugins/Runtime/Metasound/Source/MetasoundFrontend/Public/MetasoundFrontendQuery.h
--- a/Engine/Plugins/Runtime/Metasound/Source/MetasoundFrontend/Public/MetasoundFrontendQuery.h
+++ b/Engine/Plugins/Runtime/Metasound/Source/MetasoundFrontend/Public/MetasoundFrontendQuery.h
@@ -3,8 +3,8 @@
 #pragma once
 
 #include "CoreMinimal.h"
-#include "MetasoundFrontendDataLayout.h"
-#include "Misc/Variant.h"
+#include "MetasoundFrontendDocument.h"
+#include "Misc/TVariant.h"
 
 namespace Metasound
 {
@@ -13,7 +13,7 @@ namespace Metasound
 	 */
 	struct FFrontendQueryEntry
 	{
-		using FValue = TVariant<FMetasoundDocument, FMetasoundClassDescription, FMetasoundNodeDescription, FMetasoundGraphDescription>;
+		using FValue = TVariant<FMetasoundFrontendDocument, FMetasoundFrontendClass, FMetasoundFrontendNode>;
 		using FKey = int32;
 
 		static constexpr FKey InvalidKey = 0xFFFFFFFF;
diff --git a/Engine/Plugins/Runtime/Metasound/Source/MetasoundFrontend/Public/MetasoundFrontendRegistries.h b/Engine/Plugins/Runtime/Metasound/Source/MetasoundFrontend/Public/MetasoundFrontendRegistries.h
--- a/Engine/Plugins/Runtime/Metasound/Source/MetasoundFrontend/Public/MetasoundFrontendRegistries.h
+++ b/Engine/Plugins/Runtime/Metasound/Source/MetasoundFrontend/Public/MetasoundFrontendRegistries.h
@@ -5,20 +5,21 @@
 #include "CoreMinimal.h"
 #include "IAudioProxyInitializer.h"
 #include "MetasoundDataReference.h"
-#include "MetasoundFrontendDataLayout.h"
+#include "MetasoundFrontendDocument.h"
+#include "MetasoundLiteral.h"
 #include "MetasoundNodeInterface.h"
 #include "MetasoundOperatorInterface.h"
 
 namespace Metasound
 {
 	typedef TFunction<TUniquePtr<Metasound::INode>(::Metasound::FInputNodeConstructorParams&&)> FCreateInputNodeFunction;
-	typedef TFunction<TUniquePtr<Metasound::INode>(const ::Metasound::FOutputNodeConstrutorParams&)> FCreateOutputNodeFunction;
+	typedef TFunction<TUniquePtr<Metasound::INode>(const ::Metasound::FOutputNodeConstructorParams&)> FCreateOutputNodeFunction;
 
 	// This function is used to create a proxy from a datatype's base uclass.
 	typedef TFunction<Audio::IProxyDataPtr(UObject*)> FCreateAudioProxyFunction;
 
 	typedef TFunction<TUniquePtr<Metasound::INode>(const Metasound::FNodeInitData&)> FCreateMetasoundNodeFunction;
-	typedef TFunction<FMetasoundClassDescription()> FCreateMetasoundClassDescriptionFunction;
+	typedef TFunction<FMetasoundFrontendClass()> FCreateMetasoundFrontendClassFunction;
 
 	// Various elements that we pass to the frontend registry based on templated type traits.
 	struct FDataTypeRegistryInfo
@@ -27,7 +28,7 @@ namespace Metasound
 		FName DataTypeName;
 
 		// What type we should default to using for literals.
-		ELiteralArgType PreferredLiteralType = ELiteralArgType::Invalid;
+		ELiteralType PreferredLiteralType = ELiteralType::Invalid;
 
 		// This indicates the type can only be constructed with FOperatorSettings or the default constructor.
 		bool bIsDefaultParsable = false;
@@ -75,11 +76,11 @@ namespace Metasound
 			// This lambda can be used to get an INodeBase for this specific node class.
 			FCreateMetasoundNodeFunction CreateNode;
 
-			FCreateMetasoundClassDescriptionFunction CreateClassDescription;
+			FCreateMetasoundFrontendClassFunction CreateFrontendClass;
 
-			FNodeRegistryElement(FCreateMetasoundNodeFunction&& InCreateNodeFunction, FCreateMetasoundClassDescriptionFunction&& InCreateDescriptionFunction)
+			FNodeRegistryElement(FCreateMetasoundNodeFunction&& InCreateNodeFunction, FCreateMetasoundFrontendClassFunction&& InCreateDescriptionFunction)
 				: CreateNode(MoveTemp(InCreateNodeFunction))
-				, CreateClassDescription(InCreateDescriptionFunction)
+				, CreateFrontendClass(InCreateDescriptionFunction)
 			{
 			}
 		};
@@ -125,6 +126,7 @@ namespace Metasound
 			// A list of nodes that can perform a conversion between the two datatypes described in the FConverterNodeRegistryKey for this map element.
 			TArray<FConverterNodeInfo> PotentialConverterNodes;
 		};
+
 	}
 
 	struct FDataTypeConstructorCallbacks
@@ -132,9 +134,13 @@ namespace Metasound
 		// This constructs a TInputNode<> with the corresponding datatype.
 		FCreateInputNodeFunction CreateInputNode;
 
+		FCreateMetasoundFrontendClassFunction CreateFrontendInputClass;
+
 		// This constructs a TOutputNode<> with the corresponding datatype.
 		FCreateOutputNodeFunction CreateOutputNode;
 
+		FCreateMetasoundFrontendClassFunction CreateFrontendOutputClass;
+
 		// For datatypes that use a UObject literal or a UObject literal array, this lambda generates a literal from the corresponding UObject.
 		FCreateAudioProxyFunction CreateAudioProxy;
 	};
@@ -171,10 +177,10 @@ public:
 	TMap<FNodeRegistryKey, FNodeRegistryElement>& GetExternalNodeRegistry();
 
 	TUniquePtr<Metasound::INode> ConstructInputNode(const FName& InInputType, Metasound::FInputNodeConstructorParams&& InParams);
-	TUniquePtr<Metasound::INode> ConstructOutputNode(const FName& InOutputType, const Metasound::FOutputNodeConstrutorParams& InParams);
+	TUniquePtr<Metasound::INode> ConstructOutputNode(const FName& InOutputType, const Metasound::FOutputNodeConstructorParams& InParams);
 
-	Metasound::FDataTypeLiteralParam GenerateLiteralForUObject(const FName& InDataType, UObject* InObject);
-	Metasound::FDataTypeLiteralParam GenerateLiteralForUObjectArray(const FName& InDataType, TArray<UObject*> InObjectArray);
+	Metasound::FLiteral GenerateLiteralForUObject(const FName& InDataType, UObject* InObject);
+	Metasound::FLiteral GenerateLiteralForUObjectArray(const FName& InDataType, TArray<UObject*> InObjectArray);
 
 	// Create a new instance of a C++ implemented node from the registry.
 	TUniquePtr<Metasound::INode> ConstructExternalNode(const FName& InNodeType, uint32 InNodeHash, const Metasound::FNodeInitData& InInitData);
@@ -184,24 +190,29 @@ public:
 	TArray<FConverterNodeInfo> GetPossibleConverterNodes(const FName& FromDataType, const FName& ToDataType);
 
 	// Get the desired kind of literal for a given data type. Returns EConstructorArgType::Invalid if the data type couldn't be found.
-	Metasound::ELiteralArgType GetDesiredLiteralTypeForDataType(FName InDataType) const;
+	Metasound::ELiteralType GetDesiredLiteralTypeForDataType(FName InDataType) const;
 
 	UClass* GetLiteralUClassForDataType(FName InDataType) const;
 
+	template<typename ArgType>
+	bool DoesDataTypeSupportLiteralType(FName InDataType) const
+	{
+		return DoesDataTypeSupportLiteralType(InDataType, Metasound::TLiteralTypeInfo<ArgType>::GetLiteralArgTypeEnum());
+	}
 
 	// Get whether we can build a literal of this specific type for InDataType.
-	bool DoesDataTypeSupportLiteralType(FName InDataType, Metasound::ELiteralArgType InLiteralType) const;
+	bool DoesDataTypeSupportLiteralType(FName InDataType, Metasound::ELiteralType InLiteralType) const;
 
-	bool RegisterDataType(const FDataTypeRegistryInfo& InDataInfo, FDataTypeConstructorCallbacks&& InCallbacks);
+	bool RegisterDataType(const FDataTypeRegistryInfo& InDataInfo, const FDataTypeConstructorCallbacks& InCallbacks);
 
 	/** Register external node with the frontend.
 	 *
 	 * @param InCreateNode - Function for creating node from FNodeInitData.
-	 * @param InCreateDescription - Function for creating a FMetasoundClassDescription.
+	 * @param InCreateDescription - Function for creating a FMetasoundFrontendClass.
 	 *
 	 * @return True on success.
 	 */
-	bool RegisterExternalNode(Metasound::FCreateMetasoundNodeFunction&& InCreateNode, Metasound::FCreateMetasoundClassDescriptionFunction&& InCreateDescription);
+	bool RegisterExternalNode(Metasound::FCreateMetasoundNodeFunction&& InCreateNode, Metasound::FCreateMetasoundFrontendClassFunction&& InCreateDescription);
 
 	bool RegisterConversionNode(const FConverterNodeRegistryKey& InNodeKey, const FConverterNodeInfo& InNodeInfo);
 
@@ -211,7 +222,11 @@ public:
 	}
 
 	static FNodeRegistryKey GetRegistryKey(const FNodeInfo& InNodeMetadata);
+	static FNodeRegistryKey GetRegistryKey(const FMetasoundFrontendClassMetadata& InNodeMetadata);
 	static bool GetRegistryKey(const FNodeRegistryElement& InElement, FNodeRegistryKey& OutKey);
+	static bool GetFrontendClassFromRegistered(const FMetasoundFrontendClassMetadata& InMetadata, FMetasoundFrontendClass& OutClass);
+	static bool GetInputNodeClassMetadataForDataType(const FName& InDataTypeName, FMetasoundFrontendClassMetadata& OutMetadata);
+	static bool GetOutputNodeClassMetadataForDataType(const FName& InDataTypeName, FMetasoundFrontendClassMetadata& OutMetadata);
 
 	// Return any data types that can be used as a metasound input type or output type.
 	TArray<FName> GetAllValidDataTypes();
@@ -249,6 +264,7 @@ private:
 	};
 
 	TMap<FName, FDataTypeRegistryElement> DataTypeRegistry;
+	TMap<FNodeRegistryKey, FDataTypeRegistryElement> DataTypeNodeRegistry;
 };
 
 
diff --git a/Engine/Plugins/Runtime/Metasound/Source/MetasoundFrontend/Public/MetasoundFrontendTransform.h b/Engine/Plugins/Runtime/Metasound/Source/MetasoundFrontend/Public/MetasoundFrontendTransform.h
new file mode 100644
--- /dev/null
+++ b/Engine/Plugins/Runtime/Metasound/Source/MetasoundFrontend/Public/MetasoundFrontendTransform.h
@@ -0,0 +1,26 @@
+// Copyright Epic Games, Inc. All Rights Reserved.
+
+#pragma once
+
+#include "MetasoundFrontendController.h"
+
+namespace Metasound
+{
+	namespace Frontend
+	{
+		/** Interface for transforms applied to documents. */
+		class IDocumentTransform
+		{
+			public:
+				virtual ~IDocumentTransform() = default;
+				virtual bool Transform(FDocumentHandle InDocument) const = 0;
+		};
+
+		/** Adds or swaps root graph inputs and outputs to match the document's archetype. */
+		class METASOUNDFRONTEND_API FMatchRootGraphToArchetype : public IDocumentTransform
+		{
+			public:
+				bool Transform(FDocumentHandle InDocument) const override;
+		};
+	}
+}
diff --git a/Engine/Plugins/Runtime/Metasound/Source/MetasoundFrontend/Public/MetasoundInputNode.h b/Engine/Plugins/Runtime/Metasound/Source/MetasoundFrontend/Public/MetasoundInputNode.h
--- a/Engine/Plugins/Runtime/Metasound/Source/MetasoundFrontend/Public/MetasoundInputNode.h
+++ b/Engine/Plugins/Runtime/Metasound/Source/MetasoundFrontend/Public/MetasoundInputNode.h
@@ -61,7 +61,7 @@ namespace Metasound
 
 			using FDataWriteReference = TDataWriteReference<DataType>;
 
-			TInputOperatorLiteralFactory(FDataTypeLiteralParam&& InInitParam)
+			TInputOperatorLiteralFactory(FLiteral&& InInitParam)
 			:	InitParam(MoveTemp(InInitParam))
 			{
 			}
@@ -70,7 +70,7 @@ namespace Metasound
 
 		private:
 
-			FDataTypeLiteralParam InitParam;
+			FLiteral InitParam;
 	};
 
 	/** TInputOperatorFactory initializes the DataType at construction. It uses
@@ -149,7 +149,7 @@ namespace Metasound
 
 			/* Construct a TInputNode using the TInputOperatorLiteralFactory<> and moving
 			 * InParam to the TInputOperatorLiteralFactory constructor.*/
-			explicit TInputNode(const FString& InNodeDescription, const FString& InVertexName, FDataTypeLiteralParam&& InParam)
+			explicit TInputNode(const FString& InNodeDescription, const FString& InVertexName, FLiteral&& InParam)
 			:	FNode(InNodeDescription, GetNodeInfo(InVertexName))
 			,	VertexName(InVertexName)
 			,	Interface(DeclareVertexInterface(InVertexName))
diff --git a/Engine/Plugins/Runtime/Metasound/Source/MetasoundFrontend/Public/MetasoundNodeRegistrationMacro.h b/Engine/Plugins/Runtime/Metasound/Source/MetasoundFrontend/Public/MetasoundNodeRegistrationMacro.h
--- a/Engine/Plugins/Runtime/Metasound/Source/MetasoundFrontend/Public/MetasoundNodeRegistrationMacro.h
+++ b/Engine/Plugins/Runtime/Metasound/Source/MetasoundFrontend/Public/MetasoundNodeRegistrationMacro.h
@@ -3,10 +3,11 @@
 #pragma once
 
 #include "CoreMinimal.h"
-#include "MetasoundNodeInterface.h"
-#include "Templates/IntegralConstant.h"
 #include "MetasoundFrontend.h"
+#include "MetasoundFrontendDocument.h"
 #include "MetasoundFrontendRegistries.h"
+#include "MetasoundNodeInterface.h"
+#include "Templates/IntegralConstant.h"
 
 namespace Metasound
 {
@@ -44,7 +45,7 @@ namespace Metasound
 			{
 				return TUniquePtr<Metasound::INode>(new FNodeType(InInitData));
 			},
-			[=]() -> FMetasoundClassDescription
+			[=]() -> FMetasoundFrontendClass
 			{
 				return Metasound::Frontend::GenerateClassDescription(InMetadata);
 			}
diff --git a/Engine/Plugins/Runtime/Metasound/Source/MetasoundFrontend/Public/MetasoundRouter.h b/Engine/Plugins/Runtime/Metasound/Source/MetasoundFrontend/Public/MetasoundRouter.h
--- a/Engine/Plugins/Runtime/Metasound/Source/MetasoundFrontend/Public/MetasoundRouter.h
+++ b/Engine/Plugins/Runtime/Metasound/Source/MetasoundFrontend/Public/MetasoundRouter.h
@@ -242,7 +242,7 @@ namespace Metasound
 
 		virtual FName GetDataType() { return FName(); };
 
-		virtual bool PushLiteral(const FDataTypeLiteralParam& InParam) { return false; }
+		virtual bool PushLiteral(const FLiteral& InParam) { return false; }
 
 	protected:
 		virtual TUniquePtr<IReceiver> ConstructNewReceiverImplementation(const FReceiverInitParams& InitParams) { return nullptr; };
@@ -523,7 +523,7 @@ namespace Metasound
 			return GetDataTypeName();
 		}
 
-		virtual bool PushLiteral(const FDataTypeLiteralParam& InParam)
+		virtual bool PushLiteral(const FLiteral& InParam)
 		{
 			if (TLiteralTraits<TDataType>::IsParsable(InParam))
 			{
@@ -899,7 +899,7 @@ namespace Metasound
 			}
 		}
 
-		bool PushLiteral(FName InChannelName, const FDataTypeLiteralParam& InParam)
+		bool PushLiteral(FName InChannelName, const FLiteral& InParam)
 		{
 			FScopeLock ScopeLock(&DataChannelMapMutationLock);
 			if (TSharedRef<IDataChannel>* ExistingChannelPtr = DataChannelMap.Find(InChannelName))
@@ -1114,7 +1114,7 @@ namespace Metasound
 
 		// Pushes a literal parameter to a specific data channel in the global router.
 		// returns false if the literal type isn't supported.
-		bool PushLiteral(FName GlobalChannelName, const FDataTypeLiteralParam& InParam)
+		bool PushLiteral(FName GlobalChannelName, const FLiteral& InParam)
 		{
 			return GlobalRouter.PushLiteral(GlobalChannelName, InParam);
 		}
diff --git a/Engine/Plugins/Runtime/Metasound/Source/MetasoundGraphCore/Public/MetasoundDataFactory.h b/Engine/Plugins/Runtime/Metasound/Source/MetasoundGraphCore/Public/MetasoundDataFactory.h
--- a/Engine/Plugins/Runtime/Metasound/Source/MetasoundGraphCore/Public/MetasoundDataFactory.h
+++ b/Engine/Plugins/Runtime/Metasound/Source/MetasoundGraphCore/Public/MetasoundDataFactory.h
@@ -528,7 +528,7 @@ namespace Metasound
 	template<typename DataType>
 	struct TLiteralTraits
 	{
-		using TVariantConstructorTraits = DataFactoryPrivate::TDataTypeVariantConstructorTraits<DataType, FDataTypeLiteralParam::FVariantType>;
+		using TVariantConstructorTraits = DataFactoryPrivate::TDataTypeVariantConstructorTraits<DataType, FLiteral::FVariantType>;
 
 		static constexpr bool bIsParsableFromAnyLiteralType = TVariantConstructorTraits::bIsDefaultConstructible || TVariantConstructorTraits::bIsConstructibleWithSettings || TVariantConstructorTraits::bIsConstructibleWithArgs || TVariantConstructorTraits::bIsConstructibleWithSettingsAndArgs;
 
@@ -542,35 +542,35 @@ namespace Metasound
 		 * an FOperatorSettings with the literals constructor arg type, or one that
 		 * accpets the literal constructor arg type. It returns False otherwise.
 		 */
-		static const bool IsParsable(const FDataTypeLiteralParam& InLiteral)
+		static const bool IsParsable(const FLiteral& InLiteral)
 		{
-			switch (InLiteral.ConstructorArgType)
+			switch (InLiteral.GetType())
 			{
-				case ELiteralArgType::Boolean:
+				case ELiteralType::Boolean:
 
 					return TIsParsable<DataType, bool>::Value;
 
-				case ELiteralArgType::Integer:
+				case ELiteralType::Integer:
 
 					return TIsParsable<DataType, int32>::Value;
 
-				case ELiteralArgType::Float:
+				case ELiteralType::Float:
 
 					return TIsParsable<DataType, float>::Value;
 
-				case ELiteralArgType::String:
+				case ELiteralType::String:
 
 					return TIsParsable<DataType, FString>::Value;
 
-				case ELiteralArgType::UObjectProxy:
+				case ELiteralType::UObjectProxy:
 
 					return TIsParsable<DataType, Audio::IProxyDataPtr>::Value;
 
-				case ELiteralArgType::UObjectProxyArray:
+				case ELiteralType::UObjectProxyArray:
 
 					return TIsParsable<DataType, TArray<Audio::IProxyDataPtr>>::Value;
 
-				case ELiteralArgType::None:
+				case ELiteralType::None:
 
 					return TIsParsable<DataType>::Value;
 
@@ -702,7 +702,7 @@ namespace Metasound
 	/** A base factory type for creating objects related to Metasound DataTypes. 
 	 *
 	 * TDataLiteralFactory provides a unified interface for constructing objects using 
-	 * a FOperatorSettings and FDataTypeLiteralParam. The various factory 
+	 * a FOperatorSettings and FLiteral. The various factory 
 	 * methods determine how strictly the DataType constructors must match the 
 	 * arguments to the factory method.
 	 *
@@ -742,44 +742,44 @@ namespace Metasound
 		 * @return An object related to the DataType. The exact type determined by
 		 *         the CreatorType class template parameter.
 		 */
-		static auto CreateExplicitArgs(const FOperatorSettings& InSettings, const FDataTypeLiteralParam& InLiteral)
+		static auto CreateExplicitArgs(const FOperatorSettings& InSettings, const FLiteral& InLiteral)
 		{
-			switch (InLiteral.ConstructorArgType)
+			switch (InLiteral.GetType())
 			{
-				case ELiteralArgType::Boolean:
+				case ELiteralType::Boolean:
 
-					return CreateExplicitArgs<bool>(InSettings, InLiteral.ConstructorArg);
+					return CreateExplicitArgs<bool>(InSettings, InLiteral.Value);
 
-				case ELiteralArgType::Integer:
+				case ELiteralType::Integer:
 
-					return CreateExplicitArgs<int>(InSettings, InLiteral.ConstructorArg);
+					return CreateExplicitArgs<int>(InSettings, InLiteral.Value);
 
-				case ELiteralArgType::Float:
+				case ELiteralType::Float:
 
-					return CreateExplicitArgs<float>(InSettings, InLiteral.ConstructorArg);
+					return CreateExplicitArgs<float>(InSettings, InLiteral.Value);
 
-				case ELiteralArgType::String:
+				case ELiteralType::String:
 
-					return CreateExplicitArgs<FString>(InSettings, InLiteral.ConstructorArg);
+					return CreateExplicitArgs<FString>(InSettings, InLiteral.Value);
 
-				case ELiteralArgType::UObjectProxy:
+				case ELiteralType::UObjectProxy:
 
-					return CreateExplicitArgs<Audio::IProxyDataPtr>(InSettings, InLiteral.ConstructorArg);
+					return CreateExplicitArgs<Audio::IProxyDataPtr>(InSettings, InLiteral.Value);
 
-				case ELiteralArgType::UObjectProxyArray:
+				case ELiteralType::UObjectProxyArray:
 
-					return CreateExplicitArgs<TArray<Audio::IProxyDataPtr>>(InSettings, InLiteral.ConstructorArg);
+					return CreateExplicitArgs<TArray<Audio::IProxyDataPtr>>(InSettings, InLiteral.Value);
 
-				case ELiteralArgType::None:
+				case ELiteralType::None:
 
-					return CreateExplicitArgs<void>(InSettings, InLiteral.ConstructorArg);
+					return CreateExplicitArgs<void>(InSettings, InLiteral.Value);
 
-				case ELiteralArgType::Invalid:
+				case ELiteralType::Invalid:
 				default:
 
 					checkNoEntry();
 
-					return CreateAny(InSettings, InLiteral.ConstructorArg);
+					return CreateAny(InSettings, InLiteral.Value);
 			}
 		}
 	};
diff --git a/Engine/Plugins/Runtime/Metasound/Source/MetasoundGraphCore/Public/MetasoundLiteral.h b/Engine/Plugins/Runtime/Metasound/Source/MetasoundGraphCore/Public/MetasoundLiteral.h
--- a/Engine/Plugins/Runtime/Metasound/Source/MetasoundGraphCore/Public/MetasoundLiteral.h
+++ b/Engine/Plugins/Runtime/Metasound/Source/MetasoundGraphCore/Public/MetasoundLiteral.h
@@ -7,34 +7,11 @@
 #include "IAudioProxyInitializer.h"
 #include "MetasoundOperatorSettings.h" 
 #include "Misc/TVariant.h"
+#include <type_traits>
 
 namespace Metasound
 {
-	/*
-	const TCHAR* ToString(EConstructorArgType InType)
-	{
-		switch (InType)
-		{
-		case Metasound::EConstructorArgType::Boolean:
-			return TEXT("Boolean");
-			break;
-		case Metasound::EConstructorArgType::Integer:
-			return TEXT("Integer");
-			break;
-		case Metasound::EConstructorArgType::Float:
-			return TEXT("Float");
-			break;
-		case Metasound::EConstructorArgType::String:
-			return TEXT("String");
-			break;
-		default:
-			return TEXT("");
-			break;
-		}
-	}
-	*/
-
-	enum class ELiteralArgType : uint8
+	enum class ELiteralType : uint8
 	{
 		Boolean,
 		Integer,
@@ -46,9 +23,94 @@ namespace Metasound
 		Invalid,
 	};
 
+	namespace MetasoundLiteralIntrinsics
+	{
+		// Default template for converting a template parameter to a literal argument type
+		template<typename... ArgTypes>
+		ELiteralType GetLiteralArgTypeFromDecayed()
+		{
+			return ELiteralType::Invalid;
+		}
+
+		template<>
+		FORCEINLINE ELiteralType GetLiteralArgTypeFromDecayed<bool>()
+		{
+			return ELiteralType::Boolean;
+		}
+
+		template<>
+		FORCEINLINE ELiteralType GetLiteralArgTypeFromDecayed<int32>()
+		{
+			return ELiteralType::Integer;
+		}
+
+		template<>
+		FORCEINLINE ELiteralType GetLiteralArgTypeFromDecayed<float>()
+		{
+			return ELiteralType::Float;
+		}
+
+		template<>
+		FORCEINLINE ELiteralType GetLiteralArgTypeFromDecayed<FString>()
+		{
+			return ELiteralType::String;
+		}
+
+		template<>
+		FORCEINLINE ELiteralType GetLiteralArgTypeFromDecayed<Audio::IProxyDataPtr>()
+		{
+			return ELiteralType::UObjectProxy;
+		}
+
+		template<>
+		FORCEINLINE ELiteralType GetLiteralArgTypeFromDecayed<TArray<Audio::IProxyDataPtr>>()
+		{
+			return ELiteralType::UObjectProxyArray;
+		}
+
+		template<>
+		FORCEINLINE ELiteralType GetLiteralArgTypeFromDecayed<>()
+		{
+			return ELiteralType::None;
+		}
+	}
+
+
+	/** Provides literal type information for a given type. 
+	 *
+	 * @tparam ArgType - A C++ type.
+	 */
+	template<typename... ArgTypes>
+	struct TLiteralTypeInfo
+	{
+		/** Returns the associated ELiteralType for the C++ type provided in the TLiteralTypeInfo<Type> */
+		static const ELiteralType GetLiteralArgTypeEnum()
+		{
+			// Use decayed version of template arg to remove references and cv qualifiers. 
+			return MetasoundLiteralIntrinsics::GetLiteralArgTypeFromDecayed<typename std::decay<ArgTypes...>::type>();
+		}
+	};
+
+	/** Provides literal type information for a given type. 
+	 *
+	 * @tparam ArgType - A C++ type.
+	 */
+	template<>
+	struct TLiteralTypeInfo<>
+	{
+		/** Returns the associated ELiteralType for the C++ type provided in the TLiteralTypeInfo<Type> */
+		static const ELiteralType GetLiteralArgTypeEnum()
+		{
+			// Use decayed version of template arg to remove references and cv qualifiers. 
+			return MetasoundLiteralIntrinsics::GetLiteralArgTypeFromDecayed<>();
+		}
+	};
+
+
 	/**
+	 * FLiteral 
 	 * Convenience wrapper for safely invoking the correct constructor.
-	 * This can be used for int32, float, bool, or FString.
+	 * This can be used for int32, float, bool, FString, Audio::IProxyData, or TArray<Audio::IProxyData>
 	 *
 	 * For example:
 	 * // Somewhere before DECLARE_METASOUND_DATA_REFERENCE_TYPES(FAudioBuffer...), ParseFrom<FAudioBuffer>(int32) is defined
@@ -60,7 +122,7 @@ namespace Metasound
 	 *
 	 * //...
 	 * // In the frontend, we know that int and pass it to somewhere in MetasoundGraphCore...
-	 * FDataTypeLiteralParam InitParam(512);
+	 * FLiteral InitParam(512);
 	 *
 	 * //...
 	 * // In the backend, we can safely construct an FAudioBuffer.
@@ -70,68 +132,83 @@ namespace Metasound
 	 *     //...
 	 * }
 	 */
-	struct FDataTypeLiteralParam
+	struct FLiteral
 	{
-		FDataTypeLiteralParam(const FDataTypeLiteralParam& Other) = delete;
-		FDataTypeLiteralParam(FDataTypeLiteralParam&& Other) = default;
-		FDataTypeLiteralParam& operator=(FDataTypeLiteralParam&& Other) = default;
+		struct FInvalidLiteralState {};
+		struct FEmptyLiteralState {};
+
+		FLiteral(const FLiteral& Other) = delete;
+		FLiteral(FLiteral&& Other) = default;
+		FLiteral& operator=(FLiteral&& Other) = default;
 
-		typedef TVariant<bool, int, float, FString, Audio::IProxyDataPtr, TArray<Audio::IProxyDataPtr>> FVariantType;
+		typedef TVariant<FInvalidLiteralState, FEmptyLiteralState, bool, int, float, FString, Audio::IProxyDataPtr, TArray<Audio::IProxyDataPtr>> FVariantType;
 
-		ELiteralArgType ConstructorArgType;
-		FVariantType ConstructorArg;
+		FVariantType Value;
 
-		// builds an invalid FDataTypeLiteralParam.
-		static FDataTypeLiteralParam InvalidParam()
+		// builds an invalid FLiteral.
+		static FLiteral CreateInvalid()
 		{
-			FDataTypeLiteralParam InitParam;
-			InitParam.ConstructorArgType = ELiteralArgType::Invalid;
-			return MoveTemp(InitParam);
+			return FLiteral(FInvalidLiteralState());
 		}
 
 		bool IsValid() const
 		{
-			return ConstructorArgType != ELiteralArgType::Invalid;
+			return GetType() != ELiteralType::Invalid;
 		}
 
-		FDataTypeLiteralParam()
-			: ConstructorArgType(ELiteralArgType::None)
-		{}
-
-		FDataTypeLiteralParam(bool bInValue)
-			: ConstructorArgType(ELiteralArgType::Boolean)
+		ELiteralType GetType() const
 		{
-			ConstructorArg.Set<bool>(bInValue);
-		}
+			switch (Value.GetIndex())
+			{
+				case 0:
+					return ELiteralType::Invalid;
 
-		FDataTypeLiteralParam(int32 InValue)
-			: ConstructorArgType(ELiteralArgType::Integer)
-		{
-			ConstructorArg.Set<int32>(InValue);
-		}
+				case 1:
+					return ELiteralType::None;
 
-		FDataTypeLiteralParam(float InValue)
-			: ConstructorArgType(ELiteralArgType::Float)
-		{
-			ConstructorArg.Set<float>(InValue);
-		}
+				case 2:
+					return ELiteralType::Boolean;
 
-		FDataTypeLiteralParam(const FString& InString)
-			: ConstructorArgType(ELiteralArgType::String)
-		{
-			ConstructorArg.Set<FString>(InString);
+				case 3:
+					return ELiteralType::Integer;
+
+				case 4:
+					return ELiteralType::Float;
+
+				case 5:
+					return ELiteralType::String;
+
+				case 6:
+					return ELiteralType::UObjectProxy;
+
+				case 7:
+					return ELiteralType::UObjectProxyArray;
+
+				default:
+					checkNoEntry();
+					return ELiteralType::Invalid;
+			}
 		}
 
-		FDataTypeLiteralParam(Audio::IProxyDataPtr&& InProxy)
-			: ConstructorArgType(ELiteralArgType::UObjectProxy)
+
+		/** Construct a literal param with a single argument. */
+		template<
+			typename... ArgTypes,
+			typename std::enable_if<1 == sizeof...(ArgTypes), int>::type = 0
+			>
+		FLiteral(ArgTypes&&... Args)
 		{
-			ConstructorArg.Set<Audio::IProxyDataPtr>(MoveTemp(InProxy));
+			Value.Set<typename std::decay<ArgTypes...>::type>(Forward<ArgTypes>(Args)...);
 		}
 
-		FDataTypeLiteralParam(TArray<Audio::IProxyDataPtr>&& InProxyArray)
-			: ConstructorArgType(ELiteralArgType::UObjectProxyArray)
+		/** Construct a literal param with no arguments. */
+		template<
+			typename... ArgTypes,
+			typename std::enable_if<0 == sizeof...(ArgTypes), int>::type = 0
+			>
+		FLiteral(ArgTypes&&... Args)
 		{
-			ConstructorArg.Set<TArray<Audio::IProxyDataPtr>>(MoveTemp(InProxyArray));
+			Value.Set<FEmptyLiteralState>(FEmptyLiteralState());
 		}
 
 	};
@@ -143,10 +220,10 @@ namespace Metasound
 		const FString& InNodeName;
 		const FString& InVertexName;
 
-		FDataTypeLiteralParam InitParam;
+		FLiteral InitParam;
 	};
 
-	struct FOutputNodeConstrutorParams
+	struct FOutputNodeConstructorParams
 	{
 		// the instance name and name of the specific connection that we should use.
 		const FString& InNodeName;
diff --git a/Engine/Plugins/Runtime/Metasound/Source/MetasoundGraphCore/Public/MetasoundNodeInterface.h b/Engine/Plugins/Runtime/Metasound/Source/MetasoundGraphCore/Public/MetasoundNodeInterface.h
--- a/Engine/Plugins/Runtime/Metasound/Source/MetasoundGraphCore/Public/MetasoundNodeInterface.h
+++ b/Engine/Plugins/Runtime/Metasound/Source/MetasoundGraphCore/Public/MetasoundNodeInterface.h
@@ -22,7 +22,7 @@ namespace Metasound
 	struct FNodeInitData
 	{
 		FString InstanceName;
-		TMap<FName, FDataTypeLiteralParam> ParamMap;
+		TMap<FName, FLiteral> ParamMap;
 
 		template<typename ParamType>
 		ParamType GetParamValue(FName ParamName)
@@ -33,6 +33,14 @@ namespace Metasound
 		}
 	};
 
+	// TODO: update FNodeInfo with this as class name.
+	struct FNodeClassName
+	{
+		FString Namespace;
+		FString Name;
+		FString Variant;
+	};
+
 	/** Provides metadata for a given node. */
 	struct FNodeInfo
 	{
diff --git a/Engine/Plugins/Runtime/Metasound/Source/MetasoundStandardNodes/Private/MetasoundAudioFormats.cpp b/Engine/Plugins/Runtime/Metasound/Source/MetasoundStandardNodes/Private/MetasoundAudioFormats.cpp
--- a/Engine/Plugins/Runtime/Metasound/Source/MetasoundStandardNodes/Private/MetasoundAudioFormats.cpp
+++ b/Engine/Plugins/Runtime/Metasound/Source/MetasoundStandardNodes/Private/MetasoundAudioFormats.cpp
@@ -3,16 +3,15 @@
 #include "MetasoundAudioFormats.h"
 
 #include "CoreMinimal.h"
-#include "MetasoundDataReference.h"
 #include "MetasoundAudioBuffer.h"
+#include "MetasoundDataReference.h"
+#include "MetasoundLiteral.h"
+#include "MetasoundDataTypeRegistrationMacro.h"
 
-REGISTER_METASOUND_DATATYPE(Metasound::FUnformattedAudio, "Audio:Unformatted")
-REGISTER_METASOUND_DATATYPE(Metasound::FMonoAudioFormat, "Audio:Mono")
-REGISTER_METASOUND_DATATYPE(Metasound::FStereoAudioFormat, "Audio:Stereo")
-
-// FMultichannelAudio cannot be used as an input type because it's channel count
-// cannot be changed at runtime. Hence, it is not registered.
-DEFINE_METASOUND_DATA_TYPE(Metasound::FMultichannelAudioFormat, "Audio:Multichannel")
+REGISTER_METASOUND_DATATYPE(Metasound::FUnformattedAudio, "Audio:Unformatted");
+REGISTER_METASOUND_DATATYPE(Metasound::FMonoAudioFormat, "Audio:Mono");
+REGISTER_METASOUND_DATATYPE(Metasound::FStereoAudioFormat, "Audio:Stereo");
+REGISTER_METASOUND_DATATYPE(Metasound::FMultichannelAudioFormat, "Audio:Multichannel", ELiteralType::Integer);
 
 namespace Metasound
 {
