Metasound Frontend Revamp

diff --git a/Engine/Plugins/Runtime/Metasound/Source/MetasoundStandardNodes/Private/MetasoundStandardNodesModule.cpp b/Engine/Plugins/Runtime/Metasound/Source/MetasoundStandardNodes/Private/MetasoundStandardNodesModule.cpp
--- a/Engine/Plugins/Runtime/Metasound/Source/MetasoundStandardNodes/Private/MetasoundStandardNodesModule.cpp
+++ b/Engine/Plugins/Runtime/Metasound/Source/MetasoundStandardNodes/Private/MetasoundStandardNodesModule.cpp
@@ -10,11 +10,11 @@
 #include "MetasoundTime.h"
 #include "MetasoundFrequency.h"
 
-REGISTER_METASOUND_DATATYPE(Metasound::FBop, "Primitive:Bop", ::Metasound::ELiteralArgType::Boolean)
-REGISTER_METASOUND_DATATYPE(Metasound::FFloatTime, "Primitive:Time", ::Metasound::ELiteralArgType::Float)
-REGISTER_METASOUND_DATATYPE(Metasound::FDoubleTime, "Primitive:Time:HighResolution", ::Metasound::ELiteralArgType::Float)
-REGISTER_METASOUND_DATATYPE(Metasound::FSampleTime, "Primitive:Time:SampleResolution", ::Metasound::ELiteralArgType::Float)
-REGISTER_METASOUND_DATATYPE(Metasound::FFrequency, "Primitive:Frequency", ::Metasound::ELiteralArgType::Float)
+REGISTER_METASOUND_DATATYPE(Metasound::FBop, "Primitive:Bop", ::Metasound::ELiteralType::Boolean)
+REGISTER_METASOUND_DATATYPE(Metasound::FFloatTime, "Primitive:Time", ::Metasound::ELiteralType::Float)
+REGISTER_METASOUND_DATATYPE(Metasound::FDoubleTime, "Primitive:Time:HighResolution", ::Metasound::ELiteralType::Float)
+REGISTER_METASOUND_DATATYPE(Metasound::FSampleTime, "Primitive:Time:SampleResolution", ::Metasound::ELiteralType::Float)
+REGISTER_METASOUND_DATATYPE(Metasound::FFrequency, "Primitive:Frequency", ::Metasound::ELiteralType::Float)
 
 namespace Metasound 
 {
