Hiding copy constructor for incompatle tarray element types in a metasound data type.

diff --git a/Engine/Plugins/Runtime/Metasound/Source/MetasoundGraphCore/Public/MetasoundDataFactory.h b/Engine/Plugins/Runtime/Metasound/Source/MetasoundGraphCore/Public/MetasoundDataFactory.h
--- a/Engine/Plugins/Runtime/Metasound/Source/MetasoundGraphCore/Public/MetasoundDataFactory.h
+++ b/Engine/Plugins/Runtime/Metasound/Source/MetasoundGraphCore/Public/MetasoundDataFactory.h
@@ -27,6 +27,26 @@ namespace Metasound
 			static constexpr bool bIsConstructibleWithSettingsAndArgs = std::is_constructible<DataType, const FOperatorSettings&, ArgTypes...>::value;
 		};
 
+		/** Partial specialization for TArray<> types. Current implementation of the TArray<>
+		 * exposes a copy constructor which is declared even if the TArray<> element types
+		 * are incompatible. Because the copy constructor is declared, std::is_constructible 
+		 * returns true even if the constructor will not compile. This partial specialization
+		 * corrects the TDataTypeConstructorTraits so that the TArray<> copy constructor is
+		 * not used if the underlying element types are incompatible. 
+		 */
+		template<typename ElementType, typename OtherElementType>
+		struct TDataTypeConstructorTraits<TArray<ElementType>, TArray<OtherElementType>>
+		{
+			using DataType = TArray<ElementType>;
+
+			static constexpr bool bIsDefaultConstructible = std::is_constructible<DataType>::value;
+			static constexpr bool bIsConstructibleWithSettings = std::is_constructible<DataType, const FOperatorSettings&>::value;
+			// TArray copy constructor is only compilable if the element types are compatible.
+			static constexpr bool bIsConstructibleWithArgs = std::is_constructible<ElementType, OtherElementType>::value;
+			// TArray does not take operator settings.
+			static constexpr bool bIsConstructibleWithSettingsAndArgs = false;
+		};
+
 		/** Constructor traits for a for variant inputs. These traits test each
 		 * individual type supported by the variant and determines if *any* of the
 		 * types can be used as a single argument to construct the DataType.
