Metasound Math Nodes

diff --git a/Engine/Plugins/Runtime/Metasound/Source/MetasoundGenerator/Private/MetasoundGenerator.cpp b/Engine/Plugins/Runtime/Metasound/Source/MetasoundGenerator/Private/MetasoundGenerator.cpp
--- a/Engine/Plugins/Runtime/Metasound/Source/MetasoundGenerator/Private/MetasoundGenerator.cpp
+++ b/Engine/Plugins/Runtime/Metasound/Source/MetasoundGenerator/Private/MetasoundGenerator.cpp
@@ -3,16 +3,13 @@
 
 #include "DSP/Dsp.h"
 #include "MetasoundBop.h"
-#include "MetasoundOperatorInterface.h"
-#include "MetasoundOscNode.h"
-#include "MetasoundAudioMultiplyNode.h"
-#include "MetasoundADSRNode.h"
-#include "MetasoundPeriodicTriggerNode.h"
-#include "MetasoundOutputNode.h"
+#include "MetasoundFrequency.h"
+#include "MetasoundGraph.h"
 #include "MetasoundInputNode.h"
 #include "MetasoundOperatorBuilder.h"
-#include "MetasoundGraph.h"
-#include "MetasoundFrequency.h"
+#include "MetasoundOperatorInterface.h"
+#include "MetasoundOutputNode.h"
+
 
 namespace Metasound
 {
diff --git a/Engine/Plugins/Runtime/Metasound/Source/MetasoundGraphCore/Public/MetasoundDataReference.h b/Engine/Plugins/Runtime/Metasound/Source/MetasoundGraphCore/Public/MetasoundDataReference.h
--- a/Engine/Plugins/Runtime/Metasound/Source/MetasoundGraphCore/Public/MetasoundDataReference.h
+++ b/Engine/Plugins/Runtime/Metasound/Source/MetasoundGraphCore/Public/MetasoundDataReference.h
@@ -25,7 +25,7 @@
 	typedef ::Metasound::TDataReferenceTypeInfo<DataType> DataTypeInfoTypeName; \
 	\
 	typedef ::Metasound::TDataReadReference<DataType> DataReadReferenceTypeName; \
-	typedef ::Metasound::TDataWriteReference<DataType> DataWriteReferenceTypeName; \
+	typedef ::Metasound::TDataWriteReference<DataType> DataWriteReferenceTypeName;
 
 // This only needs to be called if you don't plan on calling REGISTER_METASOUND_DATATYPE.
 #define DEFINE_METASOUND_DATA_TYPE(DataType, DataTypeName) \
diff --git a/Engine/Plugins/Runtime/Metasound/Source/MetasoundStandardNodes/Private/MetasoundAudioMultiplyNode.cpp b/Engine/Plugins/Runtime/Metasound/Source/MetasoundStandardNodes/Private/MetasoundAudioMultiplyNode.cpp
deleted file mode 100644
diff --git a/Engine/Plugins/Runtime/Metasound/Source/MetasoundStandardNodes/Private/MetasoundGainMultiplyNode.cpp b/Engine/Plugins/Runtime/Metasound/Source/MetasoundStandardNodes/Private/MetasoundGainMultiplyNode.cpp
new file mode 100644
--- /dev/null
+++ b/Engine/Plugins/Runtime/Metasound/Source/MetasoundStandardNodes/Private/MetasoundGainMultiplyNode.cpp
@@ -0,0 +1,122 @@
+// Copyright Epic Games, Inc. All Rights Reserved.
+
+#include "MetasoundGainMultiplyNode.h"
+
+#include "DSP/BufferVectorOperations.h"
+#include "MetasoundAudioBuffer.h"
+#include "MetasoundExecutableOperator.h"
+#include "MetasoundFacade.h"
+#include "MetasoundNodeRegistrationMacro.h"
+#include "MetasoundOperatorSettings.h"
+#include "MetasoundPrimitives.h"
+
+#define LOCTEXT_NAMESPACE "MetasoundGainMultiplyNode"
+
+namespace Metasound
+{
+	class FGainMultiplyOperator : public TExecutableOperator<FGainMultiplyOperator>
+	{
+		public:
+			static const FNodeInfo& GetNodeInfo();
+
+			static FVertexInterface DeclareVertexInterface();
+
+			static TUniquePtr<IOperator> CreateOperator(const FCreateOperatorParams& InParams, FBuildErrorArray& OutErrors);
+
+			FGainMultiplyOperator(const FOperatorSettings& InSettings, const FAudioBufferReadRef& InBuffer, const FFloatReadRef& InGain)
+			:	Buffer(InBuffer)
+			,	Gain(InGain)
+			,	OutputBuffer(FAudioBufferWriteRef::CreateNew(InSettings))
+			,	BlockSize(InSettings.GetNumFramesPerBlock())
+			{
+				check(OutputBuffer->Num() == BlockSize);
+				check(Buffer->Num() == BlockSize);
+			}
+
+			virtual FDataReferenceCollection GetInputs() const override
+			{
+				FDataReferenceCollection InputDataReferences;
+				return InputDataReferences;
+			}
+
+			virtual FDataReferenceCollection GetOutputs() const override
+			{
+				FDataReferenceCollection OutputDataReferences;
+				OutputDataReferences.AddDataReadReference(TEXT("Out"), FAudioBufferReadRef(OutputBuffer));
+				return OutputDataReferences;
+			}
+
+			void Execute()
+			{
+				FMemory::Memcpy(OutputBuffer->GetData(), Buffer->GetData(), sizeof(float) * BlockSize);
+				Audio::MultiplyBufferByConstantInPlace(*OutputBuffer, *Gain);
+			}
+
+		private:
+			FAudioBufferReadRef Buffer;
+			FFloatReadRef Gain;
+			FAudioBufferWriteRef OutputBuffer;
+			int32 BlockSize;
+	};
+
+	FGainMultiplyNode::FGainMultiplyNode(const FString& InInstanceName)
+	:	FNodeFacade(InInstanceName, TFacadeOperatorClass<FGainMultiplyOperator>())
+	{
+	}
+
+	FGainMultiplyNode::FGainMultiplyNode(const FNodeInitData& InInitData)
+	:	FGainMultiplyNode(InInitData.InstanceName)
+	{
+	}
+
+	FVertexInterface FGainMultiplyOperator::DeclareVertexInterface()
+	{
+		static const FVertexInterface Interface(
+			FInputVertexInterface(
+				TInputDataVertexModel<FAudioBuffer>(TEXT("Buffer"), LOCTEXT("InputBuffer1Tooltip", "Input buffer.")),
+				TInputDataVertexModel<float>(TEXT("Gain"), LOCTEXT("InputGainTooltip", "Input gain."))
+			),
+			FOutputVertexInterface(
+				TOutputDataVertexModel<FAudioBuffer>(TEXT("Out"), LOCTEXT("OutpuBufferTooltip", "The output audio."))
+			)
+		);
+
+		return Interface;
+	}
+
+	TUniquePtr<IOperator> FGainMultiplyOperator::CreateOperator(const FCreateOperatorParams& InParams, FBuildErrorArray& OutErrors)
+	{
+		const FGainMultiplyNode& GainMultiplyNode = static_cast<const FGainMultiplyNode&>(InParams.Node);
+		const FDataReferenceCollection& InputCol = InParams.InputDataReferences;
+
+		FAudioBufferReadRef InputBuffer = InputCol.GetDataReadReferenceOrConstruct<FAudioBuffer>(TEXT("Buffer"), FAudioBuffer(InParams.OperatorSettings.GetNumFramesPerBlock()));
+		FFloatReadRef InputGain = InputCol.GetDataReadReferenceOrConstruct<float>(TEXT("Gain"), 1.0f);
+
+		return MakeUnique<FGainMultiplyOperator>(InParams.OperatorSettings, InputBuffer, InputGain);
+	}
+
+	const FNodeInfo& FGainMultiplyOperator::GetNodeInfo()
+	{
+		auto InitNodeInfo = []() -> FNodeInfo
+		{
+			FNodeInfo Info;
+			Info.ClassName = FName(TEXT("Gain Multiply"));
+			Info.MajorVersion = 1;
+			Info.MinorVersion = 0;
+			Info.Description = LOCTEXT("GainMultiply_NodeDescription", "Multiply audio buffer by float."),
+			Info.Author = PluginAuthor;
+			Info.PromptIfMissing = PluginNodeMissingPrompt;
+			Info.DefaultInterface = DeclareVertexInterface();
+
+			return Info;
+		};
+
+		static const FNodeInfo Info = InitNodeInfo();
+
+		return Info;
+	};
+
+	METASOUND_REGISTER_NODE(FGainMultiplyNode);
+}
+
+#undef LOCTEXT_NAMESPACE //MetasoundGainMultiplyNode
diff --git a/Engine/Plugins/Runtime/Metasound/Source/MetasoundStandardNodes/Private/MetasoundMathNodes.cpp b/Engine/Plugins/Runtime/Metasound/Source/MetasoundStandardNodes/Private/MetasoundMathNodes.cpp
new file mode 100644
--- /dev/null
+++ b/Engine/Plugins/Runtime/Metasound/Source/MetasoundStandardNodes/Private/MetasoundMathNodes.cpp
@@ -0,0 +1,296 @@
+// Copyright Epic Games, Inc. All Rights Reserved.
+#include "MetasoundMathNodes.h"
+
+#include "DSP/BufferVectorOperations.h"
+#include "MetasoundAudioBuffer.h"
+#include "MetasoundDataReference.h"
+#include "MetasoundFrequency.h"
+#include "MetasoundNodeInterface.h"
+#include "MetasoundTime.h"
+
+#define LOCTEXT_NAMESPACE "MetasoundMathNodes"
+namespace Metasound
+{
+// Specialized Math Operations
+	template <>
+	class TMathOpAdd<FAudioBuffer>
+	{
+	public:
+		static FAudioBuffer GetDefault(const FOperatorSettings& InSettings)
+		{
+			return FAudioBuffer(InSettings.GetNumFramesPerBlock());
+		}
+
+		static void Calculate(const TArray<FAudioBufferReadRef>& InOperands, FAudioBufferWriteRef& OutResult)
+		{
+			if (!ensure(!InOperands.IsEmpty()))
+			{
+				OutResult->Zero();
+				return;
+			}
+
+			const int32 NumSamples = InOperands[0]->Num();
+			if (!ensure(NumSamples == OutResult->Num()))
+			{
+				// TODO: Error
+				OutResult->Zero();
+				return;
+			}
+
+			FMemory::Memcpy(OutResult->GetData(), InOperands[0]->GetData(), sizeof(float) * NumSamples);
+
+			for (int32 i = 1; i < InOperands.Num(); ++i)
+			{
+				const FAudioBuffer& Operand = *InOperands[i];
+				if (!ensure(NumSamples == Operand.Num()))
+				{
+					// TODO: Error
+					OutResult->Zero();
+					return;
+				}
+
+				const float* OperandData = InOperands[i]->GetData();
+				float* OutData = OutResult->GetData();
+
+				if (NumSamples % AUDIO_SIMD_FLOAT_ALIGNMENT)
+				{
+					for (int32 SampleIndex = 0; SampleIndex < NumSamples; ++SampleIndex)
+					{
+						OutData[SampleIndex] += OperandData[SampleIndex];
+					}
+				}
+				else
+				{
+					const float StartGain = 1.0f;
+					const float EndGain = 1.0f;
+					Audio::MixInBufferFast(OperandData, OutData, NumSamples, StartGain, EndGain);
+				}
+			}
+		}
+	};
+
+	template <>
+	class TMathOpSubtract<FAudioBuffer>
+	{
+	public:
+		static FAudioBuffer GetDefault(const FOperatorSettings& InSettings)
+		{
+			return FAudioBuffer(InSettings.GetNumFramesPerBlock());
+		}
+
+		static void Calculate(TArray<FAudioBufferReadRef>& InOperands, FAudioBufferWriteRef& OutResult)
+		{
+			if (!ensure(!InOperands.IsEmpty()))
+			{
+				OutResult->Zero();
+				return;
+			}
+
+			const int32 NumSamples = InOperands[0]->Num();
+			if (!ensure(NumSamples == OutResult->Num()))
+			{
+				// TODO: Error
+				OutResult->Zero();
+				return;
+			}
+
+			FMemory::Memcpy(OutResult->GetData(), InOperands[0]->GetData(), sizeof(float) * NumSamples);
+
+			for (int32 i = 1; i < InOperands.Num(); ++i)
+			{
+				const FAudioBuffer& Operand = *InOperands[i];
+				if (!ensure(NumSamples == Operand.Num()))
+				{
+					// TODO: Error
+					OutResult->Zero();
+					return;
+				}
+
+				const float* OperandData = InOperands[i]->GetData();
+				float* OutData = OutResult->GetData();
+
+				if (NumSamples % AUDIO_SIMD_FLOAT_ALIGNMENT)
+				{
+					for (int32 SampleIndex = 0; SampleIndex < NumSamples; ++SampleIndex)
+					{
+						OutData[i] -= OperandData[i];
+					}
+				}
+				else
+				{
+					Audio::BufferSubtractInPlace2Fast(OutData, OperandData, NumSamples);
+				}
+			}
+		}
+	};
+
+	template <>
+	class TMathOpMultiply<FAudioBuffer>
+	{
+	public:
+		static FAudioBuffer GetDefault(const FOperatorSettings& InSettings)
+		{
+			return FAudioBuffer(InSettings.GetNumFramesPerBlock());
+		}
+
+		static void Calculate(const TArray<FAudioBufferReadRef>& InOperands, FAudioBufferWriteRef& OutResult)
+		{
+			if (!ensure(!InOperands.IsEmpty()))
+			{
+				OutResult->Zero();
+				return;
+			}
+
+			const int32 NumSamples = InOperands[0]->Num();
+			if (!ensure(NumSamples == OutResult->Num()))
+			{
+				// TODO: Error
+				OutResult->Zero();
+				return;
+			}
+
+			FMemory::Memcpy(OutResult->GetData(), InOperands[0]->GetData(), sizeof(float) * NumSamples);
+
+			for (int32 i = 1; i < InOperands.Num(); ++i)
+			{
+				const FAudioBuffer& Operand = *InOperands[i];
+				if (!ensure(NumSamples == Operand.Num()))
+				{
+					// TODO: Error
+					OutResult->Zero();
+					return;
+				}
+
+				const float* OperandData = InOperands[i]->GetData();
+				float* OutData = OutResult->GetData();
+
+				if (NumSamples % AUDIO_SIMD_FLOAT_ALIGNMENT)
+				{
+					for (int32 SampleIndex = 0; SampleIndex < NumSamples; ++SampleIndex)
+					{
+						OutData[SampleIndex] += OperandData[SampleIndex];
+					}
+				}
+				else
+				{
+					Audio::MultiplyBuffersInPlace(OperandData, OutData, NumSamples);
+				}
+			}
+		}
+	};
+
+	template <>
+	class TMathOpMultiply<FFloatTime>
+	{
+	public:
+		static FFloatTime GetDefault(const FOperatorSettings& InSettings)
+		{
+			return FFloatTime(1.0f);
+		}
+
+		static void Calculate(const TArray<FFloatTimeReadRef>& InOperands, FFloatTimeWriteRef& OutResult)
+		{
+			if (!ensure(!InOperands.IsEmpty()))
+			{
+				return;
+			}
+
+			*OutResult = *InOperands[0];
+			for (int32 i = 1; i < InOperands.Num(); ++i)
+			{
+				OutResult->SetSeconds(OutResult->GetSeconds() * InOperands[i]->GetSeconds());
+			}
+		}
+	};
+
+	template <>
+	class TMathOpDivide<FFloatTime>
+	{
+	public:
+		static FFloatTime GetDefault(const FOperatorSettings& InSettings)
+		{
+			return FFloatTime(1.0f);
+		}
+
+		static void Calculate(const TArray<FFloatTimeReadRef>& InOperands, FFloatTimeWriteRef& OutResult)
+		{
+			if (!ensure(!InOperands.IsEmpty()))
+			{
+				return;
+			}
+
+			*OutResult = *InOperands[0];
+			for (int32 i = 1; i < InOperands.Num(); ++i)
+			{
+				const FFloatTime& OperandValue = *InOperands[i];
+				if (OperandValue.GetSeconds() == 0.0f)
+				{
+					// TODO: Error here
+					return;
+				}
+
+				OutResult->SetSeconds(OutResult->GetSeconds() / OperandValue.GetSeconds());
+			}
+		}
+	};
+
+	template <>
+	class TMathOpRandRange<FFloatTime>
+	{
+	public:
+		static FFloatTime GetDefault(const FOperatorSettings& InSettings)
+		{
+			return FFloatTime(0.0f);
+		}
+
+		static void Calculate(const TArray<FFloatTimeReadRef>& InOperands, FFloatTimeWriteRef& OutResult)
+		{
+			if (InOperands.Num() > 1)
+			{
+				const FFloatTime& Min = *InOperands[0];
+				const FFloatTime& Max = *InOperands[1];
+				const float FloatResult = FMath::RandRange(Min.GetSeconds(), Max.GetSeconds());
+				OutResult->SetSeconds(FloatResult);
+			}
+			else
+			{
+				// TODO: Error
+			}
+		}
+	};
+
+// Definitions
+	DEFINE_METASOUND_MATHOP(Add, Float, float, LOCTEXT("Metasound_MathAddFloatNodeDescription", "Adds floats."))
+	DEFINE_METASOUND_MATHOP(Add, Int32, int32, LOCTEXT("Metasound_MathAddInt32NodeDescription", "Adds int32s."))
+	DEFINE_METASOUND_MATHOP(Add, Int64, int64, LOCTEXT("Metasound_MathAddInt64NodeDescription", "Adds int64s."))
+	DEFINE_METASOUND_MATHOP(Add, Frequency, FFrequency, LOCTEXT("Metasound_MathAddFrequencyNodeDescription", "Adds frequency values."))
+	DEFINE_METASOUND_MATHOP(Add, FloatTime, FFloatTime, LOCTEXT("Metasound_MathAddFloatTimeNodeDescription", "Adds time values."))
+	DEFINE_METASOUND_MATHOP(Add, AudioBuffer, FAudioBuffer, LOCTEXT("Metasound_MathAddBufferNodeDescription", "Adds buffers together by sample."))
+
+	DEFINE_METASOUND_MATHOP(Subtract, Float, float, LOCTEXT("Metasound_MathSubractFloatNodeDescription", "Subtracts floats."))
+	DEFINE_METASOUND_MATHOP(Subtract, Int32, int32, LOCTEXT("Metasound_MathSubractInt32NodeDescription", "Subtracts int32s."))
+	DEFINE_METASOUND_MATHOP(Subtract, Int64, int64, LOCTEXT("Metasound_MathSubractInt64NodeDescription", "Subtracts int64s."))
+	DEFINE_METASOUND_MATHOP(Subtract, Frequency, FFrequency, LOCTEXT("Metasound_MathSubractFrequencyNodeDescription", "Subtracts frequency values."))
+	DEFINE_METASOUND_MATHOP(Subtract, FloatTime, FFloatTime, LOCTEXT("Metasound_MathSubractFloatTimeNodeDescription", "Subtracts time values."))
+	DEFINE_METASOUND_MATHOP(Subtract, AudioBuffer, FAudioBuffer, LOCTEXT("Metasound_MathSubtractBufferNodeDescription", "Subtracts buffers by sample."))
+
+	DEFINE_METASOUND_MATHOP(Multiply, Float, float, LOCTEXT("Metasound_MathMultiplyFloatNodeDescription", "Multiplies floats."))
+	DEFINE_METASOUND_MATHOP(Multiply, Int32, int32, LOCTEXT("Metasound_MathMultiplyInt32NodeDescription", "Multiplies int32s."))
+	DEFINE_METASOUND_MATHOP(Multiply, Int64, int64, LOCTEXT("Metasound_MathMultiplyInt64NodeDescription", "Multiplies int64s."))
+	DEFINE_METASOUND_MATHOP(Multiply, Frequency, FFrequency, LOCTEXT("Metasound_MathMultiplyFrequencyNodeDescription", "Multiplies frequency values."))
+	DEFINE_METASOUND_MATHOP(Multiply, FloatTime, FFloatTime, LOCTEXT("Metasound_MathMultiplyFloatTimeNodeDescription", "Multiplies time values."))
+	DEFINE_METASOUND_MATHOP(Multiply, AudioBuffer, FAudioBuffer, LOCTEXT("Metasound_MathMultiplyBufferNodeDescription", "Multiplies buffers together by sample."))
+
+	DEFINE_METASOUND_MATHOP(Divide, Float, float, LOCTEXT("Metasound_MathDivideFloatNodeDescription", "Divide float by another float."))
+	DEFINE_METASOUND_MATHOP(Divide, Int32, int32, LOCTEXT("Metasound_MathDivideInt32NodeDescription", "Divide int32 by another int32."))
+	DEFINE_METASOUND_MATHOP(Divide, Int64, int64, LOCTEXT("Metasound_MathDivideInt64NodeDescription", "Divide int64 by another int64."))
+	DEFINE_METASOUND_MATHOP(Divide, Frequency, FFrequency, LOCTEXT("Metasound_MathDivideFrequencyNodeDescription", "Divide frequency by another frequency."))
+	DEFINE_METASOUND_MATHOP(Divide, FloatTime, FFloatTime, LOCTEXT("Metasound_MathDivideFloatTimeNodeDescription", "Divide time by another time."))
+
+	DEFINE_METASOUND_MATHOP(RandRange, Float, float, LOCTEXT("Metasound_MathRandRangeFloatNodeDescription", "Computes random float value in the range [Min, Max]."))
+	DEFINE_METASOUND_MATHOP(RandRange, Int32, int32, LOCTEXT("Metasound_MathRandRangeInt32NodeDescription", "Computes random int32 in the range [Min, Max]."))
+	DEFINE_METASOUND_MATHOP(RandRange, Int64, int64, LOCTEXT("Metasound_MathRandRangeInt64NodeDescription", "Computes random int64 in the range [Min, Max]."))
+	DEFINE_METASOUND_MATHOP(RandRange, Frequency, FFrequency, LOCTEXT("Metasound_MathRandRangeFrequencyNodeDescription", "Computes random frequency value in the range [Min, Max]."))
+	DEFINE_METASOUND_MATHOP(RandRange, FloatTime, FFloatTime, LOCTEXT("Metasound_MathRandRangeFloatTimeNodeDescription", "Computes random time value in the range [Min, Max]."))
+}
+#undef LOCTEXT_NAMESPACE // MetasoundMathNodes
diff --git a/Engine/Plugins/Runtime/Metasound/Source/MetasoundStandardNodes/Public/MetasoundAudioMultiplyNode.h b/Engine/Plugins/Runtime/Metasound/Source/MetasoundStandardNodes/Public/MetasoundAudioMultiplyNode.h
deleted file mode 100644
diff --git a/Engine/Plugins/Runtime/Metasound/Source/MetasoundStandardNodes/Public/MetasoundFrequency.h b/Engine/Plugins/Runtime/Metasound/Source/MetasoundStandardNodes/Public/MetasoundFrequency.h
--- a/Engine/Plugins/Runtime/Metasound/Source/MetasoundStandardNodes/Public/MetasoundFrequency.h
+++ b/Engine/Plugins/Runtime/Metasound/Source/MetasoundStandardNodes/Public/MetasoundFrequency.h
@@ -66,10 +66,40 @@ namespace Metasound
 				: FFrequency(InValue, EFrequencyResolution::Hertz)
 			{}
 
+			/** Implicit operator used for math operations. */
+			operator float() const
+			{
+				return FrequencyInHz;
+			}
+
+			FFrequency& operator+=(const FFrequency& InOther)
+			{
+				FrequencyInHz += InOther.GetHertz();
+				return *this;
+			}
+
+			FFrequency& operator-=(const FFrequency& InOther)
+			{
+				FrequencyInHz += InOther.GetHertz();
+				return *this;
+			}
+
+			FFrequency& operator*=(const FFrequency& InOther)
+			{
+				FrequencyInHz *= InOther.GetHertz();
+				return *this;
+			}
+
+			FFrequency& operator/=(const FFrequency& InOther)
+			{
+				FrequencyInHz /= InOther.GetHertz();
+				return *this;
+			}
+
 			/** Set the frequency in hertz. */
 			void SetHertz(float InHz)
 			{
-				FrequencyInHz = InHz;	
+				FrequencyInHz = InHz;
 			}
 
 			/** Set the frequency in kilohertz. */
@@ -81,7 +111,7 @@ namespace Metasound
 			/** Set the frequency in megahertz. */
 			void SetMegahertz(float InMHz)
 			{
-				FrequencyInHz = InMHz * 1e6f; 
+				FrequencyInHz = InMHz * 1e6f;
 			}
 
 			/** Set the frequency in radians per a sample.
diff --git a/Engine/Plugins/Runtime/Metasound/Source/MetasoundStandardNodes/Public/MetasoundGainMultiplyNode.h b/Engine/Plugins/Runtime/Metasound/Source/MetasoundStandardNodes/Public/MetasoundGainMultiplyNode.h
new file mode 100644
--- /dev/null
+++ b/Engine/Plugins/Runtime/Metasound/Source/MetasoundStandardNodes/Public/MetasoundGainMultiplyNode.h
@@ -0,0 +1,18 @@
+// Copyright Epic Games, Inc. All Rights Reserved.
+#pragma once
+
+#include "CoreMinimal.h"
+#include "Internationalization/Text.h"
+
+#include "MetasoundFacade.h"
+
+
+namespace Metasound
+{
+	class METASOUNDSTANDARDNODES_API FGainMultiplyNode : public FNodeFacade
+	{
+		public:
+			FGainMultiplyNode(const FString& InInstanceName);
+			FGainMultiplyNode(const FNodeInitData& InInitData);
+	};
+} // namespace Metasound
diff --git a/Engine/Plugins/Runtime/Metasound/Source/MetasoundStandardNodes/Public/MetasoundMathNodes.h b/Engine/Plugins/Runtime/Metasound/Source/MetasoundStandardNodes/Public/MetasoundMathNodes.h
new file mode 100644
--- /dev/null
+++ b/Engine/Plugins/Runtime/Metasound/Source/MetasoundStandardNodes/Public/MetasoundMathNodes.h
@@ -0,0 +1,274 @@
+// Copyright Epic Games, Inc. All Rights Reserved.
+
+#pragma once
+
+#include "CoreMinimal.h"
+#include "MetasoundBop.h"
+#include "MetasoundBuilderInterface.h"
+#include "MetasoundDataReferenceCollection.h"
+#include "MetasoundExecutableOperator.h"
+#include "MetasoundFacade.h"
+#include "MetasoundNode.h"
+#include "MetasoundNodeRegistrationMacro.h"
+#include "MetasoundNodeInterface.h"
+#include "MetasoundOperatorInterface.h"
+#include "MetasoundPrimitives.h"
+
+
+#define LOCTEXT_NAMESPACE "MetasoundMathOpNode"
+namespace Metasound
+{
+	template <typename TDerivedClass, typename TDataClass>
+	class TMathOpNode : public FNodeFacade
+	{
+	private:
+		class TMathOperator : public TExecutableOperator<TMathOperator>
+		{
+		private:
+			using TDataClassReadRef = TDataReadReference<TDataClass>;
+			using TDataClassWriteRef = TDataWriteReference<TDataClass>;
+
+			TArray<TDataClassReadRef> OperandRefs;
+			TDataClassWriteRef ValueRef;
+
+		public:
+			static FVertexInterface DeclareVertexInterface()
+			{
+				static const FVertexInterface Interface(
+					FInputVertexInterface(
+						TInputDataVertexModel<TDataClass>(TEXT("Operand1"), LOCTEXT("MathOpMinTooltip", "First value to operate on.")),
+						TInputDataVertexModel<TDataClass>(TEXT("Operand2"), LOCTEXT("MathOpMaxTooltip", "Second value to operate on."))
+					),
+					FOutputVertexInterface(
+						TOutputDataVertexModel<TDataClass>(TEXT("Out"), LOCTEXT("MathOpOutTooltip", "Math operation result"))
+					)
+				);
+
+				return Interface;
+			}
+
+			static const FNodeInfo& GetNodeInfo()
+			{
+				auto InitNodeInfo = []() -> FNodeInfo
+				{
+					FNodeInfo Info;
+					Info.ClassName = TDerivedClass::GetClassName();
+					Info.MajorVersion = 1;
+					Info.MinorVersion = 0;
+					Info.Description = TDerivedClass::GetDescription();
+					Info.Author = PluginAuthor;
+					Info.PromptIfMissing = PluginNodeMissingPrompt;
+					Info.DefaultInterface = DeclareVertexInterface();
+
+					return Info;
+				};
+
+				static const FNodeInfo Info = InitNodeInfo();
+				return Info;
+			}
+
+			void Execute()
+			{
+				TDerivedClass::Calculate(OperandRefs, ValueRef);
+			}
+
+			static TUniquePtr<IOperator> CreateOperator(const FCreateOperatorParams& InParams, FBuildErrorArray& OutErrors)
+			{
+				const TMathOpNode& MathOpNode = static_cast<const TMathOpNode&>(InParams.Node);
+
+				// TODO: Support dynamic number of inputs
+				TDataClassReadRef Op1 = InParams.InputDataReferences.GetDataReadReferenceOrConstruct<TDataClass>(TEXT("Operand1"), TDerivedClass::GetDefault(InParams.OperatorSettings));
+				TDataClassReadRef Op2 = InParams.InputDataReferences.GetDataReadReferenceOrConstruct<TDataClass>(TEXT("Operand2"), TDerivedClass::GetDefault(InParams.OperatorSettings));
+
+				TArray<TDataClassReadRef> OperandRefs = { Op1, Op2 };
+
+				return MakeUnique<TMathOperator>(InParams.OperatorSettings, OperandRefs);
+			}
+
+			TMathOperator(const FOperatorSettings& InSettings, const TArray<TDataClassReadRef>& InOperands)
+			:	OperandRefs(InOperands)
+			,	ValueRef(TDataClassWriteRef::CreateNew(TDerivedClass::GetDefault(InSettings)))
+			{
+			}
+
+			virtual FDataReferenceCollection GetInputs() const override
+			{
+				FDataReferenceCollection InputDataReferences;
+				InputDataReferences.AddDataReadReference(TEXT("Operand1"), OperandRefs[0]);
+				InputDataReferences.AddDataReadReference(TEXT("Operand2"), OperandRefs[1]);
+				return InputDataReferences;
+			}
+
+			virtual FDataReferenceCollection GetOutputs() const override
+			{
+				FDataReferenceCollection OutputDataReferences;
+				OutputDataReferences.AddDataReadReference(TEXT("Out"), TDataClassReadRef(ValueRef));
+
+				return OutputDataReferences;
+			}
+		};
+
+	public:
+		TMathOpNode(const FNodeInitData& InInitData)
+			: FNodeFacade(InInitData.InstanceName, TFacadeOperatorClass<TMathOperator>())
+		{
+		}
+	};
+
+// Standard Math Operations
+	template <typename TDataClass>
+	class TMathOpAdd
+	{
+	public:
+		using TDataClassReadRef = TDataReadReference<TDataClass>;
+		using TDataClassWriteRef = TDataWriteReference<TDataClass>;
+
+		static TDataClass GetDefault(const FOperatorSettings& InSettings)
+		{
+			return 0;
+		}
+
+		static void Calculate(const TArray<TDataClassReadRef>& InOperands, TDataClassWriteRef& OutResult)
+		{
+			if (!ensure(!InOperands.IsEmpty()))
+			{
+				return;
+			}
+
+			*OutResult = *InOperands[0];
+			for (int32 i = 1; i < InOperands.Num(); ++i)
+			{
+				*OutResult += *InOperands[i];
+			}
+		}
+	};
+
+	template <typename TDataClass>
+	class TMathOpSubtract
+	{
+	public:
+		using TDataClassReadRef = TDataReadReference<TDataClass>;
+		using TDataClassWriteRef = TDataWriteReference<TDataClass>;
+
+		static TDataClass GetDefault(const FOperatorSettings& InSettings)
+		{
+			return 0;
+		}
+
+		static void Calculate(const TArray<TDataClassReadRef>& InOperands, TDataClassWriteRef& OutResult)
+		{
+			if (!ensure(!InOperands.IsEmpty()))
+			{
+				return;
+			}
+
+			*OutResult = *InOperands[0];
+			for (int32 i = 1; i < InOperands.Num(); ++i)
+			{
+				*OutResult -= *InOperands[i];
+			}
+		}
+	};
+
+	template <typename TDataClass>
+	class TMathOpMultiply
+	{
+	public:
+		using TDataClassReadRef = TDataReadReference<TDataClass>;
+		using TDataClassWriteRef = TDataWriteReference<TDataClass>;
+
+		static TDataClass GetDefault(const FOperatorSettings& InSettings)
+		{
+			return 1;
+		}
+
+		static void Calculate(const TArray<TDataClassReadRef>& InOperands, TDataClassWriteRef& OutResult)
+		{
+			if (!ensure(!InOperands.IsEmpty()))
+			{
+				return;
+			}
+
+			*OutResult = *InOperands[0];
+			for (int32 i = 1; i < InOperands.Num(); ++i)
+			{
+				*OutResult *= *InOperands[i];
+			}
+		}
+	};
+
+	template <typename TDataClass>
+	class TMathOpDivide
+	{
+	public:
+		using TDataClassReadRef = TDataReadReference<TDataClass>;
+		using TDataClassWriteRef = TDataWriteReference<TDataClass>;
+
+		static TDataClass GetDefault(const FOperatorSettings& InSettings)
+		{
+			return 1;
+		}
+
+		static void Calculate(const TArray<TDataClassReadRef>& InOperands, TDataClassWriteRef& OutResult)
+		{
+			if (!ensure(!InOperands.IsEmpty()))
+			{
+				return;
+			}
+
+			*OutResult = *InOperands[0];
+			for (int32 i = 1; i < InOperands.Num(); ++i)
+			{
+				const TDataClass& OperandValue = *InOperands[i];
+				if (OperandValue == 0)
+				{
+					// TODO: Error here
+					return;
+				}
+
+				*OutResult /= OperandValue;
+			}
+		}
+	};
+
+	template <typename TDataClass>
+	class TMathOpRandRange
+	{
+	public:
+		using TDataClassReadRef = TDataReadReference<TDataClass>;
+		using TDataClassWriteRef = TDataWriteReference<TDataClass>;
+
+		static TDataClass GetDefault(const FOperatorSettings& InSettings)
+		{
+			return 0;
+		}
+
+		static void Calculate(const TArray<TDataClassReadRef>& InOperands, TDataClassWriteRef& OutResult)
+		{
+			if (InOperands.Num() > 1)
+			{
+				const TDataClass& Min = *InOperands[0];
+				const TDataClass& Max = *InOperands[1];
+				*OutResult = FMath::RandRange(Min, Max);
+			}
+			else
+			{
+				// TODO: Error
+			}
+		}
+	};
+}
+
+#define DEFINE_METASOUND_MATHOP(OpName, DataTypeName, DataClass, Description) \
+	class F##OpName##DataTypeName##Node \
+		: public ::Metasound::TMathOpNode<F##OpName##DataTypeName##Node, DataClass> \
+		, public ::Metasound::TMathOp##OpName<DataClass> \
+	{ \
+	public: \
+		static FName GetClassName() { return #OpName #DataTypeName; } \
+		static FText GetDescription() { return Description; } \
+		F##OpName##DataTypeName##Node(const FNodeInitData& InInitData) : TMathOpNode(InInitData) { } \
+	}; \
+	METASOUND_REGISTER_NODE(F##OpName##DataTypeName##Node)
+
+#undef LOCTEXT_NAMESPACE // MetasoundMathOpNode
