Fix for multiple metasound input/output vertices

diff --git a/Engine/Plugins/Runtime/Metasound/Source/MetasoundEditor/Private/MetasoundEditorGraphBuilder.cpp b/Engine/Plugins/Runtime/Metasound/Source/MetasoundEditor/Private/MetasoundEditorGraphBuilder.cpp
--- a/Engine/Plugins/Runtime/Metasound/Source/MetasoundEditor/Private/MetasoundEditorGraphBuilder.cpp
+++ b/Engine/Plugins/Runtime/Metasound/Source/MetasoundEditor/Private/MetasoundEditorGraphBuilder.cpp
@@ -96,42 +96,40 @@ namespace Metasound
 			return Categories;
 		}
 
-		FString FGraphBuilder::GenerateUniqueInputName(UObject& InMetasound, const FName InBaseName)
+		FString FGraphBuilder::GenerateUniqueInputName(const UObject& InMetasound, const FName InBaseName)
 		{
 			FString NameBase = GetDataTypeDisplayName(InBaseName);
-			FMetasoundAssetBase* MetasoundAsset = IMetasoundUObjectRegistry::Get().GetObjectAsAssetBase(&InMetasound);
+			const FMetasoundAssetBase* MetasoundAsset = IMetasoundUObjectRegistry::Get().GetObjectAsAssetBase(&InMetasound);
 			check(MetasoundAsset);
 
-			// TODO: To achieve const correctness, this needs a FConstGraphHandle.
-			const Frontend::FGraphHandle GraphHandle = MetasoundAsset->GetRootGraphHandle();
+			const Frontend::FConstGraphHandle GraphHandle = MetasoundAsset->GetRootGraphHandle();
 
 			int32 i = 1;
-			FString NewNodeName = NameBase + FString::Printf(TEXT("_%02d"), i);
-			while (GraphHandle->ContainsInputNodeWithName(NewNodeName))
+			FString NewInputName = NameBase + FString::Printf(TEXT("_%02d"), i);
+			while (GraphHandle->ContainsInputVertexWithName(NewInputName))
 			{
-				NewNodeName = NameBase + FString::Printf(TEXT("_%02d"), ++i);
+				NewInputName = NameBase + FString::Printf(TEXT("_%02d"), ++i);
 			}
 
-			return NewNodeName;
+			return NewInputName;
 		}
 
-		FString FGraphBuilder::GenerateUniqueOutputName(UObject& InMetasound, const FName InBaseName)
+		FString FGraphBuilder::GenerateUniqueOutputName(const UObject& InMetasound, const FName InBaseName)
 		{
 			FString NameBase = GetDataTypeDisplayName(InBaseName);
-			FMetasoundAssetBase* MetasoundAsset = IMetasoundUObjectRegistry::Get().GetObjectAsAssetBase(&InMetasound);
+			const FMetasoundAssetBase* MetasoundAsset = IMetasoundUObjectRegistry::Get().GetObjectAsAssetBase(&InMetasound);
 			check(MetasoundAsset);
 
-			// TODO: To achieve const correctness, this needs a FConstGraphHandle.
-			const Frontend::FGraphHandle GraphHandle = MetasoundAsset->GetRootGraphHandle();
+			const Frontend::FConstGraphHandle GraphHandle = MetasoundAsset->GetRootGraphHandle();
 
 			int32 i = 1;
-			FString NewNodeName = NameBase + FString::Printf(TEXT("_%02d"), i);
-			while (GraphHandle->ContainsOutputNodeWithName(NewNodeName))
+			FString NewOutputName = NameBase + FString::Printf(TEXT("_%02d"), i);
+			while (GraphHandle->ContainsOutputVertexWithName(NewOutputName))
 			{
-				NewNodeName = NameBase + FString::Printf(TEXT("_%02d"), ++i);
+				NewOutputName = NameBase + FString::Printf(TEXT("_%02d"), ++i);
 			}
 
-			return NewNodeName;
+			return NewOutputName;
 		}
 
 		UEdGraphNode* FGraphBuilder::AddInput(UObject& InMetasound, const FVector2D& Location, const FString& InName, const FName InTypeName, const FText& InToolTip, bool bInSelectNewNode)
diff --git a/Engine/Plugins/Runtime/Metasound/Source/MetasoundEditor/Private/MetasoundEditorGraphBuilder.h b/Engine/Plugins/Runtime/Metasound/Source/MetasoundEditor/Private/MetasoundEditorGraphBuilder.h
--- a/Engine/Plugins/Runtime/Metasound/Source/MetasoundEditor/Private/MetasoundEditorGraphBuilder.h
+++ b/Engine/Plugins/Runtime/Metasound/Source/MetasoundEditor/Private/MetasoundEditorGraphBuilder.h
@@ -38,9 +38,9 @@ namespace Metasound
 			// Adds a node handle with the given class info
 			static Frontend::FNodeHandle AddNodeHandle(UObject& InMetasound, const Frontend::FNodeClassInfo& InClassInfo);
 
-			static FString GenerateUniqueInputName(UObject& InMetasound, const FName InBaseName);
+			static FString GenerateUniqueInputName(const UObject& InMetasound, const FName InBaseName);
 
-			static FString GenerateUniqueOutputName(UObject& InMetasound, const FName InBaseName);
+			static FString GenerateUniqueOutputName(const UObject& InMetasound, const FName InBaseName);
 
 			static FString GetDataTypeDisplayName(const FName& InDataTypeName);
 
diff --git a/Engine/Plugins/Runtime/Metasound/Source/MetasoundFrontend/Private/MetasoundAssetBase.cpp b/Engine/Plugins/Runtime/Metasound/Source/MetasoundFrontend/Private/MetasoundAssetBase.cpp
--- a/Engine/Plugins/Runtime/Metasound/Source/MetasoundFrontend/Private/MetasoundAssetBase.cpp
+++ b/Engine/Plugins/Runtime/Metasound/Source/MetasoundFrontend/Private/MetasoundAssetBase.cpp
@@ -113,14 +113,12 @@ FMetasoundFrontendClassMetadata FMetasoundAssetBase::GetMetadata()
 	return GetDocumentChecked().RootGraph.Metadata;
 }
 
-/*
-Metasound::Frontend::FDescriptionAccessPoint FMetasoundAssetBase::GetGraphAccessPoint()
+Metasound::Frontend::FDocumentHandle FMetasoundAssetBase::GetDocumentHandle()
 {
-	return Metasound::Frontend::FDescriptionAccessPoint(GetDocument());
+	return Metasound::Frontend::IDocumentController::CreateDocumentHandle(GetDocument());
 }
-*/
 
-Metasound::Frontend::FDocumentHandle FMetasoundAssetBase::GetDocumentHandle()
+Metasound::Frontend::FConstDocumentHandle FMetasoundAssetBase::GetDocumentHandle() const
 {
 	return Metasound::Frontend::IDocumentController::CreateDocumentHandle(GetDocument());
 }
@@ -130,6 +128,11 @@ Metasound::Frontend::FGraphHandle FMetasoundAssetBase::GetRootGraphHandle()
 	return GetDocumentHandle()->GetRootGraph();
 }
 
+Metasound::Frontend::FConstGraphHandle FMetasoundAssetBase::GetRootGraphHandle() const
+{
+	return GetDocumentHandle()->GetRootGraph();
+}
+
 bool FMetasoundAssetBase::ImportFromJSON(const FString& InJSON)
 {
 	Metasound::Frontend::TAccessPtr<FMetasoundFrontendDocument> Document = GetDocument();
diff --git a/Engine/Plugins/Runtime/Metasound/Source/MetasoundFrontend/Private/MetasoundFrontendController.cpp b/Engine/Plugins/Runtime/Metasound/Source/MetasoundFrontend/Private/MetasoundFrontendController.cpp
--- a/Engine/Plugins/Runtime/Metasound/Source/MetasoundFrontend/Private/MetasoundFrontendController.cpp
+++ b/Engine/Plugins/Runtime/Metasound/Source/MetasoundFrontend/Private/MetasoundFrontendController.cpp
@@ -40,5 +40,11 @@ namespace Metasound
 			// Create using standard document controller.
 			return FDocumentController::CreateDocumentHandle(InDocument);
 		}
+
+		FConstDocumentHandle IDocumentController::CreateDocumentHandle(TAccessPtr<const FMetasoundFrontendDocument> InDocument)
+		{
+			// Create using standard document controller. 
+			return FDocumentController::CreateDocumentHandle(ConstCastAccessPtr<FMetasoundFrontendDocument>(InDocument));
+		}
 	}
 }
diff --git a/Engine/Plugins/Runtime/Metasound/Source/MetasoundFrontend/Private/MetasoundFrontendInvalidController.h b/Engine/Plugins/Runtime/Metasound/Source/MetasoundFrontend/Private/MetasoundFrontendInvalidController.h
--- a/Engine/Plugins/Runtime/Metasound/Source/MetasoundFrontend/Private/MetasoundFrontendInvalidController.h
+++ b/Engine/Plugins/Runtime/Metasound/Source/MetasoundFrontend/Private/MetasoundFrontendInvalidController.h
@@ -193,8 +193,8 @@ namespace Metasound
 			TArray<TSharedRef<const INodeController>> GetConstOutputNodes() const override { return TArray<TSharedRef<const INodeController>>(); }
 			TArray<TSharedRef<const INodeController>> GetConstInputNodes() const override { return TArray<TSharedRef<const INodeController>>(); }
 
-			bool ContainsOutputNodeWithName(const FString& InName) const override { return false; }
-			bool ContainsInputNodeWithName(const FString& InName) const override { return false; }
+			bool ContainsOutputVertexWithName(const FString& InName) const override { return false; }
+			bool ContainsInputVertexWithName(const FString& InName) const override { return false; }
 
 			TSharedRef<const INodeController> GetOutputNodeWithName(const FString& InName) const override { return FInvalidNodeController::GetInvalid(); }
 			TSharedRef<const INodeController> GetInputNodeWithName(const FString& InName) const override { return FInvalidNodeController::GetInvalid(); }
diff --git a/Engine/Plugins/Runtime/Metasound/Source/MetasoundFrontend/Private/MetasoundFrontendStandardController.cpp b/Engine/Plugins/Runtime/Metasound/Source/MetasoundFrontend/Private/MetasoundFrontendStandardController.cpp
--- a/Engine/Plugins/Runtime/Metasound/Source/MetasoundFrontend/Private/MetasoundFrontendStandardController.cpp
+++ b/Engine/Plugins/Runtime/Metasound/Source/MetasoundFrontend/Private/MetasoundFrontendStandardController.cpp
@@ -1982,28 +1982,28 @@ FGetClassInputFromClassWithName
 			return GetNodesByPredicate(IsInputNode);
 		}
 
-		bool FGraphController::ContainsOutputNodeWithName(const FString& InName) const
+		bool FGraphController::ContainsOutputVertexWithName(const FString& InName) const
 		{
-			if (GraphClassPtr.IsValid())
+			if (const FMetasoundFrontendGraphClass* GraphClass = GraphClassPtr.Get())
 			{
-				auto IsOutputNodeWithSameName = [&](const FMetasoundFrontendClass& NodeClass, const FMetasoundFrontendNode& Node) 
+				auto IsOutputVertexWithSameName = [&](const FMetasoundFrontendClassOutput& ClassOutput) 
 				{ 
-					return (NodeClass.Metadata.Type == EMetasoundFrontendClassType::Output) && (Node.Name == InName); 
+					return ClassOutput.Name == InName;
 				};
-				return ContainsNodesAndClassesByPredicate(IsOutputNodeWithSameName);
+				return GraphClass->Interface.Outputs.ContainsByPredicate(IsOutputVertexWithSameName);
 			}
 			return false;
 		}
 
-		bool FGraphController::ContainsInputNodeWithName(const FString& InName) const
+		bool FGraphController::ContainsInputVertexWithName(const FString& InName) const
 		{
-			if (GraphClassPtr.IsValid())
+			if (const FMetasoundFrontendGraphClass* GraphClass = GraphClassPtr.Get())
 			{
-				auto IsInputNodeWithSameName = [&](const FMetasoundFrontendClass& NodeClass, const FMetasoundFrontendNode& Node) 
+				auto IsInputVertexWithSameName = [&](const FMetasoundFrontendClassInput& ClassInput) 
 				{ 
-					return (NodeClass.Metadata.Type == EMetasoundFrontendClassType::Input) && (Node.Name == InName); 
+					return ClassInput.Name == InName;
 				};
-				return ContainsNodesAndClassesByPredicate(IsInputNodeWithSameName);
+				return GraphClass->Interface.Inputs.ContainsByPredicate(IsInputVertexWithSameName);
 			}
 			return false;
 		}
@@ -2060,6 +2060,7 @@ FGetClassInputFromClassWithName
 						{
 							FMetasoundFrontendNode& Node = GraphClass->Graph.Nodes.Emplace_GetRef(*InputClass);
 
+							Node.Name = InClassInput.Name;
 							Node.ID = NewNodeID();
 
 							// TODO: have something that checks if input node has valid interface.
@@ -2137,6 +2138,7 @@ FGetClassInputFromClassWithName
 						{
 							FMetasoundFrontendNode& Node = GraphClass->Graph.Nodes.Emplace_GetRef(*OutputClass);
 
+							Node.Name = InClassOutput.Name;
 							Node.ID = NewNodeID();
 
 							// TODO: have something that checks if input node has valid interface.
diff --git a/Engine/Plugins/Runtime/Metasound/Source/MetasoundFrontend/Private/MetasoundFrontendStandardController.h b/Engine/Plugins/Runtime/Metasound/Source/MetasoundFrontend/Private/MetasoundFrontendStandardController.h
--- a/Engine/Plugins/Runtime/Metasound/Source/MetasoundFrontend/Private/MetasoundFrontendStandardController.h
+++ b/Engine/Plugins/Runtime/Metasound/Source/MetasoundFrontend/Private/MetasoundFrontendStandardController.h
@@ -631,8 +631,8 @@ namespace Metasound
 			TArray<FNodeHandle> GetOutputNodes() override;
 			TArray<FConstNodeHandle> GetConstOutputNodes() const override;
 
-			bool ContainsInputNodeWithName(const FString& InName) const override;
-			bool ContainsOutputNodeWithName(const FString& InName) const override;
+			bool ContainsInputVertexWithName(const FString& InName) const override;
+			bool ContainsOutputVertexWithName(const FString& InName) const override;
 
 			FConstNodeHandle GetInputNodeWithName(const FString& InName) const override;
 			FConstNodeHandle GetOutputNodeWithName(const FString& InName) const override;
diff --git a/Engine/Plugins/Runtime/Metasound/Source/MetasoundFrontend/Public/MetasoundAssetBase.h b/Engine/Plugins/Runtime/Metasound/Source/MetasoundFrontend/Public/MetasoundAssetBase.h
--- a/Engine/Plugins/Runtime/Metasound/Source/MetasoundFrontend/Public/MetasoundAssetBase.h
+++ b/Engine/Plugins/Runtime/Metasound/Source/MetasoundFrontend/Public/MetasoundAssetBase.h
@@ -97,9 +97,11 @@ public:
 
 	// Returns handle for the root metasound graph of this asset.
 	Metasound::Frontend::FDocumentHandle GetDocumentHandle();
+	Metasound::Frontend::FConstDocumentHandle GetDocumentHandle() const;
 
 	// Returns handle for the root metasound graph of this asset.
 	Metasound::Frontend::FGraphHandle GetRootGraphHandle();
+	Metasound::Frontend::FConstGraphHandle GetRootGraphHandle() const;
 
 
 	// Overwrites the existing document. If the document's archetype is not supported,
diff --git a/Engine/Plugins/Runtime/Metasound/Source/MetasoundFrontend/Public/MetasoundFrontendController.h b/Engine/Plugins/Runtime/Metasound/Source/MetasoundFrontend/Public/MetasoundFrontendController.h
--- a/Engine/Plugins/Runtime/Metasound/Source/MetasoundFrontend/Public/MetasoundFrontendController.h
+++ b/Engine/Plugins/Runtime/Metasound/Source/MetasoundFrontend/Public/MetasoundFrontendController.h
@@ -346,19 +346,19 @@ namespace Metasound
 			/** Returns all input nodes in the graph. */
 			virtual TArray<FConstNodeHandle> GetConstInputNodes() const = 0;
 
-			/** Returns true if an output node with the given NodeName exists.
+			/** Returns true if an output vertex with the given Name exists.
 			 *
-			 * @param InName - Name of node.
-			 * @return True if the node exists, false otherwise. 
+			 * @param InName - Name of vertex.
+			 * @return True if the vertex exists, false otherwise. 
 			 */
-			virtual bool ContainsOutputNodeWithName(const FString& InName) const = 0;
+			virtual bool ContainsOutputVertexWithName(const FString& InName) const = 0;
 
-			/** Returns true if an input node with the given NodeName exists.
+			/** Returns true if an input vertex with the given Name exists.
 			 *
-			 * @param InName - Name of node.
-			 * @return True if the node exists, false otherwise. 
+			 * @param InName - Name of vertex.
+			 * @return True if the vertex exists, false otherwise. 
 			 */
-			virtual bool ContainsInputNodeWithName(const FString& InName) const = 0;
+			virtual bool ContainsInputVertexWithName(const FString& InName) const = 0;
 
 			/** Returns a handle to an existing output node for the given graph output name.
 			 * If no output exists with the given name, an invalid node handle is returned. 
@@ -580,6 +580,8 @@ namespace Metasound
 
 			/** Create a document from FMetasoundFrontendDocument description pointer. */
 			static FDocumentHandle CreateDocumentHandle(TAccessPtr<FMetasoundFrontendDocument> InDocument);
+			/** Create a document from FMetasoundFrontendDocument description pointer. */
+			static FConstDocumentHandle CreateDocumentHandle(TAccessPtr<const FMetasoundFrontendDocument> InDocument);
 
 			IDocumentController() = default;
 			virtual ~IDocumentController() = default;
