Fix copy/pasta in operator overload

diff --git a/Engine/Plugins/Runtime/Metasound/Source/MetasoundStandardNodes/Public/MetasoundFrequency.h b/Engine/Plugins/Runtime/Metasound/Source/MetasoundStandardNodes/Public/MetasoundFrequency.h
--- a/Engine/Plugins/Runtime/Metasound/Source/MetasoundStandardNodes/Public/MetasoundFrequency.h
+++ b/Engine/Plugins/Runtime/Metasound/Source/MetasoundStandardNodes/Public/MetasoundFrequency.h
@@ -80,7 +80,7 @@ namespace Metasound
 
 			FFrequency& operator-=(const FFrequency& InOther)
 			{
-				FrequencyInHz += InOther.GetHertz();
+				FrequencyInHz -= InOther.GetHertz();
 				return *this;
 			}
 
