Math Op Tests
- Various macro clean-up
- Shuffle code around to enable tests to access math op nodes

diff --git a/Engine/Plugins/Runtime/Metasound/Source/MetasoundStandardNodes/Private/MetasoundMathNodes.cpp b/Engine/Plugins/Runtime/Metasound/Source/MetasoundStandardNodes/Private/MetasoundMathNodes.cpp
--- a/Engine/Plugins/Runtime/Metasound/Source/MetasoundStandardNodes/Private/MetasoundMathNodes.cpp
+++ b/Engine/Plugins/Runtime/Metasound/Source/MetasoundStandardNodes/Private/MetasoundMathNodes.cpp
@@ -1,296 +1,39 @@
 // Copyright Epic Games, Inc. All Rights Reserved.
 #include "MetasoundMathNodes.h"
 
-#include "DSP/BufferVectorOperations.h"
-#include "MetasoundAudioBuffer.h"
-#include "MetasoundDataReference.h"
-#include "MetasoundFrequency.h"
-#include "MetasoundNodeInterface.h"
-#include "MetasoundTime.h"
 
-#define LOCTEXT_NAMESPACE "MetasoundMathNodes"
 namespace Metasound
 {
-// Specialized Math Operations
-	template <>
-	class TMathOpAdd<FAudioBuffer>
-	{
-	public:
-		static FAudioBuffer GetDefault(const FOperatorSettings& InSettings)
-		{
-			return FAudioBuffer(InSettings.GetNumFramesPerBlock());
-		}
-
-		static void Calculate(const TArray<FAudioBufferReadRef>& InOperands, FAudioBufferWriteRef& OutResult)
-		{
-			if (!ensure(!InOperands.IsEmpty()))
-			{
-				OutResult->Zero();
-				return;
-			}
-
-			const int32 NumSamples = InOperands[0]->Num();
-			if (!ensure(NumSamples == OutResult->Num()))
-			{
-				// TODO: Error
-				OutResult->Zero();
-				return;
-			}
-
-			FMemory::Memcpy(OutResult->GetData(), InOperands[0]->GetData(), sizeof(float) * NumSamples);
-
-			for (int32 i = 1; i < InOperands.Num(); ++i)
-			{
-				const FAudioBuffer& Operand = *InOperands[i];
-				if (!ensure(NumSamples == Operand.Num()))
-				{
-					// TODO: Error
-					OutResult->Zero();
-					return;
-				}
-
-				const float* OperandData = InOperands[i]->GetData();
-				float* OutData = OutResult->GetData();
-
-				if (NumSamples % AUDIO_SIMD_FLOAT_ALIGNMENT)
-				{
-					for (int32 SampleIndex = 0; SampleIndex < NumSamples; ++SampleIndex)
-					{
-						OutData[SampleIndex] += OperandData[SampleIndex];
-					}
-				}
-				else
-				{
-					const float StartGain = 1.0f;
-					const float EndGain = 1.0f;
-					Audio::MixInBufferFast(OperandData, OutData, NumSamples, StartGain, EndGain);
-				}
-			}
-		}
-	};
-
-	template <>
-	class TMathOpSubtract<FAudioBuffer>
-	{
-	public:
-		static FAudioBuffer GetDefault(const FOperatorSettings& InSettings)
-		{
-			return FAudioBuffer(InSettings.GetNumFramesPerBlock());
-		}
-
-		static void Calculate(TArray<FAudioBufferReadRef>& InOperands, FAudioBufferWriteRef& OutResult)
-		{
-			if (!ensure(!InOperands.IsEmpty()))
-			{
-				OutResult->Zero();
-				return;
-			}
-
-			const int32 NumSamples = InOperands[0]->Num();
-			if (!ensure(NumSamples == OutResult->Num()))
-			{
-				// TODO: Error
-				OutResult->Zero();
-				return;
-			}
-
-			FMemory::Memcpy(OutResult->GetData(), InOperands[0]->GetData(), sizeof(float) * NumSamples);
-
-			for (int32 i = 1; i < InOperands.Num(); ++i)
-			{
-				const FAudioBuffer& Operand = *InOperands[i];
-				if (!ensure(NumSamples == Operand.Num()))
-				{
-					// TODO: Error
-					OutResult->Zero();
-					return;
-				}
-
-				const float* OperandData = InOperands[i]->GetData();
-				float* OutData = OutResult->GetData();
-
-				if (NumSamples % AUDIO_SIMD_FLOAT_ALIGNMENT)
-				{
-					for (int32 SampleIndex = 0; SampleIndex < NumSamples; ++SampleIndex)
-					{
-						OutData[i] -= OperandData[i];
-					}
-				}
-				else
-				{
-					Audio::BufferSubtractInPlace2Fast(OutData, OperandData, NumSamples);
-				}
-			}
-		}
-	};
-
-	template <>
-	class TMathOpMultiply<FAudioBuffer>
-	{
-	public:
-		static FAudioBuffer GetDefault(const FOperatorSettings& InSettings)
-		{
-			return FAudioBuffer(InSettings.GetNumFramesPerBlock());
-		}
-
-		static void Calculate(const TArray<FAudioBufferReadRef>& InOperands, FAudioBufferWriteRef& OutResult)
-		{
-			if (!ensure(!InOperands.IsEmpty()))
-			{
-				OutResult->Zero();
-				return;
-			}
-
-			const int32 NumSamples = InOperands[0]->Num();
-			if (!ensure(NumSamples == OutResult->Num()))
-			{
-				// TODO: Error
-				OutResult->Zero();
-				return;
-			}
-
-			FMemory::Memcpy(OutResult->GetData(), InOperands[0]->GetData(), sizeof(float) * NumSamples);
-
-			for (int32 i = 1; i < InOperands.Num(); ++i)
-			{
-				const FAudioBuffer& Operand = *InOperands[i];
-				if (!ensure(NumSamples == Operand.Num()))
-				{
-					// TODO: Error
-					OutResult->Zero();
-					return;
-				}
-
-				const float* OperandData = InOperands[i]->GetData();
-				float* OutData = OutResult->GetData();
-
-				if (NumSamples % AUDIO_SIMD_FLOAT_ALIGNMENT)
-				{
-					for (int32 SampleIndex = 0; SampleIndex < NumSamples; ++SampleIndex)
-					{
-						OutData[SampleIndex] += OperandData[SampleIndex];
-					}
-				}
-				else
-				{
-					Audio::MultiplyBuffersInPlace(OperandData, OutData, NumSamples);
-				}
-			}
-		}
-	};
-
-	template <>
-	class TMathOpMultiply<FFloatTime>
-	{
-	public:
-		static FFloatTime GetDefault(const FOperatorSettings& InSettings)
-		{
-			return FFloatTime(1.0f);
-		}
-
-		static void Calculate(const TArray<FFloatTimeReadRef>& InOperands, FFloatTimeWriteRef& OutResult)
-		{
-			if (!ensure(!InOperands.IsEmpty()))
-			{
-				return;
-			}
-
-			*OutResult = *InOperands[0];
-			for (int32 i = 1; i < InOperands.Num(); ++i)
-			{
-				OutResult->SetSeconds(OutResult->GetSeconds() * InOperands[i]->GetSeconds());
-			}
-		}
-	};
-
-	template <>
-	class TMathOpDivide<FFloatTime>
-	{
-	public:
-		static FFloatTime GetDefault(const FOperatorSettings& InSettings)
-		{
-			return FFloatTime(1.0f);
-		}
-
-		static void Calculate(const TArray<FFloatTimeReadRef>& InOperands, FFloatTimeWriteRef& OutResult)
-		{
-			if (!ensure(!InOperands.IsEmpty()))
-			{
-				return;
-			}
-
-			*OutResult = *InOperands[0];
-			for (int32 i = 1; i < InOperands.Num(); ++i)
-			{
-				const FFloatTime& OperandValue = *InOperands[i];
-				if (OperandValue.GetSeconds() == 0.0f)
-				{
-					// TODO: Error here
-					return;
-				}
-
-				OutResult->SetSeconds(OutResult->GetSeconds() / OperandValue.GetSeconds());
-			}
-		}
-	};
-
-	template <>
-	class TMathOpRandRange<FFloatTime>
-	{
-	public:
-		static FFloatTime GetDefault(const FOperatorSettings& InSettings)
-		{
-			return FFloatTime(0.0f);
-		}
-
-		static void Calculate(const TArray<FFloatTimeReadRef>& InOperands, FFloatTimeWriteRef& OutResult)
-		{
-			if (InOperands.Num() > 1)
-			{
-				const FFloatTime& Min = *InOperands[0];
-				const FFloatTime& Max = *InOperands[1];
-				const float FloatResult = FMath::RandRange(Min.GetSeconds(), Max.GetSeconds());
-				OutResult->SetSeconds(FloatResult);
-			}
-			else
-			{
-				// TODO: Error
-			}
-		}
-	};
-
-// Definitions
-	DEFINE_METASOUND_MATHOP(Add, Float, float, LOCTEXT("Metasound_MathAddFloatNodeDescription", "Adds floats."))
-	DEFINE_METASOUND_MATHOP(Add, Int32, int32, LOCTEXT("Metasound_MathAddInt32NodeDescription", "Adds int32s."))
-	DEFINE_METASOUND_MATHOP(Add, Int64, int64, LOCTEXT("Metasound_MathAddInt64NodeDescription", "Adds int64s."))
-	DEFINE_METASOUND_MATHOP(Add, Frequency, FFrequency, LOCTEXT("Metasound_MathAddFrequencyNodeDescription", "Adds frequency values."))
-	DEFINE_METASOUND_MATHOP(Add, FloatTime, FFloatTime, LOCTEXT("Metasound_MathAddFloatTimeNodeDescription", "Adds time values."))
-	DEFINE_METASOUND_MATHOP(Add, AudioBuffer, FAudioBuffer, LOCTEXT("Metasound_MathAddBufferNodeDescription", "Adds buffers together by sample."))
-
-	DEFINE_METASOUND_MATHOP(Subtract, Float, float, LOCTEXT("Metasound_MathSubractFloatNodeDescription", "Subtracts floats."))
-	DEFINE_METASOUND_MATHOP(Subtract, Int32, int32, LOCTEXT("Metasound_MathSubractInt32NodeDescription", "Subtracts int32s."))
-	DEFINE_METASOUND_MATHOP(Subtract, Int64, int64, LOCTEXT("Metasound_MathSubractInt64NodeDescription", "Subtracts int64s."))
-	DEFINE_METASOUND_MATHOP(Subtract, Frequency, FFrequency, LOCTEXT("Metasound_MathSubractFrequencyNodeDescription", "Subtracts frequency values."))
-	DEFINE_METASOUND_MATHOP(Subtract, FloatTime, FFloatTime, LOCTEXT("Metasound_MathSubractFloatTimeNodeDescription", "Subtracts time values."))
-	DEFINE_METASOUND_MATHOP(Subtract, AudioBuffer, FAudioBuffer, LOCTEXT("Metasound_MathSubtractBufferNodeDescription", "Subtracts buffers by sample."))
-
-	DEFINE_METASOUND_MATHOP(Multiply, Float, float, LOCTEXT("Metasound_MathMultiplyFloatNodeDescription", "Multiplies floats."))
-	DEFINE_METASOUND_MATHOP(Multiply, Int32, int32, LOCTEXT("Metasound_MathMultiplyInt32NodeDescription", "Multiplies int32s."))
-	DEFINE_METASOUND_MATHOP(Multiply, Int64, int64, LOCTEXT("Metasound_MathMultiplyInt64NodeDescription", "Multiplies int64s."))
-	DEFINE_METASOUND_MATHOP(Multiply, Frequency, FFrequency, LOCTEXT("Metasound_MathMultiplyFrequencyNodeDescription", "Multiplies frequency values."))
-	DEFINE_METASOUND_MATHOP(Multiply, FloatTime, FFloatTime, LOCTEXT("Metasound_MathMultiplyFloatTimeNodeDescription", "Multiplies time values."))
-	DEFINE_METASOUND_MATHOP(Multiply, AudioBuffer, FAudioBuffer, LOCTEXT("Metasound_MathMultiplyBufferNodeDescription", "Multiplies buffers together by sample."))
-
-	DEFINE_METASOUND_MATHOP(Divide, Float, float, LOCTEXT("Metasound_MathDivideFloatNodeDescription", "Divide float by another float."))
-	DEFINE_METASOUND_MATHOP(Divide, Int32, int32, LOCTEXT("Metasound_MathDivideInt32NodeDescription", "Divide int32 by another int32."))
-	DEFINE_METASOUND_MATHOP(Divide, Int64, int64, LOCTEXT("Metasound_MathDivideInt64NodeDescription", "Divide int64 by another int64."))
-	DEFINE_METASOUND_MATHOP(Divide, Frequency, FFrequency, LOCTEXT("Metasound_MathDivideFrequencyNodeDescription", "Divide frequency by another frequency."))
-	DEFINE_METASOUND_MATHOP(Divide, FloatTime, FFloatTime, LOCTEXT("Metasound_MathDivideFloatTimeNodeDescription", "Divide time by another time."))
-
-	DEFINE_METASOUND_MATHOP(RandRange, Float, float, LOCTEXT("Metasound_MathRandRangeFloatNodeDescription", "Computes random float value in the range [Min, Max]."))
-	DEFINE_METASOUND_MATHOP(RandRange, Int32, int32, LOCTEXT("Metasound_MathRandRangeInt32NodeDescription", "Computes random int32 in the range [Min, Max]."))
-	DEFINE_METASOUND_MATHOP(RandRange, Int64, int64, LOCTEXT("Metasound_MathRandRangeInt64NodeDescription", "Computes random int64 in the range [Min, Max]."))
-	DEFINE_METASOUND_MATHOP(RandRange, Frequency, FFrequency, LOCTEXT("Metasound_MathRandRangeFrequencyNodeDescription", "Computes random frequency value in the range [Min, Max]."))
-	DEFINE_METASOUND_MATHOP(RandRange, FloatTime, FFloatTime, LOCTEXT("Metasound_MathRandRangeFloatTimeNodeDescription", "Computes random time value in the range [Min, Max]."))
+	METASOUND_REGISTER_NODE(FAddFloatNode)
+	METASOUND_REGISTER_NODE(FAddInt32Node)
+	METASOUND_REGISTER_NODE(FAddInt64Node)
+	METASOUND_REGISTER_NODE(FAddFrequencyNode)
+	METASOUND_REGISTER_NODE(FAddFloatTimeNode)
+	METASOUND_REGISTER_NODE(FAddAudioBufferNode)
+
+	METASOUND_REGISTER_NODE(FSubtractFloatNode)
+	METASOUND_REGISTER_NODE(FSubtractInt32Node)
+	METASOUND_REGISTER_NODE(FSubtractInt64Node)
+	METASOUND_REGISTER_NODE(FSubtractFrequencyNode)
+	METASOUND_REGISTER_NODE(FSubtractFloatTimeNode)
+	METASOUND_REGISTER_NODE(FSubtractAudioBufferNode)
+
+	METASOUND_REGISTER_NODE(FMultiplyFloatNode)
+	METASOUND_REGISTER_NODE(FMultiplyInt32Node)
+	METASOUND_REGISTER_NODE(FMultiplyInt64Node)
+	METASOUND_REGISTER_NODE(FMultiplyFrequencyNode)
+	METASOUND_REGISTER_NODE(FMultiplyFloatTimeNode)
+	METASOUND_REGISTER_NODE(FMultiplyAudioBufferNode)
+
+	METASOUND_REGISTER_NODE(FDivideFloatNode)
+	METASOUND_REGISTER_NODE(FDivideInt32Node)
+	METASOUND_REGISTER_NODE(FDivideInt64Node)
+	METASOUND_REGISTER_NODE(FDivideFrequencyNode)
+	METASOUND_REGISTER_NODE(FDivideFloatTimeNode)
+
+	METASOUND_REGISTER_NODE(FRandRangeFloatNode)
+	METASOUND_REGISTER_NODE(FRandRangeInt32Node)
+	METASOUND_REGISTER_NODE(FRandRangeInt64Node)
+	METASOUND_REGISTER_NODE(FRandRangeFrequencyNode)
+	METASOUND_REGISTER_NODE(FRandRangeFloatTimeNode)
 }
-#undef LOCTEXT_NAMESPACE // MetasoundMathNodes
diff --git a/Engine/Plugins/Runtime/Metasound/Source/MetasoundStandardNodes/Public/MetasoundMathNodes.h b/Engine/Plugins/Runtime/Metasound/Source/MetasoundStandardNodes/Public/MetasoundMathNodes.h
--- a/Engine/Plugins/Runtime/Metasound/Source/MetasoundStandardNodes/Public/MetasoundMathNodes.h
+++ b/Engine/Plugins/Runtime/Metasound/Source/MetasoundStandardNodes/Public/MetasoundMathNodes.h
@@ -1,24 +1,39 @@
 // Copyright Epic Games, Inc. All Rights Reserved.
-
 #pragma once
 
 #include "CoreMinimal.h"
+#include "DSP/BufferVectorOperations.h"
+#include "MetasoundAudioBuffer.h"
 #include "MetasoundBop.h"
 #include "MetasoundBuilderInterface.h"
+#include "MetasoundDataReference.h"
 #include "MetasoundDataReferenceCollection.h"
 #include "MetasoundExecutableOperator.h"
 #include "MetasoundFacade.h"
+#include "MetasoundFrequency.h"
 #include "MetasoundNode.h"
 #include "MetasoundNodeRegistrationMacro.h"
 #include "MetasoundNodeInterface.h"
 #include "MetasoundOperatorInterface.h"
 #include "MetasoundPrimitives.h"
+#include "MetasoundTime.h"
 
 
 #define LOCTEXT_NAMESPACE "MetasoundMathOpNode"
+
+#define DEFINE_METASOUND_MATHOP(OpName, DataTypeName, DataClass, Description) \
+	class F##OpName##DataTypeName##Node \
+		: public ::Metasound::TMathOpNode<F##OpName##DataTypeName##Node, ::Metasound::TMathOp##OpName<DataClass>, DataClass> \
+	{ \
+	public: \
+		static FName GetClassName() { return #OpName #DataTypeName; } \
+		static FText GetDescription() { return Description; } \
+		F##OpName##DataTypeName##Node(const FNodeInitData& InInitData) : TMathOpNode(InInitData) { } \
+	};
+
 namespace Metasound
 {
-	template <typename TDerivedClass, typename TDataClass>
+	template <typename TDerivedClass, typename TMathOpClass, typename TDataClass>
 	class TMathOpNode : public FNodeFacade
 	{
 	private:
@@ -59,6 +74,11 @@ namespace Metasound
 					Info.Author = PluginAuthor;
 					Info.PromptIfMissing = PluginNodeMissingPrompt;
 					Info.DefaultInterface = DeclareVertexInterface();
+					Info.ImageName = TEXT("MetasoundEditor.Speaker");
+					Info.bShowName = false;
+					Info.bShowInputNames = false;
+					Info.bShowOutputNames = false;
+					Info.CategoryHierarchy = { LOCTEXT("Metasound_MathCategory", "Math") };
 
 					return Info;
 				};
@@ -69,7 +89,7 @@ namespace Metasound
 
 			void Execute()
 			{
-				TDerivedClass::Calculate(OperandRefs, ValueRef);
+				TMathOpClass::Calculate(OperandRefs, ValueRef);
 			}
 
 			static TUniquePtr<IOperator> CreateOperator(const FCreateOperatorParams& InParams, FBuildErrorArray& OutErrors)
@@ -77,8 +97,8 @@ namespace Metasound
 				const TMathOpNode& MathOpNode = static_cast<const TMathOpNode&>(InParams.Node);
 
 				// TODO: Support dynamic number of inputs
-				TDataClassReadRef Op1 = InParams.InputDataReferences.GetDataReadReferenceOrConstruct<TDataClass>(TEXT("Operand1"), TDerivedClass::GetDefault(InParams.OperatorSettings));
-				TDataClassReadRef Op2 = InParams.InputDataReferences.GetDataReadReferenceOrConstruct<TDataClass>(TEXT("Operand2"), TDerivedClass::GetDefault(InParams.OperatorSettings));
+				TDataClassReadRef Op1 = InParams.InputDataReferences.GetDataReadReferenceOrConstruct<TDataClass>(TEXT("Operand1"), TMathOpClass::GetDefault(InParams.OperatorSettings));
+				TDataClassReadRef Op2 = InParams.InputDataReferences.GetDataReadReferenceOrConstruct<TDataClass>(TEXT("Operand2"), TMathOpClass::GetDefault(InParams.OperatorSettings));
 
 				TArray<TDataClassReadRef> OperandRefs = { Op1, Op2 };
 
@@ -87,7 +107,7 @@ namespace Metasound
 
 			TMathOperator(const FOperatorSettings& InSettings, const TArray<TDataClassReadRef>& InOperands)
 			:	OperandRefs(InOperands)
-			,	ValueRef(TDataClassWriteRef::CreateNew(TDerivedClass::GetDefault(InSettings)))
+			,	ValueRef(TDataClassWriteRef::CreateNew(TMathOpClass::GetDefault(InSettings)))
 			{
 			}
 
@@ -231,6 +251,229 @@ namespace Metasound
 		}
 	};
 
+// Specialized Math Operations
+	template <>
+	class TMathOpAdd<FAudioBuffer>
+	{
+	public:
+		static FAudioBuffer GetDefault(const FOperatorSettings& InSettings)
+		{
+			return FAudioBuffer(InSettings.GetNumFramesPerBlock());
+		}
+
+		static void Calculate(const TArray<FAudioBufferReadRef>& InOperands, FAudioBufferWriteRef& OutResult)
+		{
+			if (!ensure(!InOperands.IsEmpty()))
+			{
+				OutResult->Zero();
+				return;
+			}
+
+			const int32 NumSamples = InOperands[0]->Num();
+			if (!ensure(NumSamples == OutResult->Num()))
+			{
+				// TODO: Error
+				OutResult->Zero();
+				return;
+			}
+
+			FMemory::Memcpy(OutResult->GetData(), InOperands[0]->GetData(), sizeof(float) * NumSamples);
+
+			for (int32 i = 1; i < InOperands.Num(); ++i)
+			{
+				const FAudioBuffer& Operand = *InOperands[i];
+				if (!ensure(NumSamples == Operand.Num()))
+				{
+					// TODO: Error
+					OutResult->Zero();
+					return;
+				}
+
+				const float* OperandData = InOperands[i]->GetData();
+				float* OutData = OutResult->GetData();
+
+				if (NumSamples % AUDIO_SIMD_FLOAT_ALIGNMENT)
+				{
+					for (int32 SampleIndex = 0; SampleIndex < NumSamples; ++SampleIndex)
+					{
+						OutData[SampleIndex] += OperandData[SampleIndex];
+					}
+				}
+				else
+				{
+					const float StartGain = 1.0f;
+					const float EndGain = 1.0f;
+					Audio::MixInBufferFast(OperandData, OutData, NumSamples, StartGain, EndGain);
+				}
+			}
+		}
+	};
+
+	template <>
+	class TMathOpSubtract<FAudioBuffer>
+	{
+	public:
+		static FAudioBuffer GetDefault(const FOperatorSettings& InSettings)
+		{
+			return FAudioBuffer(InSettings.GetNumFramesPerBlock());
+		}
+
+		static void Calculate(TArray<FAudioBufferReadRef>& InOperands, FAudioBufferWriteRef& OutResult)
+		{
+			if (!ensure(!InOperands.IsEmpty()))
+			{
+				OutResult->Zero();
+				return;
+			}
+
+			const int32 NumSamples = InOperands[0]->Num();
+			if (!ensure(NumSamples == OutResult->Num()))
+			{
+				// TODO: Error
+				OutResult->Zero();
+				return;
+			}
+
+			FMemory::Memcpy(OutResult->GetData(), InOperands[0]->GetData(), sizeof(float) * NumSamples);
+
+			for (int32 i = 1; i < InOperands.Num(); ++i)
+			{
+				const FAudioBuffer& Operand = *InOperands[i];
+				if (!ensure(NumSamples == Operand.Num()))
+				{
+					// TODO: Error
+					OutResult->Zero();
+					return;
+				}
+
+				const float* OperandData = InOperands[i]->GetData();
+				float* OutData = OutResult->GetData();
+
+				if (NumSamples % AUDIO_SIMD_FLOAT_ALIGNMENT)
+				{
+					for (int32 SampleIndex = 0; SampleIndex < NumSamples; ++SampleIndex)
+					{
+						OutData[i] -= OperandData[i];
+					}
+				}
+				else
+				{
+					Audio::BufferSubtractInPlace2Fast(OutData, OperandData, NumSamples);
+				}
+			}
+		}
+	};
+
+	template <>
+	class TMathOpMultiply<FAudioBuffer>
+	{
+	public:
+		static FAudioBuffer GetDefault(const FOperatorSettings& InSettings)
+		{
+			return FAudioBuffer(InSettings.GetNumFramesPerBlock());
+		}
+
+		static void Calculate(const TArray<FAudioBufferReadRef>& InOperands, FAudioBufferWriteRef& OutResult)
+		{
+			if (!ensure(!InOperands.IsEmpty()))
+			{
+				OutResult->Zero();
+				return;
+			}
+
+			const int32 NumSamples = InOperands[0]->Num();
+			if (!ensure(NumSamples == OutResult->Num()))
+			{
+				// TODO: Error
+				OutResult->Zero();
+				return;
+			}
+
+			FMemory::Memcpy(OutResult->GetData(), InOperands[0]->GetData(), sizeof(float) * NumSamples);
+
+			for (int32 i = 1; i < InOperands.Num(); ++i)
+			{
+				const FAudioBuffer& Operand = *InOperands[i];
+				if (!ensure(NumSamples == Operand.Num()))
+				{
+					// TODO: Error
+					OutResult->Zero();
+					return;
+				}
+
+				const float* OperandData = InOperands[i]->GetData();
+				float* OutData = OutResult->GetData();
+
+				if (NumSamples % AUDIO_SIMD_FLOAT_ALIGNMENT)
+				{
+					for (int32 SampleIndex = 0; SampleIndex < NumSamples; ++SampleIndex)
+					{
+						OutData[SampleIndex] += OperandData[SampleIndex];
+					}
+				}
+				else
+				{
+					Audio::MultiplyBuffersInPlace(OperandData, OutData, NumSamples);
+				}
+			}
+		}
+	};
+
+	template <>
+	class TMathOpMultiply<FFloatTime>
+	{
+	public:
+		static FFloatTime GetDefault(const FOperatorSettings& InSettings)
+		{
+			return FFloatTime(1.0f);
+		}
+
+		static void Calculate(const TArray<FFloatTimeReadRef>& InOperands, FFloatTimeWriteRef& OutResult)
+		{
+			if (!ensure(!InOperands.IsEmpty()))
+			{
+				return;
+			}
+
+			*OutResult = *InOperands[0];
+			for (int32 i = 1; i < InOperands.Num(); ++i)
+			{
+				OutResult->SetSeconds(OutResult->GetSeconds() * InOperands[i]->GetSeconds());
+			}
+		}
+	};
+
+	template <>
+	class TMathOpDivide<FFloatTime>
+	{
+	public:
+		static FFloatTime GetDefault(const FOperatorSettings& InSettings)
+		{
+			return FFloatTime(1.0f);
+		}
+
+		static void Calculate(const TArray<FFloatTimeReadRef>& InOperands, FFloatTimeWriteRef& OutResult)
+		{
+			if (!ensure(!InOperands.IsEmpty()))
+			{
+				return;
+			}
+
+			*OutResult = *InOperands[0];
+			for (int32 i = 1; i < InOperands.Num(); ++i)
+			{
+				const FFloatTime& OperandValue = *InOperands[i];
+				if (OperandValue.GetSeconds() == 0.0f)
+				{
+					// TODO: Error here
+					return;
+				}
+
+				OutResult->SetSeconds(OutResult->GetSeconds() / OperandValue.GetSeconds());
+			}
+		}
+	};
+
 	template <typename TDataClass>
 	class TMathOpRandRange
 	{
@@ -257,18 +500,90 @@ namespace Metasound
 			}
 		}
 	};
-}
 
-#define DEFINE_METASOUND_MATHOP(OpName, DataTypeName, DataClass, Description) \
-	class F##OpName##DataTypeName##Node \
-		: public ::Metasound::TMathOpNode<F##OpName##DataTypeName##Node, DataClass> \
-		, public ::Metasound::TMathOp##OpName<DataClass> \
-	{ \
-	public: \
-		static FName GetClassName() { return #OpName #DataTypeName; } \
-		static FText GetDescription() { return Description; } \
-		F##OpName##DataTypeName##Node(const FNodeInitData& InInitData) : TMathOpNode(InInitData) { } \
-	}; \
-	METASOUND_REGISTER_NODE(F##OpName##DataTypeName##Node)
+	template <>
+	class TMathOpRandRange<FFloatTime>
+	{
+	public:
+		static FFloatTime GetDefault(const FOperatorSettings& InSettings)
+		{
+			return FFloatTime(0.0f);
+		}
+
+		static void Calculate(const TArray<FFloatTimeReadRef>& InOperands, FFloatTimeWriteRef& OutResult)
+		{
+			if (InOperands.Num() > 1)
+			{
+				const FFloatTime& Min = *InOperands[0];
+				const FFloatTime& Max = *InOperands[1];
+				const float FloatResult = FMath::RandRange(Min.GetSeconds(), Max.GetSeconds());
+				OutResult->SetSeconds(FloatResult);
+			}
+			else
+			{
+				// TODO: Error
+			}
+		}
+	};
+
+	template <>
+	class TMathOpRandRange<FFrequency>
+	{
+	public:
+		static FFrequency GetDefault(const FOperatorSettings& InSettings)
+		{
+			return FFrequency(0.0f);
+		}
+
+		static void Calculate(const TArray<FFrequencyReadRef>& InOperands, FFrequencyWriteRef& OutResult)
+		{
+			if (InOperands.Num() > 1)
+			{
+				const FFrequency& Min = *InOperands[0];
+				const FFrequency& Max = *InOperands[1];
+				const float FloatResult = FMath::RandRange(Min.GetHertz(), Max.GetHertz());
+				OutResult->SetHertz(FloatResult);
+			}
+			else
+			{
+				// TODO: Error
+			}
+		}
+	};
+
+// Definitions
+	DEFINE_METASOUND_MATHOP(Add, Float, float, LOCTEXT("Metasound_MathAddFloatNodeDescription", "Adds floats."))
+	DEFINE_METASOUND_MATHOP(Add, Int32, int32, LOCTEXT("Metasound_MathAddInt32NodeDescription", "Adds int32s."))
+	DEFINE_METASOUND_MATHOP(Add, Int64, int64, LOCTEXT("Metasound_MathAddInt64NodeDescription", "Adds int64s."))
+	DEFINE_METASOUND_MATHOP(Add, Frequency, FFrequency, LOCTEXT("Metasound_MathAddFrequencyNodeDescription", "Adds frequency values."))
+	DEFINE_METASOUND_MATHOP(Add, FloatTime, FFloatTime, LOCTEXT("Metasound_MathAddFloatTimeNodeDescription", "Adds time values."))
+	DEFINE_METASOUND_MATHOP(Add, AudioBuffer, FAudioBuffer, LOCTEXT("Metasound_MathAddBufferNodeDescription", "Adds buffers together by sample."))
+
+	DEFINE_METASOUND_MATHOP(Subtract, Float, float, LOCTEXT("Metasound_MathSubractFloatNodeDescription", "Subtracts floats."))
+	DEFINE_METASOUND_MATHOP(Subtract, Int32, int32, LOCTEXT("Metasound_MathSubractInt32NodeDescription", "Subtracts int32s."))
+	DEFINE_METASOUND_MATHOP(Subtract, Int64, int64, LOCTEXT("Metasound_MathSubractInt64NodeDescription", "Subtracts int64s."))
+	DEFINE_METASOUND_MATHOP(Subtract, Frequency, FFrequency, LOCTEXT("Metasound_MathSubractFrequencyNodeDescription", "Subtracts frequency values."))
+	DEFINE_METASOUND_MATHOP(Subtract, FloatTime, FFloatTime, LOCTEXT("Metasound_MathSubractFloatTimeNodeDescription", "Subtracts time values."))
+	DEFINE_METASOUND_MATHOP(Subtract, AudioBuffer, FAudioBuffer, LOCTEXT("Metasound_MathSubtractBufferNodeDescription", "Subtracts buffers by sample."))
+
+	DEFINE_METASOUND_MATHOP(Multiply, Float, float, LOCTEXT("Metasound_MathMultiplyFloatNodeDescription", "Multiplies floats."))
+	DEFINE_METASOUND_MATHOP(Multiply, Int32, int32, LOCTEXT("Metasound_MathMultiplyInt32NodeDescription", "Multiplies int32s."))
+	DEFINE_METASOUND_MATHOP(Multiply, Int64, int64, LOCTEXT("Metasound_MathMultiplyInt64NodeDescription", "Multiplies int64s."))
+	DEFINE_METASOUND_MATHOP(Multiply, Frequency, FFrequency, LOCTEXT("Metasound_MathMultiplyFrequencyNodeDescription", "Multiplies frequency values."))
+	DEFINE_METASOUND_MATHOP(Multiply, FloatTime, FFloatTime, LOCTEXT("Metasound_MathMultiplyFloatTimeNodeDescription", "Multiplies time values."))
+	DEFINE_METASOUND_MATHOP(Multiply, AudioBuffer, FAudioBuffer, LOCTEXT("Metasound_MathMultiplyBufferNodeDescription", "Multiplies buffers together by sample."))
+
+	DEFINE_METASOUND_MATHOP(Divide, Float, float, LOCTEXT("Metasound_MathDivideFloatNodeDescription", "Divide float by another float."))
+	DEFINE_METASOUND_MATHOP(Divide, Int32, int32, LOCTEXT("Metasound_MathDivideInt32NodeDescription", "Divide int32 by another int32."))
+	DEFINE_METASOUND_MATHOP(Divide, Int64, int64, LOCTEXT("Metasound_MathDivideInt64NodeDescription", "Divide int64 by another int64."))
+	DEFINE_METASOUND_MATHOP(Divide, Frequency, FFrequency, LOCTEXT("Metasound_MathDivideFrequencyNodeDescription", "Divide frequency by another frequency."))
+	DEFINE_METASOUND_MATHOP(Divide, FloatTime, FFloatTime, LOCTEXT("Metasound_MathDivideFloatTimeNodeDescription", "Divide time by another time."))
+
+	DEFINE_METASOUND_MATHOP(RandRange, Float, float, LOCTEXT("Metasound_MathRandRangeFloatNodeDescription", "Computes random float value in the range [Min, Max]."))
+	DEFINE_METASOUND_MATHOP(RandRange, Int32, int32, LOCTEXT("Metasound_MathRandRangeInt32NodeDescription", "Computes random int32 in the range [Min, Max]."))
+	DEFINE_METASOUND_MATHOP(RandRange, Int64, int64, LOCTEXT("Metasound_MathRandRangeInt64NodeDescription", "Computes random int64 in the range [Min, Max]."))
+	DEFINE_METASOUND_MATHOP(RandRange, Frequency, FFrequency, LOCTEXT("Metasound_MathRandRangeFrequencyNodeDescription", "Computes random frequency value in the range [Min, Max]."))
+	DEFINE_METASOUND_MATHOP(RandRange, FloatTime, FFloatTime, LOCTEXT("Metasound_MathRandRangeFloatTimeNodeDescription", "Computes random time value in the range [Min, Max]."))
+}
 
 #undef LOCTEXT_NAMESPACE // MetasoundMathOpNode
