Fix the build by temporarily removing references to missing fields of Metasound::FNodeInfo added in CL 15125236.

Presumably a file was missing from original CL.

diff --git a/Engine/Plugins/Runtime/Metasound/Source/MetasoundStandardNodes/Public/MetasoundMathNodes.h b/Engine/Plugins/Runtime/Metasound/Source/MetasoundStandardNodes/Public/MetasoundMathNodes.h
--- a/Engine/Plugins/Runtime/Metasound/Source/MetasoundStandardNodes/Public/MetasoundMathNodes.h
+++ b/Engine/Plugins/Runtime/Metasound/Source/MetasoundStandardNodes/Public/MetasoundMathNodes.h
@@ -74,10 +74,10 @@ namespace Metasound
 					Info.Author = PluginAuthor;
 					Info.PromptIfMissing = PluginNodeMissingPrompt;
 					Info.DefaultInterface = DeclareVertexInterface();
-					Info.ImageName = TEXT("MetasoundEditor.Speaker");
-					Info.bShowName = false;
-					Info.bShowInputNames = false;
-					Info.bShowOutputNames = false;
+					//Info.ImageName = TEXT("MetasoundEditor.Speaker");
+					//Info.bShowName = false;
+					//Info.bShowInputNames = false;
+					//Info.bShowOutputNames = false;
 					Info.CategoryHierarchy = { LOCTEXT("Metasound_MathCategory", "Math") };
 
 					return Info;
