All Trigger Node AKA "Trigger Accumulator"

diff --git a/Engine/Plugins/Runtime/Metasound/Source/MetasoundStandardNodes/Private/MetasoundTriggerAccumulatorNode.cpp b/Engine/Plugins/Runtime/Metasound/Source/MetasoundStandardNodes/Private/MetasoundTriggerAccumulatorNode.cpp
new file mode 100644
--- /dev/null
+++ b/Engine/Plugins/Runtime/Metasound/Source/MetasoundStandardNodes/Private/MetasoundTriggerAccumulatorNode.cpp
@@ -0,0 +1,230 @@
+// Copyright Epic Games, Inc. All Rights Reserved.
+
+#include "MetasoundTriggerAccumulatorNode.h"
+
+#include "Internationalization/Text.h"
+#include "MetasoundBop.h"
+#include "MetasoundExecutableOperator.h"
+#include "MetasoundNodeRegistrationMacro.h"
+#include "MetasoundPrimitives.h"
+#include "MetasoundVertex.h"
+
+#define LOCTEXT_NAMESPACE "MetasoundStandardNodes"
+
+namespace Metasound
+{
+	METASOUND_REGISTER_NODE(FTriggerAccumulatorNode)
+
+	class FTriggerAccumulatorOperator : public TExecutableOperator<FTriggerAccumulatorOperator>
+	{
+		public:
+			static const FNodeInfo& GetNodeInfo();
+			static FVertexInterface DeclareVertexInterface();
+			static TUniquePtr<IOperator> CreateOperator(const FCreateOperatorParams& InParams, FBuildErrorArray& OutErrors);
+
+			~FTriggerAccumulatorOperator() = default;
+			
+			FTriggerAccumulatorOperator(
+				const FOperatorSettings& InSettings, 
+				const FTriggerReadRef& InInputTriggerToAccumulate,
+				const FTriggerReadRef& InInputTriggerResetCount,
+				const FInt32ReadRef& InInputTriggerAtCount);
+
+			virtual FDataReferenceCollection GetInputs() const override;
+			virtual FDataReferenceCollection GetOutputs() const override;
+			
+			void Execute();
+
+		private:
+			void Reset(int32 InCurretFrameIndex);
+			void Accumulate(int32 CurrentFrameIndex);
+
+			// Input. (triggers)
+			FTriggerReadRef TriggerToAccumulate;
+			FTriggerReadRef ResetTrigger;
+			
+			// Input (values)
+			FInt32ReadRef TriggerCount;
+
+			// Outputs (triggers)
+			FTriggerWriteRef CountReached;
+
+			// Output (values).
+			FInt32WriteRef CurrentCount;
+
+			static constexpr const TCHAR* AccumulatePinName = TEXT("In");
+			static constexpr const TCHAR* TriggerCountPinName = TEXT("TriggerCount");
+			static constexpr const TCHAR* ResetPinName = TEXT("Reset");
+			static constexpr const TCHAR* CountReachedPinName = TEXT("Out");
+			static constexpr const TCHAR* CurrentCountPinName = TEXT("Count");
+	};
+
+	FTriggerAccumulatorOperator::FTriggerAccumulatorOperator(
+		const FOperatorSettings& InSettings,
+		const FTriggerReadRef& InInputTriggerToAccumulate,
+		const FTriggerReadRef& InInputTriggerResetCount,
+		const FInt32ReadRef& InInputTriggerAtCount)
+		: TriggerToAccumulate(InInputTriggerToAccumulate)
+		, ResetTrigger(InInputTriggerResetCount)
+		, TriggerCount(InInputTriggerAtCount)
+		, CountReached(FBopWriteRef::CreateNew(InSettings))
+		, CurrentCount(FInt32WriteRef::CreateNew(0))
+	{}
+
+	FDataReferenceCollection FTriggerAccumulatorOperator::GetInputs() const
+	{
+		FDataReferenceCollection InputDataReferences;
+		InputDataReferences.AddDataReadReference(AccumulatePinName, FTriggerReadRef(TriggerToAccumulate));
+		InputDataReferences.AddDataReadReference(ResetPinName,	FTriggerReadRef(ResetTrigger));
+		InputDataReferences.AddDataReadReference(TriggerCountPinName, FInt32ReadRef(TriggerCount));
+		return InputDataReferences;
+	}
+
+	FDataReferenceCollection FTriggerAccumulatorOperator::GetOutputs() const
+	{
+		FDataReferenceCollection OutputDataReferences;
+		OutputDataReferences.AddDataReadReference(CountReachedPinName, FBopWriteRef(CountReached));
+		OutputDataReferences.AddDataReadReference(CurrentCountPinName, FInt32WriteRef(CurrentCount));
+		return OutputDataReferences;
+	}
+
+	void FTriggerAccumulatorOperator::Reset(int32 InCurrentFrameIndex)
+	{
+		*CurrentCount = 0;
+	}
+	void FTriggerAccumulatorOperator::Accumulate(int32 InCurrentFrameIndex)
+	{
+		++(*CurrentCount);
+		if (*CurrentCount >= *TriggerCount)
+		{
+			Reset(InCurrentFrameIndex);
+			CountReached->BopFrame(InCurrentFrameIndex);
+		}
+	}
+
+	void FTriggerAccumulatorOperator::Execute()
+	{		
+		// We need to advance our own output triggers.
+		CountReached->AdvanceBlock();
+
+		int32 NumResets = ResetTrigger->Num();
+		int32 NumAccumulates = TriggerToAccumulate->Num();
+		int32 MaxBops = FMath::Max(NumResets, NumAccumulates);
+
+		// Indices into each trigger array
+		int32 ResetIndex = 0;
+		int32 AccumulateIndex = 0;
+		
+		// Move forwards through the bop array and return next bop frame index, or INDEX_NONE if there's none remaining.
+		static const auto GetNextTriggerFrame = [](const FTriggerReadRef& InTrigger, int32& InOutIndex, const int32 InNumTriggers) 
+			-> int32 { return InOutIndex < InNumTriggers ? (*InTrigger)[InOutIndex++] : INDEX_NONE; };
+		
+		// Convenience lambdas for getting the next Reset/Accumulate frame index from either the Trigger or Reset.
+		auto GetNextResetFrame = [this, &ResetIndex, NumResets]() -> int32 { return GetNextTriggerFrame(ResetTrigger, ResetIndex, NumResets); };
+		auto GetNextAccumulateFrame = [this, &AccumulateIndex, NumAccumulates]() -> int32 { return GetNextTriggerFrame(TriggerToAccumulate, AccumulateIndex, NumAccumulates); };
+
+		int32 NextResetFrame = GetNextResetFrame();
+		int32 NextAccumulateFrame = GetNextAccumulateFrame();
+
+		// Walk through each trigger in order.
+		for (int32 i = 0; i < MaxBops; ++i)
+		{
+			// If there's no Reset, we can just process all Accumulates in order.
+			if (NextResetFrame == INDEX_NONE)
+			{
+				Accumulate(NextAccumulateFrame);
+				NextAccumulateFrame = GetNextAccumulateFrame();
+			}
+			// ... likewise if there's no Accumulate, we can just do the resets.
+			else if (NextAccumulateFrame == INDEX_NONE)
+			{
+				Reset(NextResetFrame);
+				NextResetFrame = GetNextResetFrame();
+			}
+			// ... both resets and accumulates remaining.
+			else
+			{
+				if (NextAccumulateFrame < NextResetFrame)
+				{
+					Accumulate(NextAccumulateFrame);
+					NextAccumulateFrame = GetNextAccumulateFrame();
+				}
+				else if(NextResetFrame < NextAccumulateFrame)
+				{
+					Reset(NextResetFrame);
+					NextResetFrame = GetNextResetFrame();
+				}
+				else // Both on same frame, make sure Reset happens first.
+				{
+					Reset(NextResetFrame);
+					Accumulate(NextAccumulateFrame);
+					NextAccumulateFrame = GetNextAccumulateFrame();
+					NextResetFrame = GetNextResetFrame();
+				}
+			}
+		}
+	}
+
+	FVertexInterface FTriggerAccumulatorOperator::DeclareVertexInterface()
+	{
+		static const FVertexInterface Interface(
+			FInputVertexInterface(
+				TInputDataVertexModel<FBop>(AccumulatePinName, LOCTEXT("TriggerToAccumulateTooltop", "Trigger to accumulate")),
+				TInputDataVertexModel<FBop>(ResetPinName, LOCTEXT("ResetAccumulatorTooltip", "Reset the accumulator")),
+				TInputDataVertexModel<int32>(TriggerCountPinName, LOCTEXT("CountToTriggerAtTooltip", "Count to trigger at"))
+			),
+			FOutputVertexInterface(
+				TOutputDataVertexModel<FBop>(CountReachedPinName, LOCTEXT("CountReachedTooltop", "Triggered when the accumulated count is reached")),				
+				TOutputDataVertexModel<int32>(CurrentCountPinName, LOCTEXT("CountTooltop", "Current accumulator count"))
+			)
+		);
+
+		return Interface;
+	}
+
+	const FNodeInfo& FTriggerAccumulatorOperator::GetNodeInfo()
+	{
+		auto InitNodeInfo = []() -> FNodeInfo
+		{
+			FNodeInfo Info;
+			Info.ClassName = FName(TEXT("Trigger Accumulator"));
+			Info.MajorVersion = 1;
+			Info.MinorVersion = 0;
+			Info.Description = LOCTEXT("Metasound_TriggerAccumulatorNodeDescription", "Accumulates Triggers until a Count is hit and then it Triggers");
+			Info.Author = PluginAuthor;
+			Info.PromptIfMissing = PluginNodeMissingPrompt;
+			Info.DefaultInterface = DeclareVertexInterface();
+
+			return Info;
+		};
+
+		static const FNodeInfo Info = InitNodeInfo();
+
+		return Info;
+	}
+
+	TUniquePtr<IOperator> FTriggerAccumulatorOperator::CreateOperator(const FCreateOperatorParams& InParams, FBuildErrorArray& OutErrors)
+	{
+		const FTriggerAccumulatorNode& TaNode = static_cast<const FTriggerAccumulatorNode&>(InParams.Node);
+		const FDataReferenceCollection& InputCollection = InParams.InputDataReferences;
+				
+		return MakeUnique<FTriggerAccumulatorOperator>(
+			InParams.OperatorSettings, 
+			InputCollection.GetDataReadReferenceOrConstruct<FBop>(AccumulatePinName, InParams.OperatorSettings),
+			InputCollection.GetDataReadReferenceOrConstruct<FBop>(ResetPinName, InParams.OperatorSettings),
+			InputCollection.GetDataReadReferenceOrConstruct<int32>(TriggerCountPinName, TaNode.GetDefaultTriggerAt() )
+		);
+	}
+
+	// Node.
+	FTriggerAccumulatorNode::FTriggerAccumulatorNode(const FString& InName, int32 InDefaultTriggerAtCount)
+		: FNodeFacade(InName, TFacadeOperatorClass<FTriggerAccumulatorOperator>())
+		, DefaultTriggerAtCount(InDefaultTriggerAtCount)
+	{}
+
+	FTriggerAccumulatorNode::FTriggerAccumulatorNode(const FNodeInitData& InitData)
+		: FTriggerAccumulatorNode(InitData.InstanceName, 1)
+	{}
+}
+
+#undef LOCTEXT_NAMESPACE
\ No newline at end of file
diff --git a/Engine/Plugins/Runtime/Metasound/Source/MetasoundStandardNodes/Public/MetasoundTriggerAccumulatorNode.h b/Engine/Plugins/Runtime/Metasound/Source/MetasoundStandardNodes/Public/MetasoundTriggerAccumulatorNode.h
new file mode 100644
--- /dev/null
+++ b/Engine/Plugins/Runtime/Metasound/Source/MetasoundStandardNodes/Public/MetasoundTriggerAccumulatorNode.h
@@ -0,0 +1,43 @@
+// Copyright Epic Games, Inc. All Rights Reserved.
+#pragma once
+
+#include "CoreMinimal.h"
+
+#include "MetasoundFacade.h"
+
+namespace Metasound
+{
+	/** FTriggerAccumulatorNode
+	 *
+	 *  Creates a Trigger Accumulator Node, that accumulate triggers until hitting a count, then Trigger and reset.
+	 */
+	class METASOUNDSTANDARDNODES_API FTriggerAccumulatorNode : public FNodeFacade
+	{
+	public:
+
+		/** Trigger Accumulator node constructor.
+		 *
+		 * @param InName - Name of this node.
+		 * @param InTriggerAtCount - The count at which after accumulated triggers we trigger.
+		 */
+		FTriggerAccumulatorNode(
+			const FString& InName,
+			int32 InDefaultTriggerAtCount
+		);
+
+		/**
+		 * Constructor used by the Metasound Frontend.
+		 */
+		FTriggerAccumulatorNode(const FNodeInitData& InitData);
+
+		virtual ~FTriggerAccumulatorNode() = default;
+
+		/** 
+		 * Get Default Trigger At Value.
+		 */
+		int32 GetDefaultTriggerAt() const { return DefaultTriggerAtCount; }
+
+	private:
+		int32 DefaultTriggerAtCount = 1;
+	};
+} // namespace Metasound
