Zero FAudioBuffer on creation

diff --git a/Engine/Plugins/Runtime/Metasound/Source/MetasoundFrontend/Public/MetasoundAudioBuffer.h b/Engine/Plugins/Runtime/Metasound/Source/MetasoundFrontend/Public/MetasoundAudioBuffer.h
--- a/Engine/Plugins/Runtime/Metasound/Source/MetasoundFrontend/Public/MetasoundAudioBuffer.h
+++ b/Engine/Plugins/Runtime/Metasound/Source/MetasoundFrontend/Public/MetasoundAudioBuffer.h
@@ -29,12 +29,12 @@ namespace Metasound
 			 */
 			explicit FAudioBuffer(int32 InNumSamples)
 			{
-				Buffer.AddUninitialized(InNumSamples);
+				Buffer.AddZeroed(InNumSamples);
 
 #if METASOUNDGRAPHCORE_CHECKAUDIONUM
 				// InitialNum should not change during the life of an FAudioBuffer
 				InitialNum = Buffer.Num();
-#endif				
+#endif
 			}
 
 			/**
@@ -42,7 +42,7 @@ namespace Metasound
 			 */
 			FAudioBuffer(const FOperatorSettings& InSettings)
 			{
-				Buffer.AddUninitialized(InSettings.GetNumFramesPerBlock());
+				Buffer.AddZeroed(InSettings.GetNumFramesPerBlock());
 
 #if METASOUNDGRAPHCORE_CHECKAUDIONUM
 				// InitialNum should not change during the life of an FAudioBuffer
@@ -55,7 +55,7 @@ namespace Metasound
 #if METASOUNDGRAPHCORE_CHECKAUDIONUM
 				// InitialNum should not change during the life of an FAudioBuffer
 				InitialNum = Buffer.Num();
-#endif				
+#endif
 			}
 
 			/** Return a pointer to the audio float data. */
@@ -72,7 +72,7 @@ namespace Metasound
 			{
 #if METASOUNDGRAPHCORE_CHECKAUDIONUM
 				UE_CLOG(InitialNum != Buffer.Num(), LogMetasound, Error, TEXT("Metasound audio buffer size change detected.  Audio Buffers should not be resized."));
-#endif				
+#endif
 				return Buffer.GetData();
 			}
 
@@ -81,7 +81,7 @@ namespace Metasound
 			{
 #if METASOUNDGRAPHCORE_CHECKAUDIONUM
 				UE_CLOG(InitialNum != Buffer.Num(), LogMetasound, Error, TEXT("Metasound audio buffer size change detected.  Audio Buffers should not be resized."));
-#endif				
+#endif
 				return Buffer.Num();
 			}
 
@@ -90,8 +90,8 @@ namespace Metasound
 			{
 #if METASOUNDGRAPHCORE_CHECKAUDIONUM
 				UE_CLOG(InitialNum != Buffer.Num(), LogMetasound, Error, TEXT("Metasound audio buffer size change detected.  Audio Buffers should not be resized."));
-#endif				
-				return Buffer;
+#endif
+return Buffer;
 			}
 
 			/** Implicit conversion to Audio::AlignedFloatBuffer 
@@ -120,7 +120,7 @@ namespace Metasound
 
 #if METASOUNDGRAPHCORE_CHECKAUDIONUM
 			int32 InitialNum;
-#endif				
+#endif
 	};
 
 	DECLARE_METASOUND_DATA_REFERENCE_TYPES(FAudioBuffer, METASOUNDFRONTEND_API, FAudioBufferTypeInfo, FAudioBufferReadRef, FAudioBufferWriteRef);
