Rename Bop to Trigger in a couple more place I missed in inital review.

diff --git a/Engine/Plugins/Runtime/Metasound/Source/MetasoundStandardNodes/Private/MetasoundTriggerAccumulatorNode.cpp b/Engine/Plugins/Runtime/Metasound/Source/MetasoundStandardNodes/Private/MetasoundTriggerAccumulatorNode.cpp
--- a/Engine/Plugins/Runtime/Metasound/Source/MetasoundStandardNodes/Private/MetasoundTriggerAccumulatorNode.cpp
+++ b/Engine/Plugins/Runtime/Metasound/Source/MetasoundStandardNodes/Private/MetasoundTriggerAccumulatorNode.cpp
@@ -67,7 +67,7 @@ namespace Metasound
 		: TriggerToAccumulate(InInputTriggerToAccumulate)
 		, ResetTrigger(InInputTriggerResetCount)
 		, TriggerCount(InInputTriggerAtCount)
-		, CountReached(FBopWriteRef::CreateNew(InSettings))
+		, CountReached(FTriggerWriteRef::CreateNew(InSettings))
 		, CurrentCount(FInt32WriteRef::CreateNew(0))
 	{}
 
@@ -83,7 +83,7 @@ namespace Metasound
 	FDataReferenceCollection FTriggerAccumulatorOperator::GetOutputs() const
 	{
 		FDataReferenceCollection OutputDataReferences;
-		OutputDataReferences.AddDataReadReference(CountReachedPinName, FBopWriteRef(CountReached));
+		OutputDataReferences.AddDataReadReference(CountReachedPinName, FTriggerWriteRef(CountReached));
 		OutputDataReferences.AddDataReadReference(CurrentCountPinName, FInt32WriteRef(CurrentCount));
 		return OutputDataReferences;
 	}
@@ -109,7 +109,7 @@ namespace Metasound
 
 		int32 NumResets = ResetTrigger->Num();
 		int32 NumAccumulates = TriggerToAccumulate->Num();
-		int32 MaxBops = FMath::Max(NumResets, NumAccumulates);
+		int32 MaxTriggers = FMath::Max(NumResets, NumAccumulates);
 
 		// Indices into each trigger array
 		int32 ResetIndex = 0;
@@ -127,7 +127,7 @@ namespace Metasound
 		int32 NextAccumulateFrame = GetNextAccumulateFrame();
 
 		// Walk through each trigger in order.
-		for (int32 i = 0; i < MaxBops; ++i)
+		for (int32 i = 0; i < MaxTriggers; ++i)
 		{
 			// If there's no Reset, we can just process all Accumulates in order.
 			if (NextResetFrame == INDEX_NONE)
@@ -169,12 +169,12 @@ namespace Metasound
 	{
 		static const FVertexInterface Interface(
 			FInputVertexInterface(
-				TInputDataVertexModel<FBop>(AccumulatePinName, LOCTEXT("TriggerToAccumulateTooltop", "Trigger to accumulate")),
-				TInputDataVertexModel<FBop>(ResetPinName, LOCTEXT("ResetAccumulatorTooltip", "Reset the accumulator")),
+				TInputDataVertexModel<FTrigger>(AccumulatePinName, LOCTEXT("TriggerToAccumulateTooltop", "Trigger to accumulate")),
+				TInputDataVertexModel<FTrigger>(ResetPinName, LOCTEXT("ResetAccumulatorTooltip", "Reset the accumulator")),
 				TInputDataVertexModel<int32>(TriggerCountPinName, LOCTEXT("CountToTriggerAtTooltip", "Count to trigger at"))
 			),
 			FOutputVertexInterface(
-				TOutputDataVertexModel<FBop>(CountReachedPinName, LOCTEXT("CountReachedTooltop", "Triggered when the accumulated count is reached")),				
+				TOutputDataVertexModel<FTrigger>(CountReachedPinName, LOCTEXT("CountReachedTooltop", "Triggered when the accumulated count is reached")),				
 				TOutputDataVertexModel<int32>(CurrentCountPinName, LOCTEXT("CountTooltop", "Current accumulator count"))
 			)
 		);
@@ -210,8 +210,8 @@ namespace Metasound
 				
 		return MakeUnique<FTriggerAccumulatorOperator>(
 			InParams.OperatorSettings, 
-			InputCollection.GetDataReadReferenceOrConstruct<FBop>(AccumulatePinName, InParams.OperatorSettings),
-			InputCollection.GetDataReadReferenceOrConstruct<FBop>(ResetPinName, InParams.OperatorSettings),
+			InputCollection.GetDataReadReferenceOrConstruct<FTrigger>(AccumulatePinName, InParams.OperatorSettings),
+			InputCollection.GetDataReadReferenceOrConstruct<FTrigger>(ResetPinName, InParams.OperatorSettings),
 			InputCollection.GetDataReadReferenceOrConstruct<int32>(TriggerCountPinName, TaNode.GetDefaultTriggerAt() )
 		);
 	}
