Fix crash when deleting comment box
- Add a few new pin types (WIP)

diff --git a/Engine/Plugins/Runtime/Metasound/Source/MetasoundEditor/Private/MetasoundEditor.cpp b/Engine/Plugins/Runtime/Metasound/Source/MetasoundEditor/Private/MetasoundEditor.cpp
--- a/Engine/Plugins/Runtime/Metasound/Source/MetasoundEditor/Private/MetasoundEditor.cpp
+++ b/Engine/Plugins/Runtime/Metasound/Source/MetasoundEditor/Private/MetasoundEditor.cpp
@@ -699,50 +699,58 @@ namespace Metasound
 				Frontend::FConstDocumentHandle DocumentHandle = MetasoundAsset->GetDocumentHandle();
 				for (FGraphPanelSelectionSet::TConstIterator NodeIt(SelectedNodes); NodeIt; ++NodeIt)
 				{
-					UMetasoundEditorGraphNode* Node = CastChecked<UMetasoundEditorGraphNode>(*NodeIt);
-					Frontend::FNodeHandle NodeHandle = Node->GetNodeHandle();
+					Frontend::FNodeHandle NodeHandle = Frontend::INodeController::GetInvalidHandle();
 
-					if (NodeHandle->GetClassType() == EMetasoundFrontendClassType::Input)
+					if (UMetasoundEditorGraphNode* Node = Cast<UMetasoundEditorGraphNode>(*NodeIt))
 					{
-						auto IsRequiredInput = [&](const Frontend::FConstInputHandle& InputHandle) 
-						{ 
-							return DocumentHandle->IsRequiredInput(InputHandle->GetName()); 
-						};
-						TArray<Frontend::FConstInputHandle> NodeInputs = NodeHandle->GetConstInputs();
+						NodeHandle = Node->GetNodeHandle();
 
-						if (Frontend::FConstInputHandle* InputHandle = NodeInputs.FindByPredicate(IsRequiredInput))
+						if (NodeHandle->GetClassType() == EMetasoundFrontendClassType::Input)
 						{
-							FNotificationInfo Info(FText::Format(LOCTEXT("Metasounds_CannotDeleteRequiredInput",
-								"'Required Input '{0}' cannot be deleted."), (*InputHandle)->GetDisplayName()));
-							Info.bFireAndForget = true;
-							Info.ExpireDuration = 2.0f;
-							Info.bUseThrobber = true;
-							FSlateNotificationManager::Get().AddNotification(Info);
-							continue;
-						}
-					}
+							auto IsRequiredInput = [&](const Frontend::FConstInputHandle& InputHandle) 
+							{ 
+								return DocumentHandle->IsRequiredInput(InputHandle->GetName()); 
+							};
+							TArray<Frontend::FConstInputHandle> NodeInputs = NodeHandle->GetConstInputs();
 
-					if (NodeHandle->GetClassType() == EMetasoundFrontendClassType::Output)
-					{
-						auto IsRequiredOutput = [&](const Frontend::FConstOutputHandle& OutputHandle) 
-						{ 
-							return DocumentHandle->IsRequiredOutput(OutputHandle->GetName()); 
-						};
-						TArray<Frontend::FConstOutputHandle> NodeOutputs = NodeHandle->GetConstOutputs();
+							if (Frontend::FConstInputHandle* InputHandle = NodeInputs.FindByPredicate(IsRequiredInput))
+							{
+								FNotificationInfo Info(FText::Format(LOCTEXT("Metasounds_CannotDeleteRequiredInput",
+									"'Required Input '{0}' cannot be deleted."), (*InputHandle)->GetDisplayName()));
+								Info.bFireAndForget = true;
+								Info.ExpireDuration = 2.0f;
+								Info.bUseThrobber = true;
+								FSlateNotificationManager::Get().AddNotification(Info);
+								continue;
+							}
+						}
 
-						if (Frontend::FConstOutputHandle* OutputHandle = NodeOutputs.FindByPredicate(IsRequiredOutput))
+						if (NodeHandle->GetClassType() == EMetasoundFrontendClassType::Output)
 						{
-							FNotificationInfo Info(FText::Format(LOCTEXT("Metasounds_CannotDeleteRequiredOutput",
-								"'Required Output '{0}' cannot be deleted."), (*OutputHandle)->GetDisplayName()));
-							Info.bFireAndForget = true;
-							Info.ExpireDuration = 2.0f;
-							Info.bUseThrobber = true;
-							FSlateNotificationManager::Get().AddNotification(Info);
-							continue;
+							auto IsRequiredOutput = [&](const Frontend::FConstOutputHandle& OutputHandle) 
+							{ 
+								return DocumentHandle->IsRequiredOutput(OutputHandle->GetName()); 
+							};
+							TArray<Frontend::FConstOutputHandle> NodeOutputs = NodeHandle->GetConstOutputs();
+
+							if (Frontend::FConstOutputHandle* OutputHandle = NodeOutputs.FindByPredicate(IsRequiredOutput))
+							{
+								FNotificationInfo Info(FText::Format(LOCTEXT("Metasounds_CannotDeleteRequiredOutput",
+									"'Required Output '{0}' cannot be deleted."), (*OutputHandle)->GetDisplayName()));
+								Info.bFireAndForget = true;
+								Info.ExpireDuration = 2.0f;
+								Info.bUseThrobber = true;
+								FSlateNotificationManager::Get().AddNotification(Info);
+								continue;
+							}
 						}
 					}
 
-					FGraphBuilder::DeleteNode(*Node, false /* bInRecordTransaction */);
+					// Some nodes may not be metasound nodes (ex. comments and perhaps aliases eventually), but can be safely deleted.
+					if (UEdGraphNode* Node = Cast<UEdGraphNode>(*NodeIt))
+					{
+						FGraphBuilder::DeleteNode(*Node, NodeHandle, true /* bInRecordTransaction */);
+					}
 				}
 			}
 		}
diff --git a/Engine/Plugins/Runtime/Metasound/Source/MetasoundEditor/Private/MetasoundEditorGraphBuilder.cpp b/Engine/Plugins/Runtime/Metasound/Source/MetasoundEditor/Private/MetasoundEditorGraphBuilder.cpp
--- a/Engine/Plugins/Runtime/Metasound/Source/MetasoundEditor/Private/MetasoundEditorGraphBuilder.cpp
+++ b/Engine/Plugins/Runtime/Metasound/Source/MetasoundEditor/Private/MetasoundEditorGraphBuilder.cpp
@@ -25,12 +25,16 @@ namespace Metasound
 {
 	namespace Editor
 	{
-		const FName FGraphBuilder::PinPrimitiveBoolean = TEXT("Boolean");
-		const FName FGraphBuilder::PinPrimitiveFloat = TEXT("Float");
-		const FName FGraphBuilder::PinPrimitiveInteger = TEXT("Int");
-		const FName FGraphBuilder::PinPrimitiveString = TEXT("String");
-		const FName FGraphBuilder::PinPrimitiveUObject = TEXT("UObject");
-		const FName FGraphBuilder::PinPrimitiveUObjectArray = TEXT("UObjectArray");
+		const FName FGraphBuilder::PinAudioFormat = "Format";
+		const FName FGraphBuilder::PinAudioNumeric = "Numeric";
+		const FName FGraphBuilder::PinPrimitiveBoolean = "Boolean";
+		const FName FGraphBuilder::PinPrimitiveFloat = "Float";
+		const FName FGraphBuilder::PinPrimitiveInt32 = "Int32";
+		const FName FGraphBuilder::PinPrimitiveInt64 = "Int64";
+		const FName FGraphBuilder::PinPrimitiveString = "String";
+		const FName FGraphBuilder::PinPrimitiveTrigger = "Trigger";
+		const FName FGraphBuilder::PinPrimitiveUObject = "UObject";
+		const FName FGraphBuilder::PinPrimitiveUObjectArray = "UObjectArray";
 
 		UEdGraphNode* FGraphBuilder::AddNode(UObject& InMetasound, const FVector2D& Location, Frontend::FNodeHandle& InNodeHandle, bool bInSelectNewNode)
 		{
@@ -262,7 +266,7 @@ namespace Metasound
 			return NodeHandle;
 		}
 
-		void FGraphBuilder::DeleteNode(UMetasoundEditorGraphNode& InNode, bool bInRecordTransaction)
+		void FGraphBuilder::DeleteNode(UEdGraphNode& InNode, Frontend::FNodeHandle InNodeHandle, bool bInRecordTransaction)
 		{
 			const FScopedTransaction Transaction(LOCTEXT("DeleteMetasoundGraphNode", "Delete Metasound Node"), bInRecordTransaction);
 
@@ -273,15 +277,16 @@ namespace Metasound
 				Graph->MarkPackageDirty();
 			}
 
-			const Frontend::FNodeHandle& NodeHandle = InNode.GetNodeHandle();
-
-			FMetasoundAssetBase* MetasoundAsset = IMetasoundUObjectRegistry::Get().GetObjectAsAssetBase(Graph->GetMetasound());
-			check(MetasoundAsset);
-
-			Frontend::FGraphHandle GraphHandle = MetasoundAsset->GetRootGraphHandle();
-			if (GraphHandle->IsValid() && NodeHandle->IsValid())
+			if (InNodeHandle->IsValid())
 			{
-				GraphHandle->RemoveNode(*NodeHandle);
+				FMetasoundAssetBase* MetasoundAsset = IMetasoundUObjectRegistry::Get().GetObjectAsAssetBase(Graph->GetMetasound());
+				check(MetasoundAsset);
+
+				Frontend::FGraphHandle GraphHandle = MetasoundAsset->GetRootGraphHandle();
+				if (GraphHandle->IsValid())
+				{
+					GraphHandle->RemoveNode(*InNodeHandle);
+				}
 			}
 
 			InNode.PostEditChange();
diff --git a/Engine/Plugins/Runtime/Metasound/Source/MetasoundEditor/Private/MetasoundEditorGraphBuilder.h b/Engine/Plugins/Runtime/Metasound/Source/MetasoundEditor/Private/MetasoundEditorGraphBuilder.h
--- a/Engine/Plugins/Runtime/Metasound/Source/MetasoundEditor/Private/MetasoundEditorGraphBuilder.h
+++ b/Engine/Plugins/Runtime/Metasound/Source/MetasoundEditor/Private/MetasoundEditorGraphBuilder.h
@@ -22,10 +22,14 @@ namespace Metasound
 		class FGraphBuilder
 		{
 		public:
+			static const FName PinAudioFormat;
+			static const FName PinAudioNumeric;
 			static const FName PinPrimitiveBoolean;
 			static const FName PinPrimitiveFloat;
-			static const FName PinPrimitiveInteger;
+			static const FName PinPrimitiveInt32;
+			static const FName PinPrimitiveInt64;
 			static const FName PinPrimitiveString;
+			static const FName PinPrimitiveTrigger;
 			static const FName PinPrimitiveUObject;
 			static const FName PinPrimitiveUObjectArray;
 
@@ -56,7 +60,7 @@ namespace Metasound
 			// Adds an output node handle with the given class info
 			static Frontend::FNodeHandle AddOutputNodeHandle(UObject& InMetasound, const FString& InName, const FName InTypeName, const FText& InToolTip);
 
-			static void DeleteNode(UMetasoundEditorGraphNode& InNode, bool bInRecordTransaction = true);
+			static void DeleteNode(UEdGraphNode& InNode, Frontend::FNodeHandle InNodeHandle, bool bInRecordTransaction = true);
 
 			static void RebuildGraph(UObject& InMetasound);
 
