Fix build: Remove pin changes

diff --git a/Engine/Plugins/Runtime/Metasound/Source/MetasoundEditor/Private/MetasoundEditorGraphBuilder.cpp b/Engine/Plugins/Runtime/Metasound/Source/MetasoundEditor/Private/MetasoundEditorGraphBuilder.cpp
--- a/Engine/Plugins/Runtime/Metasound/Source/MetasoundEditor/Private/MetasoundEditorGraphBuilder.cpp
+++ b/Engine/Plugins/Runtime/Metasound/Source/MetasoundEditor/Private/MetasoundEditorGraphBuilder.cpp
@@ -266,7 +266,7 @@ namespace Metasound
 			return NodeHandle;
 		}
 
-		void FGraphBuilder::DeleteNode(UEdGraphNode& InNode, Frontend::FNodeHandle InNodeHandle, bool bInRecordTransaction)
+		void FGraphBuilder::DeleteNode(UMetasoundEditorGraphNode& InNode, bool bInRecordTransaction)
 		{
 			const FScopedTransaction Transaction(LOCTEXT("DeleteMetasoundGraphNode", "Delete Metasound Node"), bInRecordTransaction);
 
@@ -277,16 +277,15 @@ namespace Metasound
 				Graph->MarkPackageDirty();
 			}
 
-			if (InNodeHandle->IsValid())
-			{
-				FMetasoundAssetBase* MetasoundAsset = IMetasoundUObjectRegistry::Get().GetObjectAsAssetBase(Graph->GetMetasound());
-				check(MetasoundAsset);
+			const Frontend::FNodeHandle& NodeHandle = InNode.GetNodeHandle();
 
-				Frontend::FGraphHandle GraphHandle = MetasoundAsset->GetRootGraphHandle();
-				if (GraphHandle->IsValid())
-				{
-					GraphHandle->RemoveNode(*InNodeHandle);
-				}
+			FMetasoundAssetBase* MetasoundAsset = IMetasoundUObjectRegistry::Get().GetObjectAsAssetBase(Graph->GetMetasound());
+			check(MetasoundAsset);
+
+			Frontend::FGraphHandle GraphHandle = MetasoundAsset->GetRootGraphHandle();
+			if (GraphHandle->IsValid() && NodeHandle->IsValid())
+			{
+				GraphHandle->RemoveNode(*NodeHandle);
 			}
 
 			InNode.PostEditChange();
diff --git a/Engine/Plugins/Runtime/Metasound/Source/MetasoundEditor/Private/MetasoundEditorGraphBuilder.h b/Engine/Plugins/Runtime/Metasound/Source/MetasoundEditor/Private/MetasoundEditorGraphBuilder.h
--- a/Engine/Plugins/Runtime/Metasound/Source/MetasoundEditor/Private/MetasoundEditorGraphBuilder.h
+++ b/Engine/Plugins/Runtime/Metasound/Source/MetasoundEditor/Private/MetasoundEditorGraphBuilder.h
@@ -60,7 +60,7 @@ namespace Metasound
 			// Adds an output node handle with the given class info
 			static Frontend::FNodeHandle AddOutputNodeHandle(UObject& InMetasound, const FString& InName, const FName InTypeName, const FText& InToolTip);
 
-			static void DeleteNode(UEdGraphNode& InNode, Frontend::FNodeHandle InNodeHandle, bool bInRecordTransaction = true);
+			static void DeleteNode(UMetasoundEditorGraphNode& InNode, bool bInRecordTransaction = true);
 
 			static void RebuildGraph(UObject& InMetasound);
 
