Don't mind me here, just reverting the wrong part of a reverted change because I lack sleep and have children

diff --git a/Engine/Plugins/Runtime/Metasound/Source/MetasoundEditor/Private/MetasoundEditorGraphBuilder.cpp b/Engine/Plugins/Runtime/Metasound/Source/MetasoundEditor/Private/MetasoundEditorGraphBuilder.cpp
--- a/Engine/Plugins/Runtime/Metasound/Source/MetasoundEditor/Private/MetasoundEditorGraphBuilder.cpp
+++ b/Engine/Plugins/Runtime/Metasound/Source/MetasoundEditor/Private/MetasoundEditorGraphBuilder.cpp
@@ -25,16 +25,12 @@ namespace Metasound
 {
 	namespace Editor
 	{
-		const FName FGraphBuilder::PinAudioFormat = "Format";
-		const FName FGraphBuilder::PinAudioNumeric = "Numeric";
-		const FName FGraphBuilder::PinPrimitiveBoolean = "Boolean";
-		const FName FGraphBuilder::PinPrimitiveFloat = "Float";
-		const FName FGraphBuilder::PinPrimitiveInt32 = "Int32";
-		const FName FGraphBuilder::PinPrimitiveInt64 = "Int64";
-		const FName FGraphBuilder::PinPrimitiveString = "String";
-		const FName FGraphBuilder::PinPrimitiveTrigger = "Trigger";
-		const FName FGraphBuilder::PinPrimitiveUObject = "UObject";
-		const FName FGraphBuilder::PinPrimitiveUObjectArray = "UObjectArray";
+		const FName FGraphBuilder::PinPrimitiveBoolean = TEXT("Boolean");
+		const FName FGraphBuilder::PinPrimitiveFloat = TEXT("Float");
+		const FName FGraphBuilder::PinPrimitiveInteger = TEXT("Int");
+		const FName FGraphBuilder::PinPrimitiveString = TEXT("String");
+		const FName FGraphBuilder::PinPrimitiveUObject = TEXT("UObject");
+		const FName FGraphBuilder::PinPrimitiveUObjectArray = TEXT("UObjectArray");
 
 		UEdGraphNode* FGraphBuilder::AddNode(UObject& InMetasound, const FVector2D& Location, Frontend::FNodeHandle& InNodeHandle, bool bInSelectNewNode)
 		{
@@ -266,7 +262,7 @@ namespace Metasound
 			return NodeHandle;
 		}
 
-		void FGraphBuilder::DeleteNode(UMetasoundEditorGraphNode& InNode, bool bInRecordTransaction)
+		void FGraphBuilder::DeleteNode(UEdGraphNode& InNode, Frontend::FNodeHandle InNodeHandle, bool bInRecordTransaction)
 		{
 			const FScopedTransaction Transaction(LOCTEXT("DeleteMetasoundGraphNode", "Delete Metasound Node"), bInRecordTransaction);
 
@@ -277,15 +273,16 @@ namespace Metasound
 				Graph->MarkPackageDirty();
 			}
 
-			const Frontend::FNodeHandle& NodeHandle = InNode.GetNodeHandle();
-
-			FMetasoundAssetBase* MetasoundAsset = IMetasoundUObjectRegistry::Get().GetObjectAsAssetBase(Graph->GetMetasound());
-			check(MetasoundAsset);
-
-			Frontend::FGraphHandle GraphHandle = MetasoundAsset->GetRootGraphHandle();
-			if (GraphHandle->IsValid() && NodeHandle->IsValid())
+			if (InNodeHandle->IsValid())
 			{
-				GraphHandle->RemoveNode(*NodeHandle);
+				FMetasoundAssetBase* MetasoundAsset = IMetasoundUObjectRegistry::Get().GetObjectAsAssetBase(Graph->GetMetasound());
+				check(MetasoundAsset);
+
+				Frontend::FGraphHandle GraphHandle = MetasoundAsset->GetRootGraphHandle();
+				if (GraphHandle->IsValid())
+				{
+					GraphHandle->RemoveNode(*InNodeHandle);
+				}
 			}
 
 			InNode.PostEditChange();
diff --git a/Engine/Plugins/Runtime/Metasound/Source/MetasoundEditor/Private/MetasoundEditorGraphBuilder.h b/Engine/Plugins/Runtime/Metasound/Source/MetasoundEditor/Private/MetasoundEditorGraphBuilder.h
--- a/Engine/Plugins/Runtime/Metasound/Source/MetasoundEditor/Private/MetasoundEditorGraphBuilder.h
+++ b/Engine/Plugins/Runtime/Metasound/Source/MetasoundEditor/Private/MetasoundEditorGraphBuilder.h
@@ -22,14 +22,10 @@ namespace Metasound
 		class FGraphBuilder
 		{
 		public:
-			static const FName PinAudioFormat;
-			static const FName PinAudioNumeric;
 			static const FName PinPrimitiveBoolean;
 			static const FName PinPrimitiveFloat;
-			static const FName PinPrimitiveInt32;
-			static const FName PinPrimitiveInt64;
+			static const FName PinPrimitiveInteger;
 			static const FName PinPrimitiveString;
-			static const FName PinPrimitiveTrigger;
 			static const FName PinPrimitiveUObject;
 			static const FName PinPrimitiveUObjectArray;
 
@@ -60,7 +56,7 @@ namespace Metasound
 			// Adds an output node handle with the given class info
 			static Frontend::FNodeHandle AddOutputNodeHandle(UObject& InMetasound, const FString& InName, const FName InTypeName, const FText& InToolTip);
 
-			static void DeleteNode(UMetasoundEditorGraphNode& InNode, bool bInRecordTransaction = true);
+			static void DeleteNode(UEdGraphNode& InNode, Frontend::FNodeHandle InNodeHandle, bool bInRecordTransaction = true);
 
 			static void RebuildGraph(UObject& InMetasound);
 
