Add more fine grain pin Metasound specialization/colorization
- Stub in icons & node bodies
- Add SMetasoundNode & basic display option metadata
- Disallow multiple inputs/outputs with the same display name
- Various UX fixes

diff --git a/Engine/Plugins/Runtime/Metasound/Source/MetasoundEditor/MetasoundEditor.Build.cs b/Engine/Plugins/Runtime/Metasound/Source/MetasoundEditor/MetasoundEditor.Build.cs
--- a/Engine/Plugins/Runtime/Metasound/Source/MetasoundEditor/MetasoundEditor.Build.cs
+++ b/Engine/Plugins/Runtime/Metasound/Source/MetasoundEditor/MetasoundEditor.Build.cs
@@ -17,6 +17,7 @@ namespace UnrealBuildTool.Rules
 				{
 					"AudioExtensions",
 					"AudioMixer",
+					"KismetWidgets",
 					"MetasoundEngine",
 					"MetasoundFrontend",
 					"MetasoundGraphCore",
diff --git a/Engine/Plugins/Runtime/Metasound/Source/MetasoundEditor/Private/MetasoundDetailCustomization.cpp b/Engine/Plugins/Runtime/Metasound/Source/MetasoundEditor/Private/MetasoundDetailCustomization.cpp
--- a/Engine/Plugins/Runtime/Metasound/Source/MetasoundEditor/Private/MetasoundDetailCustomization.cpp
+++ b/Engine/Plugins/Runtime/Metasound/Source/MetasoundEditor/Private/MetasoundDetailCustomization.cpp
@@ -7,17 +7,20 @@
 #include "DetailCategoryBuilder.h"
 #include "DetailLayoutBuilder.h"
 #include "DetailWidgetRow.h"
+#include "Framework/Notifications/NotificationManager.h"
 #include "IDetailGroup.h"
 #include "MetasoundAssetBase.h"
 #include "MetasoundFrontend.h"
 #include "MetasoundUObjectRegistry.h"
 #include "PropertyEditorDelegates.h"
 #include "PropertyHandle.h"
+#include "PropertyRestriction.h"
 #include "SlateCore/Public/Styling/SlateColor.h"
 #include "Templates/Casts.h"
 #include "Templates/SharedPointer.h"
 #include "UObject/WeakObjectPtr.h"
 #include "UObject/WeakObjectPtrTemplates.h"
+#include "Widgets/Notifications/SNotificationList.h"
 #include "Widgets/Text/STextBlock.h"
 
 
@@ -37,10 +40,12 @@ namespace Metasound
 			return FName(InBasePath.ToString() + TEXT(".") + InPropertyName.ToString());
 		}
 
-		template <typename T, typename MetadataType>
-		void BuildIOFixedArray(IDetailLayoutBuilder& InDetailLayout, FName CategoryName, FName InPropertyName, const TSet<FString>& RequiredValues, bool bIsInput)
+		template <typename T>
+		void BuildIOFixedArray(IDetailLayoutBuilder& InDetailLayout, FName InCategoryName, FName InPropertyName, const TSet<FString>& InRequiredValues)
 		{
-			IDetailCategoryBuilder& CategoryBuilder = InDetailLayout.EditCategory(CategoryName);
+			const bool bIsInput = InCategoryName == "Inputs";
+
+			IDetailCategoryBuilder& CategoryBuilder = InDetailLayout.EditCategory(InCategoryName);
 			TSharedPtr<IPropertyHandle> ParentProperty = InDetailLayout.GetProperty(InPropertyName);
 			TSharedPtr<IPropertyHandleArray> ArrayHandle = ParentProperty->AsArray();
 
@@ -51,9 +56,9 @@ namespace Metasound
 				TSharedRef<IPropertyHandle> ArrayItemHandle = ArrayHandle->GetElement(i);
 
 				const FName TypeNamePropertyName = GET_MEMBER_NAME_CHECKED(T, TypeName);
-				const FName NamePropertyName = GET_MEMBER_NAME_CHECKED(T, TypeName);
-				const FName ToolTipPropertyName = GET_MEMBER_NAME_CHECKED(MetadataType, Description);
-				const FName DisplayNamePropertyName = GET_MEMBER_NAME_CHECKED(MetadataType, DisplayName);
+				const FName NamePropertyName = GET_MEMBER_NAME_CHECKED(T, Name);
+				const FName ToolTipPropertyName = GET_MEMBER_NAME_CHECKED(FMetasoundFrontendVertexMetadata, Description);
+				const FName DisplayNamePropertyName = GET_MEMBER_NAME_CHECKED(FMetasoundFrontendVertexMetadata, DisplayName);
 
 				TSharedPtr<IPropertyHandle> TypeProperty = ArrayItemHandle->GetChildHandle(TypeNamePropertyName);
 				TSharedPtr<IPropertyHandle> NameProperty = ArrayItemHandle->GetChildHandle(NamePropertyName);
@@ -63,9 +68,53 @@ namespace Metasound
 
 				FString Name;
 				const bool bNameFound = NameProperty->GetValue(Name) == FPropertyAccess::Success;
-				const bool bIsRequired = bNameFound && RequiredValues.Contains(Name);
+				const bool bIsRequired = bNameFound && InRequiredValues.Contains(Name);
+
+				FText DisplayName;
+				DisplayNameProperty->GetValue(DisplayName);
+				DisplayNameProperty->SetInstanceMetaData("LastValidName", DisplayName.ToString());
+
+				FSimpleDelegate NameChangeDelegate = FSimpleDelegate::CreateLambda([ChangedIndex = i, DisplayNameProperty, ArrayHandle, DisplayNamePropertyName]
+				{
+					FText MatchDisplayName;
+					if (DisplayNameProperty->GetValue(MatchDisplayName) != FPropertyAccess::Success)
+					{
+						return;
+					}
+
+					if (const FString* LastValidName = DisplayNameProperty->GetInstanceMetaData("LastValidName"))
+					{
+						uint32 NumElements = 0;
+						ArrayHandle->GetNumElements(NumElements);
+						for (int32 Elem = 0; Elem < static_cast<int32>(NumElements); ++Elem)
+						{
+							TSharedRef<IPropertyHandle> ElemHandle = ArrayHandle->GetElement(Elem);
+							TSharedPtr<IPropertyHandle> ElemNameProperty = ElemHandle->GetChildHandle(DisplayNamePropertyName, true /* bRecurse */);
+							FText ElemDisplayName;
+							if (ElemNameProperty->GetValue(ElemDisplayName) == FPropertyAccess::Success)
+							{
+								if (Elem != ChangedIndex && !ElemDisplayName.CompareTo(MatchDisplayName))
+								{
+									DisplayNameProperty->SetValue(FText::FromString(*LastValidName));
+									FNotificationInfo Info(LOCTEXT("MetasoundEditor_InvalidRename", "Rename failed: Input/output with name already exists."));
+									Info.bFireAndForget = true;
+									Info.ExpireDuration = 2.0f;
+									Info.bUseThrobber = true;
+									FSlateNotificationManager::Get().AddNotification(Info);
+									return;
+								}
+							}
+						}
+					}
+
+					DisplayNameProperty->SetInstanceMetaData("LastValidName", MatchDisplayName.ToString());
+				});
+
+				DisplayNameProperty->SetOnPropertyValueChanged(NameChangeDelegate);
 
-				CategoryBuilder.AddCustomRow(ParentProperty->GetPropertyDisplayName()).NameContent()
+				CategoryBuilder.AddCustomRow(ParentProperty->GetPropertyDisplayName())
+				.EditCondition(!bIsRequired, nullptr)
+				.NameContent()
 				[
 					SNew(STextBlock)
 					.Font(IDetailLayoutBuilder::GetDetailFontBold())
@@ -83,11 +132,11 @@ namespace Metasound
 
 						if (bIsRequired)
 						{
-							return FText::Format(LOCTEXT("Metasound_FixedIOArrayRequiredEntry_Format", "{0}. {1} ({2}, Required)"), FText::AsNumber(i + 1), DisplayName, FText::FromString(TypeNameString));
+							return FText::Format(LOCTEXT("MetasoundEditor_FixedIOArrayRequiredEntry_Format", "{0}. {1} ({2}, Required)"), FText::AsNumber(i + 1), DisplayName, FText::FromString(TypeNameString));
 						}
 						else
 						{
-							return FText::Format(LOCTEXT("Metasound_FixedIOArray_Format", "{0}. {1} ({2})"), FText::AsNumber(i + 1), DisplayName, FText::FromString(TypeNameString));
+							return FText::Format(LOCTEXT("MetasoundEditor_FixedIOArray_Format", "{0}. {1} ({2})"), FText::AsNumber(i + 1), DisplayName, FText::FromString(TypeNameString));
 						}
 					}))
 					.ToolTipText(TAttribute<FText>::Create([ToolTipProperty]()
@@ -202,8 +251,8 @@ namespace Metasound
 			const FName InputsPropertyPath = BuildChildPath(InterfacePropertyPath, GET_MEMBER_NAME_CHECKED(FMetasoundFrontendClassInterface, Inputs));
 			const FName OutputsPropertyPath = BuildChildPath(InterfacePropertyPath, GET_MEMBER_NAME_CHECKED(FMetasoundFrontendClassInterface, Outputs));
 
-			BuildIOFixedArray<FMetasoundFrontendClassInput, FMetasoundFrontendVertexMetadata>(DetailLayout, "Inputs", InputsPropertyPath, RequiredInputs, true /* bIsInput */);
-			BuildIOFixedArray<FMetasoundFrontendClassOutput, FMetasoundFrontendVertexMetadata>(DetailLayout, "Outputs", OutputsPropertyPath, RequiredOutputs, false /* bIsInput */);
+			BuildIOFixedArray<FMetasoundFrontendClassInput>(DetailLayout, "Inputs", InputsPropertyPath, RequiredInputs);
+			BuildIOFixedArray<FMetasoundFrontendClassOutput>(DetailLayout, "Outputs", OutputsPropertyPath, RequiredOutputs);
 
 			// Hack to hide parent structs for nested metadata properties
 			DetailLayout.HideCategory("CustomView");
@@ -213,6 +262,7 @@ namespace Metasound
 			DetailLayout.HideCategory("Attenuation");
 			DetailLayout.HideCategory("Curves");
 			DetailLayout.HideCategory("Debug");
+			DetailLayout.HideCategory("Developer");
 			DetailLayout.HideCategory("Effects");
 			DetailLayout.HideCategory("File Path");
 			DetailLayout.HideCategory("Format");
diff --git a/Engine/Plugins/Runtime/Metasound/Source/MetasoundEditor/Private/MetasoundDetailCustomization.h b/Engine/Plugins/Runtime/Metasound/Source/MetasoundEditor/Private/MetasoundDetailCustomization.h
--- a/Engine/Plugins/Runtime/Metasound/Source/MetasoundEditor/Private/MetasoundDetailCustomization.h
+++ b/Engine/Plugins/Runtime/Metasound/Source/MetasoundEditor/Private/MetasoundDetailCustomization.h
@@ -5,6 +5,7 @@
 #include "UObject/NameTypes.h"
 
 // Forward Declarations
+class FPropertyRestriction;
 class IDetailLayoutBuilder;
 
 namespace Metasound
diff --git a/Engine/Plugins/Runtime/Metasound/Source/MetasoundEditor/Private/MetasoundEditorGraphBuilder.cpp b/Engine/Plugins/Runtime/Metasound/Source/MetasoundEditor/Private/MetasoundEditorGraphBuilder.cpp
--- a/Engine/Plugins/Runtime/Metasound/Source/MetasoundEditor/Private/MetasoundEditorGraphBuilder.cpp
+++ b/Engine/Plugins/Runtime/Metasound/Source/MetasoundEditor/Private/MetasoundEditorGraphBuilder.cpp
@@ -25,12 +25,16 @@ namespace Metasound
 {
 	namespace Editor
 	{
-		const FName FGraphBuilder::PinPrimitiveBoolean = TEXT("Boolean");
-		const FName FGraphBuilder::PinPrimitiveFloat = TEXT("Float");
-		const FName FGraphBuilder::PinPrimitiveInteger = TEXT("Int");
-		const FName FGraphBuilder::PinPrimitiveString = TEXT("String");
-		const FName FGraphBuilder::PinPrimitiveUObject = TEXT("UObject");
-		const FName FGraphBuilder::PinPrimitiveUObjectArray = TEXT("UObjectArray");
+		const FName FGraphBuilder::PinAudioFormat = "Format";
+		const FName FGraphBuilder::PinAudioNumeric = "Numeric";
+		const FName FGraphBuilder::PinPrimitiveBoolean = "Boolean";
+		const FName FGraphBuilder::PinPrimitiveFloat = "Float";
+		const FName FGraphBuilder::PinPrimitiveInt32 = "Int32";
+		const FName FGraphBuilder::PinPrimitiveInt64 = "Int64";
+		const FName FGraphBuilder::PinPrimitiveString = "String";
+		const FName FGraphBuilder::PinPrimitiveTrigger = "Trigger";
+		const FName FGraphBuilder::PinPrimitiveUObject = "UObject";
+		const FName FGraphBuilder::PinPrimitiveUObjectArray = "UObjectArray";
 
 		UEdGraphNode* FGraphBuilder::AddNode(UObject& InMetasound, const FVector2D& Location, Frontend::FNodeHandle& InNodeHandle, bool bInSelectNewNode)
 		{
diff --git a/Engine/Plugins/Runtime/Metasound/Source/MetasoundEditor/Private/MetasoundEditorGraphBuilder.h b/Engine/Plugins/Runtime/Metasound/Source/MetasoundEditor/Private/MetasoundEditorGraphBuilder.h
--- a/Engine/Plugins/Runtime/Metasound/Source/MetasoundEditor/Private/MetasoundEditorGraphBuilder.h
+++ b/Engine/Plugins/Runtime/Metasound/Source/MetasoundEditor/Private/MetasoundEditorGraphBuilder.h
@@ -22,10 +22,15 @@ namespace Metasound
 		class FGraphBuilder
 		{
 		public:
+			// Custom pin-related styles (ex. wire color, pin heads, etc.)
+			static const FName PinAudioFormat;				// Audio formats (ex. Buffers, Mono, Stereo)
+			static const FName PinAudioNumeric;				// Audio numeric types (ex. Frequency, Time)
 			static const FName PinPrimitiveBoolean;
 			static const FName PinPrimitiveFloat;
-			static const FName PinPrimitiveInteger;
+			static const FName PinPrimitiveInt32;
+			static const FName PinPrimitiveInt64;
 			static const FName PinPrimitiveString;
+			static const FName PinPrimitiveTrigger;
 			static const FName PinPrimitiveUObject;
 			static const FName PinPrimitiveUObjectArray;
 
diff --git a/Engine/Plugins/Runtime/Metasound/Source/MetasoundEditor/Private/MetasoundEditorGraphNode.cpp b/Engine/Plugins/Runtime/Metasound/Source/MetasoundEditor/Private/MetasoundEditorGraphNode.cpp
--- a/Engine/Plugins/Runtime/Metasound/Source/MetasoundEditor/Private/MetasoundEditorGraphNode.cpp
+++ b/Engine/Plugins/Runtime/Metasound/Source/MetasoundEditor/Private/MetasoundEditorGraphNode.cpp
@@ -50,13 +50,19 @@ void UMetasoundEditorGraphNode::PostLoad()
 
 void UMetasoundEditorGraphNode::CreateInputPin()
 {
+	// TODO: Implement for nodes supporting variadic inputs
+	if (ensure(false))
+	{
+		return;
+	}
+
 	FString PinName; // get from UMetasound
 	UEdGraphPin* NewPin = CreatePin(EGPD_Input, TEXT("MetasoundEditorGraphNode"), *PinName);
 	if (NewPin->PinName.IsNone())
 	{
-		// Makes sure pin has a name for lookup purposes but user will never see it
-		NewPin->PinName = CreateUniquePinName(TEXT("Input"));
-		NewPin->PinFriendlyName = FText::FromString(TEXT(" "));
+		// Pin must have a name for lookup purposes but is not user-facing
+// 		NewPin->PinName = 
+// 		NewPin->PinFriendlyName =
 	}
 }
 
@@ -86,15 +92,16 @@ const UObject& UMetasoundEditorGraphNode::GetMetasoundChecked() const
 	return EdGraph->GetMetasoundChecked();
 }
 
-Metasound::Frontend::FGraphHandle UMetasoundEditorGraphNode::GetRootGraphHandle()
+Metasound::Frontend::FGraphHandle UMetasoundEditorGraphNode::GetRootGraphHandle() const
 {
-	FMetasoundAssetBase* MetasoundAsset = Metasound::IMetasoundUObjectRegistry::Get().GetObjectAsAssetBase(&GetMetasoundChecked());
+	const FMetasoundAssetBase* ConstMetasoundAsset = Metasound::IMetasoundUObjectRegistry::Get().GetObjectAsAssetBase(&GetMetasoundChecked());
+	FMetasoundAssetBase* MetasoundAsset = const_cast<FMetasoundAssetBase*>(ConstMetasoundAsset);
 	check(MetasoundAsset);
 
 	return MetasoundAsset->GetRootGraphHandle();
 }
 
-Metasound::Frontend::FNodeHandle UMetasoundEditorGraphNode::GetNodeHandle()
+Metasound::Frontend::FNodeHandle UMetasoundEditorGraphNode::GetNodeHandle() const
 {
 	return GetRootGraphHandle()->GetNodeWithID(NodeID);
 }
@@ -179,8 +186,7 @@ bool UMetasoundEditorGraphNode::CanUserDeleteNode() const
 	Metasound::Frontend::FDocumentHandle DocumentHandle = MetasoundAsset->GetDocumentHandle();
 	Metasound::Frontend::FNodeHandle NodeHandle = DocumentHandle->GetRootGraph()->GetNodeWithID(NodeID);
 
-	const FString& NodeName = NodeHandle->GetNodeName();
-
+	const FString& NodeName = GetNodeHandle()->GetNodeName();
 	if (DocumentHandle->IsRequiredInput(NodeName))
 	{
 		return false;
@@ -213,15 +219,7 @@ FText UMetasoundEditorGraphNode::GetNodeTitle(ENodeTitleType::Type TitleType) co
 {
 	using namespace Metasound::Frontend;
 
-	const FMetasoundAssetBase* ConstMetasoundAsset = Metasound::IMetasoundUObjectRegistry::Get().GetObjectAsAssetBase(&GetMetasoundChecked());
-
-	// TODO: create FConstGraphHandle so that this const cast is unnecessary.
-	FMetasoundAssetBase* MetasoundAsset = const_cast<FMetasoundAssetBase*>(ConstMetasoundAsset);
-	check(MetasoundAsset);
-
-	FGraphHandle GraphHandle = MetasoundAsset->GetRootGraphHandle();
-	FConstNodeHandle NodeHandle = GraphHandle->GetNodeWithID(NodeID);
-
+	FConstNodeHandle NodeHandle = GetNodeHandle();
 	switch (NodeHandle->GetClassType())
 	{
 		case EMetasoundFrontendClassType::Input:
@@ -325,17 +323,7 @@ void UMetasoundEditorGraphNode::GetNodeContextMenuActions(UToolMenu* Menu, UGrap
 
 FText UMetasoundEditorGraphNode::GetTooltipText() const
 {
-	using namespace Metasound::Frontend;
-
-	using namespace Metasound::Frontend;
-
-	const FMetasoundAssetBase* ConstMetasoundAsset = Metasound::IMetasoundUObjectRegistry::Get().GetObjectAsAssetBase(&GetMetasoundChecked());
-
-	// TODO: create FConstGraphHandle so that this const cast is unnecessary.
-	FMetasoundAssetBase* MetasoundAsset = const_cast<FMetasoundAssetBase*>(ConstMetasoundAsset);
-	check(MetasoundAsset);
-	
-	return MetasoundAsset->GetRootGraphHandle()->GetNodeWithID(NodeID)->GetClassDescription();
+	return GetRootGraphHandle()->GetNodeWithID(NodeID)->GetClassDescription();
 }
 
 FString UMetasoundEditorGraphNode::GetDocumentationExcerptName() const
diff --git a/Engine/Plugins/Runtime/Metasound/Source/MetasoundEditor/Private/MetasoundEditorGraphNodeFactory.h b/Engine/Plugins/Runtime/Metasound/Source/MetasoundEditor/Private/MetasoundEditorGraphNodeFactory.h
--- a/Engine/Plugins/Runtime/Metasound/Source/MetasoundEditor/Private/MetasoundEditorGraphNodeFactory.h
+++ b/Engine/Plugins/Runtime/Metasound/Source/MetasoundEditor/Private/MetasoundEditorGraphNodeFactory.h
@@ -1,13 +1,11 @@
 // Copyright Epic Games, Inc. All Rights Reserved.
-
 #pragma once
 
 #include "CoreMinimal.h"
 #include "EdGraphUtilities.h"
-#include "GraphNode.h"
+#include "MetasoundEditorGraphNode.h"
 #include "SGraphNode.h"
-#include "SGraphNodeSoundResult.h"
-#include "SGraphNodeSoundBase.h"
+#include "SMetasoundGraphNode.h"
 #include "Widgets/DeclarativeSyntaxSupport.h"
 
 
@@ -15,9 +13,9 @@ class FMetasoundGraphNodeFactory : public FGraphPanelNodeFactory
 {
 	virtual TSharedPtr<SGraphNode> CreateNode(UEdGraphNode* InNode) const override
 	{
-		if (UMetasoundEditorGraphNode* BaseSoundNode = Cast<UMetasoundEditorGraphNode>(InNode))
+		if (InNode->IsA<UMetasoundEditorGraphNode>())
 		{
-			return SNew(SGraphNodeSoundResult, RootSoundNode);
+			return SNew(SMetasoundGraphNode, InNode);
 		}
 
 		return nullptr;
diff --git a/Engine/Plugins/Runtime/Metasound/Source/MetasoundEditor/Private/MetasoundEditorGraphSchema.cpp b/Engine/Plugins/Runtime/Metasound/Source/MetasoundEditor/Private/MetasoundEditorGraphSchema.cpp
--- a/Engine/Plugins/Runtime/Metasound/Source/MetasoundEditor/Private/MetasoundEditorGraphSchema.cpp
+++ b/Engine/Plugins/Runtime/Metasound/Source/MetasoundEditor/Private/MetasoundEditorGraphSchema.cpp
@@ -50,6 +50,11 @@ namespace Metasound
 			const UGraphEditorSettings* Settings = GetDefault<UGraphEditorSettings>();
 			check(Settings);
 
+			if (PinType.PinCategory == FGraphBuilder::PinPrimitiveTrigger)
+			{
+				return Settings->ExecutionPinTypeColor;
+			}
+
 			if (PinType.PinCategory == FGraphBuilder::PinPrimitiveBoolean)
 			{
 				return Settings->BooleanPinTypeColor;
@@ -60,16 +65,36 @@ namespace Metasound
 				return Settings->FloatPinTypeColor;
 			}
 
-			if (PinType.PinCategory == FGraphBuilder::PinPrimitiveInteger)
+			if (PinType.PinCategory == FGraphBuilder::PinPrimitiveInt32)
 			{
 				return Settings->IntPinTypeColor;
 			}
 
+			if (PinType.PinCategory == FGraphBuilder::PinPrimitiveInt64)
+			{
+				return Settings->Int64PinTypeColor;
+			}
+
 			if (PinType.PinCategory == FGraphBuilder::PinPrimitiveString)
 			{
 				return Settings->StringPinTypeColor;
 			}
 
+			if (PinType.PinCategory == FGraphBuilder::PinPrimitiveUObject)
+			{
+				return Settings->ClassPinTypeColor;
+			}
+
+			if (PinType.PinCategory == FGraphBuilder::PinAudioNumeric)
+			{
+				return Settings->SoftObjectPinTypeColor;
+			}
+
+			if (PinType.PinCategory == FGraphBuilder::PinAudioFormat)
+			{
+				return Settings->SoftClassPinTypeColor;
+			}
+
 			return Settings->StructPinTypeColor;
 		}
 
@@ -131,7 +156,18 @@ namespace Metasound
 
 			if (!bExecuted)
 			{
-				OutParams.WireThickness = InactiveWireThickness;
+				if (InputPin->PinType.PinCategory == FGraphBuilder::PinPrimitiveTrigger)
+				{
+					OutParams.WireThickness = Settings->DefaultExecutionWireThickness;
+				}
+				else if (InputPin->PinType.PinCategory == FGraphBuilder::PinAudioFormat)
+				{
+					OutParams.bDrawBubbles = true;
+				}
+				else
+				{
+					OutParams.WireThickness = InactiveWireThickness;
+				}
 			}
 		}
 	} // namespace Editor
@@ -571,16 +607,17 @@ void UMetasoundEditorGraphSchema::GetAllMetasoundActions(FGraphActionMenuBuilder
 	using namespace Metasound::Editor;
 	using namespace Metasound::Frontend;
 
-	const FText InputMenuName = LOCTEXT("MetasoundAddInputMenu", "Add Input");
-	const FText OutputMenuName = LOCTEXT("MetasoundAddOutputMenu", "Add Output");
+	const FText InputMenuName = LOCTEXT("MetasoundActionsInputsMenu", "Inputs");
+	const FText OutputMenuName = LOCTEXT("MetasoundActionsOutputsMenu", "Outputs");
+
+	const FText MenuJoinFormat = LOCTEXT("MetasoundActionsFormatSubCategory", "{0}|{1}");
 
 	IMetasoundEditorModule& EditorModule = FModuleManager::GetModuleChecked<IMetasoundEditorModule>("MetasoundEditor");
-	EditorModule.IterateDataTypes([InputMenuName, OutputMenuName, InMenuBuilder = &ActionMenuBuilder](const FEditorDataType& DataType)
+	EditorModule.IterateDataTypes([MenuJoinFormat, InputMenuName, OutputMenuName, InMenuBuilder = &ActionMenuBuilder](const FEditorDataType& DataType)
 	{
 		const FName DataTypeName = DataType.RegistryInfo.DataTypeName;
-		const FText DataTypeDisplayName = FText::FromString(FGraphBuilder::GetDataTypeDisplayName(DataTypeName));
+		const FText DataTypeDisplayName = FText::FromString(FGraphBuilder::GetDataTypeDisplayName(DataTypeName).Replace(TEXT("Primitive:"), TEXT("")));
 		const FText DataTypeTextName = FText::FromName(DataTypeName);
-		const FText MenuJoinFormat = LOCTEXT("MetasoundFormatNodeSubCategory", "{0}|{1}");
 
 		const TArray<FString> Categories = FGraphBuilder::GetDataTypeNameCategories(DataTypeName);
 		const FText CategoriesText = FText::FromString(FString::Join(Categories, TEXT("|")));
@@ -607,7 +644,9 @@ void UMetasoundEditorGraphSchema::GetAllMetasoundActions(FGraphActionMenuBuilder
 		InMenuBuilder->AddAction(AddOutputNodeAction);
 	});
 
-	const FText NodeMenuName = LOCTEXT("MetasoundNodesMenu", "Add Node");
+	const FText NodeMenuName = LOCTEXT("MetasoundAddFunctionsMenu", "Functions");
+	const FText ConvertMenuName = LOCTEXT("MetasoundConverionsMenu", "Conversions");
+
 	const TArray<FNodeClassInfo> ClassInfos = GetAllAvailableNodeClasses();
 	for (const FNodeClassInfo& ClassInfo : ClassInfos)
 	{
@@ -617,10 +656,23 @@ void UMetasoundEditorGraphSchema::GetAllMetasoundActions(FGraphActionMenuBuilder
 			? Metadata.Description
 			: FText::Format(LOCTEXT("MetasoundTooltipAuthorFormat", "{0}\nAuthor: {1}"), Metadata.Description, Metadata.Author);
 
+		const FString DisplayName = ClassInfo.NodeName.Replace(TEXT("Primitive:"), TEXT(""));
+
+		const FText* MenuName = &NodeMenuName;
+		FText CategoriesText = FText::GetEmpty();
+		if (!Metadata.CategoryHierarchy.IsEmpty() && !Metadata.CategoryHierarchy[0].CompareTo(ConvertMenuName))
+		{
+			MenuName = &ConvertMenuName;
+		}
+		else
+		{
+			CategoriesText = FText::Join(LOCTEXT("MetasoundActionsCategoryDelim", "|"), Metadata.CategoryHierarchy);
+		}
+
 		TSharedPtr<FMetasoundGraphSchemaAction_NewNode> NewNodeAction = MakeShared<FMetasoundGraphSchemaAction_NewNode>
 		(
-			NodeMenuName,
-			FText::FromString(ClassInfo.NodeName),
+			FText::Format(MenuJoinFormat, *MenuName, CategoriesText),
+			FText::FromString(DisplayName),
 			Tooltip,
 			0
 		);
diff --git a/Engine/Plugins/Runtime/Metasound/Source/MetasoundEditor/Private/MetasoundEditorModule.cpp b/Engine/Plugins/Runtime/Metasound/Source/MetasoundEditor/Private/MetasoundEditorModule.cpp
--- a/Engine/Plugins/Runtime/Metasound/Source/MetasoundEditor/Private/MetasoundEditorModule.cpp
+++ b/Engine/Plugins/Runtime/Metasound/Source/MetasoundEditor/Private/MetasoundEditorModule.cpp
@@ -16,6 +16,7 @@
 #include "MetasoundDetailCustomization.h"
 #include "MetasoundLiteralDescriptionDetailCustomization.h"
 #include "MetasoundEditorGraphBuilder.h"
+#include "MetasoundEditorGraphNodeFactory.h"
 #include "MetasoundEditorGraphSchema.h"
 #include "Modules/ModuleInterface.h"
 #include "Modules/ModuleManager.h"
@@ -54,7 +55,7 @@ namespace Metasound
 			{
 				SetParentStyleName(FEditorStyle::GetStyleSetName());
 
-				SetContentRoot(FPaths::EngineContentDir() / TEXT("Editor/Slate"));
+				SetContentRoot(FPaths::EnginePluginsDir() / TEXT("Runtime/Metasound/Content/Editor/Slate"));
 				SetCoreContentRoot(FPaths::EngineContentDir() / TEXT("Slate"));
 
 				static const FVector2D Icon20x20(20.0f, 20.0f);
@@ -62,22 +63,39 @@ namespace Metasound
 
 				// Metasound Editor
 				{
-					Set("MetasoundEditor.Play", new FSlateImageBrush(RootToContentDir(TEXT("Icons/icon_SCueEd_PlayCue_40x.png")), Icon40x40));
-					Set("MetasoundEditor.Play.Small", new FSlateImageBrush(RootToContentDir(TEXT("Icons/icon_SCueEd_PlayCue_40x.png")), Icon20x20));
-					Set("MetasoundEditor.Stop", new FSlateImageBrush(RootToContentDir(TEXT("Icons/icon_SCueEd_Stop_40x.png")), Icon40x40));
-					Set("MetasoundEditor.Stop.Small", new FSlateImageBrush(RootToContentDir(TEXT("Icons/icon_SCueEd_Stop_40x.png")), Icon20x20));
-					Set("MetasoundEditor.Import", new FSlateImageBrush(RootToContentDir(TEXT("/Old/Kismet2/compile_40px.png")), Icon40x40));
-					Set("MetasoundEditor.Import.Small", new FSlateImageBrush(RootToContentDir(TEXT("/Old/Kismet2/compile_40px.png")), Icon20x20));
-					Set("MetasoundEditor.Export", new FSlateImageBrush(RootToContentDir(TEXT("/Old/Kismet2/compile_40px.png")), Icon40x40));
-					Set("MetasoundEditor.Export.Small", new FSlateImageBrush(RootToContentDir(TEXT("/Old/Kismet2/compile_40px.png")), Icon20x20));
-					Set("MetasoundEditor.ExportError", new FSlateImageBrush(RootToContentDir(TEXT("/Old/Kismet2/CompileStatus_Fail.png")), Icon40x40));
-					Set("MetasoundEditor.ExportError.Small", new FSlateImageBrush(RootToContentDir(TEXT("/Old/Kismet2/CompileStatus_Broken_Small.png")), Icon20x20));
+					// Actions
+					Set("MetasoundEditor.Play", new FSlateImageBrush(RootToContentDir(TEXT("Icons/play_40x.png")), Icon40x40));
+					Set("MetasoundEditor.Play.Small", new FSlateImageBrush(RootToContentDir(TEXT("Icons/play_40x.png")), Icon20x20));
+					Set("MetasoundEditor.Stop", new FSlateImageBrush(RootToContentDir(TEXT("Icons/stop_40x.png")), Icon40x40));
+					Set("MetasoundEditor.Stop.Small", new FSlateImageBrush(RootToContentDir(TEXT("Icons/stop_40x.png")), Icon20x20));
+					Set("MetasoundEditor.Import", new FSlateImageBrush(RootToContentDir(TEXT("/Icons/build_40x.png")), Icon40x40));
+					Set("MetasoundEditor.Import.Small", new FSlateImageBrush(RootToContentDir(TEXT("/Icons/build_40x.png")), Icon20x20));
+					Set("MetasoundEditor.Export", new FSlateImageBrush(RootToContentDir(TEXT("/Icons/build_40x.png")), Icon40x40));
+					Set("MetasoundEditor.Export.Small", new FSlateImageBrush(RootToContentDir(TEXT("/Icons/build_40x.png")), Icon20x20));
+					Set("MetasoundEditor.ExportError", new FSlateImageBrush(RootToContentDir(TEXT("/Icons/build_error_40x.png")), Icon40x40));
+					Set("MetasoundEditor.ExportError.Small", new FSlateImageBrush(RootToContentDir(TEXT("/Icons/build_error_40x.png")), Icon20x20));
+
+					// Graph Editor
+					Set("MetasoundEditor.Graph.Node.Body.Input", new FSlateImageBrush(RootToContentDir(TEXT("/Graph/node_input_body_64x.png")), FVector2D(114.0f, 64.0f)));
+					Set("MetasoundEditor.Graph.Node.Body.Default", new FSlateImageBrush(RootToContentDir(TEXT("/Graph/node_default_body_64x.png")), FVector2D(64.0f, 64.0f)));
+					Set("MetasoundEditor.Graph.Node.Math.Add", new FSlateImageBrush(RootToContentDir(TEXT("/Graph/node_math_add_40x.png")), Icon40x40));
+					Set("MetasoundEditor.Graph.Node.Math.Divide", new FSlateImageBrush(RootToContentDir(TEXT("/Graph/node_math_divide_40x.png")), Icon40x40));
+					Set("MetasoundEditor.Graph.Node.Math.Multiply", new FSlateImageBrush(RootToContentDir(TEXT("/Graph/node_math_multiply_40x.png")), FVector2D(32.0f, 32.0f)));
+					Set("MetasoundEditor.Graph.Node.Math.RandRange", new FSlateImageBrush(RootToContentDir(TEXT("/Graph/node_math_random_40x.png")), FVector2D(26.0f, 40.0f)));
+					Set("MetasoundEditor.Graph.Node.Math.Subtract", new FSlateImageBrush(RootToContentDir(TEXT("/Graph/node_math_subtract_40x.png")), Icon40x40));
+
+					// Misc
+					Set("MetasoundEditor.Speaker", new FSlateImageBrush(RootToContentDir(TEXT("/Icons/speaker_144x.png")), FVector2D(144.0f, 144.0f)));
 				}
 
 				FSlateStyleRegistry::RegisterSlateStyle(*this);
 			}
 		};
 
+		class FMetasoundGraphPanelPinFactory : public FGraphPanelPinFactory
+		{
+		};
+
 		class FModule : public IMetasoundEditorModule
 		{
 			virtual const FEditorDataType& FindDataType(FName InDataTypeName) const override
@@ -106,53 +124,94 @@ namespace Metasound
 					FDataTypeRegistryInfo RegistryInfo;
 					Frontend::GetTraitsForDataType(DataTypeName, RegistryInfo);
 
-					FName PinType = NAME_None;
-					switch (RegistryInfo.PreferredLiteralType)
-					{
-						case ELiteralType::Boolean:
-						{
-							PinType = FGraphBuilder::PinPrimitiveBoolean;
-						}
-						break;
+					FName PinType = DataTypeName;
 
-						case ELiteralType::Float:
-						{
-							PinType = FGraphBuilder::PinPrimitiveFloat;
-						}
-						break;
+					// Special Flowers
 
-						case ELiteralType::Integer:
-						{
-							PinType = FGraphBuilder::PinPrimitiveInteger;
-						}
-						break;
+					// Execution path triggers are specialized
+					if (DataTypeName == "Primitive:Trigger")
+					{
+						PinType = FGraphBuilder::PinPrimitiveTrigger;
+					}
 
-						case ELiteralType::String:
-						{
-							PinType = FGraphBuilder::PinPrimitiveString;
-						}
-						break;
+					// GraphEditor by default designates specialized connection
+					// specification for Int64, so use it
+					else if (DataTypeName == "Primitive:Int64")
+					{
+						PinType = FGraphBuilder::PinPrimitiveInt64;
+					}
 
-						case ELiteralType::UObjectProxy:
-						{
-							PinType = FGraphBuilder::PinPrimitiveUObject;
-						}
-						break;
+					// Differentiate stronger numeric types associated with audio
+					else if (DataTypeName == "Primitive:Frequency"
+						|| DataTypeName == "Primitive:Time"
+						|| DataTypeName == "Primitive:Time:HighResolution"
+						|| DataTypeName == "Primitive:Time:SampleResolution"
+						)
+					{
+						PinType = FGraphBuilder::PinAudioNumeric;
+					}
 
-						case ELiteralType::UObjectProxyArray:
-						{
-							PinType = FGraphBuilder::PinPrimitiveUObjectArray;
-						}
-						break;
+					// Audio types are ubiquitous, so specialize
+					else if (DataTypeName == "Audio:Buffer"
+						|| DataTypeName == "Audio:Unformatted"
+						|| DataTypeName == "Audio:Mono"
+						|| DataTypeName == "Audio:Stereo"
+						|| DataTypeName == "Audio:Multichannel"
+						)
+					{
+						PinType = FGraphBuilder::PinAudioFormat;
+					}
 
-						// Register atypical primitives
-						default:
-						case ELiteralType::None:
-						case ELiteralType::Invalid:
+					// Primitives
+					else
+					{
+						switch (RegistryInfo.PreferredLiteralType)
 						{
-							static_assert(static_cast<int32>(ELiteralType::Invalid) == 7, "Possible missing binding of pin category to primitive type");
+							case ELiteralType::Boolean:
+							{
+								PinType = FGraphBuilder::PinPrimitiveBoolean;
+							}
+							break;
+
+							case ELiteralType::Float:
+							{
+								PinType = FGraphBuilder::PinPrimitiveFloat;
+							}
+							break;
+
+							case ELiteralType::Integer:
+							{
+								PinType = FGraphBuilder::PinPrimitiveInt32;
+							}
+							break;
+
+							case ELiteralType::String:
+							{
+								PinType = FGraphBuilder::PinPrimitiveString;
+							}
+							break;
+
+							case ELiteralType::UObjectProxy:
+							{
+								PinType = FGraphBuilder::PinPrimitiveUObject;
+							}
+							break;
+
+							case ELiteralType::UObjectProxyArray:
+							{
+								PinType = FGraphBuilder::PinPrimitiveUObjectArray;
+							}
+							break;
+
+							// Register atypical primitives
+							default:
+							case ELiteralType::None:
+							case ELiteralType::Invalid:
+							{
+								static_assert(static_cast<int32>(ELiteralType::Invalid) == 7, "Possible missing binding of pin category to primitive type");
+							}
+							break;
 						}
-						break;
 					}
 
 					DataTypeInfo.Add(DataTypeName, FEditorDataType(PinType, RegistryInfo));
@@ -186,6 +245,12 @@ namespace Metasound
 
 				GraphConnectionFactory = MakeShared<FGraphConnectionDrawingPolicyFactory>();
 				FEdGraphUtilities::RegisterVisualPinConnectionFactory(GraphConnectionFactory);
+
+				GraphNodeFactory = MakeShared<FMetasoundGraphNodeFactory>();
+				FEdGraphUtilities::RegisterVisualNodeFactory(GraphNodeFactory);
+
+				GraphPanelPinFactory = MakeShared<FMetasoundGraphPanelPinFactory>();
+				FEdGraphUtilities::RegisterVisualPinFactory(GraphPanelPinFactory);
 			}
 
 			virtual void ShutdownModule() override
@@ -204,6 +269,18 @@ namespace Metasound
 					FEdGraphUtilities::UnregisterVisualPinConnectionFactory(GraphConnectionFactory);
 				}
 
+				if (GraphNodeFactory.IsValid())
+				{
+					FEdGraphUtilities::UnregisterVisualNodeFactory(GraphNodeFactory);
+					GraphNodeFactory.Reset();
+				}
+
+				if (GraphPanelPinFactory.IsValid())
+				{
+					FEdGraphUtilities::UnregisterVisualPinFactory(GraphPanelPinFactory);
+					GraphPanelPinFactory.Reset();
+				}
+
 				AssetActions.Reset();
 				DataTypeInfo.Reset();
 			}
@@ -211,7 +288,9 @@ namespace Metasound
 			TArray<TSharedPtr<FAssetTypeActions_Base>> AssetActions;
 			TMap<FName, FEditorDataType> DataTypeInfo;
 
+			TSharedPtr<FMetasoundGraphNodeFactory> GraphNodeFactory;
 			TSharedPtr<FGraphPanelPinConnectionFactory> GraphConnectionFactory;
+			TSharedPtr<FMetasoundGraphPanelPinFactory> GraphPanelPinFactory;
 			TSharedPtr<FSlateStyleSet> StyleSet;
 		};
 	} // namespace Editor
diff --git a/Engine/Plugins/Runtime/Metasound/Source/MetasoundEditor/Private/SMetasoundGraphNode.cpp b/Engine/Plugins/Runtime/Metasound/Source/MetasoundEditor/Private/SMetasoundGraphNode.cpp
new file mode 100644
--- /dev/null
+++ b/Engine/Plugins/Runtime/Metasound/Source/MetasoundEditor/Private/SMetasoundGraphNode.cpp
@@ -0,0 +1,271 @@
+// Copyright Epic Games, Inc. All Rights Reserved.
+#include "SMetasoundGraphNode.h"
+
+#include "EditorStyleSet.h"
+#include "GraphEditorSettings.h"
+#include "IDocumentation.h"
+#include "MetasoundEditorGraphNode.h"
+#include "MetasoundEditorModule.h"
+#include "ScopedTransaction.h"
+#include "SCommentBubble.h"
+#include "SGraphNode.h"
+#include "SGraphPin.h"
+#include "SLevelOfDetailBranchNode.h"
+#include "Styling/SlateStyleRegistry.h"
+#include "Styling/SlateColor.h"
+#include "TutorialMetaData.h"
+#include "Widgets/Images/SImage.h"
+#include "Widgets/SBoxPanel.h"
+#include "Widgets/SOverlay.h"
+
+#define LOCTEXT_NAMESPACE "MetasoundGraphNode"
+
+void SMetasoundGraphNode::Construct(const FArguments& InArgs, class UEdGraphNode* InNode)
+{
+	GraphNode = InNode;
+	SetCursor(EMouseCursor::CardinalCross);
+	UpdateGraphNode();
+}
+
+void SMetasoundGraphNode::CreateOutputSideAddButton(TSharedPtr<SVerticalBox> OutputBox)
+{
+	TSharedRef<SWidget> AddPinButton = AddPinButtonContent(
+		LOCTEXT("MetasoundGraphNode_AddPinButton", "Add input"),
+		LOCTEXT("MetasoundGraphNode_AddPinButton_Tooltip", "Add an input to the parent Metasound node.")
+	);
+
+	FMargin AddPinPadding = Settings->GetOutputPinPadding();
+	AddPinPadding.Top += 6.0f;
+
+	OutputBox->AddSlot()
+	.AutoHeight()
+	.VAlign(VAlign_Center)
+	.Padding(AddPinPadding)
+	[
+		AddPinButton
+	];
+}
+
+UMetasoundEditorGraphNode& SMetasoundGraphNode::GetMetasoundNode()
+{
+	check(GraphNode);
+	return *Cast<UMetasoundEditorGraphNode>(GraphNode);
+}
+
+const UMetasoundEditorGraphNode& SMetasoundGraphNode::GetMetasoundNode() const
+{
+	check(GraphNode);
+	return *Cast<UMetasoundEditorGraphNode>(GraphNode);
+}
+
+void SMetasoundGraphNode::CreateStandardPinWidget(UEdGraphPin* CurPin)
+{
+	const bool bShowPin = ShouldPinBeHidden(CurPin);
+	if (bShowPin)
+	{
+		TSharedPtr<SGraphPin> NewPin = CreatePinWidget(CurPin);
+		check(NewPin.IsValid());
+
+		Metasound::Frontend::FNodeHandle NodeHandle = GetMetasoundNode().GetNodeHandle();
+		if (CurPin->Direction == EGPD_Input)
+		{
+			if (!NodeHandle->GetClassDisplayInfo().bShowInputName)
+			{
+				NewPin->SetShowLabel(false);
+			}
+		}
+		else if (CurPin->Direction == EGPD_Output)
+		{
+			if (!NodeHandle->GetClassDisplayInfo().bShowOutputName)
+			{
+				NewPin->SetShowLabel(false);
+			}
+		}
+
+		AddPin(NewPin.ToSharedRef());
+	}
+}
+
+TSharedRef<SWidget> SMetasoundGraphNode::CreateTitleWidget(TSharedPtr<SNodeTitle> NodeTitle)
+{
+	Metasound::Frontend::FNodeHandle NodeHandle = GetMetasoundNode().GetNodeHandle();
+	if (!NodeHandle->GetClassDisplayInfo().bShowName)
+	{
+		return SNullWidget::NullWidget;
+	}
+
+	return SGraphNode::CreateTitleWidget(NodeTitle);
+}
+
+void SMetasoundGraphNode::SetDefaultTitleAreaWidget(TSharedRef<SOverlay> DefaultTitleAreaWidget)
+{
+	SGraphNode::SetDefaultTitleAreaWidget(DefaultTitleAreaWidget);
+
+	Metasound::Frontend::FNodeHandle NodeHandle = GetMetasoundNode().GetNodeHandle();
+	if (NodeHandle->GetClassDisplayInfo().bShowName)
+	{
+		DefaultTitleAreaWidget->ClearChildren();
+		TSharedPtr<SNodeTitle> NodeTitle = SNew(SNodeTitle, GraphNode);
+
+		DefaultTitleAreaWidget->AddSlot()
+		.HAlign(HAlign_Fill)
+		.VAlign(VAlign_Center)
+		[
+			SNew(SHorizontalBox)
+			+ SHorizontalBox::Slot()
+			.HAlign(HAlign_Fill)
+			[
+				SNew(SBorder)
+				.BorderImage(FEditorStyle::GetBrush("NoBorder"))
+				[
+					SNew(SHorizontalBox)
+					+ SHorizontalBox::Slot()
+					[
+						SNew(SVerticalBox)
+						+ SVerticalBox::Slot()
+						.AutoHeight()
+						.HAlign(HAlign_Center)
+						[
+							CreateTitleWidget(NodeTitle)
+						]
+						+ SVerticalBox::Slot()
+						.AutoHeight()
+						[
+							NodeTitle.ToSharedRef()
+						]
+					]
+				]
+			]
+			+ SHorizontalBox::Slot()
+			.HAlign(HAlign_Right)
+			.VAlign(VAlign_Center)
+			.Padding(0, 0, 5, 0)
+			.AutoWidth()
+			[
+				CreateTitleRightWidget()
+			]
+		];
+
+		DefaultTitleAreaWidget->AddSlot()
+		.VAlign(VAlign_Top)
+		[
+			SNew(SBorder)
+			.Visibility(EVisibility::HitTestInvisible)
+			.BorderImage( FEditorStyle::GetBrush( "Graph.Node.TitleHighlight" ) )
+			.BorderBackgroundColor( this, &SGraphNode::GetNodeTitleIconColor )
+			[
+				SNew(SSpacer)
+				.Size(FVector2D(20,20))
+			]
+		];
+	}
+	else
+	{
+		DefaultTitleAreaWidget->SetVisibility(EVisibility::Collapsed);
+	}
+}
+
+const FSlateBrush* SMetasoundGraphNode::GetNodeBodyBrush() const
+{
+// 	if (const ISlateStyle* MetasoundStyle = FSlateStyleRegistry::FindSlateStyle("MetasoundStyle"))
+// 	{
+// 		if (GetMetasoundNode().GetNodeHandle()->GetClassType() == EMetasoundFrontendClassType::Input)
+// 		{
+// 			if (const FSlateBrush* InputBrush = MetasoundStyle->GetBrush("MetasoundEditor.Graph.Node.Body.Input"))
+// 			{
+// 				return InputBrush;
+// 			}
+// 		}
+// 
+// 		if (const FSlateBrush* DefaultBrush = MetasoundStyle->GetBrush("MetasoundEditor.Graph.Node.Body.Default"))
+// 		{
+// 			return DefaultBrush;
+// 		}
+// 	}
+
+	return FEditorStyle::GetBrush("Graph.Node.Body");
+}
+
+EVisibility SMetasoundGraphNode::IsAddPinButtonVisible() const
+{
+	EVisibility DefaultVisibility = SGraphNode::IsAddPinButtonVisible();
+	if (DefaultVisibility == EVisibility::Visible)
+	{
+		if (!GetMetasoundNode().CanAddInputPin())
+		{
+			return EVisibility::Collapsed;
+		}
+	}
+
+	return DefaultVisibility;
+}
+
+FReply SMetasoundGraphNode::OnAddPin()
+{
+	GetMetasoundNode().CreateInputPin();
+
+	return FReply::Handled();
+}
+
+TSharedRef<SWidget> SMetasoundGraphNode::CreateNodeContentArea()
+{
+	using namespace Metasound::Editor;
+	using namespace Metasound::Frontend;
+
+	FNodeHandle NodeHandle = GetMetasoundNode().GetNodeHandle();
+	FMetasoundFrontendClassDisplayInfo DisplayInfo = NodeHandle->GetClassDisplayInfo();
+
+	TSharedRef<SHorizontalBox> ContentBox = SNew(SHorizontalBox);
+
+	if (DisplayInfo.ImageName.IsNone())
+	{
+		ContentBox->AddSlot()
+			.HAlign(HAlign_Left)
+			.FillWidth(1.0f)
+			[
+				SAssignNew(LeftNodeBox, SVerticalBox)
+			];
+	}
+	else
+	{
+		ContentBox->AddSlot()
+			.HAlign(HAlign_Left)
+			.VAlign(VAlign_Center)
+			.AutoWidth()
+			[
+				SAssignNew(LeftNodeBox, SVerticalBox)
+			];
+
+		if (const ISlateStyle* MetasoundStyle = FSlateStyleRegistry::FindSlateStyle("MetasoundStyle"))
+		{
+			const FSlateBrush* ImageBrush = MetasoundStyle->GetBrush(DisplayInfo.ImageName);
+			ContentBox->AddSlot()
+				.AutoWidth()
+				.HAlign(HAlign_Center)
+				.VAlign(VAlign_Center)
+				[
+					SNew(SImage)
+					.Image(ImageBrush)
+					.ColorAndOpacity(FSlateColor::UseForeground())
+				];
+		}
+	}
+
+	ContentBox->AddSlot()
+		.AutoWidth()
+		.HAlign(HAlign_Right)
+		.VAlign(VAlign_Center)
+		[
+			SAssignNew(RightNodeBox, SVerticalBox)
+		];
+
+	return SNew(SBorder)
+		.BorderImage(FEditorStyle::GetBrush("NoBorder"))
+		.HAlign(HAlign_Fill)
+		.VAlign(VAlign_Fill)
+		.Padding(FMargin(0,3))
+		[
+			ContentBox
+		];
+}
+#undef LOCTEXT_NAMESPACE // MetasoundGraphNode
diff --git a/Engine/Plugins/Runtime/Metasound/Source/MetasoundEditor/Private/SMetasoundGraphNode.h b/Engine/Plugins/Runtime/Metasound/Source/MetasoundEditor/Private/SMetasoundGraphNode.h
new file mode 100644
--- /dev/null
+++ b/Engine/Plugins/Runtime/Metasound/Source/MetasoundEditor/Private/SMetasoundGraphNode.h
@@ -0,0 +1,41 @@
+// Copyright Epic Games, Inc. All Rights Reserved.
+#pragma once
+
+#include "CoreMinimal.h"
+#include "Input/Reply.h"
+#include "Layout/Visibility.h"
+#include "SGraphNode.h"
+#include "Widgets/DeclarativeSyntaxSupport.h"
+#include "Widgets/SOverlay.h"
+
+
+// Forward Declarations
+class SVerticalBox;
+class UMetasoundEditorGraphNode;
+
+
+class SMetasoundGraphNode : public SGraphNode
+{
+public:
+	SLATE_BEGIN_ARGS(SMetasoundGraphNode)
+	{
+	}
+
+	SLATE_END_ARGS()
+
+	void Construct(const FArguments& InArgs, class UEdGraphNode* InNode);
+
+protected:
+	// SGraphNode Interface
+	virtual void CreateOutputSideAddButton(TSharedPtr<SVerticalBox> OutputBox) override;
+	virtual void CreateStandardPinWidget(UEdGraphPin* CurPin) override;
+	virtual TSharedRef<SWidget> CreateNodeContentArea() override;
+	virtual TSharedRef<SWidget> CreateTitleWidget(TSharedPtr<SNodeTitle> NodeTitle) override;
+	virtual const FSlateBrush* GetNodeBodyBrush() const override;
+	virtual EVisibility IsAddPinButtonVisible() const override;
+	virtual FReply OnAddPin() override;
+	virtual void SetDefaultTitleAreaWidget(TSharedRef<SOverlay> DefaultTitleAreaWidget) override;
+
+	UMetasoundEditorGraphNode& GetMetasoundNode();
+	const UMetasoundEditorGraphNode& GetMetasoundNode() const;
+};
diff --git a/Engine/Plugins/Runtime/Metasound/Source/MetasoundEditor/Public/MetasoundEditorGraphNode.h b/Engine/Plugins/Runtime/Metasound/Source/MetasoundEditor/Public/MetasoundEditorGraphNode.h
--- a/Engine/Plugins/Runtime/Metasound/Source/MetasoundEditor/Public/MetasoundEditorGraphNode.h
+++ b/Engine/Plugins/Runtime/Metasound/Source/MetasoundEditor/Public/MetasoundEditorGraphNode.h
@@ -52,11 +52,16 @@ public:
 	virtual void PostDuplicate(bool bDuplicateForPIE) override;
 	// End of UObject interface
 
+	virtual bool CanAddInputPin() const
+	{
+		return false;
+	}
+
 	UObject& GetMetasoundChecked();
 	const UObject& GetMetasoundChecked() const;
 
-	Metasound::Frontend::FGraphHandle GetRootGraphHandle();
-	Metasound::Frontend::FNodeHandle GetNodeHandle();
+	Metasound::Frontend::FGraphHandle GetRootGraphHandle() const;
+	Metasound::Frontend::FNodeHandle GetNodeHandle() const;
 
 	void SetNodeID(uint32 InNodeID);
 	uint32 GetNodeID() const;
diff --git a/Engine/Plugins/Runtime/Metasound/Source/MetasoundEditor/Public/MetasoundEditorGraphSchema.h b/Engine/Plugins/Runtime/Metasound/Source/MetasoundEditor/Public/MetasoundEditorGraphSchema.h
--- a/Engine/Plugins/Runtime/Metasound/Source/MetasoundEditor/Public/MetasoundEditorGraphSchema.h
+++ b/Engine/Plugins/Runtime/Metasound/Source/MetasoundEditor/Public/MetasoundEditorGraphSchema.h
@@ -132,8 +132,8 @@ struct METASOUNDEDITOR_API FMetasoundGraphSchemaAction_NewNode : public FEdGraph
 		: FEdGraphSchemaAction()
 	{}
 
-	FMetasoundGraphSchemaAction_NewNode(FText InNodeCategory, FText InMenuDesc, FText InToolTip, const int32 InGrouping)
-		: FEdGraphSchemaAction(MoveTemp(InNodeCategory), MoveTemp(InMenuDesc), MoveTemp(InToolTip), InGrouping)
+	FMetasoundGraphSchemaAction_NewNode(const FText& InNodeCategory, const FText& InMenuDesc, const FText& InToolTip, const int32 InGrouping)
+		: FEdGraphSchemaAction(InNodeCategory, InMenuDesc, InToolTip, InGrouping)
 	{}
 
 	//~ Begin FEdGraphSchemaAction Interface
diff --git a/Engine/Plugins/Runtime/Metasound/Source/MetasoundFrontend/Private/MetasoundFrontend.cpp b/Engine/Plugins/Runtime/Metasound/Source/MetasoundFrontend/Private/MetasoundFrontend.cpp
--- a/Engine/Plugins/Runtime/Metasound/Source/MetasoundFrontend/Private/MetasoundFrontend.cpp
+++ b/Engine/Plugins/Runtime/Metasound/Source/MetasoundFrontend/Private/MetasoundFrontend.cpp
@@ -81,7 +81,15 @@ namespace Metasound
 				InNodeMetadata.Description,
 				InNodeMetadata.PromptIfMissing,
 				InNodeMetadata.Author,
-				InNodeMetadata.Keywords	
+				InNodeMetadata.Keywords,
+				InNodeMetadata.CategoryHierarchy,
+				FMetasoundFrontendClassDisplayInfo
+				{
+					InNodeMetadata.ImageName,
+					InNodeMetadata.bShowName,
+					InNodeMetadata.bShowInputNames,
+					InNodeMetadata.bShowOutputNames
+				},
 			};
 
 			int32 RunningPointID = 0;
@@ -240,6 +248,6 @@ REGISTER_METASOUND_DATATYPE(double, "Primitive:Double", ::Metasound::ELiteralTyp
 REGISTER_METASOUND_DATATYPE(FString, "Primitive:String", ::Metasound::ELiteralType::String)
 
 REGISTER_METASOUND_DATATYPE(Metasound::FAudioBuffer, "Audio:Buffer")
-REGISTER_METASOUND_DATATYPE(Metasound::FSendAddress, "Send:Address")
+REGISTER_METASOUND_DATATYPE(Metasound::FSendAddress, "Transmission:Address")
 
 IMPLEMENT_MODULE(FMetasoundFrontendModule, MetasoundFrontend);
diff --git a/Engine/Plugins/Runtime/Metasound/Source/MetasoundFrontend/Private/MetasoundFrontendInvalidController.h b/Engine/Plugins/Runtime/Metasound/Source/MetasoundFrontend/Private/MetasoundFrontendInvalidController.h
--- a/Engine/Plugins/Runtime/Metasound/Source/MetasoundFrontend/Private/MetasoundFrontendInvalidController.h
+++ b/Engine/Plugins/Runtime/Metasound/Source/MetasoundFrontend/Private/MetasoundFrontendInvalidController.h
@@ -144,7 +144,7 @@ namespace Metasound
 			const FString& GetClassName() const override { return MetasoundFrontendInvalidControllerPrivate::GetInvalid<FString>(); }
 			FMetasoundFrontendVersionNumber GetClassVersionNumber() const override { return MetasoundFrontendInvalidControllerPrivate::GetInvalid<FMetasoundFrontendVersionNumber>(); }
 			const FText& GetClassDescription() const override { return MetasoundFrontendInvalidControllerPrivate::GetInvalid<FText>(); }
-
+			FMetasoundFrontendClassDisplayInfo GetClassDisplayInfo() const { return FMetasoundFrontendClassDisplayInfo(); }
 
 			TSharedRef<IGraphController> AsGraph() override;
 			TSharedRef<const IGraphController> AsGraph() const override;
diff --git a/Engine/Plugins/Runtime/Metasound/Source/MetasoundFrontend/Private/MetasoundFrontendStandardController.cpp b/Engine/Plugins/Runtime/Metasound/Source/MetasoundFrontend/Private/MetasoundFrontendStandardController.cpp
--- a/Engine/Plugins/Runtime/Metasound/Source/MetasoundFrontend/Private/MetasoundFrontendStandardController.cpp
+++ b/Engine/Plugins/Runtime/Metasound/Source/MetasoundFrontend/Private/MetasoundFrontendStandardController.cpp
@@ -915,6 +915,16 @@ namespace Metasound
 			return FMetasoundFrontendVersionNumber();
 		}
 
+		FMetasoundFrontendClassDisplayInfo FBaseNodeController::GetClassDisplayInfo() const
+		{
+			if (ClassPtr.IsValid())
+			{
+				return ClassPtr->Metadata.DisplayInfo;
+			}
+
+			return FMetasoundFrontendClassDisplayInfo();
+		}
+
 		const FText& FBaseNodeController::GetClassDescription() const
 		{
 			if (ClassPtr.IsValid())
diff --git a/Engine/Plugins/Runtime/Metasound/Source/MetasoundFrontend/Private/MetasoundFrontendStandardController.h b/Engine/Plugins/Runtime/Metasound/Source/MetasoundFrontend/Private/MetasoundFrontendStandardController.h
--- a/Engine/Plugins/Runtime/Metasound/Source/MetasoundFrontend/Private/MetasoundFrontendStandardController.h
+++ b/Engine/Plugins/Runtime/Metasound/Source/MetasoundFrontend/Private/MetasoundFrontendStandardController.h
@@ -295,6 +295,7 @@ namespace Metasound
 			int32 GetClassID() const override;
 			FMetasoundFrontendVersionNumber GetClassVersionNumber() const override;
 			const FText& GetClassDescription() const override;
+			FMetasoundFrontendClassDisplayInfo GetClassDisplayInfo() const override;
 
 			const FString& GetNodeName() const override;
 			EMetasoundFrontendClassType GetClassType() const override;
diff --git a/Engine/Plugins/Runtime/Metasound/Source/MetasoundFrontend/Public/MetasoundAutoConverterNode.h b/Engine/Plugins/Runtime/Metasound/Source/MetasoundFrontend/Public/MetasoundAutoConverterNode.h
--- a/Engine/Plugins/Runtime/Metasound/Source/MetasoundFrontend/Public/MetasoundAutoConverterNode.h
+++ b/Engine/Plugins/Runtime/Metasound/Source/MetasoundFrontend/Public/MetasoundAutoConverterNode.h
@@ -61,7 +61,7 @@ namespace Metasound
 				Info.PromptIfMissing = PluginNodeMissingPrompt;
 				Info.DefaultInterface = DeclareVertexInterface();
 
-				Info.CategoryHierarchy = {LOCTEXT("Metasound_ConvertNodeCategory", "Convert"), FText::AsCultureInvariant(GetMetasoundDataTypeString<FromDataType>())};
+				Info.CategoryHierarchy = { LOCTEXT("Metasound_ConvertNodeCategory", "Conversions") };
 				Info.Keywords = {TEXT("convert"), GetMetasoundDataTypeName<FromDataType>(), GetMetasoundDataTypeName<ToDataType>()};
 
 				return Info;
diff --git a/Engine/Plugins/Runtime/Metasound/Source/MetasoundFrontend/Public/MetasoundFrontendController.h b/Engine/Plugins/Runtime/Metasound/Source/MetasoundFrontend/Public/MetasoundFrontendController.h
--- a/Engine/Plugins/Runtime/Metasound/Source/MetasoundFrontend/Public/MetasoundFrontendController.h
+++ b/Engine/Plugins/Runtime/Metasound/Source/MetasoundFrontend/Public/MetasoundFrontendController.h
@@ -275,6 +275,7 @@ namespace Metasound
 			virtual const FString& GetClassName() const = 0;
 			virtual FMetasoundFrontendVersionNumber GetClassVersionNumber() const = 0;
 			virtual const FText& GetClassDescription() const = 0;
+			virtual FMetasoundFrontendClassDisplayInfo GetClassDisplayInfo() const = 0;
 
 			/** If the node is also a graph, this returns a graph handle.
 			 * If the node is not also a graph, it will return an invalid handle.
diff --git a/Engine/Plugins/Runtime/Metasound/Source/MetasoundFrontend/Public/MetasoundFrontendDocument.h b/Engine/Plugins/Runtime/Metasound/Source/MetasoundFrontend/Public/MetasoundFrontendDocument.h
--- a/Engine/Plugins/Runtime/Metasound/Source/MetasoundFrontend/Public/MetasoundFrontendDocument.h
+++ b/Engine/Plugins/Runtime/Metasound/Source/MetasoundFrontend/Public/MetasoundFrontendDocument.h
@@ -257,7 +257,6 @@ struct FMetasoundFrontendStyleEdge
 	EMetasoundFrontendStyleEdgeDisplay Display;
 };
 
-
 // Styling for a class of edges dependent upon edge data type.
 USTRUCT()
 struct FMetasoundFrontendStyleEdgeClass
@@ -396,6 +395,25 @@ struct METASOUNDFRONTEND_API FMetasoundFrontendClassVertex : public FMetasoundFr
 	static bool IsFunctionalEquivalent(const FMetasoundFrontendClassVertex& InLHS, const FMetasoundFrontendClassVertex& InRHS);
 };
 
+// Information regarding how to display a node class
+USTRUCT()
+struct METASOUNDFRONTEND_API FMetasoundFrontendClassDisplayInfo
+{
+	GENERATED_BODY()
+
+	UPROPERTY()
+	FName ImageName;
+
+	UPROPERTY()
+	bool bShowName = true;
+
+	UPROPERTY()
+	bool bShowInputName = true;
+
+	UPROPERTY()
+	bool bShowOutputName = true;
+};
+
 // Contains info for input vertex of a Metasound class.
 USTRUCT() 
 struct METASOUNDFRONTEND_API FMetasoundFrontendClassInput : public FMetasoundFrontendClassVertex
@@ -562,6 +580,12 @@ struct FMetasoundFrontendClassMetadata
 
 	UPROPERTY()
 	TArray<FName> Keywords;
+
+	UPROPERTY()
+	TArray<FText> CategoryHierarchy;
+
+	UPROPERTY()
+	FMetasoundFrontendClassDisplayInfo DisplayInfo;
 };
 
 UENUM()
diff --git a/Engine/Plugins/Runtime/Metasound/Source/MetasoundFrontend/Public/MetasoundReceiveNode.h b/Engine/Plugins/Runtime/Metasound/Source/MetasoundFrontend/Public/MetasoundReceiveNode.h
--- a/Engine/Plugins/Runtime/Metasound/Source/MetasoundFrontend/Public/MetasoundReceiveNode.h
+++ b/Engine/Plugins/Runtime/Metasound/Source/MetasoundFrontend/Public/MetasoundReceiveNode.h
@@ -59,7 +59,7 @@ namespace Metasound
 				Info.Author = PluginAuthor;
 				Info.PromptIfMissing = PluginNodeMissingPrompt;
 				Info.DefaultInterface = DeclareVertexInterface();
-				Info.CategoryHierarchy = {LOCTEXT("Metasound_ReceiveNodeCategory", "Receive"), FText::AsCultureInvariant(GetMetasoundDataTypeString<TDataType>())};
+				Info.CategoryHierarchy = { LOCTEXT("Metasound_TransmissionNodeCategory", "Transmission") };
 				Info.Keywords = {TEXT("receive"), GetMetasoundDataTypeName<TDataType>()};
 
 				return Info;
diff --git a/Engine/Plugins/Runtime/Metasound/Source/MetasoundFrontend/Public/MetasoundSendNode.h b/Engine/Plugins/Runtime/Metasound/Source/MetasoundFrontend/Public/MetasoundSendNode.h
--- a/Engine/Plugins/Runtime/Metasound/Source/MetasoundFrontend/Public/MetasoundSendNode.h
+++ b/Engine/Plugins/Runtime/Metasound/Source/MetasoundFrontend/Public/MetasoundSendNode.h
@@ -58,7 +58,7 @@ namespace Metasound
 				Info.Author = PluginAuthor;
 				Info.PromptIfMissing = PluginNodeMissingPrompt;
 				Info.DefaultInterface = DeclareVertexInterface();
-				Info.CategoryHierarchy = {LOCTEXT("Metasound_SendNodeCategory", "Send"), FText::AsCultureInvariant(GetMetasoundDataTypeString<TDataType>())};
+				Info.CategoryHierarchy = { LOCTEXT("Metasound_TransmissionNodeCategory", "Transmission") };
 				Info.Keywords = {TEXT("send"), GetMetasoundDataTypeName<TDataType>()};
 
 				return Info;
diff --git a/Engine/Plugins/Runtime/Metasound/Source/MetasoundGraphCore/Public/MetasoundNodeInterface.h b/Engine/Plugins/Runtime/Metasound/Source/MetasoundGraphCore/Public/MetasoundNodeInterface.h
--- a/Engine/Plugins/Runtime/Metasound/Source/MetasoundGraphCore/Public/MetasoundNodeInterface.h
+++ b/Engine/Plugins/Runtime/Metasound/Source/MetasoundGraphCore/Public/MetasoundNodeInterface.h
@@ -70,6 +70,14 @@ namespace Metasound
 		/** Hierarchy of categories for displaying node. */
 		TArray<FText> CategoryHierarchy;
 
+		FName ImageName;
+
+		bool bShowName = true;
+
+		bool bShowInputNames = true;
+
+		bool bShowOutputNames = true;
+
 		/** List of keywords for contextual node searching. */
 		TArray<FName> Keywords;
 
diff --git a/Engine/Plugins/Runtime/Metasound/Source/MetasoundStandardNodes/Private/MetasoundStandardNodesModule.cpp b/Engine/Plugins/Runtime/Metasound/Source/MetasoundStandardNodes/Private/MetasoundStandardNodesModule.cpp
--- a/Engine/Plugins/Runtime/Metasound/Source/MetasoundStandardNodes/Private/MetasoundStandardNodesModule.cpp
+++ b/Engine/Plugins/Runtime/Metasound/Source/MetasoundStandardNodes/Private/MetasoundStandardNodesModule.cpp
@@ -10,7 +10,7 @@
 #include "MetasoundTime.h"
 #include "MetasoundFrequency.h"
 
-REGISTER_METASOUND_DATATYPE(Metasound::FBop, "Primitive:Bop", ::Metasound::ELiteralType::Boolean)
+REGISTER_METASOUND_DATATYPE(Metasound::FTrigger, "Primitive:Trigger", ::Metasound::ELiteralType::Boolean)
 REGISTER_METASOUND_DATATYPE(Metasound::FFloatTime, "Primitive:Time", ::Metasound::ELiteralType::Float)
 REGISTER_METASOUND_DATATYPE(Metasound::FDoubleTime, "Primitive:Time:HighResolution", ::Metasound::ELiteralType::Float)
 REGISTER_METASOUND_DATATYPE(Metasound::FSampleTime, "Primitive:Time:SampleResolution", ::Metasound::ELiteralType::Float)
diff --git a/Engine/Plugins/Runtime/Metasound/Source/MetasoundStandardNodes/Public/MetasoundMathNodes.h b/Engine/Plugins/Runtime/Metasound/Source/MetasoundStandardNodes/Public/MetasoundMathNodes.h
--- a/Engine/Plugins/Runtime/Metasound/Source/MetasoundStandardNodes/Public/MetasoundMathNodes.h
+++ b/Engine/Plugins/Runtime/Metasound/Source/MetasoundStandardNodes/Public/MetasoundMathNodes.h
@@ -26,8 +26,9 @@
 		: public ::Metasound::TMathOpNode<F##OpName##DataTypeName##Node, ::Metasound::TMathOp##OpName<DataClass>, DataClass> \
 	{ \
 	public: \
-		static FName GetClassName() { return #OpName #DataTypeName; } \
+		static FName GetClassName() { return #OpName " " #DataTypeName; } \
 		static FText GetDescription() { return Description; } \
+		static FName GetImageName() { return "MetasoundEditor.Graph.Node.Math." #OpName; } \
 		F##OpName##DataTypeName##Node(const FNodeInitData& InInitData) : TMathOpNode(InInitData) { } \
 	};
 
@@ -74,10 +75,10 @@ namespace Metasound
 					Info.Author = PluginAuthor;
 					Info.PromptIfMissing = PluginNodeMissingPrompt;
 					Info.DefaultInterface = DeclareVertexInterface();
-					//Info.ImageName = TEXT("MetasoundEditor.Speaker");
-					//Info.bShowName = false;
-					//Info.bShowInputNames = false;
-					//Info.bShowOutputNames = false;
+					Info.ImageName = TDerivedClass::GetImageName();
+					Info.bShowName = false;
+					Info.bShowInputNames = false;
+					Info.bShowOutputNames = false;
 					Info.CategoryHierarchy = { LOCTEXT("Metasound_MathCategory", "Math") };
 
 					return Info;
diff --git a/Engine/Plugins/Runtime/Metasound/Source/MetasoundStandardNodes/Public/MetasoundTime.h b/Engine/Plugins/Runtime/Metasound/Source/MetasoundStandardNodes/Public/MetasoundTime.h
--- a/Engine/Plugins/Runtime/Metasound/Source/MetasoundStandardNodes/Public/MetasoundTime.h
+++ b/Engine/Plugins/Runtime/Metasound/Source/MetasoundStandardNodes/Public/MetasoundTime.h
@@ -37,7 +37,7 @@ namespace Metasound
 
 		public:
 			using FQuantizationType = QuantizationType;
-			using FTimeTYpe = TimeType;
+			using FTimeType = TimeType;
 
 			/** Default constructor */
 			TTime()
