Fix-up relative include path

diff --git a/Engine/Plugins/Runtime/Metasound/Source/MetasoundEditor/Private/MetasoundEditorGraphSchema.cpp b/Engine/Plugins/Runtime/Metasound/Source/MetasoundEditor/Private/MetasoundEditorGraphSchema.cpp
--- a/Engine/Plugins/Runtime/Metasound/Source/MetasoundEditor/Private/MetasoundEditorGraphSchema.cpp
+++ b/Engine/Plugins/Runtime/Metasound/Source/MetasoundEditor/Private/MetasoundEditorGraphSchema.cpp
@@ -12,6 +12,7 @@
 #include "GraphEditorSettings.h"
 #include "Layout/SlateRect.h"
 #include "Metasound.h"
+#include "MetasoundAssetBase.h"
 #include "MetasoundEditor.h"
 #include "MetasoundEditorGraph.h"
 #include "MetasoundEditorGraphBuilder.h"
@@ -24,7 +25,6 @@
 #include "UObject/NoExportTypes.h"
 #include "UObject/UObjectHash.h"
 #include "UObject/UObjectIterator.h"
-#include "../../MetasoundFrontend/Public/MetasoundAssetBase.h"
 
 #define LOCTEXT_NAMESPACE "MetasoundEditor"
 
