CIS build fix for Metasound access ptr define

diff --git a/Engine/Plugins/Runtime/Metasound/Source/MetasoundFrontend/Public/MetasoundAccessPtr.h b/Engine/Plugins/Runtime/Metasound/Source/MetasoundFrontend/Public/MetasoundAccessPtr.h
--- a/Engine/Plugins/Runtime/Metasound/Source/MetasoundFrontend/Public/MetasoundAccessPtr.h
+++ b/Engine/Plugins/Runtime/Metasound/Source/MetasoundFrontend/Public/MetasoundAccessPtr.h
@@ -7,6 +7,8 @@
 
 #if !(UE_BUILD_SHIPPING || UE_BUILD_TEST)
 #define METASOUND_FRONTEND_ACCESSPTR_DEBUG_INFO 1
+#else
+#define METASOUND_FRONTEND_ACCESSPTR_DEBUG_INFO 0
 #endif
 namespace Metasound
 {
