metasounds: gate and switch nodes

diff --git a/Engine/Plugins/Runtime/Metasound/Source/MetasoundStandardNodes/Private/MetasoundSwitchGateNode.cpp b/Engine/Plugins/Runtime/Metasound/Source/MetasoundStandardNodes/Private/MetasoundSwitchGateNode.cpp
new file mode 100644
--- /dev/null
+++ b/Engine/Plugins/Runtime/Metasound/Source/MetasoundStandardNodes/Private/MetasoundSwitchGateNode.cpp
@@ -0,0 +1,460 @@
+// Copyright Epic Games, Inc. All Rights Reserved.
+
+#include "MetasoundSwitchGateNode.h"
+
+#include "MetasoundAudioBuffer.h"
+#include "MetasoundExecutableOperator.h"
+#include "MetasoundFacade.h"
+#include "MetasoundFrequency.h"
+#include "MetasoundNodeRegistrationMacro.h"
+#include "MetasoundPrimitives.h"
+#include "MetasoundVertex.h"
+
+#define LOCTEXT_NAMESPACE "MetasoundStandardNodes"
+
+// pending array support...
+//REGISTER_METASOUND_DATATYPE(Metasound::FAudioBufferArray, "Audio:BufferArray")
+
+// since we don't have array support in the UI, this is the number of inputs to
+// expose.
+static const int PREVIEW_INPUT_COUNT = 4;
+
+namespace Metasound
+{
+	// Names have to match everywhere or bad things happen.
+	// FString::Printf doesn't allow variables to be passed as the
+	// format argument, so those are #defines.
+	static constexpr const TCHAR* SWITCH_SELECTOR_NAME = TEXT("Selector");
+	static constexpr const TCHAR* SWITCH_OUTPUT_NAME = TEXT("Output");
+	static constexpr const TCHAR* SWITCH_CROSSFADE_NAME = TEXT("Crossfade");
+	#define SWITCH_INPUT_NAME_FMT TEXT("Input%d")
+
+	static constexpr const TCHAR* GATE_SELECTOR_NAME = TEXT("Selector");
+	static constexpr const TCHAR* GATE_INPUT_NAME = TEXT("Input");
+	static constexpr const TCHAR* GATE_CROSSFADE_NAME = TEXT("Crossfade");
+	#define GATE_OUTPUT_NAME_FMT TEXT("Output%d")
+
+	static const int32 CROSSFADE_FRAMES = 64;
+
+	class FSwitchOperator : public TExecutableOperator<FSwitchOperator>
+	{
+		public:
+			static TUniquePtr<IOperator> CreateOperator(const FCreateOperatorParams& InParams, FBuildErrorArray& OutErrors);
+			static const FNodeInfo& GetNodeInfo();
+			static FVertexInterface DeclareVertexInterface();
+
+			FSwitchOperator(const FOperatorSettings& InSettings, FInt32ReadRef InSelector, FBoolReadRef InShouldCrossfade, FAudioBufferReadRef InBuffers[PREVIEW_INPUT_COUNT])
+				: BlockSize(InSettings.GetNumFramesPerBlock())
+				, ShouldCrossfade(InShouldCrossfade)
+				, Selector(InSelector)				
+				, TargetAudioBuffer(FAudioBufferWriteRef::CreateNew(InSettings))
+			{
+				for (int32 i = 0; i < PREVIEW_INPUT_COUNT; i++)
+				{
+					SourceAudioBuffers.Add(InBuffers[i]);
+				}
+
+				check(TargetAudioBuffer->Num() == InSettings.GetNumFramesPerBlock());
+			}
+
+			virtual FDataReferenceCollection GetInputs() const override;
+			virtual FDataReferenceCollection GetOutputs() const override;
+
+			void Execute();
+
+		private:
+			int32 BlockSize;
+
+			// so we know if we've mixed before.
+			int32 LastSelector = -1;
+
+			// Determines whether running changes should crossfade, defaults to true,
+			// see comments in Execute()
+			FBoolReadRef ShouldCrossfade;
+			FInt32ReadRef Selector;
+			TArray<FAudioBufferReadRef> SourceAudioBuffers;
+			FAudioBufferWriteRef TargetAudioBuffer;
+	};
+
+	FDataReferenceCollection FSwitchOperator::GetInputs() const
+	{
+		FDataReferenceCollection InputDataReferences;	
+		InputDataReferences.AddDataReadReference(SWITCH_SELECTOR_NAME, FInt32ReadRef(Selector));
+		for (int32 i=0; i<PREVIEW_INPUT_COUNT; i++)
+		{
+			InputDataReferences.AddDataReadReference(*FString::Printf(SWITCH_INPUT_NAME_FMT, i), FAudioBufferReadRef(SourceAudioBuffers[i]));
+		}		
+		return InputDataReferences;
+	}
+
+	FDataReferenceCollection FSwitchOperator::GetOutputs() const
+	{
+		FDataReferenceCollection OutputDataReferences;
+		OutputDataReferences.AddDataReadReference(SWITCH_OUTPUT_NAME, FAudioBufferReadRef(TargetAudioBuffer));
+		return OutputDataReferences;
+	}
+
+	void FSwitchOperator::Execute()
+	{
+		int32 CurrentSelector = *Selector;
+
+		if (CurrentSelector < 0 ||
+			CurrentSelector >= PREVIEW_INPUT_COUNT)
+		{
+			return;
+		}
+
+		//
+		// Generally speaking when swapping buffers you want to do a short crossfade
+		// as you don't really know whether you're on a sane zero crossing. This is
+		// solely to prevent pops. However, if you're at the very beginning of a sound
+		// you don't want to do this as you mute the beginning of the sound. The
+		// difficulty is knowing when you're at the beginning of a sound. At this point
+		// I'm expecting the usage case to be almost always sounds start at the beginning
+		// of the graph evaluation timeframe and end at the end, with no delays. As
+		// a result, the nodes are constructed such that the default behavior is to
+		// snap to the first buffer, and then crossfade further changes.
+		//
+		// Whether to crossfade or not is exposed as pins to disable runtime crossfading.
+		//
+		// The case where you'd want this is when the graph is swapping to a new
+		// audio source that is starting up mid evaluation. For example some trigger
+		// starts up a cue while at the same time selecting that cue on a switch, some
+		// 10 seconds in to a sound. If left to defaults, this will be crossfaded, causing
+		// the first samples of the new cue to be muted, which is unlikely to be the
+		// desired behavior.
+		//
+		// The end-goal best case design I believe would be for branches of the tree to know 
+		// when they are being evaluated "for the first time" i.e. audio sources on a natural
+		// boundary (loop, start, or a designer specified marker). This knowledge would be
+		// passed to the operator, and crossfading controlled thusly.
+		//
+
+		if (LastSelector == -1 || *ShouldCrossfade == false) // we've never mixed, so don't xfade
+		{
+			LastSelector = CurrentSelector;
+		}
+
+		FMemory::Memcpy(TargetAudioBuffer->GetData(), SourceAudioBuffers[CurrentSelector]->GetData(), sizeof(float) * BlockSize);
+		if (LastSelector != CurrentSelector) // expected to be rare
+		{
+			if (LastSelector >=  0 &&
+				LastSelector < PREVIEW_INPUT_COUNT) // can't fail with a static count, but maybe if there's ever a dynamic buffer count...?
+			{
+				int32 CrossfadeFrames = CROSSFADE_FRAMES;
+				if (CrossfadeFrames > BlockSize)
+				{
+					CrossfadeFrames = BlockSize;
+				}
+
+				float const* PreviousBuffer = SourceAudioBuffers[LastSelector]->GetData();
+				float const* CurrentBuffer = SourceAudioBuffers[CurrentSelector]->GetData();
+				float* DestinationBuffer = TargetAudioBuffer->GetData();
+
+				float VolumeStep = 1.0f / CrossfadeFrames;
+				float In = 0;
+				float Out = 1.0f;
+				for (int32 i = 0; i < CrossfadeFrames; i++)
+				{
+					DestinationBuffer[i] = In*CurrentBuffer[i] + Out*PreviousBuffer[i];
+					In += VolumeStep;
+					Out -= VolumeStep;
+				}
+			}
+			LastSelector = CurrentSelector;
+		}
+	}
+
+	FVertexInterface FSwitchOperator::DeclareVertexInterface()
+	{
+		static FVertexInterface CachedInterface;
+		static int32 CachedInterfaceInputCount;
+
+		if (CachedInterfaceInputCount != PREVIEW_INPUT_COUNT)
+		{
+			// Create an interface that represents the # of inputs we are previewing.
+
+			CachedInterface = FVertexInterface();
+			
+			CachedInterface.GetInputInterface().Add(TInputDataVertexModel<int32>(SWITCH_SELECTOR_NAME, LOCTEXT("SwitchSelectorDescription", "The index of the audio buffer to use (0 based).")));
+			CachedInterface.GetInputInterface().Add(TInputDataVertexModel<bool>(SWITCH_CROSSFADE_NAME, LOCTEXT("SwitchCrossfadeDescription", "Whether changing the buffer while running causes a short crossfade to prevent pops. This should almost always be true.")));
+			for (int32 i=0; i < PREVIEW_INPUT_COUNT; i++)
+			{				
+				CachedInterface.GetInputInterface().Add(TInputDataVertexModel<FAudioBuffer>(*FString::Printf(SWITCH_INPUT_NAME_FMT, i), LOCTEXT("SwitchInputDescription", "An input buffer to choose from.")));
+			}			
+
+			CachedInterface.GetOutputInterface().Add(TOutputDataVertexModel<FAudioBuffer>(SWITCH_OUTPUT_NAME, LOCTEXT("SwitchOutputDescription", "The output audio buffer.")));
+			
+			CachedInterfaceInputCount = PREVIEW_INPUT_COUNT;
+		}
+
+		return CachedInterface;
+	}
+
+	const FNodeInfo& FSwitchOperator::GetNodeInfo()
+	{
+		auto InitNodeInfo = []() -> FNodeInfo
+		{
+			FNodeInfo Info;
+			Info.ClassName = FName(TEXT("Switch"));
+			Info.MajorVersion = 1;
+			Info.MinorVersion = 0;
+			Info.Description = LOCTEXT("Metasound_SwitchNodeDescription", "Selects one audio buffer from any number of inputs.");
+			Info.Author = PluginAuthor;
+			Info.PromptIfMissing = PluginNodeMissingPrompt;
+			Info.DefaultInterface = DeclareVertexInterface();
+
+			return Info;
+		};
+
+		static const FNodeInfo Info = InitNodeInfo();
+		return Info;
+	}
+
+	TUniquePtr<IOperator> FSwitchOperator::CreateOperator(const FCreateOperatorParams& InParams, FBuildErrorArray& OutErrors)
+	{
+		const FSwitchNode& SwitchNode = static_cast<const FSwitchNode&>(InParams.Node);
+		const FDataReferenceCollection& InputCol = InParams.InputDataReferences;
+
+		FAudioBuffer Zeroes(InParams.OperatorSettings.GetNumFramesPerBlock());
+		Zeroes.Zero();
+
+		FInt32ReadRef Selector = InputCol.GetDataReadReferenceOrConstruct<int32>(SWITCH_SELECTOR_NAME, 0);
+		FBoolReadRef Crossfade = InputCol.GetDataReadReferenceOrConstruct<bool>(SWITCH_CROSSFADE_NAME, true);
+
+		FAudioBufferReadRef PreviewBuffers[PREVIEW_INPUT_COUNT] = {
+			InputCol.GetDataReadReferenceOrConstruct<FAudioBuffer>(*FString::Printf(SWITCH_INPUT_NAME_FMT, 0), Zeroes),
+			InputCol.GetDataReadReferenceOrConstruct<FAudioBuffer>(*FString::Printf(SWITCH_INPUT_NAME_FMT, 1), Zeroes),
+			InputCol.GetDataReadReferenceOrConstruct<FAudioBuffer>(*FString::Printf(SWITCH_INPUT_NAME_FMT, 2), Zeroes),
+			InputCol.GetDataReadReferenceOrConstruct<FAudioBuffer>(*FString::Printf(SWITCH_INPUT_NAME_FMT, 3), Zeroes)
+		};
+
+		return MakeUnique<FSwitchOperator>(InParams.OperatorSettings, Selector, Crossfade, PreviewBuffers);
+	}
+
+	FSwitchNode::FSwitchNode(const FString& InName)
+		: FNodeFacade(InName, TFacadeOperatorClass<FSwitchOperator>())
+	{
+	}
+
+	FSwitchNode::FSwitchNode(const FNodeInitData& InInitData)
+		: FSwitchNode(InInitData.InstanceName)
+	{
+	}
+
+	METASOUND_REGISTER_NODE(FSwitchNode);
+
+	class FGateOperator : public TExecutableOperator<FGateOperator>
+	{
+	public:
+		static TUniquePtr<IOperator> CreateOperator(const FCreateOperatorParams& InParams, FBuildErrorArray& OutErrors);
+		static const FNodeInfo& GetNodeInfo();
+		static FVertexInterface DeclareVertexInterface();
+
+		FGateOperator(const FOperatorSettings& InSettings, FInt32ReadRef InSelector, FBoolReadRef InShouldCrossfade, FAudioBufferReadRef InBuffer)
+			: BlockSize(InSettings.GetNumFramesPerBlock())
+			, ShouldCrossfade(InShouldCrossfade)
+			, Selector(InSelector)			
+			, SourceAudioBuffer(InBuffer)
+		{
+			for (int32 i = 0; i < PREVIEW_INPUT_COUNT; i++)
+			{
+				TargetAudioBuffers.Add(FAudioBufferWriteRef::CreateNew(InSettings));
+			}
+		}
+
+		virtual FDataReferenceCollection GetInputs() const override;
+		virtual FDataReferenceCollection GetOutputs() const override;
+
+		void Execute();
+
+	private:
+		int32 BlockSize;
+
+		// so we know if we've mixed before.
+		int32 LastSelector = -1;
+
+		// Determines whether running changes should crossfade, defaults to true,
+		// see comments in Execute()
+		FBoolReadRef ShouldCrossfade;
+		FInt32ReadRef Selector;
+		FAudioBufferReadRef SourceAudioBuffer;
+		TArray<FAudioBufferWriteRef> TargetAudioBuffers;
+	};
+
+	FDataReferenceCollection FGateOperator::GetInputs() const
+	{
+		FDataReferenceCollection InputDataReferences;
+		InputDataReferences.AddDataReadReference(GATE_SELECTOR_NAME, FInt32ReadRef(Selector));
+		InputDataReferences.AddDataReadReference(GATE_INPUT_NAME, FAudioBufferReadRef(SourceAudioBuffer));
+		InputDataReferences.AddDataReadReference(GATE_CROSSFADE_NAME, FBoolReadRef(ShouldCrossfade));
+		return InputDataReferences;
+	}
+
+	FDataReferenceCollection FGateOperator::GetOutputs() const
+	{
+		FDataReferenceCollection OutputDataReferences;
+		for (int32 i=0; i<PREVIEW_INPUT_COUNT; i++)
+		{
+			OutputDataReferences.AddDataReadReference(*FString::Printf(GATE_OUTPUT_NAME_FMT, i), FAudioBufferReadRef(TargetAudioBuffers[i]));
+		}
+		return OutputDataReferences;
+	}
+
+
+	void FGateOperator::Execute()
+	{
+		int32 CurrentSelector = *Selector;
+		if (CurrentSelector < 0 ||
+			CurrentSelector >= PREVIEW_INPUT_COUNT)
+		{
+			return;
+		}
+		
+		//
+		// The crossfade problem is a bit different with gates, as it acts more like a stop/start
+		// sequence, with caveats.
+		//
+		// We can likely assume that both branches reach the output, and if both branches are equal
+		// then technically no crossfade is necessary as the output gets the same view. However,
+		// if there are filters in between, we can introduce a pop in to the interim operators, which
+		// can cause badness, depending on the operator (e.g. IIR filters).
+		//
+		// So, we continue to output a few frames of fadeout to the old buffer, and fade in a bit on
+		// the new buffer.
+		//
+		// Of course, this has the same issue as the switch where if the sound is on a natural boundary, 
+		// this can cause unwanted muting of samples. This is even worse because if the old branch is
+		// actually cut _out_ of the output via volume or even a delay, the new buffer has a fade in
+		// and the fade out isn't synced.
+		//
+		// Since we can't know that information, we expose the crossfade, and like the switch, default
+		// it to true.
+		//
+		if (LastSelector == -1 || *ShouldCrossfade == false) // we've never mixed, so don't xfade
+		{
+			LastSelector = CurrentSelector;
+		}
+
+		// All outputs must be written to.
+		for (int32 i=0; i < PREVIEW_INPUT_COUNT; i++)
+		{
+			if (i != CurrentSelector)
+			{
+				FMemory::Memset(TargetAudioBuffers[i]->GetData(), 0, sizeof(float) * BlockSize);
+			}
+			else
+			{
+				FMemory::Memcpy(TargetAudioBuffers[CurrentSelector]->GetData(), SourceAudioBuffer->GetData(), sizeof(float) * BlockSize);
+			}
+		}
+
+		if (LastSelector != CurrentSelector) // rare
+		{
+			if (LastSelector >= 0 &&
+				LastSelector < PREVIEW_INPUT_COUNT) // can't fail with a static count, but maybe if there's ever a dynamic buffer count...?
+			{
+				int32 CrossfadeFrames = CROSSFADE_FRAMES;
+				if (CrossfadeFrames > BlockSize)
+				{
+					CrossfadeFrames = BlockSize;
+				}
+
+				float const VolumeStep = 1.0f / CrossfadeFrames;
+				float const* SourceBuffer = SourceAudioBuffer->GetData();
+				
+				float* OldDestinationBuffer = TargetAudioBuffers[LastSelector]->GetData();
+				float* NewDestinationBuffer = TargetAudioBuffers[CurrentSelector]->GetData();
+
+				float Out = 1.0f;
+				float In = 0;
+				for (int32 i = 0; i < CrossfadeFrames; i++)
+				{
+					// Fade out the old buffer
+					OldDestinationBuffer[i] = Out * SourceBuffer[i];
+
+					// Fade in the new buffer.	
+					NewDestinationBuffer[i] = In * SourceBuffer[i];
+					Out -= VolumeStep;
+					In += VolumeStep;
+				}
+			}
+			LastSelector = CurrentSelector;
+		}		
+	}
+
+	FVertexInterface FGateOperator::DeclareVertexInterface()
+	{
+		static FVertexInterface CachedInterface;
+		static int32 CachedInterfaceInputCount;
+
+		if (CachedInterfaceInputCount != PREVIEW_INPUT_COUNT)
+		{
+			// Create an interface that represents the # of inputs we are previewing.
+
+			CachedInterface = FVertexInterface();
+
+			CachedInterface.GetInputInterface().Add(TInputDataVertexModel<int32>(GATE_SELECTOR_NAME, LOCTEXT("GateSelectorDescription", "The index of the output audio buffer to use (0 based).")));
+			CachedInterface.GetInputInterface().Add(TInputDataVertexModel<FAudioBuffer>(GATE_INPUT_NAME, LOCTEXT("GateInputDescription", "The input buffer to route.")));
+			CachedInterface.GetInputInterface().Add(TInputDataVertexModel<bool>(GATE_CROSSFADE_NAME, LOCTEXT("GateCrossfadeDescription", "Whether changing the buffer while running causes a short crossfade to prevent pops. This should almost always be true.")));
+
+			for (int32 i = 0; i < PREVIEW_INPUT_COUNT; i++)
+			{
+				CachedInterface.GetOutputInterface().Add(TOutputDataVertexModel<FAudioBuffer>(*FString::Printf(GATE_OUTPUT_NAME_FMT, i), LOCTEXT("GateOutputDescription", "An output buffer to choose from.")));
+			}
+
+			CachedInterfaceInputCount = PREVIEW_INPUT_COUNT;
+		}
+
+		return CachedInterface;
+	}
+
+	const FNodeInfo& FGateOperator::GetNodeInfo()
+	{
+		auto InitNodeInfo = []() -> FNodeInfo
+		{
+			FNodeInfo Info;
+			Info.ClassName = FName(TEXT("Gate"));
+			Info.MajorVersion = 1;
+			Info.MinorVersion = 0;
+			Info.Description = LOCTEXT("Metasound_GateNodeDescription", "Routes an audio buffer to one of a set of output buffers based on an input value.");
+			Info.Author = PluginAuthor;
+			Info.PromptIfMissing = PluginNodeMissingPrompt;
+			Info.DefaultInterface = DeclareVertexInterface();
+
+			return Info;
+		};
+
+		static const FNodeInfo Info = InitNodeInfo();
+		return Info;
+	}
+
+	TUniquePtr<IOperator> FGateOperator::CreateOperator(const FCreateOperatorParams& InParams, FBuildErrorArray& OutErrors)
+	{
+		const FGateNode& GateNode = static_cast<const FGateNode&>(InParams.Node);
+		const FDataReferenceCollection& InputCol = InParams.InputDataReferences;
+
+		FAudioBuffer Zeroes(InParams.OperatorSettings.GetNumFramesPerBlock());
+		Zeroes.Zero();
+
+		FInt32ReadRef Selector = InputCol.GetDataReadReferenceOrConstruct<int32>(GATE_SELECTOR_NAME, 0);
+		FBoolReadRef Crossfade = InputCol.GetDataReadReferenceOrConstruct<bool>(GATE_CROSSFADE_NAME, true);
+		FAudioBufferReadRef Input = InputCol.GetDataReadReferenceOrConstruct<FAudioBuffer>(GATE_INPUT_NAME, Zeroes);
+
+		return MakeUnique<FGateOperator>(InParams.OperatorSettings, Selector, Crossfade, Input);
+	}
+
+	FGateNode::FGateNode(const FString& InName)
+		: FNodeFacade(InName, TFacadeOperatorClass<FGateOperator>())
+	{
+	}
+
+	FGateNode::FGateNode(const FNodeInitData& InInitData)
+		: FGateNode(InInitData.InstanceName)
+	{
+	}
+
+	METASOUND_REGISTER_NODE(FGateNode);
+
+}
+#undef LOCTEXT_NAMESPACE
diff --git a/Engine/Plugins/Runtime/Metasound/Source/MetasoundStandardNodes/Public/MetasoundSwitchGateNode.h b/Engine/Plugins/Runtime/Metasound/Source/MetasoundStandardNodes/Public/MetasoundSwitchGateNode.h
new file mode 100644
--- /dev/null
+++ b/Engine/Plugins/Runtime/Metasound/Source/MetasoundStandardNodes/Public/MetasoundSwitchGateNode.h
@@ -0,0 +1,46 @@
+// Copyright Epic Games, Inc. All Rights Reserved.
+
+#pragma once
+
+#include "CoreMinimal.h"
+
+#include "MetasoundAudioBuffer.h"
+#include "MetasoundBuilderInterface.h"
+#include "MetasoundDataReferenceCollection.h"
+#include "MetasoundExecutableOperator.h"
+#include "MetasoundFacade.h"
+#include "MetasoundFrequency.h"
+#include "MetasoundNode.h"
+#include "MetasoundOperatorInterface.h"
+
+namespace Metasound
+{
+	// Pending array support...
+	//typedef TArray<FAudioBuffer> FAudioBufferArray;
+	//DECLARE_METASOUND_DATA_REFERENCE_TYPES(FAudioBufferArray, METASOUNDSTANDARDNODES_API, FAudioBufferArrayTypeInfo, FAudioBufferArrayReadRef, FAudioBufferArrayWriteRef)
+
+	/** FGateNode
+	 *
+	 *	The switch node selects a single audio buffer from an array of options.
+	 *	The selection can change at any time during execution, however there is no
+	 *	blending between nodes.
+	*/
+	class METASOUNDSTANDARDNODES_API FSwitchNode : public FNodeFacade
+	{
+	public:
+		FSwitchNode(const FString& InName);
+		FSwitchNode(const FNodeInitData& InInitData);
+	};
+
+	/** FGateNode
+	 *
+	 *  The gate node routes an audio buffer to a selected output buffer.
+	 *	The output buffer can be changed during execution.
+	 */
+	class METASOUNDSTANDARDNODES_API FGateNode : public FNodeFacade
+	{
+	public:
+		FGateNode(const FString& InName);
+		FGateNode(const FNodeInitData& InInitData);
+	};
+}
