Metasound Document GUID and Document Accessor

diff --git a/Engine/Plugins/Runtime/Metasound/Source/MetasoundEditor/Private/MetasoundEditorGraphBuilder.cpp b/Engine/Plugins/Runtime/Metasound/Source/MetasoundEditor/Private/MetasoundEditorGraphBuilder.cpp
--- a/Engine/Plugins/Runtime/Metasound/Source/MetasoundEditor/Private/MetasoundEditorGraphBuilder.cpp
+++ b/Engine/Plugins/Runtime/Metasound/Source/MetasoundEditor/Private/MetasoundEditorGraphBuilder.cpp
@@ -155,7 +155,7 @@ namespace Metasound
 			Description.Name = InName;
 			Description.TypeName = InTypeName;
 			Description.Metadata.Description = InToolTip;
-			int32 PointID = GraphHandle->GetNewPointID();
+			FGuid PointID = GraphHandle->GetNewPointID();
 
 			Description.PointIDs.Add(PointID);
 			FMetasoundFrontendVertexLiteral DefaultValue;
@@ -316,7 +316,7 @@ namespace Metasound
 				UEdGraphNode* GraphNode = nullptr;
 			};
 
-			TMap<uint32, FNodePair> NewIdNodeMap;
+			TMap<FGuid, FNodePair> NewIdNodeMap;
 			TArray<FNodeHandle> NodeHandles = GraphHandle->GetNodes();
 			for (FNodeHandle& NodeHandle : NodeHandles)
 			{
@@ -341,7 +341,7 @@ namespace Metasound
 				NewIdNodeMap.Add(NodeHandle->GetID(), FNodePair { NodeHandle, NewNode });
 			}
 
-			for (const TPair<uint32, FNodePair>& IdNodePair : NewIdNodeMap)
+			for (const TPair<FGuid, FNodePair>& IdNodePair : NewIdNodeMap)
 			{
 				UEdGraphNode* GraphNode = IdNodePair.Value.GraphNode;
 				check(GraphNode);
@@ -528,7 +528,7 @@ namespace Metasound
 				UMetasoundEditorGraphNode* EditorNode = nullptr;
 				FNodeHandle Node = INodeController::GetInvalidHandle();;
 			};
-			TMap<int32, FNodePair> PairedNodes;
+			TMap<FGuid, FNodePair> PairedNodes;
 
 			// Reverse iterate so paired nodes can safely be removed from the array.
 			for (int32 i = Nodes.Num() - 1; i >= 0; i--)
@@ -573,7 +573,7 @@ namespace Metasound
 			}
 
 			// Synchronize pins on node pairs.
-			for (const TPair<int32, FNodePair>& IdNodePair : PairedNodes)
+			for (const TPair<FGuid, FNodePair>& IdNodePair : PairedNodes)
 			{
 				UMetasoundEditorGraphNode* EditorNode = IdNodePair.Value.EditorNode;
 				bIsEditorGraphDirty |= SynchronizeNodePins(*IdNodePair.Value.EditorNode, IdNodePair.Value.Node);
@@ -682,7 +682,7 @@ namespace Metasound
 			TArray<UMetasoundEditorGraphNode*> EditorNodes;
 
 			// Cache map of editor nodes indexed by node id.
-			TMap<int32, UMetasoundEditorGraphNode*> EditorNodesByID;
+			TMap<FGuid, UMetasoundEditorGraphNode*> EditorNodesByID;
 			for (UEdGraphNode* EdGraphNode : EditorGraph->Nodes)
 			{
 				UMetasoundEditorGraphNode* MetasoundEditorNode = Cast<UMetasoundEditorGraphNode>(EdGraphNode);
diff --git a/Engine/Plugins/Runtime/Metasound/Source/MetasoundEditor/Private/MetasoundEditorGraphNode.cpp b/Engine/Plugins/Runtime/Metasound/Source/MetasoundEditor/Private/MetasoundEditorGraphNode.cpp
--- a/Engine/Plugins/Runtime/Metasound/Source/MetasoundEditor/Private/MetasoundEditorGraphNode.cpp
+++ b/Engine/Plugins/Runtime/Metasound/Source/MetasoundEditor/Private/MetasoundEditorGraphNode.cpp
@@ -205,12 +205,12 @@ FString UMetasoundEditorGraphNode::GetDocumentationLink() const
 	return TEXT("Shared/GraphNodes/Metasound");
 }
 
-void UMetasoundEditorGraphNode::SetNodeID(uint32 InNodeID)
+void UMetasoundEditorGraphNode::SetNodeID(FGuid InNodeID)
 {
 	NodeID = InNodeID;
 }
 
-uint32 UMetasoundEditorGraphNode::GetNodeID() const
+FGuid UMetasoundEditorGraphNode::GetNodeID() const
 {
 	return NodeID;
 }
diff --git a/Engine/Plugins/Runtime/Metasound/Source/MetasoundEditor/Private/MetasoundEditorGraphSchema.cpp b/Engine/Plugins/Runtime/Metasound/Source/MetasoundEditor/Private/MetasoundEditorGraphSchema.cpp
--- a/Engine/Plugins/Runtime/Metasound/Source/MetasoundEditor/Private/MetasoundEditorGraphSchema.cpp
+++ b/Engine/Plugins/Runtime/Metasound/Source/MetasoundEditor/Private/MetasoundEditorGraphSchema.cpp
@@ -533,7 +533,7 @@ void UMetasoundEditorGraphSchema::BreakNodeLinks(UEdGraphNode& TargetNode) const
 	Super::BreakNodeLinks(TargetNode);
 
 	FNodeHandle NodeHandle = CastChecked<UMetasoundEditorGraphNode>(&TargetNode)->GetNodeHandle();
-	const uint32 NodeID = NodeHandle->GetID();
+	const FGuid NodeID = NodeHandle->GetID();
 
 	FGraphHandle GraphHandle = CastChecked<UMetasoundEditorGraphNode>(&TargetNode)->GetRootGraphHandle();
 	TArray<FNodeHandle> AllNodes = GraphHandle->GetNodes();
diff --git a/Engine/Plugins/Runtime/Metasound/Source/MetasoundEditor/Public/MetasoundEditorGraphNode.h b/Engine/Plugins/Runtime/Metasound/Source/MetasoundEditor/Public/MetasoundEditorGraphNode.h
--- a/Engine/Plugins/Runtime/Metasound/Source/MetasoundEditor/Public/MetasoundEditorGraphNode.h
+++ b/Engine/Plugins/Runtime/Metasound/Source/MetasoundEditor/Public/MetasoundEditorGraphNode.h
@@ -8,6 +8,7 @@
 #include "MetasoundFrontend.h"
 #include "MetasoundFrontendController.h"
 #include "MetasoundFrontendRegistries.h"
+#include "Misc/Guid.h"
 
 #include "MetasoundEditorGraphNode.generated.h"
 
@@ -63,10 +64,10 @@ public:
 	Metasound::Frontend::FGraphHandle GetRootGraphHandle() const;
 	Metasound::Frontend::FNodeHandle GetNodeHandle() const;
 
-	void SetNodeID(uint32 InNodeID);
-	uint32 GetNodeID() const;
+	void SetNodeID(FGuid InNodeID);
+	FGuid GetNodeID() const;
 
 protected:
 	UPROPERTY()
-	uint32 NodeID = INDEX_NONE;
+	FGuid NodeID;
 };
diff --git a/Engine/Plugins/Runtime/Metasound/Source/MetasoundEngine/Private/MetasoundSource.cpp b/Engine/Plugins/Runtime/Metasound/Source/MetasoundEngine/Private/MetasoundSource.cpp
--- a/Engine/Plugins/Runtime/Metasound/Source/MetasoundEngine/Private/MetasoundSource.cpp
+++ b/Engine/Plugins/Runtime/Metasound/Source/MetasoundEngine/Private/MetasoundSource.cpp
@@ -228,7 +228,7 @@ const FMetasoundFrontendArchetype& UMetasoundSource::GetBaseArchetype()
 		OnPlayTrigger.Metadata.DisplayName = FText::FromString(OnPlayTrigger.Name);
 		OnPlayTrigger.TypeName = Metasound::Frontend::GetDataTypeName<Metasound::FBop>();
 		OnPlayTrigger.Metadata.Description = LOCTEXT("OnPlayTriggerToolTip", "Trigger executed when this source is played.");
-		OnPlayTrigger.PointIDs.Add(2);
+		OnPlayTrigger.PointIDs.Add(FGuid::NewGuid());
 
 		Archetype.Interface.Inputs.Add(OnPlayTrigger);
 
@@ -237,7 +237,7 @@ const FMetasoundFrontendArchetype& UMetasoundSource::GetBaseArchetype()
 		OnFinished.Metadata.DisplayName = FText::FromString(OnFinished.Name);
 		OnFinished.TypeName = Metasound::Frontend::GetDataTypeName<Metasound::FBop>();
 		OnFinished.Metadata.Description = LOCTEXT("OnFinishedToolTip", "Bop executed to initiate stopping the source.");
-		OnFinished.PointIDs.Add(3);
+		OnFinished.PointIDs.Add(FGuid::NewGuid());
 
 		Archetype.Interface.Outputs.Add(OnFinished);
 
@@ -268,7 +268,7 @@ const FMetasoundFrontendArchetype& UMetasoundSource::GetMonoSourceArchetype()
 		GeneratedAudio.TypeName = Metasound::Frontend::GetDataTypeName<Metasound::FMonoAudioFormat>();
 		GeneratedAudio.Metadata.DisplayName = LOCTEXT("GeneratedMono", "Mono");
 		GeneratedAudio.Metadata.Description = LOCTEXT("GeneratedAudioToolTip", "The resulting output audio from this source.");
-		GeneratedAudio.PointIDs.Add(4);
+		GeneratedAudio.PointIDs.Add(FGuid::NewGuid());
 
 		Archetype.Interface.Outputs.Add(GeneratedAudio);
 
@@ -292,7 +292,7 @@ const FMetasoundFrontendArchetype& UMetasoundSource::GetStereoSourceArchetype()
 		GeneratedAudio.TypeName = Metasound::Frontend::GetDataTypeName<Metasound::FStereoAudioFormat>();
 		GeneratedAudio.Metadata.DisplayName = LOCTEXT("GeneratedStereo", "Stereo");
 		GeneratedAudio.Metadata.Description = LOCTEXT("GeneratedAudioToolTip", "The resulting output audio from this source.");
-		GeneratedAudio.PointIDs.Add(5);
+		GeneratedAudio.PointIDs.Add(FGuid::NewGuid());
 
 		Archetype.Interface.Outputs.Add(GeneratedAudio);
 
diff --git a/Engine/Plugins/Runtime/Metasound/Source/MetasoundFrontend/Private/MetasoundFrontend.cpp b/Engine/Plugins/Runtime/Metasound/Source/MetasoundFrontend/Private/MetasoundFrontend.cpp
--- a/Engine/Plugins/Runtime/Metasound/Source/MetasoundFrontend/Private/MetasoundFrontend.cpp
+++ b/Engine/Plugins/Runtime/Metasound/Source/MetasoundFrontend/Private/MetasoundFrontend.cpp
@@ -92,8 +92,6 @@ namespace Metasound
 				},
 			};
 
-			int32 RunningPointID = 0;
-
 			// External metasounds aren't dependent on any other nodes by definition, so all we need to do
 			// is populate the Input and Output sets.
 			for (auto& InputTuple : InNodeMetadata.DefaultInterface.GetInputInterface())
@@ -103,7 +101,7 @@ namespace Metasound
 
 				ClassInput.Name = InputTuple.Value.GetVertexName();
 				ClassInput.TypeName = InputTuple.Value.GetDataTypeName();
-				ClassInput.PointIDs.Add(RunningPointID++);
+				ClassInput.PointIDs.Add(FGuid::NewGuid());
 				ClassInput.Metadata.DisplayName = FText::FromString(InputTuple.Value.GetVertexName());
 				ClassInput.Metadata.Description = InputTuple.Value.GetDescription();
 
@@ -117,7 +115,7 @@ namespace Metasound
 
 				ClassOutput.Name = OutputTuple.Value.GetVertexName();
 				ClassOutput.TypeName = OutputTuple.Value.GetDataTypeName();
-				ClassOutput.PointIDs.Add(RunningPointID++);
+				ClassOutput.PointIDs.Add(FGuid::NewGuid());
 				ClassOutput.Metadata.DisplayName = FText::FromString(OutputTuple.Value.GetVertexName());
 				ClassOutput.Metadata.Description = OutputTuple.Value.GetDescription();
 
diff --git a/Engine/Plugins/Runtime/Metasound/Source/MetasoundFrontend/Private/MetasoundFrontendDocument.cpp b/Engine/Plugins/Runtime/Metasound/Source/MetasoundFrontend/Private/MetasoundFrontendDocument.cpp
--- a/Engine/Plugins/Runtime/Metasound/Source/MetasoundFrontend/Private/MetasoundFrontendDocument.cpp
+++ b/Engine/Plugins/Runtime/Metasound/Source/MetasoundFrontend/Private/MetasoundFrontendDocument.cpp
@@ -120,7 +120,7 @@ FMetasoundFrontendClassInput::FMetasoundFrontendClassInput(const FMetasoundFront
 {
 	EMetasoundFrontendLiteralType DefaultType = Metasound::Frontend::GetMetasoundLiteralType(FMetasoundFrontendRegistryContainer::Get()->GetDesiredLiteralTypeForDataType(InOther.TypeName));
 
-	for (int32 PointID : PointIDs)
+	for (const FGuid& PointID : PointIDs)
 	{
 		FMetasoundFrontendVertexLiteral Default;
 		Default.PointID = PointID;
@@ -136,7 +136,7 @@ FMetasoundFrontendGraphClass::FMetasoundFrontendGraphClass()
 
 FMetasoundFrontendDocument::FMetasoundFrontendDocument()
 {
-	RootGraph.ID = 0;
+	RootGraph.ID = FGuid::NewGuid();
 	RootGraph.Metadata.Type = EMetasoundFrontendClassType::Graph;
 }
 
diff --git a/Engine/Plugins/Runtime/Metasound/Source/MetasoundFrontend/Private/MetasoundFrontendGraph.cpp b/Engine/Plugins/Runtime/Metasound/Source/MetasoundFrontend/Private/MetasoundFrontendGraph.cpp
--- a/Engine/Plugins/Runtime/Metasound/Source/MetasoundFrontend/Private/MetasoundFrontendGraph.cpp
+++ b/Engine/Plugins/Runtime/Metasound/Source/MetasoundFrontend/Private/MetasoundFrontendGraph.cpp
@@ -18,7 +18,7 @@ namespace Metasound
 	{
 	}
 
-	void FFrontendGraph::AddInputNode(int32 InDependencyId, int32 InIndex, const FVertexKey& InVertexKey, TUniquePtr<INode> InNode)
+	void FFrontendGraph::AddInputNode(FGuid InDependencyId, int32 InIndex, const FVertexKey& InVertexKey, TUniquePtr<INode> InNode)
 	{
 		if (InNode.IsValid())
 		{
@@ -32,7 +32,7 @@ namespace Metasound
 		}
 	}
 
-	void FFrontendGraph::AddOutputNode(int32 InNodeID, int32 InIndex, const FVertexKey& InVertexKey, TUniquePtr<INode> InNode)
+	void FFrontendGraph::AddOutputNode(FGuid InNodeID, int32 InIndex, const FVertexKey& InVertexKey, TUniquePtr<INode> InNode)
 	{
 		if (InNode.IsValid())
 		{
@@ -47,7 +47,7 @@ namespace Metasound
 	}
 
 	/** Store a node on this graph. */
-	void FFrontendGraph::AddNode(int32 InNodeID, TUniquePtr<INode> InNode)
+	void FFrontendGraph::AddNode(FGuid InNodeID, TUniquePtr<INode> InNode)
 	{
 		if (InNode.IsValid())
 		{
@@ -59,7 +59,7 @@ namespace Metasound
 		}
 	}
 
-	const INode* FFrontendGraph::FindNode(int32 InNodeID) const
+	const INode* FFrontendGraph::FindNode(FGuid InNodeID) const
 	{
 		INode* const* NodePtr = NodeMap.Find(InNodeID);
 
@@ -171,13 +171,13 @@ namespace Metasound
 				}
 				else
 				{
-					UE_LOG(LogMetasound, Error, TEXT("Cannot create input node [NodeID:%d]. [Vertex:%s] cannot be constructed with the provided literal type."), InNode.ID, *InputVertex.Name);
+					UE_LOG(LogMetasound, Error, TEXT("Cannot create input node [NodeID:%s]. [Vertex:%s] cannot be constructed with the provided literal type."), *InNode.ID.ToString(), *InputVertex.Name);
 				}
 			}
 		}
 		else
 		{
-			UE_LOG(LogMetasound, Error, TEXT("Cannot create input node [NodeID:%d]. No default literal set for input node."), InNode.ID);
+			UE_LOG(LogMetasound, Error, TEXT("Cannot create input node [NodeID:%s]. No default literal set for input node."), *InNode.ID.ToString());
 		}
 
 		return TUniquePtr<INode>(nullptr);
@@ -213,7 +213,7 @@ namespace Metasound
 
 		InitData.InstanceName.Append(InNode.Name);
 		InitData.InstanceName.AppendChar('_');
-		InitData.InstanceName.AppendInt(InNode.ID);
+		InitData.InstanceName.Append(InNode.ID.ToString());
 
 		// Copy over our initialization params.
 		/*
@@ -283,6 +283,57 @@ namespace Metasound
 		return nullptr;
 	}
 
+		/*
+	const FMetasoundFrontendVertexLiteral* FFrontendGraphBuilder::FindInputLiteralForInput(FGuid InPointID, const FMetasoundFrontendNode& InNode, const FMetasoundFrontendClass& InNodeClass) const
+	{
+		// TODO: Need utility for matching class input to node input. Will need to be cognisant of input type [Point, Array, ....]
+
+		// Default value priority is:
+		// 1. A value set directly on the node
+		// 3. A default value on the node class.
+
+		const FMetasoundFrontendVertexLiteral* VertexLiteral = nullptr;
+
+		// Check for default value directly on node.
+		if (ensure(InNode.Interface.Inputs.Num() == 1))
+		{
+			const FMetasoundFrontendVertex& InputVertex = InNode.Interface.Inputs[0];
+
+			// Currently only support one point on input vertex.
+			if (ensure(InputVertex.PointIDs.Num() == 1))
+			{
+				VertexLiteral = InNode.InputLiterals.FindByPredicate(
+					[&](const FMetasoundFrontendVertexLiteral& InVertexLiteral)
+					{
+						return InVertexLiteral.PointID == InputVertex.PointIDs[0];
+					}
+				);
+			}
+		}
+
+		// Check for default value on input node class
+		if (nullptr == VertexLiteral && ensure(InNodeClass.Interface.Inputs.Num() == 1))
+		{
+			const FMetasoundFrontendClassInput& InputNodeClassInput = InNodeClass.Interface.Inputs[0];
+
+			if (ensure(InputNodeClassInput.PointIDs.Num() == 1))
+			{
+				int32 PointID = InputNodeClassInput.PointIDs[0];
+
+				VertexLiteral = InputNodeClassInput.Defaults.FindByPredicate(
+					[&](const FMetasoundFrontendVertexLiteral& InVertexLiteral)
+					{
+						return InVertexLiteral.PointID == PointID;
+					}
+				);
+			}
+		}
+
+		return VertexLiteral;
+		return nullptr;
+	}
+		*/
+
 	const FMetasoundFrontendVertexLiteral* FFrontendGraphBuilder::FindInputLiteralForInputNode(const FMetasoundFrontendNode& InInputNode, const FMetasoundFrontendClass& InInputNodeClass, const FMetasoundFrontendClassInput& InOwningGraphClassInput) const
 	{
 		// Default value priority is:
@@ -314,7 +365,7 @@ namespace Metasound
 		{
 			if (ensure(InOwningGraphClassInput.PointIDs.Num() == 1))
 			{
-				int32 PointID = InOwningGraphClassInput.PointIDs[0];
+				const FGuid& PointID = InOwningGraphClassInput.PointIDs[0];
 
 				VertexLiteral = InOwningGraphClassInput.Defaults.FindByPredicate(
 					[&](const FMetasoundFrontendVertexLiteral& InVertexLiteral)
@@ -332,7 +383,7 @@ namespace Metasound
 
 			if (ensure(InputNodeClassInput.PointIDs.Num() == 1))
 			{
-				int32 PointID = InputNodeClassInput.PointIDs[0];
+				const FGuid& PointID = InputNodeClassInput.PointIDs[0];
 
 				VertexLiteral = InputNodeClassInput.Defaults.FindByPredicate(
 					[&](const FMetasoundFrontendVertexLiteral& InVertexLiteral)
@@ -347,7 +398,7 @@ namespace Metasound
 	}
 
 	// TODO: add errors here. Most will be a "PromptIfMissing"...
-	void FFrontendGraphBuilder::AddNodesToGraph(const FMetasoundFrontendGraphClass& InGraphClass, const TMap<int32, const FMetasoundFrontendClass*>& InClasses, FFrontendGraph& OutGraph) const
+	void FFrontendGraphBuilder::AddNodesToGraph(const FMetasoundFrontendGraphClass& InGraphClass, const TMap<FGuid, const FMetasoundFrontendClass*>& InClasses, FFrontendGraph& OutGraph) const
 	{
 		for (const FMetasoundFrontendNode& Node : InGraphClass.Graph.Nodes)
 		{
@@ -369,7 +420,7 @@ namespace Metasound
 							}
 							else
 							{
-								UE_LOG(LogMetasound, Error, TEXT("Failed to match input node [NodeID:%d, NodeName:%s] to owning graph [ClassID:%d] output."), Node.ID, *Node.Name, InGraphClass.ID);
+								UE_LOG(LogMetasound, Error, TEXT("Failed to match input node [NodeID:%s, NodeName:%s] to owning graph [ClassID:%s] output."), *Node.ID.ToString(), *Node.Name, *InGraphClass.ID.ToString());
 							}
 						}
 
@@ -386,7 +437,7 @@ namespace Metasound
 							}
 							else
 							{
-								UE_LOG(LogMetasound, Error, TEXT("Failed to match output node [NodeID:%d, NodeName:%s] to owning graph [ClassID:%d] output."), Node.ID, *Node.Name, InGraphClass.ID);
+								UE_LOG(LogMetasound, Error, TEXT("Failed to match output node [NodeID:%s, NodeName:%s] to owning graph [ClassID:%s] output."), *Node.ID.ToString(), *Node.Name, *InGraphClass.ID.ToString());
 							}
 						}
 
@@ -410,7 +461,7 @@ namespace Metasound
 		};
 
 		// TODO: add support for array vertices.
-		typedef TTuple<int32, int32> FNodeIDPointID;
+		typedef TTuple<FGuid, FGuid> FNodeIDPointID;
 
 		TMap<FNodeIDPointID, FCoreNodeAndFrontendVertex> NodeSourcesByID;
 		TMap<FNodeIDPointID, FCoreNodeAndFrontendVertex> NodeDestinationsByID;
@@ -421,13 +472,13 @@ namespace Metasound
 			const INode* CoreNode = OutGraph.FindNode(Node.ID);
 			if (nullptr == CoreNode)
 			{
-				UE_LOG(LogMetasound, Display, TEXT("Could not find referenced node [ID:%d]"), Node.ID);
+				UE_LOG(LogMetasound, Display, TEXT("Could not find referenced node [NodeID:%s]"), *Node.ID.ToString());
 				continue;
 			}
 
 			for (const FMetasoundFrontendVertex& Vertex : Node.Interface.Inputs)
 			{
-				for (int32 PointID : Vertex.PointIDs)
+				for (const FGuid& PointID : Vertex.PointIDs)
 				{
 					NodeDestinationsByID.Add(FNodeIDPointID(Node.ID, PointID), FCoreNodeAndFrontendVertex({CoreNode, &Vertex}));
 				}
@@ -435,7 +486,7 @@ namespace Metasound
 
 			for (const FMetasoundFrontendVertex& Vertex : Node.Interface.Outputs)
 			{
-				for (int32 PointID : Vertex.PointIDs)
+				for (const FGuid& PointID : Vertex.PointIDs)
 				{
 					NodeSourcesByID.Add(FNodeIDPointID(Node.ID, PointID), FCoreNodeAndFrontendVertex({CoreNode, &Vertex}));
 				}
@@ -451,14 +502,14 @@ namespace Metasound
 			if (nullptr == DestinationNodeAndVertex)
 			{
 				// TODO: bubble up error
-				UE_LOG(LogMetasound, Error, TEXT("Failed to add edge. Could not find destination [NodeID:%d, PointID:%d]"), Edge.ToNodeID, Edge.ToPointID);
+				UE_LOG(LogMetasound, Error, TEXT("Failed to add edge. Could not find destination [NodeID:%s, PointID:%s]"), *Edge.ToNodeID.ToString(), *Edge.ToPointID.ToString());
 				continue;
 			}
 
 			if (nullptr == DestinationNodeAndVertex->Node)
 			{
 				// TODO: bubble up error
-				UE_LOG(LogMetasound, Warning, TEXT("Skipping edge. Null destination node [ID:%d]"), Edge.ToNodeID);
+				UE_LOG(LogMetasound, Warning, TEXT("Skipping edge. Null destination node [NodeID:%s]"), *Edge.ToNodeID.ToString());
 				continue;
 			}
 
@@ -467,14 +518,14 @@ namespace Metasound
 
 			if (nullptr == SourceNodeAndVertex)
 			{
-				UE_LOG(LogMetasound, Error, TEXT("Failed to add edge. Could not find source [NodeID:%d, PointID:%d]"), Edge.FromNodeID, Edge.FromPointID);
+				UE_LOG(LogMetasound, Error, TEXT("Failed to add edge. Could not find source [NodeID:%s, PointID:%s]"), *Edge.FromNodeID.ToString(), *Edge.FromPointID.ToString());
 				continue;
 			}
 
 			if (nullptr == SourceNodeAndVertex->Node)
 			{
 				// TODO: bubble up error
-				UE_LOG(LogMetasound, Warning, TEXT("Skipping edge. Null source node [ID:%d]"), Edge.FromNodeID);
+				UE_LOG(LogMetasound, Warning, TEXT("Skipping edge. Null source node [NodeID:%s]"), *Edge.FromNodeID.ToString());
 				continue;
 			}
 
@@ -488,7 +539,7 @@ namespace Metasound
 
 			if (!bSuccess)
 			{
-				UE_LOG(LogMetasound, Error, TEXT("Failed to connect edge from [NodeID:%d, PointID:%d] to [NodeID:%d, PointID:%d]"), Edge.FromNodeID, Edge.FromPointID, Edge.ToNodeID, Edge.ToPointID);
+				UE_LOG(LogMetasound, Error, TEXT("Failed to connect edge from [NodeID:%s, PointID:%s] to [NodeID:%s, PointID:%s]"), *Edge.FromNodeID.ToString(), *Edge.FromPointID.ToString(), *Edge.ToNodeID.ToString(), *Edge.ToPointID.ToString());
 			}
 		}
 	}
@@ -522,7 +573,7 @@ namespace Metasound
 		}
 
 		// All the dependencies are met 
-		TSet<int32> AvailableDependencies;
+		TSet<FGuid> AvailableDependencies;
 		Algo::Transform(InDependencies, AvailableDependencies, [](const FMetasoundFrontendClass& InDesc) { return InDesc.ID; });
 
 		auto IsDependencyMet = [&](const FMetasoundFrontendNode& InNode) 
@@ -546,7 +597,7 @@ namespace Metasound
 		}
 
 
-		TMap<int32, const FMetasoundFrontendClass*> ClassMap;
+		TMap<FGuid, const FMetasoundFrontendClass*> ClassMap;
 
 		for (const FMetasoundFrontendClass& ExtClass : InDependencies)
 		{
diff --git a/Engine/Plugins/Runtime/Metasound/Source/MetasoundFrontend/Private/MetasoundFrontendGraph.h b/Engine/Plugins/Runtime/Metasound/Source/MetasoundFrontend/Private/MetasoundFrontendGraph.h
--- a/Engine/Plugins/Runtime/Metasound/Source/MetasoundFrontend/Private/MetasoundFrontendGraph.h
+++ b/Engine/Plugins/Runtime/Metasound/Source/MetasoundFrontend/Private/MetasoundFrontendGraph.h
@@ -33,7 +33,7 @@ namespace Metasound
 			 * @param InVertexKey - The key for the graph input vertex.
 			 * @param InNode - A unique pointer to an input node. 
 			 */
-			void AddInputNode(int32 InNodeID, int32 InIndex, const FVertexKey& InVertexKey, TUniquePtr<INode> InNode);
+			void AddInputNode(FGuid InNodeID, int32 InIndex, const FVertexKey& InVertexKey, TUniquePtr<INode> InNode);
 
 			/** Add an output node to this graph.
 			 *
@@ -42,14 +42,14 @@ namespace Metasound
 			 * @param InVertexKey - The key for the graph output vertex.
 			 * @param InNode - A unique pointer to an output node. 
 			 */
-			void AddOutputNode(int32 InNodeID, int32 InIndex, const FVertexKey& InVertexKey, TUniquePtr<INode> InNode);
+			void AddOutputNode(FGuid InNodeID, int32 InIndex, const FVertexKey& InVertexKey, TUniquePtr<INode> InNode);
 
 			/** Store a node on this graph. 
 			 *
 			 * @param InNodeID - The NodeID related to the parent FMetasoundFrontendClass.
 			 * @param InNode - A unique pointer to a node. 
 			 */
-			void AddNode(int32 InNodeID, TUniquePtr<INode> InNode);
+			void AddNode(FGuid InNodeID, TUniquePtr<INode> InNode);
 
 			/** Retrieve node by dependency ID.
 			 *
@@ -57,7 +57,7 @@ namespace Metasound
 			 *
 			 * @return Pointer to the Node if it is stored on this graph. nullptr otherwise. 
 			 */
-			const INode* FindNode(int32 InNodeID) const;
+			const INode* FindNode(FGuid InNodeID) const;
 
 			/** Retrieve node by input index.
 			 *
@@ -86,7 +86,7 @@ namespace Metasound
 			TMap<int32, INode*> InputNodes;
 			TMap<int32, INode*> OutputNodes;
 
-			TMap<int32, INode*> NodeMap;
+			TMap<FGuid, INode*> NodeMap;
 			TSet<const INode*> StoredNodes;
 			TArray<TUniquePtr<INode>> Storage;
 	};
@@ -122,6 +122,7 @@ namespace Metasound
 		const FMetasoundFrontendClassOutput* FindClassOutputForOutputNode(const FMetasoundFrontendGraphClass& InOwningGraph, const FMetasoundFrontendNode& InOutputNode, int32& OutClassOutputIndex) const;
 
 		const FMetasoundFrontendVertexLiteral* FindInputLiteralForInputNode(const FMetasoundFrontendNode& InInputNode, const FMetasoundFrontendClass& InInputNodeClass, const FMetasoundFrontendClassInput& InOwningGraphClassInput) const;
+		//const FMetasoundFrontendVertexLiteral* FindInputLiteralForInput(FGuid InPointID, const FMetasoundFrontendNode& InNode, const FMetasoundFrontendClass& InNodeClass) const;
 
 		TUniquePtr<INode> CreateInputNode(const FMetasoundFrontendNode& InNode, const FMetasoundFrontendClass& InClass, const FMetasoundFrontendClassInput& InOwningGraphClassInput) const;
 
@@ -130,7 +131,7 @@ namespace Metasound
 		TUniquePtr<INode> CreateExternalNode(const FMetasoundFrontendNode& InNode, const FMetasoundFrontendClass& InClass) const;
 
 		// TODO: add errors here. Most will be a "PromptIfMissing"...
-		void AddNodesToGraph(const FMetasoundFrontendGraphClass& InGraphClass, const TMap<int32, const FMetasoundFrontendClass*>& InClasses, FFrontendGraph& OutGraph) const;
+		void AddNodesToGraph(const FMetasoundFrontendGraphClass& InGraphClass, const TMap<FGuid, const FMetasoundFrontendClass*>& InClasses, FFrontendGraph& OutGraph) const;
 
 		void AddEdgesToGraph(const FMetasoundFrontendGraph& InGraphDescription, FFrontendGraph& OutGraph) const;
 	};
diff --git a/Engine/Plugins/Runtime/Metasound/Source/MetasoundFrontend/Private/MetasoundFrontendInvalidController.h b/Engine/Plugins/Runtime/Metasound/Source/MetasoundFrontend/Private/MetasoundFrontendInvalidController.h
--- a/Engine/Plugins/Runtime/Metasound/Source/MetasoundFrontend/Private/MetasoundFrontendInvalidController.h
+++ b/Engine/Plugins/Runtime/Metasound/Source/MetasoundFrontend/Private/MetasoundFrontendInvalidController.h
@@ -43,13 +43,13 @@ namespace Metasound
 			}
 
 			bool IsValid() const override { return false; }
-			int32 GetID() const override { return Metasound::FrontendInvalidID; }
+			FGuid GetID() const override { return Metasound::FrontendInvalidID; }
 			const FName& GetDataType() const override { return MetasoundFrontendInvalidControllerPrivate::GetInvalid<FName>(); }  
 			const FString& GetName() const override { return MetasoundFrontendInvalidControllerPrivate::GetInvalid<FString>(); }
 			const FText& GetDisplayName() const override { return FText::GetEmpty(); }
 			const FText& GetTooltip() const override { return FText::GetEmpty(); }
 
-			int32 GetOwningNodeID() const override { return Metasound::FrontendInvalidID; }
+			FGuid GetOwningNodeID() const override { return Metasound::FrontendInvalidID; }
 			TSharedRef<INodeController> GetOwningNode() override;
 			TSharedRef<const INodeController> GetOwningNode() const override;
 
@@ -59,6 +59,10 @@ namespace Metasound
 			bool ConnectWithConverterNode(IInputController& InController, const FConverterNodeInfo& InNodeClassName) override { return false; }
 			bool Disconnect(IInputController& InController) override { return false; }
 
+		protected:
+			FDocumentAccess ShareAccess() override { return FDocumentAccess(); }
+			FConstDocumentAccess ShareAccess() const override { return FConstDocumentAccess(); }
+
 		};
 
 		/** FInvalidInputController is a controller which is always invalid. 
@@ -78,14 +82,14 @@ namespace Metasound
 			}
 
 			bool IsValid() const override { return false; }
-			int32 GetID() const override { return Metasound::FrontendInvalidID; }
+			FGuid GetID() const override { return Metasound::FrontendInvalidID; }
 			bool IsConnected() const override { return false; }
 			const FName& GetDataType() const override { return MetasoundFrontendInvalidControllerPrivate::GetInvalid<FName>(); }  
 			const FString& GetName() const override { return MetasoundFrontendInvalidControllerPrivate::GetInvalid<FString>(); }
 			const FText& GetDisplayName() const override { return FText::GetEmpty(); }
 			const FText& GetTooltip() const override { return FText::GetEmpty(); }
 
-			int32 GetOwningNodeID() const override { return Metasound::FrontendInvalidID; }
+			FGuid GetOwningNodeID() const override { return Metasound::FrontendInvalidID; }
 			TSharedRef<INodeController> GetOwningNode() override;
 			TSharedRef<const INodeController> GetOwningNode() const override;
 
@@ -97,6 +101,9 @@ namespace Metasound
 			virtual bool ConnectWithConverterNode(IOutputController& InController, const FConverterNodeInfo& InNodeClassName) override { return false; }
 			virtual bool Disconnect(IOutputController& InController) override { return false; }
 			virtual bool Disconnect() override { return false; }
+		protected:
+			FDocumentAccess ShareAccess() override { return FDocumentAccess(); }
+			FConstDocumentAccess ShareAccess() const override { return FConstDocumentAccess(); }
 		};
 
 		/** FInvalidNodeController is a controller which is always invalid. 
@@ -127,18 +134,18 @@ namespace Metasound
 			TArray<FConstInputHandle> GetInputsWithVertexName(const FString& InName) const override { return TArray<FConstInputHandle>(); }
 			TArray<FOutputHandle> GetOutputsWithVertexName(const FString& InName) override { return TArray<FOutputHandle>(); }
 			TArray<FConstOutputHandle> GetOutputsWithVertexName(const FString& InName) const override { return TArray<FConstOutputHandle>(); }
-			TSharedRef<IInputController> GetInputWithID(int32 InPointID) override { return FInvalidInputController::GetInvalid(); }
-			TSharedRef<IOutputController> GetOutputWithID(int32 InPointID) override { return FInvalidOutputController::GetInvalid(); }
-			TSharedRef<const IInputController> GetInputWithID(int32 InPointID) const override { return FInvalidInputController::GetInvalid(); }
-			TSharedRef<const IOutputController> GetOutputWithID(int32 InPointID) const override { return FInvalidOutputController::GetInvalid(); }
+			TSharedRef<IInputController> GetInputWithID(FGuid InPointID) override { return FInvalidInputController::GetInvalid(); }
+			TSharedRef<IOutputController> GetOutputWithID(FGuid InPointID) override { return FInvalidOutputController::GetInvalid(); }
+			TSharedRef<const IInputController> GetInputWithID(FGuid InPointID) const override { return FInvalidInputController::GetInvalid(); }
+			TSharedRef<const IOutputController> GetOutputWithID(FGuid InPointID) const override { return FInvalidOutputController::GetInvalid(); }
 
 			bool CanAddInput(const FString& InVertexName) const override { return false; }
 			FInputHandle AddInput(const FString& InVertexName, const FMetasoundFrontendLiteral* InDefault) override { return FInvalidInputController::GetInvalid(); }
-			bool RemoveInput(int32 InPointID) override { return false; }
+			bool RemoveInput(FGuid InPointID) override { return false; }
 
 			bool CanAddOutput(const FString& InVertexName) const override { return false; }
 			FInputHandle AddOutput(const FString& InVertexName, const FMetasoundFrontendLiteral* InDefault) override { return FInvalidInputController::GetInvalid(); }
-			bool RemoveOutput(int32 InPointID) override { return false; }
+			bool RemoveOutput(FGuid InPointID) override { return false; }
 
 			EMetasoundFrontendClassType GetClassType() const override { return EMetasoundFrontendClassType::Invalid; }
 			const FString& GetClassName() const override { return MetasoundFrontendInvalidControllerPrivate::GetInvalid<FString>(); }
@@ -149,14 +156,17 @@ namespace Metasound
 			TSharedRef<IGraphController> AsGraph() override;
 			TSharedRef<const IGraphController> AsGraph() const override;
 
-			int32 GetID() const override { return Metasound::FrontendInvalidID; }
-			int32 GetClassID() const override { return Metasound::FrontendInvalidID; }
+			FGuid GetID() const override { return Metasound::FrontendInvalidID; }
+			FGuid GetClassID() const override { return Metasound::FrontendInvalidID; }
 
-			int32 GetOwningGraphClassID() const override { return Metasound::FrontendInvalidID; }
+			FGuid GetOwningGraphClassID() const override { return Metasound::FrontendInvalidID; }
 			TSharedRef<IGraphController> GetOwningGraph() override;
 			TSharedRef<const IGraphController> GetOwningGraph() const override;
 
 			const FString& GetNodeName() const override { return MetasoundFrontendInvalidControllerPrivate::GetInvalid<FString>(); }
+		protected:
+			FDocumentAccess ShareAccess() override { return FDocumentAccess(); }
+			FConstDocumentAccess ShareAccess() const override { return FConstDocumentAccess(); }
 		};
 
 		/** FInvalidGraphController is a controller which is always invalid. 
@@ -177,16 +187,16 @@ namespace Metasound
 
 			bool IsValid() const override { return false; }
 
-			int32 GetNewPointID() const override { return Metasound::FrontendInvalidID; }
+			FGuid GetNewPointID() const override { return Metasound::FrontendInvalidID; }
 			TArray<FString> GetInputVertexNames() const override { return TArray<FString>(); }
 			TArray<FString> GetOutputVertexNames() const override { return TArray<FString>(); }
 
 			TArray<TSharedRef<INodeController>> GetNodes() override { return TArray<TSharedRef<INodeController>>(); }
 			TArray<TSharedRef<const INodeController>> GetConstNodes() const override { return TArray<TSharedRef<const INodeController>>(); }
 
-			int32 GetClassID() const override { return Metasound::FrontendInvalidID; }
-			TSharedRef<const INodeController> GetNodeWithID(int32 InNodeID) const override { return FInvalidNodeController::GetInvalid(); }
-			TSharedRef<INodeController> GetNodeWithID(int32 InNodeID) override { return FInvalidNodeController::GetInvalid(); }
+			FGuid GetClassID() const override { return Metasound::FrontendInvalidID; }
+			TSharedRef<const INodeController> GetNodeWithID(FGuid InNodeID) const override { return FInvalidNodeController::GetInvalid(); }
+			TSharedRef<INodeController> GetNodeWithID(FGuid InNodeID) override { return FInvalidNodeController::GetInvalid(); }
 
 			TArray<TSharedRef<INodeController>> GetOutputNodes() override { return TArray<TSharedRef<INodeController>>(); }
 			TArray<TSharedRef<INodeController>> GetInputNodes() override { return TArray<TSharedRef<INodeController>>(); }
@@ -220,17 +230,17 @@ namespace Metasound
 			// For inputs whose preferred literal type is UObject or UObjectArray, this can be used to determine the UObject corresponding to that input's datatype.
 			UClass* GetSupportedClassForInputVertex(const FString& InInputName) override { return nullptr; }
 
-			TArray<int32> GetDefaultIDsForInputVertex(const FString& InInputName) const { return TArray<int32>(); }
-			TArray<int32> GetDefaultIDsForOutputVertex(const FString& InInputName) const { return TArray<int32>(); }
+			TArray<FGuid> GetDefaultIDsForInputVertex(const FString& InInputName) const { return TArray<FGuid>(); }
+			TArray<FGuid> GetDefaultIDsForOutputVertex(const FString& InInputName) const { return TArray<FGuid>(); }
 
 			// These can be used to set the default value for a given input on this graph.
 			// @returns false if the input name couldn't be found, or if the literal type was incompatible with the Data Type of this input.
-			bool SetDefaultInputToLiteral(const FString& InInputName, int32 InPointID, bool bInValue) override { return false; }
-			bool SetDefaultInputToLiteral(const FString& InInputName, int32 InPointID, int32 InValue) override { return false; }
-			bool SetDefaultInputToLiteral(const FString& InInputName, int32 InPointID, float InValue) override { return false; }
-			bool SetDefaultInputToLiteral(const FString& InInputName, int32 InPointID, const FString& InValue) override { return false; }
-			bool SetDefaultInputToLiteral(const FString& InInputName, int32 InPointID, UObject* InValue) override { return false; }
-			bool SetDefaultInputToLiteral(const FString& InInputName, int32 InPointID, const TArray<UObject*>& InValue) override { return false; }
+			bool SetDefaultInputToLiteral(const FString& InInputName, FGuid InPointID, bool bInValue) override { return false; }
+			bool SetDefaultInputToLiteral(const FString& InInputName, FGuid InPointID, int32 InValue) override { return false; }
+			bool SetDefaultInputToLiteral(const FString& InInputName, FGuid InPointID, float InValue) override { return false; }
+			bool SetDefaultInputToLiteral(const FString& InInputName, FGuid InPointID, const FString& InValue) override { return false; }
+			bool SetDefaultInputToLiteral(const FString& InInputName, FGuid InPointID, UObject* InValue) override { return false; }
+			bool SetDefaultInputToLiteral(const FString& InInputName, FGuid InPointID, const TArray<UObject*>& InValue) override { return false; }
 
 
 			// Set the display name for the input with the given name
@@ -241,7 +251,7 @@ namespace Metasound
 
 			// This can be used to clear the current literal for a given input.
 			// @returns false if the input name couldn't be found.
-			bool ClearLiteralForInput(const FString& InInputName, int32 InPointID) override { return false; }
+			bool ClearLiteralForInput(const FString& InInputName, FGuid InPointID) override { return false; }
 
 			TSharedRef<INodeController> AddNode(const FNodeClassInfo& InNodeClass) override { return FInvalidNodeController::GetInvalid(); }
 			TSharedRef<INodeController> AddNode(const FNodeRegistryKey& InNodeClass) override { return FInvalidNodeController::GetInvalid(); }
@@ -270,6 +280,9 @@ namespace Metasound
 
 			TSharedRef<IDocumentController> GetOwningDocument() override;
 			TSharedRef<const IDocumentController> GetOwningDocument() const override;
+		protected:
+			FDocumentAccess ShareAccess() override { return FDocumentAccess(); }
+			FConstDocumentAccess ShareAccess() const override { return FConstDocumentAccess(); }
 		};
 
 		/** FInvalidDocumentController is a controller which is always invalid. 
@@ -299,9 +312,9 @@ namespace Metasound
 				TArray<FMetasoundFrontendGraphClass> GetSubgraphs() const override { return TArray<FMetasoundFrontendGraphClass>(); }
 				TArray<FMetasoundFrontendClass> GetClasses() const override { return TArray<FMetasoundFrontendClass>(); }
 
-				FConstClassAccessPtr FindDependencyWithID(int32 InClassID) const override { return FConstClassAccessPtr(); }
-				FConstGraphClassAccessPtr FindSubgraphWithID(int32 InClassID) const override { return FConstGraphClassAccessPtr(); }
-				FConstClassAccessPtr FindClassWithID(int32 InClassID) const override { return FConstClassAccessPtr(); }
+				FConstClassAccessPtr FindDependencyWithID(FGuid InClassID) const override { return FConstClassAccessPtr(); }
+				FConstGraphClassAccessPtr FindSubgraphWithID(FGuid InClassID) const override { return FConstGraphClassAccessPtr(); }
+				FConstClassAccessPtr FindClassWithID(FGuid InClassID) const override { return FConstClassAccessPtr(); }
 
 				FConstClassAccessPtr FindClass(const FNodeClassInfo& InNodeClass) const override { return FConstClassAccessPtr(); }
 				FConstClassAccessPtr FindOrAddClass(const FNodeClassInfo& InNodeClass) override { return FConstClassAccessPtr(); }
@@ -313,14 +326,19 @@ namespace Metasound
 
 				TArray<FConstGraphHandle> GetSubgraphHandles() const override { return TArray<FConstGraphHandle>(); }
 
-				FGraphHandle GetSubgraphWithClassID(int32 InClassID) { return FInvalidGraphController::GetInvalid(); }
+				FGraphHandle GetSubgraphWithClassID(FGuid InClassID) { return FInvalidGraphController::GetInvalid(); }
 
-				FConstGraphHandle GetSubgraphWithClassID(int32 InClassID) const { return FInvalidGraphController::GetInvalid(); }
+				FConstGraphHandle GetSubgraphWithClassID(FGuid InClassID) const { return FInvalidGraphController::GetInvalid(); }
 
 				TSharedRef<IGraphController> GetRootGraph() override;
 				TSharedRef<const IGraphController> GetRootGraph() const override;
 				bool ExportToJSONAsset(const FString& InAbsolutePath) const override { return false; }
 				FString ExportToJSON() const override { return FString(); }
+
+			protected:
+
+				FDocumentAccess ShareAccess() override { return FDocumentAccess(); }
+				FConstDocumentAccess ShareAccess() const override { return FConstDocumentAccess(); }
 		};
 	}
 }
diff --git a/Engine/Plugins/Runtime/Metasound/Source/MetasoundFrontend/Private/MetasoundFrontendStandardController.cpp b/Engine/Plugins/Runtime/Metasound/Source/MetasoundFrontend/Private/MetasoundFrontendStandardController.cpp
--- a/Engine/Plugins/Runtime/Metasound/Source/MetasoundFrontend/Private/MetasoundFrontendStandardController.cpp
+++ b/Engine/Plugins/Runtime/Metasound/Source/MetasoundFrontend/Private/MetasoundFrontendStandardController.cpp
@@ -80,7 +80,7 @@ namespace Metasound
 			// specific queries. 
 			struct FGetSubgraphFromDocumentByID
 			{
-				FGetSubgraphFromDocumentByID(int32 InID)
+				FGetSubgraphFromDocumentByID(FGuid InID)
 				: ID(InID)
 				{
 				}
@@ -96,7 +96,7 @@ namespace Metasound
 				}
 
 			private:
-				int32 ID;
+				FGuid ID;
 			};
 
 			struct FGetClassInputFromClassWithName
@@ -179,7 +179,7 @@ namespace Metasound
 				};
 			}
 
-			auto FGetNodeInputFromNodeWithPoint(int32 InPointID)
+			auto FGetNodeInputFromNodeWithPoint(FGuid InPointID)
 			{
 				return [=](FMetasoundFrontendNode& InNode) -> FMetasoundFrontendVertex*
 				{
@@ -188,7 +188,7 @@ namespace Metasound
 				};
 			}
 
-			auto FGetConstNodeInputFromNodeWithPoint(int32 InPointID)
+			auto FGetConstNodeInputFromNodeWithPoint(FGuid InPointID)
 			{
 				return [=](const FMetasoundFrontendNode& InNode) -> const FMetasoundFrontendVertex*
 				{
@@ -197,7 +197,7 @@ namespace Metasound
 				};
 			}
 
-			auto FGetNodeOutputFromNodeWithPoint(int32 InPointID)
+			auto FGetNodeOutputFromNodeWithPoint(FGuid InPointID)
 			{
 				return [=](FMetasoundFrontendNode& InNode) -> FMetasoundFrontendVertex*
 				{
@@ -206,7 +206,7 @@ namespace Metasound
 				};
 			}
 
-			auto FGetConstNodeOutputFromNodeWithPoint(int32 InPointID)
+			auto FGetConstNodeOutputFromNodeWithPoint(FGuid InPointID)
 			{
 				return [=](const FMetasoundFrontendNode& InNode) -> const FMetasoundFrontendVertex*
 				{
@@ -215,7 +215,7 @@ namespace Metasound
 				};
 			}
 
-			auto FGetConstClassFromDocumentDependenciesWithID(int32 InClassID)
+			auto FGetConstClassFromDocumentDependenciesWithID(FGuid InClassID)
 			{
 				return [=](const FMetasoundFrontendDocument& InDocument) -> const FMetasoundFrontendClass*
 				{
@@ -224,7 +224,7 @@ namespace Metasound
 				};
 			}
 
-			auto FGetConstGraphClassFromDocumentWithID(int32 InClassID)
+			auto FGetConstGraphClassFromDocumentWithID(FGuid InClassID)
 			{
 				return [=](const FMetasoundFrontendDocument& InDocument) -> const FMetasoundFrontendGraphClass*
 				{
@@ -269,7 +269,7 @@ namespace Metasound
 				};
 			}
 
-			auto FGetNodeFromGraphClassByNodeID(int32 InNodeID)
+			auto FGetNodeFromGraphClassByNodeID(FGuid InNodeID)
 			{
 				return [=](FMetasoundFrontendGraphClass& InGraphClass) -> FMetasoundFrontendNode*
 				{
@@ -287,7 +287,7 @@ namespace Metasound
 				return [](FMetasoundFrontendGraphClass& InGraphClass) -> FMetasoundFrontendGraph* { return &InGraphClass.Graph; };
 			}
 
-			auto FGetConstNodeFromGraphClassByNodeID(int32 InNodeID)
+			auto FGetConstNodeFromGraphClassByNodeID(FGuid InNodeID)
 			{
 				return [=](const FMetasoundFrontendGraphClass& InGraphClass) -> const FMetasoundFrontendNode*
 				{
@@ -296,7 +296,7 @@ namespace Metasound
 				};
 			}
 
-			auto FGetConstClassInputFromGraphClassWithNodeID(int32 InNodeID)
+			auto FGetConstClassInputFromGraphClassWithNodeID(FGuid InNodeID)
 			{
 				return [=](const FMetasoundFrontendGraphClass& InGraphClass) -> const FMetasoundFrontendClassInput*
 				{
@@ -304,7 +304,7 @@ namespace Metasound
 				};
 			}
 
-			auto FGetClassInputFromGraphClassWithNodeID(int32 InNodeID)
+			auto FGetClassInputFromGraphClassWithNodeID(FGuid InNodeID)
 			{
 				return [=](FMetasoundFrontendGraphClass& InGraphClass) -> FMetasoundFrontendClassInput*
 				{
@@ -312,7 +312,7 @@ namespace Metasound
 				};
 			}
 
-			auto FGetConstClassOutputFromGraphClassWithNodeID(int32 InNodeID)
+			auto FGetConstClassOutputFromGraphClassWithNodeID(FGuid InNodeID)
 			{
 				return [=](const FMetasoundFrontendGraphClass& InGraphClass) -> const FMetasoundFrontendClassOutput*
 				{
@@ -320,7 +320,7 @@ namespace Metasound
 				};
 			}
 
-			auto FGetClassOutputFromGraphClassWithNodeID(int32 InNodeID)
+			auto FGetClassOutputFromGraphClassWithNodeID(FGuid InNodeID)
 			{
 				return [=](FMetasoundFrontendGraphClass& InGraphClass) -> FMetasoundFrontendClassOutput*
 				{
@@ -329,57 +329,16 @@ namespace Metasound
 			}
 
 
-			/*
-			struct FGetNodeInputFromNodeWithName
-			{
-				FGetNodeInputFromNodeWithName(const FString& InName)
-				: Name(InName)
-				{
-				}
-
-				FMetasoundFrontendVertex* operator () (FMetasoundFrontendNode& InNode)
-				{
-					auto IsVertexWithName = [&](const FMetasoundFrontendVertex& Vertex) { return Vertex.Name == Name; };
-					return InNode.Interface.Inputs.FindByPredicate(IsVertexWithName);
-				}
-
-				const FMetasoundFrontendVertex* operator() (const FMetasoundFrontendNode& InNode)
-				{
-					auto IsVertexWithName = [&](const FMetasoundFrontendVertex& Vertex) { return Vertex.Name == Name; };
-					return InNode.Interface.Inputs.FindByPredicate(IsVertexWithName);
-				}
-			};
-			*/
-
-			int32 GetNewID(const TSet<int32> InExistingIDs)
-			{
-				// Assumption here is that we will never need more than ten thousand nodes,
-				// and four digits are easy enough to read/remember when looking at metasound graph documents.
-				int32 IDMax = 9999;
-
-				int32 NewID = Metasound::FrontendInvalidID;
-
-				while (InExistingIDs.Num() > (IDMax / 10))
-				{
-					IDMax *= 10;
-					IDMax += 9;
-
-					if (!ensure((IDMax > 0)))
-					{
-						// IDMax has rolled over to negative values. now is a good
-						// time to consider a new ID type.
-						return Metasound::FrontendInvalidID;
-					}
-				}
+		}
 
-				do 
-				{
-					NewID = FMath::RandRange(1, IDMax);
-				}
-				while (InExistingIDs.Contains(NewID));
+		FDocumentAccess IDocumentAccessor::GetSharedAccess(IDocumentAccessor& InDocumentAccessor)
+		{
+			return InDocumentAccessor.ShareAccess();
+		}
 
-				return NewID;
-			}
+		FConstDocumentAccess IDocumentAccessor::GetSharedAccess(const IDocumentAccessor& InDocumentAccessor)
+		{
+			return InDocumentAccessor.ShareAccess();
 		}
 
 		//
@@ -397,7 +356,7 @@ namespace Metasound
 			return OwningNode->IsValid() && NodeVertexPtr.IsValid();
 		}
 
-		int32 FBaseOutputController::GetID() const
+		FGuid FBaseOutputController::GetID() const
 		{
 			return ID;
 		}
@@ -420,7 +379,7 @@ namespace Metasound
 			return FrontendControllerIntrinsics::GetInvalidValueConstRef<FString>();
 		}
 
-		int32 FBaseOutputController::GetOwningNodeID() const
+		FGuid FBaseOutputController::GetOwningNodeID() const
 		{
 			return OwningNode->GetID();
 		}
@@ -455,6 +414,24 @@ namespace Metasound
 			return InController.Disconnect(*this);
 		}
 
+		FDocumentAccess FBaseOutputController::ShareAccess() 
+		{
+			FDocumentAccess Access;
+
+			Access.ConstVertex = NodeVertexPtr;
+
+			return Access;
+		}
+
+		FConstDocumentAccess FBaseOutputController::ShareAccess() const
+		{
+			FConstDocumentAccess Access;
+
+			Access.ConstVertex = NodeVertexPtr;
+
+			return Access;
+		}
+
 		//
 		// FNodeOutputController
 		//
@@ -489,6 +466,24 @@ namespace Metasound
 			return FrontendControllerIntrinsics::GetInvalidValueConstRef<FText>();
 		}
 
+		FDocumentAccess FNodeOutputController::ShareAccess() 
+		{
+			FDocumentAccess Access = FBaseOutputController::ShareAccess();
+
+			Access.ConstClassOutput = ClassOutputPtr;
+
+			return Access;
+		}
+
+		FConstDocumentAccess FNodeOutputController::ShareAccess() const
+		{
+			FConstDocumentAccess Access = FBaseOutputController::ShareAccess();
+
+			Access.ConstClassOutput = ClassOutputPtr;
+
+			return Access;
+		}
+
 		//
 		// FInputNodeOutputController
 		// 
@@ -523,6 +518,24 @@ namespace Metasound
 			return FrontendControllerIntrinsics::GetInvalidValueConstRef<FText>();
 		}
 
+		FDocumentAccess FInputNodeOutputController::ShareAccess() 
+		{
+			FDocumentAccess Access = FBaseOutputController::ShareAccess();
+
+			Access.ConstClassInput = OwningGraphClassInputPtr;
+
+			return Access;
+		}
+
+		FConstDocumentAccess FInputNodeOutputController::ShareAccess() const
+		{
+			FConstDocumentAccess Access = FBaseOutputController::ShareAccess();
+
+			Access.ConstClassInput = OwningGraphClassInputPtr;
+
+			return Access;
+		}
+
 		//
 		// FBaseInputController
 		// 
@@ -539,7 +552,7 @@ namespace Metasound
 			return OwningNode->IsValid() && NodeVertexPtr.IsValid() &&  GraphPtr.IsValid();
 		}
 
-		int32 FBaseInputController::GetID() const
+		FGuid FBaseInputController::GetID() const
 		{
 			return ID;
 		}
@@ -569,7 +582,7 @@ namespace Metasound
 			return (nullptr != FindEdge());
 		}
 
-		int32 FBaseInputController::GetOwningNodeID() const
+		FGuid FBaseInputController::GetOwningNodeID() const
 		{
 			return OwningNode->GetID();
 		}
@@ -708,10 +721,10 @@ namespace Metasound
 		{
 			if (GraphPtr.IsValid())
 			{
-				int32 FromNodeID = InController.GetOwningNodeID();
-				int32 FromPointID = InController.GetID();
-				int32 ToNodeID = GetOwningNodeID();
-				int32 ToPointID = GetID();
+				FGuid FromNodeID = InController.GetOwningNodeID();
+				FGuid FromPointID = InController.GetID();
+				FGuid ToNodeID = GetOwningNodeID();
+				FGuid ToPointID = GetID();
 
 				auto IsMatchingEdge = [&](const FMetasoundFrontendEdge& Edge)
 				{
@@ -729,8 +742,8 @@ namespace Metasound
 		{
 			if (GraphPtr.IsValid())
 			{
-				const int32 NodeID = GetOwningNodeID();
-				const int32 PointID = GetID();
+				const FGuid NodeID = GetOwningNodeID();
+				FGuid PointID = GetID();
 
 				auto EdgeHasMatchingDestination = [&](const FMetasoundFrontendEdge& Edge)
 				{
@@ -749,8 +762,8 @@ namespace Metasound
 		{
 			if (GraphPtr.IsValid())
 			{
-				const int32 NodeID = GetOwningNodeID();
-				const int32 PointID = GetID();
+				const FGuid NodeID = GetOwningNodeID();
+				FGuid PointID = GetID();
 
 				auto EdgeHasMatchingDestination = [&](const FMetasoundFrontendEdge& Edge)
 				{
@@ -767,8 +780,8 @@ namespace Metasound
 		{
 			if (GraphPtr.IsValid())
 			{
-				const int32 NodeID = GetOwningNodeID();
-				const int32 PointID = GetID();
+				const FGuid NodeID = GetOwningNodeID();
+				FGuid PointID = GetID();
 
 				auto EdgeHasMatchingDestination = [&](const FMetasoundFrontendEdge& Edge)
 				{
@@ -781,6 +794,27 @@ namespace Metasound
 			return nullptr;
 		}
 
+		FDocumentAccess FBaseInputController::ShareAccess() 
+		{
+			FDocumentAccess Access;
+
+			Access.ConstVertex = NodeVertexPtr;
+			Access.Graph = GraphPtr;
+			Access.ConstGraph = GraphPtr;
+
+			return Access;
+		}
+
+		FConstDocumentAccess FBaseInputController::ShareAccess() const
+		{
+			FConstDocumentAccess Access;
+
+			Access.ConstVertex = NodeVertexPtr;
+			Access.ConstGraph = GraphPtr;
+
+			return Access;
+		}
+
 		//
 		// FNodeInputController
 		//
@@ -815,6 +849,24 @@ namespace Metasound
 			return FrontendControllerIntrinsics::GetInvalidValueConstRef<FText>();
 		}
 
+		FDocumentAccess FNodeInputController::ShareAccess() 
+		{
+			FDocumentAccess Access = FBaseInputController::ShareAccess();
+
+			Access.ConstClassInput = OwningGraphClassInputPtr;
+
+			return Access;
+		}
+
+		FConstDocumentAccess FNodeInputController::ShareAccess() const
+		{
+			FConstDocumentAccess Access = FBaseInputController::ShareAccess();
+
+			Access.ConstClassInput = OwningGraphClassInputPtr;
+
+			return Access;
+		}
+
 		//
 		// FOutputNodeInputController
 		//
@@ -849,6 +901,24 @@ namespace Metasound
 			return FrontendControllerIntrinsics::GetInvalidValueConstRef<FText>();
 		}
 
+		FDocumentAccess FOutputNodeInputController::ShareAccess() 
+		{
+			FDocumentAccess Access = FBaseInputController::ShareAccess();
+
+			Access.ConstClassOutput = OwningGraphClassOutputPtr;
+
+			return Access;
+		}
+
+		FConstDocumentAccess FOutputNodeInputController::ShareAccess() const
+		{
+			FConstDocumentAccess Access = FBaseInputController::ShareAccess();
+
+			Access.ConstClassOutput = OwningGraphClassOutputPtr;
+
+			return Access;
+		}
+
 
 		//
 		// FBaseNodeController
@@ -862,7 +932,7 @@ namespace Metasound
 			{
 				if (NodePtr->ClassID != ClassPtr->ID)
 				{
-					UE_LOG(LogMetasound, Warning, TEXT("Changing node's class id from [ClassID:%d] to [ClassID:%d]"), NodePtr->ClassID, ClassPtr->ID);
+					UE_LOG(LogMetasound, Warning, TEXT("Changing node's class id from [ClassID:%s] to [ClassID:%s]"), *NodePtr->ClassID.ToString(), *ClassPtr->ID.ToString());
 					NodePtr->ClassID = ClassPtr->ID;
 				}
 			}
@@ -873,7 +943,7 @@ namespace Metasound
 			return OwningGraph->IsValid() && NodePtr.IsValid() && ClassPtr.IsValid();
 		}
 
-		int32 FBaseNodeController::GetOwningGraphClassID() const
+		FGuid FBaseNodeController::GetOwningGraphClassID() const
 		{
 			return OwningGraph->GetClassID();
 		}
@@ -888,7 +958,7 @@ namespace Metasound
 			return OwningGraph;
 		}
 
-		int32 FBaseNodeController::GetID() const
+		FGuid FBaseNodeController::GetID() const
 		{
 			if (NodePtr.IsValid())
 			{
@@ -897,7 +967,7 @@ namespace Metasound
 			return Metasound::FrontendInvalidID;
 		}
 
-		int32 FBaseNodeController::GetClassID() const
+		FGuid FBaseNodeController::GetClassID() const
 		{
 			if (ClassPtr.IsValid())
 			{
@@ -975,7 +1045,7 @@ namespace Metasound
 			return FInvalidInputController::GetInvalid();
 		}
 
-		bool FBaseNodeController::RemoveInput(int32 InPointID)
+		bool FBaseNodeController::RemoveInput(FGuid InPointID)
 		{
 			checkNoEntry();
 			// TODO: not yet supported
@@ -995,7 +1065,7 @@ namespace Metasound
 			return FInvalidInputController::GetInvalid();
 		}
 
-		bool FBaseNodeController::RemoveOutput(int32 InPointID)
+		bool FBaseNodeController::RemoveOutput(FGuid InPointID)
 		{
 			checkNoEntry();
 			// TODO: not yet supported
@@ -1157,7 +1227,7 @@ namespace Metasound
 			return Outputs;
 		}
 
-		FInputHandle FBaseNodeController::GetInputWithID(int32 InPointID)
+		FInputHandle FBaseNodeController::GetInputWithID(FGuid InPointID)
 		{
 			FInputControllerParams Params;
 
@@ -1170,7 +1240,7 @@ namespace Metasound
 			return FInvalidInputController::GetInvalid();
 		}
 
-		FConstInputHandle FBaseNodeController::GetInputWithID(int32 InPointID) const
+		FConstInputHandle FBaseNodeController::GetInputWithID(FGuid InPointID) const
 		{
 			FInputControllerParams Params;
 
@@ -1184,7 +1254,7 @@ namespace Metasound
 			return FInvalidInputController::GetInvalid();
 		}
 
-		FOutputHandle FBaseNodeController::GetOutputWithID(int32 InPointID)
+		FOutputHandle FBaseNodeController::GetOutputWithID(FGuid InPointID)
 		{
 			FOutputControllerParams Params;
 
@@ -1197,7 +1267,7 @@ namespace Metasound
 			return FInvalidOutputController::GetInvalid();
 		}
 
-		FConstOutputHandle FBaseNodeController::GetOutputWithID(int32 InPointID) const
+		FConstOutputHandle FBaseNodeController::GetOutputWithID(FGuid InPointID) const
 		{
 			FOutputControllerParams Params;
 
@@ -1227,47 +1297,12 @@ namespace Metasound
 
 					FConstClassInputAccessPtr ClassInputPtr = ClassPtr.GetMemberAccessPtr<const FMetasoundFrontendClassInput>(FGetClassInputFromClassWithName(NodeInputVertex.Name));
 
-					for (int32 PointID : NodeInputVertex.PointIDs)
+					for (const FGuid& PointID : NodeInputVertex.PointIDs)
 					{
 						Inputs.Add({PointID, NodeVertexPtr, ClassInputPtr});
 					}
 				}
 			}
-			/*
-
-
-
-
-
-
-			if (ClassPtr.IsValid() && NodePtr.IsValid())
-			{
-				for (const FMetasoundFrontendVertex& NodeInputVertex : NodePtr->Interface.Inputs)
-				{
-
-					FConstVertexAccessPtr NodeVertexPtr = FrontendControllerIntrinsics::MakeAutoAccessPtr(NodeInputVertex);
-
-
-					if (const FMetasoundFrontendClassInput* ClassInputVertex = FindClassInputWithName(NodeInputVertex.Name))
-					{
-
-
-FGetClassInputFromClassWithName
-
-						FConstClassInputAccessPtr ClassInputPtr = FrontendControllerIntrinsics::MakeAutoAccessPtr(*ClassInputVertex);
-
-						for (int32 PointID : NodeInputVertex.PointIDs)
-						{
-							Inputs.Add({PointID, NodeVertexPtr, ClassInputPtr});
-						}
-					}
-					else
-					{
-						UE_LOG(LogMetasound, Error, TEXT("Node vertex [NodeID:%d, InputName:%s] does not exist in class [ClassID:%d]"), NodePtr->ID, *NodeInputVertex.Name, ClassPtr->ID);
-					}
-				}
-			}
-			*/
 
 			return Inputs;
 		}
@@ -1285,37 +1320,12 @@ FGetClassInputFromClassWithName
 				FConstVertexAccessPtr NodeVertexPtr = NodePtr.GetMemberAccessPtr<const FMetasoundFrontendVertex>(FGetConstNodeOutputFromNodeWithName(VertexName));
 				FConstClassOutputAccessPtr ClassOutputPtr = ClassPtr.GetMemberAccessPtr<const FMetasoundFrontendClassOutput>(FGetClassOutputFromClassWithName(VertexName));
 
-				for (int32 PointID : NodeOutputVertex.PointIDs)
+				for (const FGuid& PointID : NodeOutputVertex.PointIDs)
 				{
 					Outputs.Add({PointID, NodeVertexPtr, ClassOutputPtr});
 				}
 			}
 
-
-			/*
-			if (ClassPtr.IsValid() && NodePtr.IsValid())
-			{
-				for (const FMetasoundFrontendVertex& NodeOutputVertex : NodePtr->Interface.Outputs)
-				{
-					FConstVertexAccessPtr NodeVertexPtr = FrontendControllerIntrinsics::MakeAutoAccessPtr(NodeOutputVertex);
-
-					if (const FMetasoundFrontendClassOutput* ClassOutputVertex = FindClassOutputWithName(NodeOutputVertex.Name))
-					{
-						FConstClassOutputAccessPtr ClassOutputPtr = FrontendControllerIntrinsics::MakeAutoAccessPtr(*ClassOutputVertex);
-
-						for (int32 PointID : NodeOutputVertex.PointIDs)
-						{
-							Outputs.Add({PointID, NodeVertexPtr, ClassOutputPtr});
-						}
-					}
-					else
-					{
-						UE_LOG(LogMetasound, Error, TEXT("Node vertex [NodeID:%d, OutputName:%s] does not exist in class [ClassID:%d]"), NodePtr->ID, *NodeOutputVertex.Name, ClassPtr->ID);
-					}
-				}
-			}
-			*/
-
 			return Outputs;
 		}
 
@@ -1331,37 +1341,12 @@ FGetClassInputFromClassWithName
 			{
 				FConstClassInputAccessPtr ClassInputPtr = ClassPtr.GetMemberAccessPtr<const FMetasoundFrontendClassInput>(FGetClassInputFromClassWithName(InName));
 
-				for (int32 PointID : Vertex->PointIDs)
+				for (const FGuid& PointID : Vertex->PointIDs)
 				{
 					Inputs.Add({PointID, NodeVertexPtr, ClassInputPtr});
 				}
 			}
 
-
-			/*
-			if (ClassPtr.IsValid() && NodePtr.IsValid())
-			{
-				if (const FMetasoundFrontendVertex* NodeInputVertex = FindNodeInputWithName(InName))
-				{
-					FConstVertexAccessPtr NodeVertexPtr = FrontendControllerIntrinsics::MakeAutoAccessPtr(*NodeInputVertex);
-
-					if (const FMetasoundFrontendClassInput* ClassInputVertex = FindClassInputWithName(InName))
-					{
-						FConstClassInputAccessPtr ClassInputPtr = FrontendControllerIntrinsics::MakeAutoAccessPtr(*ClassInputVertex);
-
-						for (int32 PointID : NodeInputVertex->PointIDs)
-						{
-							Inputs.Add({PointID, NodeVertexPtr, ClassInputPtr});
-						}
-					}
-					else
-					{
-						UE_LOG(LogMetasound, Error, TEXT("Node vertex [NodeID:%d, InputName:%s] does not exist in class [ClassID:%d]"), NodePtr->ID, *NodeInputVertex->Name, ClassPtr->ID);
-					}
-				}
-			}
-			*/
-
 			return Inputs;
 		}
 
@@ -1377,39 +1362,16 @@ FGetClassInputFromClassWithName
 			{
 				FConstClassOutputAccessPtr ClassOutputPtr = ClassPtr.GetMemberAccessPtr<const FMetasoundFrontendClassOutput>(FGetClassOutputFromClassWithName(InName));
 
-				for (int32 PointID : Vertex->PointIDs)
+				for (const FGuid& PointID : Vertex->PointIDs)
 				{
 					Outputs.Add({PointID, NodeVertexPtr, ClassOutputPtr});
 				}
 			}
-			/*
-			if (ClassPtr.IsValid() && NodePtr.IsValid())
-			{
-				if (const FMetasoundFrontendVertex* NodeOutputVertex = FindNodeOutputWithName(InName))
-				{
-					FConstVertexAccessPtr NodeVertexPtr = FrontendControllerIntrinsics::MakeAutoAccessPtr(*NodeOutputVertex);
-
-					if (const FMetasoundFrontendClassOutput* ClassOutputVertex = FindClassOutputWithName(InName))
-					{
-						FConstClassOutputAccessPtr ClassOutputPtr = FrontendControllerIntrinsics::MakeAutoAccessPtr(*ClassOutputVertex);
-
-						for (int32 PointID : NodeOutputVertex->PointIDs)
-						{
-							Outputs.Add({PointID, NodeVertexPtr, ClassOutputPtr});
-						}
-					}
-					else
-					{
-						UE_LOG(LogMetasound, Error, TEXT("Node vertex [NodeID:%d, OutputName:%s] does not exist in class [ClassID:%d]"), NodePtr->ID, *NodeOutputVertex->Name, ClassPtr->ID);
-					}
-				}
-			}
-			*/
 
 			return Outputs;
 		}
 
-		bool FBaseNodeController::FindInputControllerParamsWithID(int32 InPointID, FInputControllerParams& OutParams) const
+		bool FBaseNodeController::FindInputControllerParamsWithID(FGuid InPointID, FInputControllerParams& OutParams) const
 		{
 			using namespace FrontendControllerIntrinsics;
 
@@ -1422,34 +1384,11 @@ FGetClassInputFromClassWithName
 				OutParams = FInputControllerParams{InPointID, NodeVertexPtr, ClassInputPtr};
 				return true;
 			}
-			/*
-			if (ClassPtr.IsValid() && NodePtr.IsValid())
-			{
-				for (const FMetasoundFrontendVertex& NodeInputVertex : NodePtr->Interface.Inputs)
-				{
-					if (NodeInputVertex.PointIDs.Contains(InPointID))
-					{
-						if (const FMetasoundFrontendClassInput* ClassInputVertex = FindClassInputWithName(NodeInputVertex.Name))
-						{
-							FConstVertexAccessPtr NodeVertexPtr = FrontendControllerIntrinsics::MakeAutoAccessPtr(NodeInputVertex);
-							FConstClassInputAccessPtr ClassInputPtr = FrontendControllerIntrinsics::MakeAutoAccessPtr(*ClassInputVertex);
-
-							OutParams = FInputControllerParams{InPointID, NodeVertexPtr, ClassInputPtr};
-							return true;
-						}
-						else
-						{
-							UE_LOG(LogMetasound, Error, TEXT("Node vertex [NodeID:%d, InputName:%s] does not exist in class [ClassID:%d]"), NodePtr->ID, *NodeInputVertex.Name, ClassPtr->ID);
-						}
-					}
-				}
-			}
-			*/
 
 			return false;
 		}
 
-		bool FBaseNodeController::FindOutputControllerParamsWithID(int32 InPointID, FOutputControllerParams& OutParams) const
+		bool FBaseNodeController::FindOutputControllerParamsWithID(FGuid InPointID, FOutputControllerParams& OutParams) const
 		{
 			using namespace FrontendControllerIntrinsics;
 
@@ -1462,29 +1401,6 @@ FGetClassInputFromClassWithName
 				OutParams = FOutputControllerParams{InPointID, NodeVertexPtr, ClassOutputPtr};
 				return true;
 			}
-			/*
-			if (ClassPtr.IsValid() && NodePtr.IsValid())
-			{
-				for (const FMetasoundFrontendVertex& NodeOutputVertex : NodePtr->Interface.Outputs)
-				{
-					if (NodeOutputVertex.PointIDs.Contains(InPointID))
-					{
-						if (const FMetasoundFrontendClassOutput* ClassOutputVertex = FindClassOutputWithName(NodeOutputVertex.Name))
-						{
-							FConstVertexAccessPtr NodeVertexPtr = FrontendControllerIntrinsics::MakeAutoAccessPtr(NodeOutputVertex);
-							FConstClassOutputAccessPtr ClassOutputPtr = FrontendControllerIntrinsics::MakeAutoAccessPtr(*ClassOutputVertex);
-
-							OutParams = FOutputControllerParams{InPointID, NodeVertexPtr, ClassOutputPtr};
-							return true;
-						}
-						else
-						{
-							UE_LOG(LogMetasound, Error, TEXT("Node vertex [NodeID:%d, OutputName:%s] does not exist in class [ClassID:%d]"), NodePtr->ID, *NodeOutputVertex.Name, ClassPtr->ID);
-						}
-					}
-				}
-			}
-			*/
 
 			return false;
 		}
@@ -1562,6 +1478,27 @@ FGetClassInputFromClassWithName
 			return FInvalidGraphController::GetInvalid();
 		}
 
+		FDocumentAccess FBaseNodeController::ShareAccess() 
+		{
+			FDocumentAccess Access;
+
+			Access.Node = NodePtr;
+			Access.ConstNode = NodePtr;
+			Access.ConstClass = ClassPtr;
+
+			return Access;
+		}
+
+		FConstDocumentAccess FBaseNodeController::ShareAccess() const
+		{
+			FConstDocumentAccess Access;
+
+			Access.ConstNode = NodePtr;
+			Access.ConstClass = ClassPtr;
+
+			return Access;
+		}
+
 		// 
 		// FNodeController
 		//
@@ -1582,7 +1519,7 @@ FGetClassInputFromClassWithName
 				}
 				else
 				{
-					UE_LOG(LogMetasound, Warning, TEXT("Frontend Node [NodeId:%d, ClassID: %d] is not of expected class class [ClassId:%d]"), InParams.NodePtr->ID, InParams.NodePtr->ClassID, InParams.ClassPtr->ID);
+					UE_LOG(LogMetasound, Warning, TEXT("Frontend Node [NodeID:%s, ClassID:%s] is not of expected class class [ClassID:%s]"), *InParams.NodePtr->ID.ToString(), *InParams.NodePtr->ClassID.ToString(), *InParams.ClassPtr->ID.ToString());
 				}
 			}
 			return FInvalidNodeController::GetInvalid();
@@ -1599,7 +1536,7 @@ FGetClassInputFromClassWithName
 				}
 				else
 				{
-					UE_LOG(LogMetasound, Warning, TEXT("Frontend Node [NodeId:%d, ClassID:%d] is not of expected class class [ClassId:%d]"), InParams.NodePtr->ID, InParams.NodePtr->ClassID, InParams.ClassPtr->ID);
+					UE_LOG(LogMetasound, Warning, TEXT("Frontend Node [NodeID:%s, ClassID:%s] is not of expected class class [ClassID:%s]"), *InParams.NodePtr->ID.ToString(), *InParams.NodePtr->ClassID.ToString(), *InParams.ClassPtr->ID.ToString());
 				}
 			}
 			return FInvalidNodeController::GetInvalid();
@@ -1610,16 +1547,35 @@ FGetClassInputFromClassWithName
 			return FBaseNodeController::IsValid() && GraphPtr.IsValid();
 		}
 
-		FInputHandle FNodeController::CreateInputController(int32 InPointID, FConstVertexAccessPtr InNodeVertexPtr, FConstClassInputAccessPtr InClassInputPtr, FNodeHandle InOwningNode) const
+		FInputHandle FNodeController::CreateInputController(FGuid InPointID, FConstVertexAccessPtr InNodeVertexPtr, FConstClassInputAccessPtr InClassInputPtr, FNodeHandle InOwningNode) const
 		{
 			return MakeShared<FNodeInputController>(FNodeInputController::FInitParams{InPointID, InNodeVertexPtr, InClassInputPtr, GraphPtr, InOwningNode});
 		}
 
-		FOutputHandle FNodeController::CreateOutputController(int32 InPointID, FConstVertexAccessPtr InNodeVertexPtr, FConstClassOutputAccessPtr InClassOutputPtr, FNodeHandle InOwningNode) const
+		FOutputHandle FNodeController::CreateOutputController(FGuid InPointID, FConstVertexAccessPtr InNodeVertexPtr, FConstClassOutputAccessPtr InClassOutputPtr, FNodeHandle InOwningNode) const
 		{
 			return MakeShared<FNodeOutputController>(FNodeOutputController::FInitParams{InPointID, InNodeVertexPtr, InClassOutputPtr, InOwningNode});
 		}
 
+		FDocumentAccess FNodeController::ShareAccess() 
+		{
+			FDocumentAccess Access = FBaseNodeController::ShareAccess();
+
+			Access.Graph = GraphPtr;
+			Access.ConstGraph = GraphPtr;
+
+			return Access;
+		}
+
+		FConstDocumentAccess FNodeController::ShareAccess() const
+		{
+			FConstDocumentAccess Access = FBaseNodeController::ShareAccess();
+
+			Access.ConstGraph = GraphPtr;
+
+			return Access;
+		}
+
 		//
 		// FOutputNodeController
 		//
@@ -1642,12 +1598,12 @@ FGetClassInputFromClassWithName
 					}
 					else
 					{
-						UE_LOG(LogMetasound, Warning, TEXT("Frontend Node [NodeId:%d, ClassID: %d] is not of expected class class [ClassId:%d]"), InParams.NodePtr->ID, InParams.NodePtr->ClassID, InParams.ClassPtr->ID);
+						UE_LOG(LogMetasound, Warning, TEXT("Frontend Node [NodeID:%s, ClassID:%s] is not of expected class class [ClassID:%s]"), *InParams.NodePtr->ID.ToString(), *InParams.NodePtr->ClassID.ToString(), *InParams.ClassPtr->ID.ToString());
 					}
 				}
 				else
 				{
-					UE_LOG(LogMetasound, Warning, TEXT("Frontend Class of incorrect class type [ClassID:%d] while creating output node.. Must be EMetasoundFrontendClassType::Output."), InParams.ClassPtr->ID);
+					UE_LOG(LogMetasound, Warning, TEXT("Frontend Class of incorrect class type [ClassID:%s] while creating output node.. Must be EMetasoundFrontendClassType::Output."), *InParams.ClassPtr->ID.ToString());
 				}
 			}
 
@@ -1666,12 +1622,12 @@ FGetClassInputFromClassWithName
 					}
 					else
 					{
-						UE_LOG(LogMetasound, Warning, TEXT("Frontend Node [NodeId:%d, ClassID: %d] is not of expected class class [ClassId:%d]"), InParams.NodePtr->ID, InParams.NodePtr->ClassID, InParams.ClassPtr->ID);
+						UE_LOG(LogMetasound, Warning, TEXT("Frontend Node [NodeID:%s, ClassID:%s] is not of expected class class [ClassID:%s]"), *InParams.NodePtr->ID.ToString(), *InParams.NodePtr->ClassID.ToString(), *InParams.ClassPtr->ID.ToString());
 					}
 				}
 				else
 				{
-					UE_LOG(LogMetasound, Warning, TEXT("Frontend Class of incorrect class type [ClassID:%d] while creating output node. Must be EMetasoundFrontendClassType::Output."), InParams.ClassPtr->ID);
+					UE_LOG(LogMetasound, Warning, TEXT("Frontend Class of incorrect class type [ClassID:%s] while creating output node. Must be EMetasoundFrontendClassType::Output."), *InParams.ClassPtr->ID.ToString());
 				}
 			}
 
@@ -1684,16 +1640,37 @@ FGetClassInputFromClassWithName
 			return OwningGraph->IsValid() && OwningGraphClassOutputPtr.IsValid() && GraphPtr.IsValid();
 		}
 
-		FInputHandle FOutputNodeController::CreateInputController(int32 InPointID, FConstVertexAccessPtr InNodeVertexPtr, FConstClassInputAccessPtr InClassInputPtr, FNodeHandle InOwningNode) const
+		FInputHandle FOutputNodeController::CreateInputController(FGuid InPointID, FConstVertexAccessPtr InNodeVertexPtr, FConstClassInputAccessPtr InClassInputPtr, FNodeHandle InOwningNode) const
 		{
 			return MakeShared<FOutputNodeInputController>(FOutputNodeInputController::FInitParams{InPointID, InNodeVertexPtr, OwningGraphClassOutputPtr, GraphPtr, InOwningNode});
 		}
 
-		FOutputHandle FOutputNodeController::CreateOutputController(int32 InPointID, FConstVertexAccessPtr InNodeVertexPtr, FConstClassOutputAccessPtr InClassOutputPtr, FNodeHandle InOwningNode) const
+		FOutputHandle FOutputNodeController::CreateOutputController(FGuid InPointID, FConstVertexAccessPtr InNodeVertexPtr, FConstClassOutputAccessPtr InClassOutputPtr, FNodeHandle InOwningNode) const
 		{
 			return FInvalidOutputController::GetInvalid();
 		}
 
+		FDocumentAccess FOutputNodeController::ShareAccess() 
+		{
+			FDocumentAccess Access = FBaseNodeController::ShareAccess();
+
+			Access.Graph = GraphPtr;
+			Access.ConstGraph = GraphPtr;
+			Access.ConstClassOutput = OwningGraphClassOutputPtr;
+
+			return Access;
+		}
+
+		FConstDocumentAccess FOutputNodeController::ShareAccess() const
+		{
+			FConstDocumentAccess Access = FBaseNodeController::ShareAccess();
+
+			Access.ConstGraph = GraphPtr;
+			Access.ConstClassOutput = OwningGraphClassOutputPtr;
+
+			return Access;
+		}
+
 		//
 		// FInputNodeController
 		//
@@ -1716,12 +1693,12 @@ FGetClassInputFromClassWithName
 					}
 					else
 					{
-						UE_LOG(LogMetasound, Warning, TEXT("Frontend Node [NodeId:%d, ClassID: %d] is not of expected class class [ClassId:%d]"), InParams.NodePtr->ID, InParams.NodePtr->ClassID, InParams.ClassPtr->ID);
+						UE_LOG(LogMetasound, Warning, TEXT("Frontend Node [NodeID:%s, ClassID:%s] is not of expected class class [ClassID:%s]"), *InParams.NodePtr->ID.ToString(), *InParams.NodePtr->ClassID.ToString(), *InParams.ClassPtr->ID.ToString());
 					}
 				}
 				else
 				{
-					UE_LOG(LogMetasound, Warning, TEXT("Frontend Class of incorrect class type [ClassID:%d] while creating input node. Must be EMetasoundFrontendClassType::Input."), InParams.ClassPtr->ID);
+					UE_LOG(LogMetasound, Warning, TEXT("Frontend Class of incorrect class type [ClassID:%s] while creating input node. Must be EMetasoundFrontendClassType::Input."), *InParams.ClassPtr->ID.ToString());
 				}
 			}
 
@@ -1740,12 +1717,12 @@ FGetClassInputFromClassWithName
 					}
 					else
 					{
-						UE_LOG(LogMetasound, Warning, TEXT("Frontend Node [NodeId:%d, ClassID: %d] is not of expected class class [ClassId:%d]"), InParams.NodePtr->ID, InParams.NodePtr->ClassID, InParams.ClassPtr->ID);
+						UE_LOG(LogMetasound, Warning, TEXT("Frontend Node [NodeID:%s, ClassID:%s] is not of expected class class [ClassID:%s]"), *InParams.NodePtr->ID.ToString(), *InParams.NodePtr->ClassID.ToString(), *InParams.ClassPtr->ID.ToString());
 					}
 				}
 				else
 				{
-					UE_LOG(LogMetasound, Warning, TEXT("Frontend Class of incorrect class type [ClassID:%d] while creating input node. Must be EMetasoundFrontendClassType::Input."), InParams.ClassPtr->ID);
+					UE_LOG(LogMetasound, Warning, TEXT("Frontend Class of incorrect class type [ClassID:%s] while creating input node. Must be EMetasoundFrontendClassType::Input."), *InParams.ClassPtr->ID.ToString());
 				}
 			}
 
@@ -1758,16 +1735,37 @@ FGetClassInputFromClassWithName
 			return OwningGraph->IsValid() && OwningGraphClassInputPtr.IsValid() && GraphPtr.IsValid();
 		}
 
-		FInputHandle FInputNodeController::CreateInputController(int32 InPointID, FConstVertexAccessPtr InNodeVertexPtr, FConstClassInputAccessPtr InClassInputPtr, FNodeHandle InOwningNode) const
+		FInputHandle FInputNodeController::CreateInputController(FGuid InPointID, FConstVertexAccessPtr InNodeVertexPtr, FConstClassInputAccessPtr InClassInputPtr, FNodeHandle InOwningNode) const
 		{
 			return FInvalidInputController::GetInvalid();
 		}
 
-		FOutputHandle FInputNodeController::CreateOutputController(int32 InPointID, FConstVertexAccessPtr InNodeVertexPtr, FConstClassOutputAccessPtr InClassOutputPtr, FNodeHandle InOwningNode) const
+		FOutputHandle FInputNodeController::CreateOutputController(FGuid InPointID, FConstVertexAccessPtr InNodeVertexPtr, FConstClassOutputAccessPtr InClassOutputPtr, FNodeHandle InOwningNode) const
 		{
 			return MakeShared<FInputNodeOutputController>(FInputNodeOutputController::FInitParams{InPointID, InNodeVertexPtr, OwningGraphClassInputPtr, InOwningNode});
 		}
 
+		FDocumentAccess FInputNodeController::ShareAccess() 
+		{
+			FDocumentAccess Access = FBaseNodeController::ShareAccess();
+
+			Access.Graph = GraphPtr;
+			Access.ConstGraph = GraphPtr;
+			Access.ConstClassInput = OwningGraphClassInputPtr;
+
+			return Access;
+		}
+
+		FConstDocumentAccess FInputNodeController::ShareAccess() const
+		{
+			FConstDocumentAccess Access = FBaseNodeController::ShareAccess();
+
+			Access.ConstGraph = GraphPtr;
+			Access.ConstClassInput = OwningGraphClassInputPtr;
+
+			return Access;
+		}
+
 
 		// 
 		// FGraphController
@@ -1788,7 +1786,7 @@ FGetClassInputFromClassWithName
 				}
 				else
 				{
-					UE_LOG(LogMetasound, Warning, TEXT("Failed to make graph controller [ClassId:%d]. Class must be EMeatsoundFrontendClassType::Graph."), InParams.GraphClassPtr->ID)
+					UE_LOG(LogMetasound, Warning, TEXT("Failed to make graph controller [ClassID:%s]. Class must be EMeatsoundFrontendClassType::Graph."), *InParams.GraphClassPtr->ID.ToString())
 				}
 			}
 			return FInvalidGraphController::GetInvalid();
@@ -1804,7 +1802,7 @@ FGetClassInputFromClassWithName
 				}
 				else
 				{
-					UE_LOG(LogMetasound, Warning, TEXT("Failed to make graph controller [ClassId:%d]. Class must be EMeatsoundFrontendClassType::Graph."), InParams.GraphClassPtr->ID)
+					UE_LOG(LogMetasound, Warning, TEXT("Failed to make graph controller [ClassID:%s]. Class must be EMeatsoundFrontendClassType::Graph."), *InParams.GraphClassPtr->ID.ToString())
 				}
 			}
 			return FInvalidGraphController::GetInvalid();
@@ -1815,7 +1813,7 @@ FGetClassInputFromClassWithName
 			return GraphClassPtr.IsValid() && OwningDocument->IsValid();
 		}
 
-		int32 FGraphController::GetClassID() const
+		FGuid FGraphController::GetClassID() const
 		{
 			if (GraphClassPtr.IsValid())
 			{
@@ -1825,26 +1823,9 @@ FGetClassInputFromClassWithName
 			return Metasound::FrontendInvalidID;
 		}
 
-		int32 FGraphController::GetNewPointID() const
+		FGuid FGraphController::GetNewPointID() const
 		{
-			if (const FMetasoundFrontendGraphClass* ThisClass = GraphClassPtr.Get())
-			{
-				TSet<int32> UsedIDs;
-
-				for (const FMetasoundFrontendClassInput& Input : ThisClass->Interface.Inputs)
-				{
-					UsedIDs.Append(Input.PointIDs);
-				}
-
-				for (const FMetasoundFrontendClassOutput& Output : ThisClass->Interface.Outputs)
-				{
-					UsedIDs.Append(Output.PointIDs);
-				}
-
-				return FrontendControllerIntrinsics::GetNewID(UsedIDs);
-			}
-
-			return Metasound::FrontendInvalidID;
+			return FGuid::NewGuid();
 		}
 
 		TArray<FString> FGraphController::GetInputVertexNames() const 
@@ -1877,27 +1858,6 @@ FGetClassInputFromClassWithName
 			return Names;
 		}
 
-		/*
-		const FMetasoundFrontendClassInput* FGraphController::FindClassInputWithName(const FString& InName) const
-		{
-			if (GraphClassPtr.IsValid())
-			{
-				auto IsInputWithName = [&](const FMetasoundFrontendClassInput& Input) { return Input.Name == InName; };
-				return GraphClassPtr->Interface.Inputs.FindByPredicate(IsInputWithName);
-			}
-			return nullptr;
-		}
-
-		const FMetasoundFrontendClassOutput* FGraphController::FindClassOutputWithName(const FString& InName) const
-		{
-			if (GraphClassPtr.IsValid())
-			{
-				auto IsOutputWithName = [&](const FMetasoundFrontendClassOutput& Output) { return Output.Name == InName; };
-				return GraphClassPtr->Interface.Outputs.FindByPredicate(IsOutputWithName);
-			}
-			return nullptr;
-		}
-		*/
 		FConstClassInputAccessPtr FGraphController::FindClassInputWithName(const FString& InName) const
 		{
 			using namespace FrontendControllerIntrinsics;
@@ -1910,22 +1870,22 @@ FGetClassInputFromClassWithName
 			return GraphClassPtr.GetMemberAccessPtr<FMetasoundFrontendClassOutput>(FGetClassOutputFromClassWithName(InName));
 		}
 
-		TArray<int32> FGraphController::GetDefaultIDsForInputVertex(const FString& InInputName) const
+		TArray<FGuid> FGraphController::GetDefaultIDsForInputVertex(const FString& InInputName) const
 		{
 			if (const FMetasoundFrontendClassInput* Input = FindClassInputWithName(InInputName).Get())
 			{
 				return Input->PointIDs;
 			}
-			return TArray<int32>();
+			return TArray<FGuid>();
 		}
 
-		TArray<int32> FGraphController::GetDefaultIDsForOutputVertex(const FString& InOutputName) const
+		TArray<FGuid> FGraphController::GetDefaultIDsForOutputVertex(const FString& InOutputName) const
 		{
 			if (const FMetasoundFrontendClassOutput* Output = FindClassOutputWithName(InOutputName).Get())
 			{
 				return Output->PointIDs;
 			}
-			return TArray<int32>();
+			return TArray<FGuid>();
 		}
 
 		TArray<FNodeHandle> FGraphController::GetNodes()
@@ -1938,7 +1898,7 @@ FGetClassInputFromClassWithName
 			return GetNodeHandles(GetNodesAndClasses());
 		}
 
-		FConstNodeHandle FGraphController::GetNodeWithID(int32 InNodeID) const
+		FConstNodeHandle FGraphController::GetNodeWithID(FGuid InNodeID) const
 		{
 			auto IsNodeWithSameID = [&](const FMetasoundFrontendClass& NodeClas, const FMetasoundFrontendNode& Node) 
 			{ 
@@ -1947,7 +1907,7 @@ FGetClassInputFromClassWithName
 			return GetNodeByPredicate(IsNodeWithSameID);
 		}
 
-		FNodeHandle FGraphController::GetNodeWithID(int32 InNodeID)
+		FNodeHandle FGraphController::GetNodeWithID(FGuid InNodeID)
 		{
 			auto IsNodeWithSameID = [&](const FMetasoundFrontendClass& NodeClas, const FMetasoundFrontendNode& Node) 
 			{ 
@@ -2107,7 +2067,7 @@ FGetClassInputFromClassWithName
 				}
 				else
 				{
-					UE_LOG(LogMetasound, Display, TEXT("Failed to add input. Input with same name \"%s\" exists in class [ClassID:%d]"), *InClassInput.Name, GraphClassPtr->ID);
+					UE_LOG(LogMetasound, Display, TEXT("Failed to add input. Input with same name \"%s\" exists in class [ClassID:%s]"), *InClassInput.Name, *GraphClassPtr->ID.ToString());
 				}
 			}
 			return FInvalidNodeController::GetInvalid();
@@ -2143,7 +2103,6 @@ FGetClassInputFromClassWithName
 					FMetasoundFrontendClassMetadata ClassMetadata;
 					if (FMetasoundFrontendRegistryContainer::GetOutputNodeClassMetadataForDataType(InClassOutput.TypeName, ClassMetadata))
 					{
-						//if (const FMetasoundFrontendClass* OutputClass = OwningDocument->FindOrAddClass(ClassMetadata))
 						if (FConstClassAccessPtr OutputClass = OwningDocument->FindOrAddClass(ClassMetadata))
 						{
 							FMetasoundFrontendNode& Node = GraphClass->Graph.Nodes.Emplace_GetRef(*OutputClass);
@@ -2185,7 +2144,7 @@ FGetClassInputFromClassWithName
 				}
 				else
 				{
-					UE_LOG(LogMetasound, Display, TEXT("Failed to add output. Output with same name \"%s\" exists in class [ClassID:%d]"), *InClassOutput.Name, GraphClassPtr->ID);
+					UE_LOG(LogMetasound, Display, TEXT("Failed to add output. Output with same name \"%s\" exists in class [ClassID:%s]"), *InClassOutput.Name, *GraphClassPtr->ID.ToString());
 				}
 			}
 
@@ -2234,32 +2193,32 @@ FGetClassInputFromClassWithName
 
 		// These can be used to set the default value for a given input on this graph.
 		// @returns false if the input name couldn't be found, or if the literal type was incompatible with the Data Type of this input.
-		bool FGraphController::SetDefaultInputToLiteral(const FString& InInputName, int32 InPointID, bool bInValue)
+		bool FGraphController::SetDefaultInputToLiteral(const FString& InInputName, FGuid InPointID, bool bInValue)
 		{
 			return SetDefaultInputToLiteralInternal(InInputName, InPointID, bInValue);
 		}
 
-		bool FGraphController::SetDefaultInputToLiteral(const FString& InInputName, int32 InPointID, int32 InValue)
+		bool FGraphController::SetDefaultInputToLiteral(const FString& InInputName, FGuid InPointID, int32 InValue)
 		{
 			return SetDefaultInputToLiteralInternal(InInputName, InPointID, InValue);
 		}
 
-		bool FGraphController::SetDefaultInputToLiteral(const FString& InInputName, int32 InPointID, float InValue)
+		bool FGraphController::SetDefaultInputToLiteral(const FString& InInputName, FGuid InPointID, float InValue)
 		{
 			return SetDefaultInputToLiteralInternal(InInputName, InPointID, InValue);
 		}
 
-		bool FGraphController::SetDefaultInputToLiteral(const FString& InInputName, int32 InPointID, const FString& InValue)
+		bool FGraphController::SetDefaultInputToLiteral(const FString& InInputName, FGuid InPointID, const FString& InValue)
 		{
 			return SetDefaultInputToLiteralInternal(InInputName, InPointID, InValue);
 		}
 
-		bool FGraphController::SetDefaultInputToLiteral(const FString& InInputName, int32 InPointID, UObject* InValue)
+		bool FGraphController::SetDefaultInputToLiteral(const FString& InInputName, FGuid InPointID, UObject* InValue)
 		{
 			return SetDefaultInputToLiteralInternal(InInputName, InPointID, InValue);
 		}
 
-		bool FGraphController::SetDefaultInputToLiteral(const FString& InInputName, int32 InPointID, const TArray<UObject*>& InValue)
+		bool FGraphController::SetDefaultInputToLiteral(const FString& InInputName, FGuid InPointID, const TArray<UObject*>& InValue)
 		{
 			return SetDefaultInputToLiteralInternal(InInputName, InPointID, InValue);
 		}
@@ -2299,7 +2258,7 @@ FGetClassInputFromClassWithName
 
 		// This can be used to clear the current literal for a given input.
 		// @returns false if the input name couldn't be found.
-		bool FGraphController::ClearLiteralForInput(const FString& InInputName, int32 InPointID)
+		bool FGraphController::ClearLiteralForInput(const FString& InInputName, FGuid InPointID)
 		{
 			if (FMetasoundFrontendClassInput* Desc = FindInputDescriptionWithName(InInputName))
 			{
@@ -2322,7 +2281,6 @@ FGetClassInputFromClassWithName
 		{
 			if (IsValid())
 			{
-				//if (const FMetasoundFrontendClass* DependencyDescription = OwningDocument->FindOrAddClass(InNodeClass))
 				if (FConstClassAccessPtr DependencyDescription = OwningDocument->FindOrAddClass(InNodeClass))
 				{
 					return AddNode(DependencyDescription);
@@ -2347,7 +2305,7 @@ FGetClassInputFromClassWithName
 		// On success, invalidates the received node handle.
 		bool FGraphController::RemoveNode(INodeController& InNode)
 		{
-			int32 NodeID = InNode.GetID();
+			FGuid NodeID = InNode.GetID();
 			const EMetasoundFrontendClassType NodeType = InNode.GetClassType();
 
 			auto IsNodeWithSameID = [&](const FMetasoundFrontendNode& InDesc) { return InDesc.ID == NodeID; };
@@ -2465,28 +2423,9 @@ FGetClassInputFromClassWithName
 
 					FNodeAccessPtr NodePtr = GraphClassPtr.GetMemberAccessPtr<FMetasoundFrontendNode>(FGetNodeFromGraphClassByNodeID(Node.ID));
 
-					//FGraphAccessPtr GraphPtr = GraphClassPtr.GetMemberAccessPtr<FMetasoundFrontendGraph>(FGetGraphFromGraphClass());
-
-					//return FNodeController::CreateNodeHandle(FNodeController::FInitParams{NodePtr, InExistingDependency, GraphPtr, this->AsShared()});
 					return GetNodeHandle(FGraphController::FNodeAndClass{NodePtr, InExistingDependency});
 				}
 			}
-			/*
-			if (GraphClassPtr.IsValid())
-			{
-				FMetasoundFrontendNode& NewNodeDescription = GraphClassPtr->Graph.Nodes.AddDefaulted_GetRef();
-
-				NewNodeDescription.Name = InExistingDependency.Metadata.Name.Name;
-				NewNodeDescription.ID = NewNodeID();
-				NewNodeDescription.ClassID = InExistingDependency.ID;
-
-				FNodeAccessPtr NodePtr = FrontendControllerIntrinsics::MakeAutoAccessPtr(NewNodeDescription);
-				FConstClassAccessPtr ClassPtr = FrontendControllerIntrinsics::MakeAutoAccessPtr(InExistingDependency);
-				FGraphAccessPtr GraphPtr = FrontendControllerIntrinsics::MakeAutoAccessPtr(GraphClassPtr->Graph);
-
-				return FNodeController::CreateNodeHandle(FNodeController::FInitParams{NodePtr, ClassPtr, GraphPtr, this->AsShared()});
-			}
-			*/
 
 			return FInvalidNodeController::GetInvalid();
 		}
@@ -2543,21 +2482,9 @@ FGetClassInputFromClassWithName
 			return false;
 		}
 
-		int32 FGraphController::NewNodeID() const
+		FGuid FGraphController::NewNodeID() const
 		{
-			if (GraphClassPtr.IsValid())
-			{
-				TSet<int32> UsedIDs;
-
-				for (const FMetasoundFrontendNode& Node : GraphClassPtr->Graph.Nodes)
-				{
-					UsedIDs.Add(Node.ID);
-				}
-
-				return FrontendControllerIntrinsics::GetNewID(UsedIDs);
-			}
-
-			return Metasound::FrontendInvalidID;
+			return FGuid::NewGuid();
 		}
 
 		bool FGraphController::ContainsNodesAndClassesByPredicate(TFunctionRef<bool (const FMetasoundFrontendClass&, const FMetasoundFrontendNode&)> InPredicate) const
@@ -2576,7 +2503,7 @@ FGetClassInputFromClassWithName
 					}
 					else
 					{
-						UE_LOG(LogMetasound, Warning, TEXT("Failed to find class for node [NodeID:%d, ClassID:%d]"), Node.ID, Node.ClassID);
+						UE_LOG(LogMetasound, Warning, TEXT("Failed to find class for node [NodeID:%s, ClassID:%s]"), *Node.ID.ToString(), *Node.ClassID.ToString());
 					}
 				}
 			}
@@ -2594,15 +2521,15 @@ FGetClassInputFromClassWithName
 				for (FMetasoundFrontendNode& Node : GraphClassPtr->Graph.Nodes)
 				{
 					FNodeAccessPtr NodePtr = GraphClassPtr.GetMemberAccessPtr<FMetasoundFrontendNode>(FGetNodeFromGraphClassByNodeID(Node.ID));
-					//const FMetasoundFrontendClass* NodeClass = OwningDocument->FindClassWithID(Node.ClassID);
 					FConstClassAccessPtr NodeClassPtr = OwningDocument->FindClassWithID(Node.ClassID);
+
 					if (NodeClassPtr && NodePtr)
 					{
 						NodesAndClasses.Add(FGraphController::FNodeAndClass{ NodePtr, NodeClassPtr });
 					}
 					else
 					{
-						UE_LOG(LogMetasound, Warning, TEXT("Failed to find class for node [NodeID:%d, ClassID:%d]"), Node.ID, Node.ClassID);
+						UE_LOG(LogMetasound, Warning, TEXT("Failed to find class for node [NodeID:%s, ClassID:%s]"), *Node.ID.ToString(), *Node.ClassID.ToString());
 					}
 				}
 			}
@@ -2620,28 +2547,15 @@ FGetClassInputFromClassWithName
 				for (const FMetasoundFrontendNode& Node : GraphClassPtr->Graph.Nodes)
 				{
 					FConstNodeAccessPtr NodePtr = GraphClassPtr.GetMemberAccessPtr<const FMetasoundFrontendNode>(FGetConstNodeFromGraphClassByNodeID(Node.ID));
-					//const FMetasoundFrontendClass* NodeClass = OwningDocument->FindClassWithID(Node.ClassID);
 					FConstClassAccessPtr NodeClassPtr = OwningDocument->FindClassWithID(Node.ClassID);
+
 					if (NodeClassPtr && NodePtr)
 					{
 						NodesAndClasses.Add(FGraphController::FConstNodeAndClass{ NodePtr, NodeClassPtr });
 					}
-					/*
-					const FMetasoundFrontendClass* NodeClass = OwningDocument->FindClassWithID(Node.ClassID);
-					if (nullptr != NodeClass)
-					{
-						NodesAndClasses.Emplace(
-							FGraphController::FConstNodeAndClass
-							{
-								FrontendControllerIntrinsics::MakeAutoAccessPtr(Node),
-								FrontendControllerIntrinsics::MakeAutoAccessPtr(*NodeClass)
-							}
-						);
-					}
-					*/
 					else
 					{
-						UE_LOG(LogMetasound, Warning, TEXT("Failed to find class for node [NodeID:%d, ClassID:%d]"), Node.ID, Node.ClassID);
+						UE_LOG(LogMetasound, Warning, TEXT("Failed to find class for node [NodeID:%s, ClassID:%s]"), *Node.ID.ToString(), *Node.ClassID.ToString());
 					}
 				}
 			}
@@ -2670,7 +2584,7 @@ FGetClassInputFromClassWithName
 					}
 					else
 					{
-						UE_LOG(LogMetasound, Warning, TEXT("Failed to find class for node [NodeID:%d, ClassID:%d]"), Node.ID, Node.ClassID);
+						UE_LOG(LogMetasound, Warning, TEXT("Failed to find class for node [NodeID:%s, ClassID:%s]"), *Node.ID.ToString(), *Node.ClassID.ToString());
 					}
 				}
 			}
@@ -2699,7 +2613,7 @@ FGetClassInputFromClassWithName
 					}
 					else
 					{
-						UE_LOG(LogMetasound, Warning, TEXT("Failed to find class for node [NodeID:%d, ClassID:%d]"), Node.ID, Node.ClassID);
+						UE_LOG(LogMetasound, Warning, TEXT("Failed to find class for node [NodeID:%s, ClassID:%s]"), *Node.ID.ToString(), *Node.ClassID.ToString());
 					}
 				}
 			}
@@ -2795,7 +2709,6 @@ FGetClassInputFromClassWithName
 			{
 				FGraphHandle OwningGraph = this->AsShared();
 				FGraphAccessPtr GraphPtr = GraphClassPtr.GetMemberAccessPtr<FMetasoundFrontendGraph>(FGetGraphFromGraphClass());
-				//FGraphAccessPtr GraphPtr = FrontendControllerIntrinsics::MakeAutoAccessPtr(GraphClassPtr->Graph);
 
 				switch (InNodeAndClass.Class->Metadata.Type)
 				{
@@ -2858,8 +2771,8 @@ FGetClassInputFromClassWithName
 			if (InNodeAndClass.IsValid() && GraphClassPtr.IsValid())
 			{
 				FConstGraphHandle OwningGraph = this->AsShared();
-				//FConstGraphAccessPtr GraphPtr = FrontendControllerIntrinsics::MakeAutoAccessPtr(GraphClassPtr->Graph);
 				FConstGraphAccessPtr GraphPtr = GraphClassPtr.GetMemberAccessPtr<const FMetasoundFrontendGraph>(FGetConstGraphFromGraphClass());
+
 				switch (InNodeAndClass.Class->Metadata.Type)
 				{
 					case EMetasoundFrontendClassType::Input:
@@ -2951,83 +2864,52 @@ FGetClassInputFromClassWithName
 			return nullptr;
 		}
 
-		FClassInputAccessPtr FGraphController::FindInputDescriptionWithNodeID(int32 InNodeID)
+		FClassInputAccessPtr FGraphController::FindInputDescriptionWithNodeID(FGuid InNodeID)
 		{
 			using namespace FrontendControllerIntrinsics;
 
 			return GraphClassPtr.GetMemberAccessPtr<FMetasoundFrontendClassInput>(FGetClassInputFromGraphClassWithNodeID(InNodeID));
-			/*
-			if (GraphClassPtr.IsValid())
-			{
-				return GraphClassPtr->Interface.Inputs.FindByPredicate([&](const FMetasoundFrontendClassInput& Desc) { return Desc.NodeID == InNodeID; });
-			}
-			return nullptr;
-			*/
 		}
 
-		FConstClassInputAccessPtr FGraphController::FindInputDescriptionWithNodeID(int32 InNodeID) const
+		FConstClassInputAccessPtr FGraphController::FindInputDescriptionWithNodeID(FGuid InNodeID) const
 		{
 			using namespace FrontendControllerIntrinsics;
 
 			return GraphClassPtr.GetMemberAccessPtr<const FMetasoundFrontendClassInput>(FGetConstClassInputFromGraphClassWithNodeID(InNodeID));
-			/*
-			if (GraphClassPtr.IsValid())
-			{
-				return GraphClassPtr->Interface.Inputs.FindByPredicate([&](const FMetasoundFrontendClassInput& Desc) { return Desc.NodeID == InNodeID; });
-			}
-			return nullptr;
-			*/
 		}
 
-		FClassOutputAccessPtr FGraphController::FindOutputDescriptionWithNodeID(int32 InNodeID)
+		FClassOutputAccessPtr FGraphController::FindOutputDescriptionWithNodeID(FGuid InNodeID)
 		{
 			using namespace FrontendControllerIntrinsics;
 
 			return GraphClassPtr.GetMemberAccessPtr<FMetasoundFrontendClassOutput>(FGetClassOutputFromGraphClassWithNodeID(InNodeID));
-			/*
-			if (GraphClassPtr.IsValid())
-			{
-				return GraphClassPtr->Interface.Outputs.FindByPredicate([&](const FMetasoundFrontendClassOutput& Desc) { return Desc.NodeID == InNodeID; });
-			}
-			return nullptr;
-			*/
 		}
 
-		FConstClassOutputAccessPtr FGraphController::FindOutputDescriptionWithNodeID(int32 InNodeID) const
+		FConstClassOutputAccessPtr FGraphController::FindOutputDescriptionWithNodeID(FGuid InNodeID) const
 		{
 			using namespace FrontendControllerIntrinsics;
 
 			return GraphClassPtr.GetMemberAccessPtr<const FMetasoundFrontendClassOutput>(FGetConstClassOutputFromGraphClassWithNodeID(InNodeID));
-			/*
-			if (GraphClassPtr.IsValid())
-			{
-				return GraphClassPtr->Interface.Outputs.FindByPredicate([&](const FMetasoundFrontendClassOutput& Desc) { return Desc.NodeID == InNodeID; });
-			}
-			return nullptr;
-			*/
 		}
 
-		/*
-		FMetasoundFrontendNode* FGraphController::FindNodeByID(int32 InNodeID)
+		FDocumentAccess FGraphController::ShareAccess() 
 		{
-			if (GraphClassPtr.IsValid())
-			{
-				auto IsNodeWithID = [&](const FMetasoundFrontendNode& InNode) { return InNode.ID == InNodeID; };
-				return GraphClassPtr->Graph.Nodes.FindByPredicate(IsNodeWithID);
-			}
-			return nullptr;
+			FDocumentAccess Access;
+
+			Access.GraphClass = GraphClassPtr;
+			Access.ConstGraphClass = GraphClassPtr;
+
+			return Access;
 		}
 
-		const FMetasoundFrontendNode* FGraphController::FindNodeByID(int32 InNodeID) const
+		FConstDocumentAccess FGraphController::ShareAccess() const
 		{
-			if (GraphClassPtr.IsValid())
-			{
-				auto IsNodeWithID = [&](const FMetasoundFrontendNode& InNode) { return InNode.ID == InNodeID; };
-				return GraphClassPtr->Graph.Nodes.FindByPredicate(IsNodeWithID);
-			}
-			return nullptr;
+			FConstDocumentAccess Access;
+
+			Access.ConstGraphClass = GraphClassPtr;
+
+			return Access;
 		}
-		*/
 
 		//
 		// FDocumentController
@@ -3109,36 +2991,19 @@ FGetClassInputFromClassWithName
 			return Classes;
 		}
 
-		FConstClassAccessPtr FDocumentController::FindDependencyWithID(int32 InClassID) const 
+		FConstClassAccessPtr FDocumentController::FindDependencyWithID(FGuid InClassID) const 
 		{
 			using namespace FrontendControllerIntrinsics;
 			return DocumentPtr.GetMemberAccessPtr<const FMetasoundFrontendClass>(FGetConstClassFromDocumentDependenciesWithID(InClassID));
-
-			/*
-			if (DocumentPtr.IsValid())
-			{
-				auto IsMatchingID = [&](const FMetasoundFrontendClass& InClass) { return InClass.ID == InClassID; };
-				return DocumentPtr->Dependencies.FindByPredicate(IsMatchingID);
-			}
-			return nullptr;
-			*/
 		}
 
-		FConstGraphClassAccessPtr FDocumentController::FindSubgraphWithID(int32 InClassID) const
+		FConstGraphClassAccessPtr FDocumentController::FindSubgraphWithID(FGuid InClassID) const
 		{
 			using namespace FrontendControllerIntrinsics;
 			return DocumentPtr.GetMemberAccessPtr<const FMetasoundFrontendGraphClass>(FGetConstGraphClassFromDocumentWithID(InClassID));
-			/*
-			if (DocumentPtr.IsValid())
-			{
-				auto IsMatchingID = [&](const FMetasoundFrontendGraphClass& InClass) { return InClass.ID == InClassID; };
-				return DocumentPtr->Subgraphs.FindByPredicate(IsMatchingID);
-			}
-			return nullptr;
-			*/
 		}
 
-		FConstClassAccessPtr FDocumentController::FindClassWithID(int32 InClassID) const
+		FConstClassAccessPtr FDocumentController::FindClassWithID(FGuid InClassID) const
 		{
 			FConstClassAccessPtr MetasoundClass = FindDependencyWithID(InClassID);
 
@@ -3155,26 +3020,6 @@ FGetClassInputFromClassWithName
 			using namespace FrontendControllerIntrinsics;
 
 			return DocumentPtr.GetMemberAccessPtr<const FMetasoundFrontendClass>(FGetConstClassFromDocumentWithClassInfo(InNodeClass));
-			/*
-			if (DocumentPtr.IsValid())
-			{
-				auto IsMatchingClass = [&](const FMetasoundFrontendClass& Desc)
-				{
-					return IsMatchingMetasoundClass(InNodeClass, Desc.Metadata);
-				};
-
-				FConstClassAccessPtr MetasoundClass = DocumentPtr->Dependencies.FindByPredicate(IsMatchingClass);
-
-				if (nullptr == MetasoundClass)
-				{
-					MetasoundClass = DocumentPtr->Subgraphs.FindByPredicate(IsMatchingClass);
-				}
-
-				return MetasoundClass;
-			}
-
-			return nullptr;
-			*/
 		}
 
 		FConstClassAccessPtr FDocumentController::FindOrAddClass(const FNodeClassInfo& InNodeClass)
@@ -3204,23 +3049,6 @@ FGetClassInputFromClassWithName
 			using namespace FrontendControllerIntrinsics;
 
 			return DocumentPtr.GetMemberAccessPtr<const FMetasoundFrontendClass>(FGetConstClassFromDocumentWithMetadata(InMetadata));
-			/*
-			if (DocumentPtr.IsValid())
-			{
-				auto IsMatchingClass = [&](const FMetasoundFrontendClass& Desc)
-				{
-					return IsMatchingMetasoundClass(InMetadata, Desc.Metadata);
-				};
-
-				FConstClassAccessPtr MetasoundClass = DocumentPtr->Dependencies.FindByPredicate(IsMatchingClass);
-
-				if (nullptr == MetasoundClass)
-				{
-					MetasoundClass = DocumentPtr->Subgraphs.FindByPredicate(IsMatchingClass);
-				}
-			}
-			return nullptr;
-			*/
 		}
 
 		FConstClassAccessPtr FDocumentController::FindOrAddClass(const FMetasoundFrontendClassMetadata& InMetadata)
@@ -3262,143 +3090,6 @@ FGetClassInputFromClassWithName
 			return ClassPtr;
 		}
 
-		/*
-		const FMetasoundFrontendClass* FDocumentController::FindDependencyWithID(int32 InClassID) const 
-		{
-			if (DocumentPtr.IsValid())
-			{
-				auto IsMatchingID = [&](const FMetasoundFrontendClass& InClass) { return InClass.ID == InClassID; };
-				return DocumentPtr->Dependencies.FindByPredicate(IsMatchingID);
-			}
-			return nullptr;
-		}
-
-		const FMetasoundFrontendGraphClass* FDocumentController::FindSubgraphWithID(int32 InClassID) const
-		{
-			if (DocumentPtr.IsValid())
-			{
-				auto IsMatchingID = [&](const FMetasoundFrontendGraphClass& InClass) { return InClass.ID == InClassID; };
-				return DocumentPtr->Subgraphs.FindByPredicate(IsMatchingID);
-			}
-			return nullptr;
-		}
-
-		const FMetasoundFrontendClass* FDocumentController::FindClassWithID(int32 InClassID) const
-		{
-			const FMetasoundFrontendClass* MetasoundClass = FindDependencyWithID(InClassID);
-
-			if (nullptr == MetasoundClass)
-			{
-				MetasoundClass = FindSubgraphWithID(InClassID);
-			}
-
-			return MetasoundClass;
-		}
-
-		const FMetasoundFrontendClass* FDocumentController::FindClass(const FNodeClassInfo& InNodeClass) const
-		{
-			if (DocumentPtr.IsValid())
-			{
-				auto IsMatchingClass = [&](const FMetasoundFrontendClass& Desc)
-				{
-					return IsMatchingMetasoundClass(InNodeClass, Desc.Metadata);
-				};
-
-				const FMetasoundFrontendClass* MetasoundClass = DocumentPtr->Dependencies.FindByPredicate(IsMatchingClass);
-
-				if (nullptr == MetasoundClass)
-				{
-					MetasoundClass = DocumentPtr->Subgraphs.FindByPredicate(IsMatchingClass);
-				}
-
-				return MetasoundClass;
-			}
-
-			return nullptr;
-		}
-
-		const FMetasoundFrontendClass* FDocumentController::FindOrAddClass(const FNodeClassInfo& InNodeClass)
-		{
-			if (DocumentPtr.IsValid())
-			{
-
-				if (const FMetasoundFrontendClass* MetasoundClass = FindClass(InNodeClass))
-				{
-					return MetasoundClass;
-				}
-
-				// FNodeClassInfo does not contain enough info add a subgraph.
-				check(EMetasoundFrontendClassType::Graph != InNodeClass.NodeType);
-
-				FMetasoundFrontendClass NewClass = GenerateClassDescription(InNodeClass);
-				NewClass.ID = NewClassID();
-
-				return &(DocumentPtr->Dependencies.Add_GetRef(NewClass));
-			}
-
-			return nullptr;
-		}
-
-		const FMetasoundFrontendClass* FDocumentController::FindClass(const FMetasoundFrontendClassMetadata& InMetadata) const
-		{
-			if (DocumentPtr.IsValid())
-			{
-				auto IsMatchingClass = [&](const FMetasoundFrontendClass& Desc)
-				{
-					return IsMatchingMetasoundClass(InMetadata, Desc.Metadata);
-				};
-
-				const FMetasoundFrontendClass* MetasoundClass = DocumentPtr->Dependencies.FindByPredicate(IsMatchingClass);
-
-				if (nullptr == MetasoundClass)
-				{
-					MetasoundClass = DocumentPtr->Subgraphs.FindByPredicate(IsMatchingClass);
-				}
-			}
-			return nullptr;
-		}
-
-		const FMetasoundFrontendClass* FDocumentController::FindOrAddClass(const FMetasoundFrontendClassMetadata& InMetadata)
-		{
-			if (DocumentPtr.IsValid())
-			{
-				if (const FMetasoundFrontendClass* MetasoundClass = FindClass(InMetadata))
-				{
-					return MetasoundClass;
-				}
-
-				if ((EMetasoundFrontendClassType::External == InMetadata.Type) || (EMetasoundFrontendClassType::Input == InMetadata.Type) || (EMetasoundFrontendClassType::Output == InMetadata.Type))
-				{
-					FMetasoundFrontendClass NewClass;
-					if (FMetasoundFrontendRegistryContainer::GetFrontendClassFromRegistered(InMetadata, NewClass))
-					{
-						NewClass.ID = NewClassID();
-						return &(DocumentPtr->Dependencies.Add_GetRef(NewClass));
-					}
-					else
-					{
-						UE_LOG(LogMetasound, Error, TEXT("Cannot add external dependency. No Metasound class found with matching metadata Name: \"%s\", Version %d.%d. Suggested solution \"%s\" by %s."), *InMetadata.Name.GetFullName(), InMetadata.Version.Major, InMetadata.Version.Minor, *InMetadata.PromptIfMissing.ToString(), *InMetadata.Author.ToString());
-					}
-				} 
-				else if (EMetasoundFrontendClassType::Graph == InMetadata.Type)
-				{
-					FMetasoundFrontendGraphClass NewClass;
-					NewClass.ID = NewClassID();
-					NewClass.Metadata = InMetadata;
-
-					return &(DocumentPtr->Subgraphs.Add_GetRef(NewClass));
-				}
-				else
-				{
-					UE_LOG(LogMetasound, Error, TEXT("Unsupported metasound class type for node: \"%s\", Version %d.%d."), *InMetadata.Name.GetFullName(), InMetadata.Version.Major, InMetadata.Version.Minor);
-					checkNoEntry();
-				}
-			}
-
-			return nullptr;
-		}
-		*/
-
 		void FDocumentController::RemoveUnreferencedDependencies()
 		{
 			if (DocumentPtr.IsValid())
@@ -3409,7 +3100,7 @@ FGetClassInputFromClassWithName
 				// no unreferenced dependencies left. 
 				do
 				{
-					TSet<int32> ReferencedDependencyIDs;
+					TSet<FGuid> ReferencedDependencyIDs;
 					auto AddNodeClassIDToSet = [&](const FMetasoundFrontendNode& Node)
 					{
 						ReferencedDependencyIDs.Add(Node.ClassID);
@@ -3440,7 +3131,6 @@ FGetClassInputFromClassWithName
 		{
 			if (DocumentPtr.IsValid())
 			{
-				//FGraphClassAccessPtr GraphClass = FrontendControllerIntrinsics::MakeAutoAccessPtr(DocumentPtr->RootGraph);
 				FGraphClassAccessPtr GraphClass = DocumentPtr.GetMemberAccessPtr<FMetasoundFrontendGraphClass>([](FMetasoundFrontendDocument& Doc) { return &Doc.RootGraph; });
 				return FGraphController::CreateGraphHandle(FGraphController::FInitParams{GraphClass, this->AsShared()});
 			}
@@ -3452,7 +3142,6 @@ FGetClassInputFromClassWithName
 		{
 			if (DocumentPtr.IsValid())
 			{
-				//FGraphClassAccessPtr GraphClass = FrontendControllerIntrinsics::MakeAutoAccessPtr(DocumentPtr->RootGraph);
 				FConstGraphClassAccessPtr GraphClass = DocumentPtr.GetMemberAccessPtr<const FMetasoundFrontendGraphClass>([](const FMetasoundFrontendDocument& Doc) { return &Doc.RootGraph; });
 				return FGraphController::CreateConstGraphHandle(FGraphController::FInitParams
 					{
@@ -3471,10 +3160,6 @@ FGetClassInputFromClassWithName
 			{
 				for (FMetasoundFrontendGraphClass& GraphClass : DocumentPtr->Subgraphs)
 				{
-					/*
-					FGraphClassAccessPtr GraphClassPtr = FrontendControllerIntrinsics::MakeAutoAccessPtr(GraphClass);
-					FGraphController::CreateGraphHandle(FGraphController::FInitParams{GraphClassPtr, this->AsShared()}));
-					*/
 					Subgraphs.Add(GetSubgraphWithClassID(GraphClass.ID));
 				}
 			}
@@ -3490,16 +3175,6 @@ FGetClassInputFromClassWithName
 			{
 				for (const FMetasoundFrontendGraphClass& GraphClass : DocumentPtr->Subgraphs)
 				{
-					/*
-					FConstGraphClassAccessPtr GraphClassPtr = FrontendControllerIntrinsics::MakeAutoAccessPtr(GraphClass);
-					Subgraphs.Add(FGraphController::CreateConstGraphHandle(
-						FGraphController::FInitParams
-						{
-							ConstCastAccessPtr<FMetasoundFrontendGraphClass>(GraphClassPtr), 
-							ConstCastSharedRef<IDocumentController>(this->AsShared())
-						}
-					));
-					*/
 					Subgraphs.Add(GetSubgraphWithClassID(GraphClass.ID));
 				}
 			}
@@ -3507,60 +3182,17 @@ FGetClassInputFromClassWithName
 			return Subgraphs;
 		}
 
-		FGraphHandle FDocumentController::GetSubgraphWithClassID(int32 InClassID)
+		FGraphHandle FDocumentController::GetSubgraphWithClassID(FGuid InClassID)
 		{
 			FGraphClassAccessPtr GraphClassPtr = DocumentPtr.GetMemberAccessPtr<FMetasoundFrontendGraphClass>(FrontendControllerIntrinsics::FGetSubgraphFromDocumentByID(InClassID));
 			return FGraphController::CreateGraphHandle(FGraphController::FInitParams{GraphClassPtr, this->AsShared()});
-			/*
-			if (DocumentPtr.IsValid())
-			{
-
-
-					[InClassID](FMetasoundFrontendDocument& Doc)
-					{
-						return Doc.Subgraphs.FindByPredicate(
-							[&](const FMetasoundFrontendGraphClass& InGraphClass) 
-							{ 
-								return InGraphClass.ID == InClassID; 
-							}
-						);
-					}
-				);
-				*/
-
-				/*
-
-
-				auto IsMatchingID = [&](const FMetasoundFrontendGraphClass& InGraphClass) { return InGraphClass.ID == InClassID; };
-				if (FMetasoundFrontendGraphClass* GraphClass = DocumentPtr->Subgraphs.FindByPredicate(IsMatchingID))
-				{
-					//FGraphClassAccessPtr GraphClassPtr = FrontendControllerIntrinsics::MakeAutoAccessPtr(*GraphClass);
-					return FGraphController::CreateGraphHandle(FGraphController::FInitParams{GraphClassPtr, this->AsShared()});
-				}
-			}
-			return FInvalidGraphController::GetInvalid();
-				*/
 		}
 
-		FConstGraphHandle FDocumentController::GetSubgraphWithClassID(int32 InClassID) const
+		FConstGraphHandle FDocumentController::GetSubgraphWithClassID(FGuid InClassID) const
 		{
 			FConstGraphClassAccessPtr GraphClassPtr = DocumentPtr.GetMemberAccessPtr<const FMetasoundFrontendGraphClass>(FrontendControllerIntrinsics::FGetSubgraphFromDocumentByID(InClassID));
 
 			return FGraphController::CreateConstGraphHandle(FGraphController::FInitParams{ConstCastAccessPtr<FMetasoundFrontendGraphClass>(GraphClassPtr), ConstCastSharedRef<IDocumentController>(this->AsShared())});
-
-			/*
-			if (DocumentPtr.IsValid())
-			{
-				auto IsMatchingID = [&](const FMetasoundFrontendGraphClass& InGraphClass) { return InGraphClass.ID == InClassID; };
-
-				if (const FMetasoundFrontendGraphClass* GraphClass = DocumentPtr->Subgraphs.FindByPredicate(IsMatchingID))
-				{
-					FConstGraphClassAccessPtr GraphClassPtr = FrontendControllerIntrinsics::MakeAutoAccessPtr(*GraphClass);
-					return FGraphController::CreateConstGraphHandle(FGraphController::FInitParams{ConstCastAccessPtr<FMetasoundFrontendGraphClass>(GraphClassPtr), ConstCastSharedRef<IDocumentController>(this->AsShared())});
-				}
-			}
-			return FInvalidGraphController::GetInvalid();
-			*/
 		}
 
 		bool FDocumentController::ExportToJSONAsset(const FString& InAbsolutePath) const
@@ -3629,24 +3261,28 @@ FGetClassInputFromClassWithName
 			return false;
 		}
 
-		int32 FDocumentController::NewClassID() const 
+		FDocumentAccess FDocumentController::ShareAccess() 
 		{
-			if (DocumentPtr.IsValid())
-			{
-				TSet<int32> UsedIDs;
-				for (const FMetasoundFrontendClass& InDesc : DocumentPtr->Dependencies)
-				{
-					UsedIDs.Add(InDesc.ID);
-				}
-				for (const FMetasoundFrontendClass& InDesc : DocumentPtr->Subgraphs)
-				{
-					UsedIDs.Add(InDesc.ID);
-				}
+			FDocumentAccess Access;
 
-				return FrontendControllerIntrinsics::GetNewID(UsedIDs);
-			}
+			Access.Document = DocumentPtr;
+			Access.ConstDocument = DocumentPtr;
 
-			return Metasound::FrontendInvalidID;
+			return Access;
+		}
+
+		FConstDocumentAccess FDocumentController::ShareAccess() const 
+		{
+			FConstDocumentAccess Access;
+
+			Access.ConstDocument = DocumentPtr;
+
+			return Access;
+		}
+
+		FGuid FDocumentController::NewClassID() const 
+		{
+			return FGuid::NewGuid();
 		}
 	}
 }
diff --git a/Engine/Plugins/Runtime/Metasound/Source/MetasoundFrontend/Private/MetasoundFrontendStandardController.h b/Engine/Plugins/Runtime/Metasound/Source/MetasoundFrontend/Private/MetasoundFrontendStandardController.h
--- a/Engine/Plugins/Runtime/Metasound/Source/MetasoundFrontend/Private/MetasoundFrontendStandardController.h
+++ b/Engine/Plugins/Runtime/Metasound/Source/MetasoundFrontend/Private/MetasoundFrontendStandardController.h
@@ -14,25 +14,6 @@ namespace Metasound
 {
 	namespace Frontend
 	{
-		/*
-		using FVertexAccessPtr = TAccessPtr<FMetasoundFrontendVertex>;
-		using FConstVertexAccessPtr = TAccessPtr<const FMetasoundFrontendVertex>;
-		using FClassOutputAccessPtr = TAccessPtr<FMetasoundFrontendClassOutput>;
-		using FConstClassOutputAccessPtr = TAccessPtr<const FMetasoundFrontendClassOutput>;
-		using FClassInputAccessPtr = TAccessPtr<FMetasoundFrontendClassInput>;
-		using FConstClassInputAccessPtr = TAccessPtr<const FMetasoundFrontendClassInput>;
-		using FNodeAccessPtr = TAccessPtr<FMetasoundFrontendNode>;
-		using FConstNodeAccessPtr = TAccessPtr<const FMetasoundFrontendNode>;
-		using FGraphAccessPtr = TAccessPtr<FMetasoundFrontendGraph>;
-		using FConstGraphAccessPtr = TAccessPtr<const FMetasoundFrontendGraph>;
-		using FClassAccessPtr = TAccessPtr<FMetasoundFrontendClass>;
-		using FConstClassAccessPtr = TAccessPtr<const FMetasoundFrontendClass>;
-		using FGraphClassAccessPtr = TAccessPtr<FMetasoundFrontendGraphClass>;
-		using FConstGraphClassAccessPtr = TAccessPtr<const FMetasoundFrontendGraphClass>;
-		using FDocumentAccessPtr = TAccessPtr<FMetasoundFrontendDocument>;
-		using FConstDocumentAccessPtr = TAccessPtr<const FMetasoundFrontendDocument>;
-		*/
-
 		/** FBaseOutputController provides common functionality for multiple derived
 		 * output controllers.
 		 */
@@ -42,7 +23,7 @@ namespace Metasound
 
 			struct FInitParams
 			{
-				int32 ID;
+				FGuid ID;
 
 				FConstVertexAccessPtr NodeVertexPtr;
 
@@ -57,12 +38,12 @@ namespace Metasound
 
 			bool IsValid() const override;
 
-			int32 GetID() const override;
+			FGuid GetID() const override;
 			const FName& GetDataType() const override;
 			const FString& GetName() const override;
 
 			// Return info on containing node. 
-			int32 GetOwningNodeID() const override;
+			FGuid GetOwningNodeID() const override;
 			FNodeHandle GetOwningNode() override;
 			FConstNodeHandle GetOwningNode() const override;
 
@@ -73,8 +54,10 @@ namespace Metasound
 			bool Disconnect(IInputController& InController) override;
 
 		protected:
+			FDocumentAccess ShareAccess() override;
+			FConstDocumentAccess ShareAccess() const override;
 
-			int32 ID;
+			FGuid ID;
 			FConstVertexAccessPtr NodeVertexPtr;	
 			FNodeHandle OwningNode;
 		};
@@ -85,7 +68,7 @@ namespace Metasound
 		public:
 			struct FInitParams
 			{
-				int32 ID;
+				FGuid ID;
 
 				FConstVertexAccessPtr NodeVertexPtr;
 			
@@ -105,8 +88,11 @@ namespace Metasound
 			const FText& GetDisplayName() const override;
 			const FText& GetTooltip() const override;
 
-		private:
+		protected:
+			FDocumentAccess ShareAccess() override;
+			FConstDocumentAccess ShareAccess() const override;
 
+		private:
 			FConstClassOutputAccessPtr ClassOutputPtr;
 		};
 
@@ -121,7 +107,7 @@ namespace Metasound
 		public:
 			struct FInitParams
 			{
-				int32 ID;
+				FGuid ID;
 
 				FConstVertexAccessPtr NodeVertexPtr;
 			
@@ -142,6 +128,11 @@ namespace Metasound
 			const FText& GetDisplayName() const override;
 			const FText& GetTooltip() const override;
 
+
+		protected:
+			FDocumentAccess ShareAccess() override;
+			FConstDocumentAccess ShareAccess() const override;
+
 		private:
 
 			FConstClassInputAccessPtr OwningGraphClassInputPtr;
@@ -157,7 +148,7 @@ namespace Metasound
 
 			struct FInitParams
 			{
-				int32 ID; 
+				FGuid ID; 
 				FConstVertexAccessPtr NodeVertexPtr;
 				FGraphAccessPtr GraphPtr; 
 				FNodeHandle OwningNode;
@@ -170,12 +161,12 @@ namespace Metasound
 
 			bool IsValid() const override;
 
-			int32 GetID() const override;
+			FGuid GetID() const override;
 			const FName& GetDataType() const override;
 			const FString& GetName() const override;
 
 			// Owning node info
-			int32 GetOwningNodeID() const override;
+			FGuid GetOwningNodeID() const override;
 			FNodeHandle GetOwningNode() override;
 			FConstNodeHandle GetOwningNode() const override;
 
@@ -195,12 +186,15 @@ namespace Metasound
 
 		protected:
 
+			FDocumentAccess ShareAccess() override;
+			FConstDocumentAccess ShareAccess() const override;
+
 			const FMetasoundFrontendEdge* FindEdge() const;
 			FMetasoundFrontendEdge* FindEdge();
 
 		private:
 
-			int32 ID;
+			FGuid ID;
 			FConstVertexAccessPtr NodeVertexPtr;
 			FGraphAccessPtr GraphPtr;
 			FNodeHandle OwningNode;
@@ -213,7 +207,7 @@ namespace Metasound
 
 			struct FInitParams
 			{
-				int32 ID; 
+				FGuid ID; 
 				FConstVertexAccessPtr NodeVertexPtr;
 				FConstClassInputAccessPtr OwningGraphClassInputPtr;
 				FGraphAccessPtr GraphPtr; 
@@ -229,6 +223,11 @@ namespace Metasound
 			const FText& GetDisplayName() const override;
 			const FText& GetTooltip() const override;
 
+		protected:
+
+			FDocumentAccess ShareAccess() override;
+			FConstDocumentAccess ShareAccess() const override;
+
 		private:
 
 			FConstClassInputAccessPtr OwningGraphClassInputPtr;
@@ -245,7 +244,7 @@ namespace Metasound
 		public:
 			struct FInitParams
 			{
-				int32 ID; 
+				FGuid ID; 
 				FConstVertexAccessPtr NodeVertexPtr;
 				FConstClassOutputAccessPtr OwningGraphClassOutputPtr;
 				FGraphAccessPtr GraphPtr; 
@@ -261,6 +260,11 @@ namespace Metasound
 			const FText& GetDisplayName() const override;
 			const FText& GetTooltip() const override;
 
+		protected:
+
+			FDocumentAccess ShareAccess() override;
+			FConstDocumentAccess ShareAccess() const override;
+
 		private:
 
 			FConstClassOutputAccessPtr OwningGraphClassOutputPtr;
@@ -286,13 +290,13 @@ namespace Metasound
 			bool IsValid() const override;
 
 			// Owning graph info
-			int32 GetOwningGraphClassID() const override;
+			FGuid GetOwningGraphClassID() const override;
 			FGraphHandle GetOwningGraph() override;
 			FConstGraphHandle GetOwningGraph() const override;
 
 			// Info about this node.
-			int32 GetID() const override;
-			int32 GetClassID() const override;
+			FGuid GetID() const override;
+			FGuid GetClassID() const override;
 			FMetasoundFrontendVersionNumber GetClassVersionNumber() const override;
 			const FText& GetClassDescription() const override;
 			FMetasoundFrontendClassDisplayInfo GetClassDisplayInfo() const override;
@@ -302,11 +306,11 @@ namespace Metasound
 
 			bool CanAddInput(const FString& InVertexName) const override;
 			FInputHandle AddInput(const FString& InVertexName, const FMetasoundFrontendLiteral* InDefault) override;
-			bool RemoveInput(int32 InPointID) override;
+			bool RemoveInput(FGuid InPointID) override;
 
 			bool CanAddOutput(const FString& InVertexName) const override;
 			FInputHandle AddOutput(const FString& InVertexName, const FMetasoundFrontendLiteral* InDefault) override;
-			bool RemoveOutput(int32 InPointID) override;
+			bool RemoveOutput(FGuid InPointID) override;
 
 			const FString& GetClassName() const override;
 
@@ -332,31 +336,34 @@ namespace Metasound
 			 *
 			 * If the input does not exist, an invalid handle is returned. 
 			 */
-			FInputHandle GetInputWithID(int32 InPointID) override;
+			FInputHandle GetInputWithID(FGuid InPointID) override;
 
 			/** Returns an input with the given name. 
 			 *
 			 * If the input does not exist, an invalid handle is returned. 
 			 */
-			FConstInputHandle GetInputWithID(int32 InPointID) const override;
+			FConstInputHandle GetInputWithID(FGuid InPointID) const override;
 
 			/** Returns an output with the given name. 
 			 *
 			 * If the output does not exist, an invalid handle is returned. 
 			 */
-			FOutputHandle GetOutputWithID(int32 InPointID) override;
+			FOutputHandle GetOutputWithID(FGuid InPointID) override;
 
 			/** Returns an output with the given name. 
 			 *
 			 * If the output does not exist, an invalid handle is returned. 
 			 */
-			FConstOutputHandle GetOutputWithID(int32 InPointID) const override;
+			FConstOutputHandle GetOutputWithID(FGuid InPointID) const override;
 
 			FGraphHandle AsGraph() override;
 			FConstGraphHandle AsGraph() const override;
 
 		protected:
 
+			FDocumentAccess ShareAccess() override;
+			FConstDocumentAccess ShareAccess() const override;
+
 			FNodeAccessPtr NodePtr;
 			FConstClassAccessPtr ClassPtr;
 			FGraphHandle OwningGraph;
@@ -365,14 +372,14 @@ namespace Metasound
 
 			struct FInputControllerParams
 			{
-				int32 PointID;
+				FGuid PointID;
 				FConstVertexAccessPtr NodeVertexPtr;
 				FConstClassInputAccessPtr ClassInputPtr;
 			};
 
 			struct FOutputControllerParams
 			{
-				int32 PointID;
+				FGuid PointID;
 				FConstVertexAccessPtr NodeVertexPtr;
 				FConstClassOutputAccessPtr ClassOutputPtr;
 			};
@@ -383,8 +390,8 @@ namespace Metasound
 			TArray<FInputControllerParams> GetInputControllerParamsWithVertexName(const FString& InName) const;
 			TArray<FOutputControllerParams> GetOutputControllerParamsWithVertexName(const FString& InName) const;
 
-			bool FindInputControllerParamsWithID(int32 InPointID, FInputControllerParams& OutParams) const;
-			bool FindOutputControllerParamsWithID(int32 InPointID, FOutputControllerParams& OutParams) const;
+			bool FindInputControllerParamsWithID(FGuid InPointID, FInputControllerParams& OutParams) const;
+			bool FindOutputControllerParamsWithID(FGuid InPointID, FOutputControllerParams& OutParams) const;
 
 			const FMetasoundFrontendClassInput* FindClassInputWithName(const FString& InName) const;
 			const FMetasoundFrontendClassOutput* FindClassOutputWithName(const FString& InName) const;
@@ -394,9 +401,9 @@ namespace Metasound
 
 
 
-			virtual FInputHandle CreateInputController(int32 InPointID, FConstVertexAccessPtr InNodeVertexPtr, FConstClassInputAccessPtr InClassInputPtr, FNodeHandle InOwningNode) const = 0;
+			virtual FInputHandle CreateInputController(FGuid InPointID, FConstVertexAccessPtr InNodeVertexPtr, FConstClassInputAccessPtr InClassInputPtr, FNodeHandle InOwningNode) const = 0;
 
-			virtual FOutputHandle CreateOutputController(int32 InPointID, FConstVertexAccessPtr InNodeVertexPtr, FConstClassOutputAccessPtr InClassOutputPtr, FNodeHandle InOwningNode) const = 0;
+			virtual FOutputHandle CreateOutputController(FGuid InPointID, FConstVertexAccessPtr InNodeVertexPtr, FConstClassOutputAccessPtr InClassOutputPtr, FNodeHandle InOwningNode) const = 0;
 		};
 
 		/** FNodeController represents a external or subgraph node. */
@@ -437,11 +444,17 @@ namespace Metasound
 
 			bool IsValid() const override;
 
+		protected:
+
+			FDocumentAccess ShareAccess() override;
+			FConstDocumentAccess ShareAccess() const override;
+
+
 		private:
 
-			FInputHandle CreateInputController(int32 InPointID, FConstVertexAccessPtr InNodeVertexPtr, FConstClassInputAccessPtr InClassInputPtr, FNodeHandle InOwningNode) const override;
+			FInputHandle CreateInputController(FGuid InPointID, FConstVertexAccessPtr InNodeVertexPtr, FConstClassInputAccessPtr InClassInputPtr, FNodeHandle InOwningNode) const override;
 
-			FOutputHandle CreateOutputController(int32 InPointID, FConstVertexAccessPtr InNodeVertexPtr, FConstClassOutputAccessPtr InClassOutputPtr, FNodeHandle InOwningNode) const override;
+			FOutputHandle CreateOutputController(FGuid InPointID, FConstVertexAccessPtr InNodeVertexPtr, FConstClassOutputAccessPtr InClassOutputPtr, FNodeHandle InOwningNode) const override;
 
 			FGraphAccessPtr GraphPtr;
 
@@ -486,10 +499,16 @@ namespace Metasound
 
 			bool IsValid() const override;
 
+		protected:
+
+			FDocumentAccess ShareAccess() override;
+			FConstDocumentAccess ShareAccess() const override;
+
+
 		private:
-			FInputHandle CreateInputController(int32 InPointID, FConstVertexAccessPtr InNodeVertexPtr, FConstClassInputAccessPtr InClassInputPtr, FNodeHandle InOwningNode) const override;
+			FInputHandle CreateInputController(FGuid InPointID, FConstVertexAccessPtr InNodeVertexPtr, FConstClassInputAccessPtr InClassInputPtr, FNodeHandle InOwningNode) const override;
 
-			FOutputHandle CreateOutputController(int32 InPointID, FConstVertexAccessPtr InNodeVertexPtr, FConstClassOutputAccessPtr InClassOutputPtr, FNodeHandle InOwningNode) const override;
+			FOutputHandle CreateOutputController(FGuid InPointID, FConstVertexAccessPtr InNodeVertexPtr, FConstClassOutputAccessPtr InClassOutputPtr, FNodeHandle InOwningNode) const override;
 
 			FGraphAccessPtr GraphPtr;
 			FConstClassOutputAccessPtr OwningGraphClassOutputPtr; 
@@ -534,11 +553,17 @@ namespace Metasound
 
 			bool IsValid() const override;
 
+		protected:
+
+			FDocumentAccess ShareAccess() override;
+			FConstDocumentAccess ShareAccess() const override;
+
+
 
 		private:
-			FInputHandle CreateInputController(int32 InPointID, FConstVertexAccessPtr InNodeVertexPtr, FConstClassInputAccessPtr InClassInputPtr, FNodeHandle InOwningNode) const override;
+			FInputHandle CreateInputController(FGuid InPointID, FConstVertexAccessPtr InNodeVertexPtr, FConstClassInputAccessPtr InClassInputPtr, FNodeHandle InOwningNode) const override;
 
-			FOutputHandle CreateOutputController(int32 InPointID, FConstVertexAccessPtr InNodeVertexPtr, FConstClassOutputAccessPtr InClassOutputPtr, FNodeHandle InOwningNode) const override;
+			FOutputHandle CreateOutputController(FGuid InPointID, FConstVertexAccessPtr InNodeVertexPtr, FConstClassOutputAccessPtr InClassOutputPtr, FNodeHandle InOwningNode) const override;
 
 			FConstClassInputAccessPtr OwningGraphClassInputPtr;
 			FGraphAccessPtr GraphPtr;
@@ -549,7 +574,7 @@ namespace Metasound
 		{
 			// Util for setting template to literal based on template argument
 			template<typename ArgType>
-			bool SetDefaultInputToLiteralInternal(const FString& InInputName, int32 InPointID, ArgType InValue)
+			bool SetDefaultInputToLiteralInternal(const FString& InInputName, FGuid InPointID, ArgType InValue)
 			{
 				if (FMetasoundFrontendClassInput* Desc = FindInputDescriptionWithName(InInputName))
 				{
@@ -609,22 +634,22 @@ namespace Metasound
 
 			bool IsValid() const override;
 
-			int32 GetClassID() const override;
+			FGuid GetClassID() const override;
 
-			int32 GetNewPointID() const override;
+			FGuid GetNewPointID() const override;
 
 			TArray<FString> GetInputVertexNames() const override;
 			TArray<FString> GetOutputVertexNames() const override;
 			FConstClassInputAccessPtr FindClassInputWithName(const FString& InName) const override;
 			FConstClassOutputAccessPtr FindClassOutputWithName(const FString& InName) const override;
-			TArray<int32> GetDefaultIDsForInputVertex(const FString& InInputName) const override;
-			TArray<int32> GetDefaultIDsForOutputVertex(const FString& InOutputName) const override;
+			TArray<FGuid> GetDefaultIDsForInputVertex(const FString& InInputName) const override;
+			TArray<FGuid> GetDefaultIDsForOutputVertex(const FString& InOutputName) const override;
 
 			TArray<FNodeHandle> GetNodes() override;
 			TArray<FConstNodeHandle> GetConstNodes() const override;
 
-			FConstNodeHandle GetNodeWithID(int32 InNodeID) const override;
-			FNodeHandle GetNodeWithID(int32 InNodeID) override;
+			FConstNodeHandle GetNodeWithID(FGuid InNodeID) const override;
+			FNodeHandle GetNodeWithID(FGuid InNodeID) override;
 
 			TArray<FNodeHandle> GetInputNodes() override;
 			TArray<FConstNodeHandle> GetConstInputNodes() const override;
@@ -656,12 +681,12 @@ namespace Metasound
 
 			// These can be used to set the default value for a given input on this graph.
 			// @returns false if the input name couldn't be found, or if the literal type was incompatible with the Data Type of this input.
-			bool SetDefaultInputToLiteral(const FString& InInputName, int32 InPointID, bool bInValue) override;
-			bool SetDefaultInputToLiteral(const FString& InInputName, int32 InPointID, int32 InValue) override;
-			bool SetDefaultInputToLiteral(const FString& InInputName, int32 InPointID, float InValue) override;
-			bool SetDefaultInputToLiteral(const FString& InInputName, int32 InPointID, const FString& InValue) override;
-			bool SetDefaultInputToLiteral(const FString& InInputName, int32 InPointID, UObject* InValue) override;
-			bool SetDefaultInputToLiteral(const FString& InInputName, int32 InPointID, const TArray<UObject*>& InValue) override;
+			bool SetDefaultInputToLiteral(const FString& InInputName, FGuid InPointID, bool bInValue) override;
+			bool SetDefaultInputToLiteral(const FString& InInputName, FGuid InPointID, int32 InValue) override;
+			bool SetDefaultInputToLiteral(const FString& InInputName, FGuid InPointID, float InValue) override;
+			bool SetDefaultInputToLiteral(const FString& InInputName, FGuid InPointID, const FString& InValue) override;
+			bool SetDefaultInputToLiteral(const FString& InInputName, FGuid InPointID, UObject* InValue) override;
+			bool SetDefaultInputToLiteral(const FString& InInputName, FGuid InPointID, const TArray<UObject*>& InValue) override;
 
 			// Set the display name for the input with the given name
 			void SetInputDisplayName(FString InName, const FText& InDisplayName) override;
@@ -671,7 +696,7 @@ namespace Metasound
 
 			// This can be used to clear the current literal for a given input.
 			// @returns false if the input name couldn't be found.
-			bool ClearLiteralForInput(const FString& InInputName, int32 InPointID) override;
+			bool ClearLiteralForInput(const FString& InInputName, FGuid InPointID) override;
 
 			FNodeHandle AddNode(const FNodeClassInfo& InNodeClass) override;
 			FNodeHandle AddNode(const FNodeRegistryKey& InNodeClass) override;
@@ -692,10 +717,14 @@ namespace Metasound
 			FDocumentHandle GetOwningDocument() override;
 			FConstDocumentHandle GetOwningDocument() const override;
 
+		protected:
+
+			FDocumentAccess ShareAccess() override;
+			FConstDocumentAccess ShareAccess() const override;
+
 		private:
 
 			// Add/remove nodes
-			//FNodeHandle AddNode(const FMetasoundFrontendClass& InExistingDependency);
 			FNodeHandle AddNode(FConstClassAccessPtr InExistingDependency);
 			bool RemoveNode(const FMetasoundFrontendNode& InDesc);
 
@@ -703,7 +732,7 @@ namespace Metasound
 			bool RemoveInput(const FMetasoundFrontendNode& InNode);
 			bool RemoveOutput(const FMetasoundFrontendNode& InNode);
 
-			int32 NewNodeID() const;
+			FGuid NewNodeID() const;
 
 			FNodeHandle GetNodeByPredicate(TFunctionRef<bool (const FMetasoundFrontendClass&, const FMetasoundFrontendNode&)> InPredicate);
 			FConstNodeHandle GetNodeByPredicate(TFunctionRef<bool (const FMetasoundFrontendClass&, const FMetasoundFrontendNode&)> InPredicate) const;
@@ -748,16 +777,11 @@ namespace Metasound
 			FMetasoundFrontendClassOutput* FindOutputDescriptionWithName(const FString& InName);
 			const FMetasoundFrontendClassOutput* FindOutputDescriptionWithName(const FString& InName) const;
 
-			FClassInputAccessPtr FindInputDescriptionWithNodeID(int32 InNodeID);
-			FConstClassInputAccessPtr FindInputDescriptionWithNodeID(int32 InNodeID) const;
-
-			FClassOutputAccessPtr FindOutputDescriptionWithNodeID(int32 InNodeID);
-			FConstClassOutputAccessPtr FindOutputDescriptionWithNodeID(int32 InNodeID) const;
+			FClassInputAccessPtr FindInputDescriptionWithNodeID(FGuid InNodeID);
+			FConstClassInputAccessPtr FindInputDescriptionWithNodeID(FGuid InNodeID) const;
 
-			/*
-			FMetasoundFrontendNode* FindNodeByID(int32 InNodeID);
-			const FMetasoundFrontendNode* FindNodeByID(int32 InNodeID) const;
-			*/
+			FClassOutputAccessPtr FindOutputDescriptionWithNodeID(FGuid InNodeID);
+			FConstClassOutputAccessPtr FindOutputDescriptionWithNodeID(FGuid InNodeID) const;
 
 			FMetasoundFrontendClassMetadata CreateInputClassMetadata(const FMetasoundFrontendClassInput& InClassInput);
 			FMetasoundFrontendClassMetadata CreateOutputClassMetadata(const FMetasoundFrontendClassOutput& InClassOutput);
@@ -803,9 +827,9 @@ namespace Metasound
 			TArray<FMetasoundFrontendGraphClass> GetSubgraphs() const override;
 			TArray<FMetasoundFrontendClass> GetClasses() const override;
 
-			FConstClassAccessPtr FindDependencyWithID(int32 InClassID) const override;
-			FConstGraphClassAccessPtr FindSubgraphWithID(int32 InClassID) const override;
-			FConstClassAccessPtr FindClassWithID(int32 InClassID) const override;
+			FConstClassAccessPtr FindDependencyWithID(FGuid InClassID) const override;
+			FConstGraphClassAccessPtr FindSubgraphWithID(FGuid InClassID) const override;
+			FConstClassAccessPtr FindClassWithID(FGuid InClassID) const override;
 
 			FConstClassAccessPtr FindClass(const FNodeClassInfo& InNodeClass) const override;
 			FConstClassAccessPtr FindClass(const FMetasoundFrontendClassMetadata& InMetadata) const override;
@@ -825,19 +849,23 @@ namespace Metasound
 			TArray<FConstGraphHandle> GetSubgraphHandles() const override;
 
 			/** Returns a graphs in the document with the given class ID.*/
-			FGraphHandle GetSubgraphWithClassID(int32 InClassID) override;
+			FGraphHandle GetSubgraphWithClassID(FGuid InClassID) override;
 
 			/** Returns a graphs in the document with the given class ID.*/
-			FConstGraphHandle GetSubgraphWithClassID(int32 InClassID) const override;
+			FConstGraphHandle GetSubgraphWithClassID(FGuid InClassID) const override;
 
 			bool ExportToJSONAsset(const FString& InAbsolutePath) const override;
 			FString ExportToJSON() const override;
 
 			static bool IsMatchingMetasoundClass(const FMetasoundFrontendClassMetadata& InMetadataA, const FMetasoundFrontendClassMetadata& InMetadataB);
 			static bool IsMatchingMetasoundClass(const FNodeClassInfo& InNodeClass, const FMetasoundFrontendClassMetadata& InMetadata);
+
+		protected:
+			FDocumentAccess ShareAccess() override;
+			FConstDocumentAccess ShareAccess() const override;
 		private:
 
-			int32 NewClassID() const;
+			FGuid NewClassID() const;
 
 			FDocumentAccessPtr DocumentPtr;
 		};
diff --git a/Engine/Plugins/Runtime/Metasound/Source/MetasoundFrontend/Public/MetasoundFrontendController.h b/Engine/Plugins/Runtime/Metasound/Source/MetasoundFrontend/Public/MetasoundFrontendController.h
--- a/Engine/Plugins/Runtime/Metasound/Source/MetasoundFrontend/Public/MetasoundFrontendController.h
+++ b/Engine/Plugins/Runtime/Metasound/Source/MetasoundFrontend/Public/MetasoundFrontendController.h
@@ -71,9 +71,88 @@ namespace Metasound
 		using FConstGraphClassAccessPtr = TAccessPtr<const FMetasoundFrontendGraphClass>;
 		using FDocumentAccessPtr = TAccessPtr<FMetasoundFrontendDocument>;
 		using FConstDocumentAccessPtr = TAccessPtr<const FMetasoundFrontendDocument>;
+
+		// Container holding various access pointers to the FMetasoundFrontendDocument
+		struct FConstDocumentAccess
+		{
+			FConstVertexAccessPtr ConstVertex;
+
+			FConstClassInputAccessPtr ConstClassInput;
+
+			FConstClassOutputAccessPtr ConstClassOutput;
+		
+			FConstNodeAccessPtr ConstNode;
+
+			FConstClassAccessPtr ConstClass;
+
+			FConstGraphClassAccessPtr ConstGraphClass;
+
+			FConstGraphAccessPtr ConstGraph;
+
+			FConstDocumentAccessPtr ConstDocument;
+		};
+
+		// Container holding various access pointers to the FMetasoundFrontendDocument
+		struct FDocumentAccess : public FConstDocumentAccess
+		{
+			FVertexAccessPtr Vertex; 
+
+			FClassInputAccessPtr ClassInput;
+
+			FClassOutputAccessPtr ClassOutput;
+		
+			FNodeAccessPtr Node;
+
+			FClassAccessPtr Class;
+
+			FGraphClassAccessPtr GraphClass;
+
+			FGraphAccessPtr Graph;
+
+			FDocumentAccessPtr Document;
+		};
+
+
+		/** IDocumentAccessor describes an interface for various I*Controllers to interact with
+		 * each other without exposing that interface publicly or requiring friendship 
+		 * between various controller implementation classes. 
+		 */
+		class METASOUNDFRONTEND_API IDocumentAccessor
+		{
+			protected:
+				/** Share access to FMetasoundFrontendDocument objects. 
+				 *
+				 * Derived classes must implement this method. In the implementation,
+				 * derived classes should set the various TAccessPtrs on FDocumentAccess 
+				 * to the TAccessPtrs held internal in the derived class.
+				 *
+				 * Sharing access simplifies operations involving multiple frontend controllers
+				 * by providing direct access to the FMetasoundFrontendDocument objects to be
+				 * edited. 
+				 */
+				virtual FDocumentAccess ShareAccess() = 0;
+
+				/** Share access to FMetasoundFrontendDocument objects. 
+				 *
+				 * Derived classes must implement this method. In the implementation,
+				 * derived classes should set the various TAccessPtrs on FDocumentAccess 
+				 * to the TAccessPtrs held internal in the derived class.
+				 *
+				 * Sharing access simplifies operations involving multiple frontend controllers
+				 * by providing direct access to the FMetasoundFrontendDocument objects to be
+				 * edited. 
+				 */
+				virtual FConstDocumentAccess ShareAccess() const = 0;
+
+				/** Returns the shared access from an IDocumentAccessor. */
+				static FDocumentAccess GetSharedAccess(IDocumentAccessor& InDocumentAccessor);
+
+				/** Returns the shared access from an IDocumentAccessor. */
+				static FConstDocumentAccess GetSharedAccess(const IDocumentAccessor& InDocumentAccessor);
+		};
 		
 		/* An IOutputController provides methods for querying and manipulating a metasound output vertex. */
-		class METASOUNDFRONTEND_API IOutputController : public TSharedFromThis<IOutputController>
+		class METASOUNDFRONTEND_API IOutputController : public TSharedFromThis<IOutputController>, public IDocumentAccessor
 		{
 		public:
 			static FOutputHandle GetInvalidHandle();
@@ -84,7 +163,7 @@ namespace Metasound
 			/** Returns true if the controller is in a valid state. */
 			virtual bool IsValid() const = 0;
 
-			virtual int32 GetID() const = 0;
+			virtual FGuid GetID() const = 0;
 			
 			/** Returns the data type name associated with this output. */
 			virtual const FName& GetDataType() const = 0;
@@ -99,7 +178,7 @@ namespace Metasound
 			virtual const FText& GetTooltip() const = 0;
 
 			/** Returns the ID of the node which owns this output. */
-			virtual int32 GetOwningNodeID() const = 0;
+			virtual FGuid GetOwningNodeID() const = 0;
 			
 			/** Returns a FNodeHandle to the node which owns this output. */
 			virtual FNodeHandle GetOwningNode() = 0;
@@ -135,7 +214,7 @@ namespace Metasound
 		};
 
 		/* An IInputController provides methods for querying and manipulating a metasound input vertex. */
-		class METASOUNDFRONTEND_API IInputController : public TSharedFromThis<IInputController>
+		class METASOUNDFRONTEND_API IInputController : public TSharedFromThis<IInputController>, public IDocumentAccessor
 		{
 		public:
 			static FInputHandle GetInvalidHandle();
@@ -146,7 +225,7 @@ namespace Metasound
 			/** Returns true if the controller is in a valid state. */
 			virtual bool IsValid() const = 0;
 
-			virtual int32 GetID() const = 0;
+			virtual FGuid GetID() const = 0;
 
 			/** Return true if the input is connect to an output. */
 			virtual bool IsConnected() const = 0;
@@ -164,7 +243,7 @@ namespace Metasound
 			virtual const FText& GetTooltip() const = 0;
 
 			/** Returns the ID of the node which owns this output. */
-			virtual int32 GetOwningNodeID() const = 0;
+			virtual FGuid GetOwningNodeID() const = 0;
 			
 			/** Returns a FNodeHandle to the node which owns this output. */
 			virtual FNodeHandle GetOwningNode() = 0;
@@ -202,10 +281,18 @@ namespace Metasound
 			 * @return True on success, false on failure. 
 			 */
 			virtual bool Disconnect() = 0;
+
+			/*
+			virtual bool IsLiteralSet() const = 0;
+			virtual FMetasoundFrontendLiteral GetLiteral() const = 0;
+			virtual bool SetLiteral(const FMetasoundFrontendLiteral& InLiteral) = 0;
+			*/
+
+
 		};
 
 		/* An INodeController provides methods for querying and manipulating a Metasound node. */
-		class METASOUNDFRONTEND_API INodeController : public TSharedFromThis<INodeController>
+		class METASOUNDFRONTEND_API INodeController : public TSharedFromThis<INodeController>, public IDocumentAccessor
 		{
 
 		public:
@@ -239,34 +326,34 @@ namespace Metasound
 			 *
 			 * If the input does not exist, an invalid handle is returned. 
 			 */
-			virtual FInputHandle GetInputWithID(int32 InPointID) = 0;
+			virtual FInputHandle GetInputWithID(FGuid InPointID) = 0;
 
 			/** Returns an input with the given name. 
 			 *
 			 * If the input does not exist, an invalid handle is returned. 
 			 */
-			virtual FConstInputHandle GetInputWithID(int32 InPointID) const = 0;
+			virtual FConstInputHandle GetInputWithID(FGuid InPointID) const = 0;
 
 			/** Returns an output with the given name. 
 			 *
 			 * If the output does not exist, an invalid handle is returned. 
 			 */
-			virtual FOutputHandle GetOutputWithID(int32 InPointID) = 0;
+			virtual FOutputHandle GetOutputWithID(FGuid InPointID) = 0;
 
 			/** Returns an output with the given name. 
 			 *
 			 * If the output does not exist, an invalid handle is returned. 
 			 */
-			virtual FConstOutputHandle GetOutputWithID(int32 InPointID) const = 0;
+			virtual FConstOutputHandle GetOutputWithID(FGuid InPointID) const = 0;
 
 			virtual bool CanAddInput(const FString& InVertexName) const = 0;
 			virtual FInputHandle AddInput(const FString& InVertexName, const FMetasoundFrontendLiteral* InDefault) = 0;
-			virtual bool RemoveInput(int32 InPointID) = 0;
+			virtual bool RemoveInput(FGuid InPointID) = 0;
 
 			// TODO: PointID -> PinID? PointID? SlotID? RefID? 
 			virtual bool CanAddOutput(const FString& InVertexName) const = 0;
 			virtual FInputHandle AddOutput(const FString& InVertexName, const FMetasoundFrontendLiteral* InDefault) = 0;
-			virtual bool RemoveOutput(int32 InPointID) = 0;
+			virtual bool RemoveOutput(FGuid InPointID) = 0;
 
 			/** Returns the class type for this node. */
 			virtual EMetasoundFrontendClassType GetClassType() const = 0;
@@ -291,13 +378,13 @@ namespace Metasound
 			virtual const FString& GetNodeName() const = 0;
 
 			/** Returns the ID associated with this node. */
-			virtual int32 GetID() const = 0;
+			virtual FGuid GetID() const = 0;
 
 			/** Returns the ID associated with the node class. */
-			virtual int32 GetClassID() const = 0;
+			virtual FGuid GetClassID() const = 0;
 
 			/** Returns the class ID of the metasound class which owns this node. */
-			virtual int32 GetOwningGraphClassID() const = 0;
+			virtual FGuid GetOwningGraphClassID() const = 0;
 
 			/** Returns the graph which owns this node. */
 			virtual FGraphHandle GetOwningGraph() = 0;
@@ -307,7 +394,7 @@ namespace Metasound
 		};
 
 		/* An IGraphController provides methods for querying and manipulating a Metasound graph. */
-		class METASOUNDFRONTEND_API IGraphController : public TSharedFromThis<IGraphController>
+		class METASOUNDFRONTEND_API IGraphController : public TSharedFromThis<IGraphController>, public IDocumentAccessor
 		{
 			public:
 			static FGraphHandle GetInvalidHandle();
@@ -318,7 +405,7 @@ namespace Metasound
 			/** Returns true if the controller is in a valid state. */
 			virtual bool IsValid() const = 0;
 
-			virtual int32 GetNewPointID() const = 0;
+			virtual FGuid GetNewPointID() const = 0;
 
 			virtual TArray<FString> GetInputVertexNames() const = 0;
 			virtual TArray<FString> GetOutputVertexNames() const = 0;
@@ -327,13 +414,13 @@ namespace Metasound
 			virtual TArray<FNodeHandle> GetNodes() = 0;
 
 			/** Returns a node by NodeID. If the node does not exist, an invalid handle is returned. */
-			virtual FConstNodeHandle GetNodeWithID(int32 InNodeID) const = 0;
+			virtual FConstNodeHandle GetNodeWithID(FGuid InNodeID) const = 0;
 
 			/** Returns all nodes in the graph. */
 			virtual TArray<FConstNodeHandle> GetConstNodes() const = 0;
 
 			/** Returns a node by NodeID. If the node does not exist, an invalid handle is returned. */
-			virtual FNodeHandle GetNodeWithID(int32 InNodeID) = 0;
+			virtual FNodeHandle GetNodeWithID(FGuid InNodeID) = 0;
 
 			/** Returns all output nodes in the graph. */
 			virtual TArray<FNodeHandle> GetOutputNodes() = 0;
@@ -435,8 +522,8 @@ namespace Metasound
 			 */
 			virtual UClass* GetSupportedClassForInputVertex(const FString& InInputName) = 0;
 
-			virtual TArray<int32> GetDefaultIDsForInputVertex(const FString& InInputName) const = 0;
-			virtual TArray<int32> GetDefaultIDsForOutputVertex(const FString& InOutputName) const = 0;
+			virtual TArray<FGuid> GetDefaultIDsForInputVertex(const FString& InInputName) const = 0;
+			virtual TArray<FGuid> GetDefaultIDsForOutputVertex(const FString& InOutputName) const = 0;
 
 			/** Set the default value for the graph input.
 			 *
@@ -445,7 +532,7 @@ namespace Metasound
 			 *
 			 * @return True on success. False if the input does not exist or if the literal type was incompatible with the input.
 			 */
-			virtual bool SetDefaultInputToLiteral(const FString& InInputName, int32 InPointID, bool InValue) = 0;
+			virtual bool SetDefaultInputToLiteral(const FString& InInputName, FGuid InPointID, bool InValue) = 0;
 
 			/** Set the default value for the graph input.
 			 *
@@ -454,7 +541,7 @@ namespace Metasound
 			 *
 			 * @return True on success. False if the input does not exist or if the literal type was incompatible with the input.
 			 */
-			virtual bool SetDefaultInputToLiteral(const FString& InInputName, int32 InPointID, int32 InValue) = 0;
+			virtual bool SetDefaultInputToLiteral(const FString& InInputName, FGuid InPointID, int32 InValue) = 0;
 
 			/** Set the default value for the graph input.
 			 *
@@ -463,7 +550,7 @@ namespace Metasound
 			 *
 			 * @return True on success. False if the input does not exist or if the literal type was incompatible with the input.
 			 */
-			virtual bool SetDefaultInputToLiteral(const FString& InInputName, int32 InPointID, float InValue) = 0;
+			virtual bool SetDefaultInputToLiteral(const FString& InInputName, FGuid InPointID, float InValue) = 0;
 
 			/** Set the default value for the graph input.
 			 *
@@ -472,7 +559,7 @@ namespace Metasound
 			 *
 			 * @return True on success. False if the input does not exist or if the literal type was incompatible with the input.
 			 */
-			virtual bool SetDefaultInputToLiteral(const FString& InInputName, int32 InPointID, const FString& InValue) = 0;
+			virtual bool SetDefaultInputToLiteral(const FString& InInputName, FGuid InPointID, const FString& InValue) = 0;
 
 			/** Set the default value for the graph input.
 			 *
@@ -481,7 +568,7 @@ namespace Metasound
 			 *
 			 * @return True on success. False if the input does not exist or if the literal type was incompatible with the input.
 			 */
-			virtual bool SetDefaultInputToLiteral(const FString& InInputName, int32 InPointID, UObject* InValue) = 0;
+			virtual bool SetDefaultInputToLiteral(const FString& InInputName, FGuid InPointID, UObject* InValue) = 0;
 
 			/** Set the default value for the graph input.
 			 *
@@ -490,7 +577,7 @@ namespace Metasound
 			 *
 			 * @return True on success. False if the input does not exist or if the literal type was incompatible with the input.
 			 */
-			virtual bool SetDefaultInputToLiteral(const FString& InInputName, int32 InPointID, const TArray<UObject*>& InValue) = 0;
+			virtual bool SetDefaultInputToLiteral(const FString& InInputName, FGuid InPointID, const TArray<UObject*>& InValue) = 0;
 
 
 			/** Set the display name for the input with the given name. */
@@ -503,7 +590,7 @@ namespace Metasound
 			 *
 			 * @return True on success, false on failure.
 			 */
-			virtual bool ClearLiteralForInput(const FString& InInputName, int32 InPointID) = 0;
+			virtual bool ClearLiteralForInput(const FString& InInputName, FGuid InPointID) = 0;
 
 			/** Add a new node to this graph.
 			 *
@@ -528,7 +615,7 @@ namespace Metasound
 			virtual bool RemoveNode(INodeController& InNode) = 0;
 
 			/** Returns the ClassID associated with this graph. */
-			virtual int32 GetClassID() const = 0;
+			virtual FGuid GetClassID() const = 0;
 
 			/** Return the metadata for the current graph. */
 			virtual const FMetasoundFrontendClassMetadata& GetGraphMetadata() const = 0;
@@ -574,7 +661,7 @@ namespace Metasound
 		};
 
 		/* An IDocumentController provides methods for querying and manipulating a Metasound document. */
-		class METASOUNDFRONTEND_API IDocumentController : public TSharedFromThis<IDocumentController>
+		class METASOUNDFRONTEND_API IDocumentController : public TSharedFromThis<IDocumentController>, public IDocumentAccessor
 		{
 		public:
 			static FDocumentHandle GetInvalidHandle();
@@ -611,9 +698,9 @@ namespace Metasound
 			virtual TArray<FMetasoundFrontendGraphClass> GetSubgraphs() const = 0;
 			virtual TArray<FMetasoundFrontendClass> GetClasses() const = 0;
 
-			virtual FConstClassAccessPtr FindDependencyWithID(int32 InClassID) const = 0;
-			virtual FConstGraphClassAccessPtr FindSubgraphWithID(int32 InClassID) const = 0;
-			virtual FConstClassAccessPtr FindClassWithID(int32 InClassID) const = 0;
+			virtual FConstClassAccessPtr FindDependencyWithID(FGuid InClassID) const = 0;
+			virtual FConstGraphClassAccessPtr FindSubgraphWithID(FGuid InClassID) const = 0;
+			virtual FConstClassAccessPtr FindClassWithID(FGuid InClassID) const = 0;
 
 
 			/** Returns an existing Metasound class description corresponding to 
@@ -658,10 +745,10 @@ namespace Metasound
 			virtual TArray<FConstGraphHandle> GetSubgraphHandles() const = 0;
 
 			/** Returns a graphs in the document with the given class ID.*/
-			virtual FGraphHandle GetSubgraphWithClassID(int32 InClassID) = 0;
+			virtual FGraphHandle GetSubgraphWithClassID(FGuid InClassID) = 0;
 
 			/** Returns a graphs in the document with the given class ID.*/
-			virtual FConstGraphHandle GetSubgraphWithClassID(int32 InClassID) const = 0;
+			virtual FConstGraphHandle GetSubgraphWithClassID(FGuid InClassID) const = 0;
 
 			/** Returns the root graph of this document. */
 			virtual FGraphHandle GetRootGraph() = 0;
diff --git a/Engine/Plugins/Runtime/Metasound/Source/MetasoundFrontend/Public/MetasoundFrontendDocument.h b/Engine/Plugins/Runtime/Metasound/Source/MetasoundFrontend/Public/MetasoundFrontendDocument.h
--- a/Engine/Plugins/Runtime/Metasound/Source/MetasoundFrontend/Public/MetasoundFrontendDocument.h
+++ b/Engine/Plugins/Runtime/Metasound/Source/MetasoundFrontend/Public/MetasoundFrontendDocument.h
@@ -5,12 +5,13 @@
 #include "CoreMinimal.h"
 #include "Internationalization/Text.h"
 #include "MetasoundAccessPtr.h"
+#include "Misc/Guid.h"
 
 #include "MetasoundFrontendDocument.generated.h"
 
 namespace Metasound
 {
-	static constexpr int32 FrontendInvalidID = INDEX_NONE;
+	static const FGuid FrontendInvalidID;
 }
 
 // Forward declare
@@ -134,7 +135,7 @@ struct METASOUNDFRONTEND_API FMetasoundFrontendVertex
 
 	// IDs of connection points supported by the vertex.
 	UPROPERTY()	
-	TArray<int32> PointIDs;
+	TArray<FGuid> PointIDs;
 
 	// Returns true if vertexes have equal name, type and number of IDs. 
 	static bool IsFunctionalEquivalent(const FMetasoundFrontendVertex& InLHS, const FMetasoundFrontendVertex& InRHS);
@@ -148,7 +149,7 @@ struct FMetasoundFrontendVertexLiteral
 
 	// ID of vertex.
 	UPROPERTY(EditAnywhere, Category = Parameters)
-	int32 PointID = Metasound::FrontendInvalidID;
+	FGuid PointID = Metasound::FrontendInvalidID;
 
 	// Value to use when constructing input. 
 	UPROPERTY(EditAnywhere, Category = Parameters)
@@ -193,11 +194,11 @@ struct METASOUNDFRONTEND_API FMetasoundFrontendNode
 
 	// Unique ID of this node.
 	UPROPERTY()
-	int32 ID = Metasound::FrontendInvalidID;
+	FGuid ID = Metasound::FrontendInvalidID;
 
 	// ID of FMetasoundFrontendClass corresponding to this node.
 	UPROPERTY()
-	int32 ClassID = Metasound::FrontendInvalidID;
+	FGuid ClassID = Metasound::FrontendInvalidID;
 
 	// Name of node instance. 
 	UPROPERTY()
@@ -223,19 +224,19 @@ struct FMetasoundFrontendEdge
 
 	// ID of source node.
 	UPROPERTY()
-	int32 FromNodeID = Metasound::FrontendInvalidID;
+	FGuid FromNodeID = Metasound::FrontendInvalidID;
 
 	// ID of source point on source node.
 	UPROPERTY()
-	int32 FromPointID = Metasound::FrontendInvalidID;
+	FGuid FromPointID = Metasound::FrontendInvalidID;
 
 	// ID of destination node.
 	UPROPERTY()
-	int32 ToNodeID = Metasound::FrontendInvalidID;
+	FGuid ToNodeID = Metasound::FrontendInvalidID;
 
 	// ID of destination point on destination node.
 	UPROPERTY()
-	int32 ToPointID = Metasound::FrontendInvalidID;
+	FGuid ToPointID = Metasound::FrontendInvalidID;
 };
 
 // Display style for an edge.
@@ -382,7 +383,7 @@ struct METASOUNDFRONTEND_API FMetasoundFrontendClassVertex : public FMetasoundFr
 	virtual ~FMetasoundFrontendClassVertex() = default;
 
 	UPROPERTY()
-	int32 NodeID = Metasound::FrontendInvalidID;
+	FGuid NodeID = Metasound::FrontendInvalidID;
 
 	// Metadata associated with input.
 	UPROPERTY(EditAnywhere, Category = CustomView)
@@ -630,7 +631,7 @@ struct METASOUNDFRONTEND_API FMetasoundFrontendClass
 	virtual ~FMetasoundFrontendClass() = default;
 
 	UPROPERTY()
-	int32 ID = Metasound::FrontendInvalidID;
+	FGuid ID = Metasound::FrontendInvalidID;
 
 	UPROPERTY(EditAnywhere, Category = CustomView)
 	FMetasoundFrontendClassMetadata Metadata;
diff --git a/Engine/Plugins/Runtime/Metasound/Source/MetasoundGraphCore/Public/MetasoundDataReference.h b/Engine/Plugins/Runtime/Metasound/Source/MetasoundGraphCore/Public/MetasoundDataReference.h
--- a/Engine/Plugins/Runtime/Metasound/Source/MetasoundGraphCore/Public/MetasoundDataReference.h
+++ b/Engine/Plugins/Runtime/Metasound/Source/MetasoundGraphCore/Public/MetasoundDataReference.h
@@ -145,9 +145,7 @@ namespace Metasound
 		static const void* const TypeId = GetMetasoundDataTypeId<DataType>();
 
 		bool bEqualTypeName = InReference.GetDataTypeName() == TypeName;
-		// TODO: need to move data type ID definition to traits implementation.
-		//bool bEqualTypeId = InReference.GetDataTypeId() == TypeId;
-		bool bEqualTypeId = true;
+		bool bEqualTypeId = InReference.GetDataTypeId() == TypeId;
 
 		return (bEqualTypeName && bEqualTypeId);
 	}
@@ -360,7 +358,6 @@ namespace Metasound
 			TDataReadReference<DataType>& operator=(const TDataWriteReference<DataType>& Other)
 			{
 				TDataReference<DataType>::ObjectReference = Other.ObjectReference;
-				//ObjectReference = Other.ObjectReference;
 				return *this;
 			}
 
@@ -374,14 +371,12 @@ namespace Metasound
 			FORCEINLINE const DataType& operator*() const
 			{
 				return *TDataReference<DataType>::ObjectReference;
-				//return *ObjectReference;
 			}
 
 			/** Const access to the underlying parameter object. */
 			FORCEINLINE const DataType* operator->() const
 			{
 				return TDataReference<DataType>::ObjectReference.operator->();
-				//return ObjectReference.operator->();
 			}
 
 			/** Create a clone of this parameter reference. */
