BugFix: Osc Node produces unexpected frequencies when given high values

diff --git a/Engine/Plugins/Runtime/Metasound/Source/MetasoundStandardNodes/Private/MetasoundOscNode.cpp b/Engine/Plugins/Runtime/Metasound/Source/MetasoundStandardNodes/Private/MetasoundOscNode.cpp
--- a/Engine/Plugins/Runtime/Metasound/Source/MetasoundStandardNodes/Private/MetasoundOscNode.cpp
+++ b/Engine/Plugins/Runtime/Metasound/Source/MetasoundStandardNodes/Private/MetasoundOscNode.cpp
@@ -9,6 +9,8 @@
 #include "MetasoundNodeRegistrationMacro.h"
 #include "MetasoundPrimitives.h"
 #include "MetasoundVertex.h"
+#include "MetasoundBop.h"
+#include "DSP/Dsp.h"
 
 #define LOCTEXT_NAMESPACE "MetasoundStandardNodes"
 
@@ -16,38 +18,63 @@ namespace Metasound
 {
 	class FOscOperator : public TExecutableOperator<FOscOperator>
 	{
-		public:
+	public:
 
-			static TUniquePtr<IOperator> CreateOperator(const FCreateOperatorParams& InParams, FBuildErrorArray& OutErrors);
-			static const FNodeInfo& GetNodeInfo();
-			static FVertexInterface DeclareVertexInterface();
+		static TUniquePtr<IOperator> CreateOperator(const FCreateOperatorParams& InParams, FBuildErrorArray& OutErrors);
+		static const FNodeInfo& GetNodeInfo();
+		static FVertexInterface DeclareVertexInterface();
 
-			FOscOperator(const FOperatorSettings& InSettings, const FFrequencyReadRef& InFrequency);
+		FOscOperator(const FOperatorSettings& InSettings, const FFrequencyReadRef& InFrequency, const FBoolReadRef& InActivate);
 
-			virtual FDataReferenceCollection GetInputs() const override;
+		virtual FDataReferenceCollection GetInputs() const override;
 
-			virtual FDataReferenceCollection GetOutputs() const override;
+		virtual FDataReferenceCollection GetOutputs() const override;
 
-			void Execute();
+		void Execute();
 
-		private:
-			const float TwoPi;
-			float Phase;
-			float SampleRate;
-			int32 BlockSize;
+	private:
+		// Returns true if there was a phase wrap this update
+		FORCEINLINE bool WrapPhase(float InPhaseInc, float& OutPhase)
+		{
+			bool Result = false;
+			if (InPhaseInc > 0.0f && OutPhase >= 1.0f)
+			{
+				OutPhase = FMath::Fmod(OutPhase, 1.0f);
+				Result = true;
+			}
+			else if (InPhaseInc < 0.0f && OutPhase <= 0.0f)
+			{
+				OutPhase = FMath::Fmod(OutPhase, 1.0f) + 1.0f;
+				Result = true;
+			}
+			return Result;
+		}
+
+		static constexpr const float TwoPi = 2.f * PI;
 
-			FFrequencyReadRef Frequency;
-			FAudioBufferWriteRef AudioBuffer;
+		// The current phase of oscillator (between 0.0 and 1.0)
+		float Phase;
+		float OneOverSampleRate;
+		int32 BlockSize;		
+		float Nyquist;
 
+		FFrequencyReadRef Frequency;
+		FBoolReadRef Enabled;
+		FAudioBufferWriteRef AudioBuffer;
+
+		static constexpr const TCHAR* EnabledPinName = TEXT("Enabled");
+		static constexpr const TCHAR* FrequencyPinName = TEXT("Frequency");
+		static constexpr const TCHAR* AudioOutPinName = TEXT("Audio");		
 	};
 
-	FOscOperator::FOscOperator(const FOperatorSettings& InSettings, const FFrequencyReadRef& InFrequency)
-	:	TwoPi(2.f * PI)
-	,	Phase(0.f)
-	,	SampleRate(InSettings.GetSampleRate())
-	,	BlockSize(InSettings.GetNumFramesPerBlock())
-	,	Frequency(InFrequency)
-	,	AudioBuffer(FAudioBufferWriteRef::CreateNew(InSettings))
+	FOscOperator::FOscOperator(const FOperatorSettings& InSettings, const FFrequencyReadRef& InFrequency, const FBoolReadRef& InEnabled)
+		: Phase(0.f)
+		, OneOverSampleRate(1.f / InSettings.GetSampleRate())
+		, BlockSize(InSettings.GetNumFramesPerBlock())
+		, Nyquist(InSettings.GetSampleRate() / 2.0f)
+		, Frequency(InFrequency)
+		, Enabled(InEnabled)
+		, AudioBuffer(FAudioBufferWriteRef::CreateNew(InSettings))
 	{
 		check(AudioBuffer->Num() == InSettings.GetNumFramesPerBlock());
 	}
@@ -55,36 +82,51 @@ namespace Metasound
 	FDataReferenceCollection FOscOperator::GetInputs() const
 	{
 		FDataReferenceCollection InputDataReferences;
-		InputDataReferences.AddDataReadReference(TEXT("Audio"), FFrequencyReadRef(Frequency));
+		InputDataReferences.AddDataReadReference(FrequencyPinName, FFrequencyReadRef(Frequency));
+		InputDataReferences.AddDataReadReference(EnabledPinName, FBoolReadRef(Enabled));
+
 		return InputDataReferences;
 	}
 
 	FDataReferenceCollection FOscOperator::GetOutputs() const
 	{
 		FDataReferenceCollection OutputDataReferences;
-		OutputDataReferences.AddDataReadReference(TEXT("Audio"), FAudioBufferReadRef(AudioBuffer));
+		OutputDataReferences.AddDataReadReference(AudioOutPinName, FAudioBufferReadRef(AudioBuffer));
 		return OutputDataReferences;
 	}
 
 	void FOscOperator::Execute()
 	{
-		const float PhaseDelta = Frequency->GetRadiansPerSample(SampleRate);
+		// Clamp frequencies into Nyquist range
+		float Freq = FMath::Clamp(Frequency->GetHertz(), -Nyquist, Nyquist);
+		
+		float PhaseInc = Freq * OneOverSampleRate;
+		
 		float* Data = AudioBuffer->GetData();
+		FMemory::Memzero(Data, BlockSize);
 
-		for (int32 i = 0; i < BlockSize; i++)
+		if (*Enabled)
 		{
-			Data[i] = FMath::Sin(Phase);
-			Phase += PhaseDelta;
+			for (int32 i = 0; i < BlockSize; i++)
+			{
+				// This is borrowed from the FOsc class with the intention to eventually recreate that functionality here. 
+				// We don't wish to use it directly as it has a virtual call for each sample. Phase is in cents with the intent
+				// to support other oscilator types in time.
+
+				const float Radians = (Phase * TwoPi) - PI;
+				Data[i] = Audio::FastSin3(-1.0f * Radians);
+				Phase += PhaseInc;
+				WrapPhase(PhaseInc, Phase);
+			}
 		}
-
-		Phase -= FMath::FloorToFloat(Phase / TwoPi) * TwoPi;
 	}
 
 	FVertexInterface FOscOperator::DeclareVertexInterface()
 	{
 		static const FVertexInterface Interface(
 			FInputVertexInterface(
-				TInputDataVertexModel<FFrequency>(TEXT("Frequency"), LOCTEXT("OscFrequencyDescription", "The frequency of oscillator."))
+				TInputDataVertexModel<FFrequency>(FrequencyPinName, LOCTEXT("OscFrequencyDescription", "The frequency of oscillator.")),
+				TInputDataVertexModel<bool>(EnabledPinName, LOCTEXT("OscActivateDescription", "Enable the oscilator."))
 			),
 			FOutputVertexInterface(
 				TOutputDataVertexModel<FAudioBuffer>(TEXT("Audio"), LOCTEXT("AudioTooltip", "The output audio"))
@@ -116,15 +158,15 @@ namespace Metasound
 		return Info;
 	}
 
-
-	FOscNode::FOscNode(const FString& InName, float InDefaultFrequency)
+	FOscNode::FOscNode(const FString& InName, float InDefaultFrequency, bool bInDefaultEnablement)
 	:	FNodeFacade(InName, TFacadeOperatorClass<FOscOperator>())
 	,	DefaultFrequency(InDefaultFrequency)
+	,	bDefaultEnablement(bInDefaultEnablement)
 	{
 	}
 
 	FOscNode::FOscNode(const FNodeInitData& InInitData)
-		: FOscNode(InInitData.InstanceName, 440.0f)
+		: FOscNode(InInitData.InstanceName, 440.0f, true)
 	{
 	}
 
@@ -133,17 +175,24 @@ namespace Metasound
 		return DefaultFrequency;
 	}
 
+	bool FOscNode::GetDefaultEnablement() const
+	{
+		return bDefaultEnablement;
+	}
 
 	TUniquePtr<IOperator> FOscOperator::CreateOperator(const FCreateOperatorParams& InParams, FBuildErrorArray& OutErrors) 
 	{
 		const FOscNode& OscNode = static_cast<const FOscNode&>(InParams.Node);
 		const FDataReferenceCollection& InputCol = InParams.InputDataReferences;
+		const FOperatorSettings& Settings = InParams.OperatorSettings;
 
-		FFrequencyReadRef Frequency = InputCol.GetDataReadReferenceOrConstruct<FFrequency>(TEXT("Frequency"), OscNode.GetDefaultFrequency(), EFrequencyResolution::Hertz);
+		FFrequencyReadRef Frequency = InputCol.GetDataReadReferenceOrConstruct<FFrequency>(FrequencyPinName, OscNode.GetDefaultFrequency(), EFrequencyResolution::Hertz);
+		FBoolReadRef Enabled = InputCol.GetDataReadReferenceOrConstruct<bool>(EnabledPinName, OscNode.GetDefaultEnablement());
 
-		return MakeUnique<FOscOperator>(InParams.OperatorSettings, Frequency);
+		return MakeUnique<FOscOperator>(InParams.OperatorSettings, Frequency, Enabled);
 	}
 
 	METASOUND_REGISTER_NODE(FOscNode);
 }
 #undef LOCTEXT_NAMESPACE //MetasoundOscNode
+
diff --git a/Engine/Plugins/Runtime/Metasound/Source/MetasoundStandardNodes/Public/MetasoundOscNode.h b/Engine/Plugins/Runtime/Metasound/Source/MetasoundStandardNodes/Public/MetasoundOscNode.h
--- a/Engine/Plugins/Runtime/Metasound/Source/MetasoundStandardNodes/Public/MetasoundOscNode.h
+++ b/Engine/Plugins/Runtime/Metasound/Source/MetasoundStandardNodes/Public/MetasoundOscNode.h
@@ -19,14 +19,17 @@ namespace Metasound
 	{
 		public:
 
-		FOscNode(const FString& InName, float InDefaultFrequency);
+		FOscNode(const FString& InName, float InDefaultFrequency, bool bInDefaultEnablement);
 
 		FOscNode(const FNodeInitData& InInitData);
 
+		
 		float GetDefaultFrequency() const;
+		bool GetDefaultEnablement() const;
 
 		private:
 
-		float DefaultFrequency = 440.f;
+		float DefaultFrequency = 440.f;		
+		bool bDefaultEnablement = true;
 	};
 }
