Fix TA operator to use Triggers in block not total Triggers.

diff --git a/Engine/Plugins/Runtime/Metasound/Source/MetasoundStandardNodes/Private/MetasoundBop.cpp b/Engine/Plugins/Runtime/Metasound/Source/MetasoundStandardNodes/Private/MetasoundBop.cpp
--- a/Engine/Plugins/Runtime/Metasound/Source/MetasoundStandardNodes/Private/MetasoundBop.cpp
+++ b/Engine/Plugins/Runtime/Metasound/Source/MetasoundStandardNodes/Private/MetasoundBop.cpp
@@ -83,6 +83,11 @@ namespace Metasound
 		return BoppedFrames.Num();
 	}
 
+	int32 FBop::NumBoppedInBlock() const
+	{
+		return LastBopIndexInBlock;
+	}
+
 	int32 FBop::operator[](int32 InBopIndex) const
 	{
 		return BoppedFrames[InBopIndex];
@@ -106,7 +111,7 @@ namespace Metasound
 	void FBop::Reset()
 	{
 		BoppedFrames.Reset();
-		LastBopIndexInBlock = INDEX_NONE;
+		LastBopIndexInBlock = 0;
 		bHasBop = false;
 	}
 
diff --git a/Engine/Plugins/Runtime/Metasound/Source/MetasoundStandardNodes/Private/MetasoundTriggerAccumulatorNode.cpp b/Engine/Plugins/Runtime/Metasound/Source/MetasoundStandardNodes/Private/MetasoundTriggerAccumulatorNode.cpp
--- a/Engine/Plugins/Runtime/Metasound/Source/MetasoundStandardNodes/Private/MetasoundTriggerAccumulatorNode.cpp
+++ b/Engine/Plugins/Runtime/Metasound/Source/MetasoundStandardNodes/Private/MetasoundTriggerAccumulatorNode.cpp
@@ -107,8 +107,9 @@ namespace Metasound
 		// We need to advance our own output triggers.
 		CountReached->AdvanceBlock();
 
-		int32 NumResets = ResetTrigger->Num();
-		int32 NumAccumulates = TriggerToAccumulate->Num();
+		// Query the number of triggers in the current block.
+		int32 NumResets = ResetTrigger->NumBoppedInBlock();
+		int32 NumAccumulates = TriggerToAccumulate->NumBoppedInBlock();
 		int32 MaxTriggers = FMath::Max(NumResets, NumAccumulates);
 
 		// Indices into each trigger array
diff --git a/Engine/Plugins/Runtime/Metasound/Source/MetasoundStandardNodes/Public/MetasoundBop.h b/Engine/Plugins/Runtime/Metasound/Source/MetasoundStandardNodes/Public/MetasoundBop.h
--- a/Engine/Plugins/Runtime/Metasound/Source/MetasoundStandardNodes/Public/MetasoundBop.h
+++ b/Engine/Plugins/Runtime/Metasound/Source/MetasoundStandardNodes/Public/MetasoundBop.h
@@ -65,6 +65,9 @@ namespace Metasound
 			/** Number of bopped frames. */
 			int32 Num() const;
 
+			/** Returns true if there is a bop in the current block of audio */
+			int32 NumBoppedInBlock() const;
+
 			/** Returns frame index for a given bop index. 
 			 *
 			 * @param InIndex - Index of bop. Must be a value between 0 and Num().
@@ -163,7 +166,7 @@ namespace Metasound
 			bool bHasBop = false;
 			int32 NumFramesPerBlock = 0;
 			float SampleRate;
-			int32 LastBopIndexInBlock = INDEX_NONE;
+			int32 LastBopIndexInBlock = 0;
 	};
 
 	DECLARE_METASOUND_DATA_REFERENCE_TYPES(FBop, METASOUNDSTANDARDNODES_API, FBopTypeInfo, FBopReadRef, FBopWriteRef);
