Minor "de-bopping" of Periodic trigger to avoid UX displaying it (and a few low-hanging variable name fix-ups), initialize member for safety

diff --git a/Engine/Plugins/Runtime/Metasound/Source/MetasoundStandardNodes/Private/MetasoundPeriodicTriggerNode.cpp b/Engine/Plugins/Runtime/Metasound/Source/MetasoundStandardNodes/Private/MetasoundPeriodicTriggerNode.cpp
--- a/Engine/Plugins/Runtime/Metasound/Source/MetasoundStandardNodes/Private/MetasoundPeriodicTriggerNode.cpp
+++ b/Engine/Plugins/Runtime/Metasound/Source/MetasoundStandardNodes/Private/MetasoundPeriodicTriggerNode.cpp
@@ -98,7 +98,7 @@ namespace Metasound
 			}
 		);
 
-		// Advance internal counter to get rid of old bops.
+		// Advance internal counter to get rid of old triggers.
 		TriggerOut->AdvanceBlock();
 
 		if (bEnabled)
@@ -121,17 +121,17 @@ namespace Metasound
 
 	TUniquePtr<IOperator> FPeriodicTriggerOperator::CreateOperator(const FCreateOperatorParams& InParams, FBuildErrorArray& OutErrors)
 	{
-		const FPeriodicTriggerNode& PeriodicBopNode = static_cast<const FPeriodicTriggerNode&>(InParams.Node);
+		const FPeriodicTriggerNode& PeriodicTriggerNode = static_cast<const FPeriodicTriggerNode&>(InParams.Node);
 
-		FFloatTimeReadRef Period = FFloatTimeReadRef::CreateNew(PeriodicBopNode.GetDefaultPeriodInSeconds(), ETimeResolution::Seconds);
+		FFloatTimeReadRef Period = FFloatTimeReadRef::CreateNew(PeriodicTriggerNode.GetDefaultPeriodInSeconds(), ETimeResolution::Seconds);
 
 		if (InParams.InputDataReferences.ContainsDataReadReference<FFloatTime>(TEXT("Period")))
 		{
 			Period = InParams.InputDataReferences.GetDataReadReference<FFloatTime>(TEXT("Period"));
 		}
 
-		FTriggerReadRef TriggerEnable = InParams.InputDataReferences.GetDataReadReferenceOrConstruct<FBop>(TEXT("Activate"), InParams.OperatorSettings);
-		FTriggerReadRef TriggerDisable = InParams.InputDataReferences.GetDataReadReferenceOrConstruct<FBop>(TEXT("Deactivate"), InParams.OperatorSettings);
+		FTriggerReadRef TriggerEnable = InParams.InputDataReferences.GetDataReadReferenceOrConstruct<FTrigger>(TEXT("Activate"), InParams.OperatorSettings);
+		FTriggerReadRef TriggerDisable = InParams.InputDataReferences.GetDataReadReferenceOrConstruct<FTrigger>(TEXT("Deactivate"), InParams.OperatorSettings);
 
 		return MakeUnique<FPeriodicTriggerOperator>(InParams.OperatorSettings, TriggerEnable, TriggerDisable, Period);
 	}
@@ -141,11 +141,11 @@ namespace Metasound
 		static const FVertexInterface Interface(
 			FInputVertexInterface(
 				TInputDataVertexModel<FFloatTime>(TEXT("Period"), LOCTEXT("PeriodTooltip", "The period to trigger in seconds.")),
-				TInputDataVertexModel<FBop>(TEXT("Activate"), LOCTEXT("TriggerEnableTooltip", "Enables executing periodic output triggers.")),
-				TInputDataVertexModel<FBop>(TEXT("Deactivate"), LOCTEXT("TriggerDisableTooltip", "Disables executing periodic output triggers."))
+				TInputDataVertexModel<FTrigger>(TEXT("Activate"), LOCTEXT("TriggerEnableTooltip", "Enables executing periodic output triggers.")),
+				TInputDataVertexModel<FTrigger>(TEXT("Deactivate"), LOCTEXT("TriggerDisableTooltip", "Disables executing periodic output triggers."))
 			),
 			FOutputVertexInterface(
-				TOutputDataVertexModel<FBop>(TEXT("Out"), LOCTEXT("TriggerOutTooltip", "The periodically generated output trigger"))
+				TOutputDataVertexModel<FTrigger>(TEXT("Out"), LOCTEXT("TriggerOutTooltip", "The periodically generated output trigger"))
 			)
 		);
 
@@ -157,10 +157,10 @@ namespace Metasound
 		auto InitNodeInfo = []() -> FNodeInfo
 		{
 			FNodeInfo Info;
-			Info.ClassName = "PeriodicBop";
+			Info.ClassName = "Periodic Trigger";
 			Info.MajorVersion = 1;
 			Info.MinorVersion = 0;
-			Info.Description = LOCTEXT("Metasound_PeriodicBopNodeDescription", "Emits a bop periodically based on the period duration given.");
+			Info.Description = LOCTEXT("Metasound_PeriodicTriggerNodeDescription", "Emits a trigger periodically based on the period duration given.");
 			Info.Author = PluginAuthor;
 			Info.PromptIfMissing = PluginNodeMissingPrompt;
 			Info.DefaultInterface = DeclareVertexInterface();
@@ -190,4 +190,4 @@ namespace Metasound
 	}
 }
 
-#undef LOCTEXT_NAMESPACE //MetasoundPeriodicBopNode
+#undef LOCTEXT_NAMESPACE //MetasoundPeriodicTriggerNode
diff --git a/Engine/Plugins/Runtime/Metasound/Source/MetasoundStandardNodes/Public/MetasoundPeriodicTriggerNode.h b/Engine/Plugins/Runtime/Metasound/Source/MetasoundStandardNodes/Public/MetasoundPeriodicTriggerNode.h
--- a/Engine/Plugins/Runtime/Metasound/Source/MetasoundStandardNodes/Public/MetasoundPeriodicTriggerNode.h
+++ b/Engine/Plugins/Runtime/Metasound/Source/MetasoundStandardNodes/Public/MetasoundPeriodicTriggerNode.h
@@ -24,6 +24,6 @@ namespace Metasound
 			float GetDefaultPeriodInSeconds() const;
 
 		private:
-			float DefaultPeriod;
+			float DefaultPeriod = 1.0f;
 	};
 }
