Metasound BP interface
- Adding IAudioInstanceTransmitter interface to Audio extensions
- Adding IAudioInstanceTransmitter to FActiveSound
- Renaming functions with "Archetype" to "MetasoundArchetype" to avoid conflict with UOBject archetype.
- Adding transmittable info to data type registration
- Adding default input pin to receive node.
- Adding FMetasoundInstanceTransmitter for communicating to single metasound source instance.

diff --git a/Engine/Plugins/Runtime/Metasound/Source/MetasoundEditor/Private/MetasoundFactory.cpp b/Engine/Plugins/Runtime/Metasound/Source/MetasoundEditor/Private/MetasoundFactory.cpp
--- a/Engine/Plugins/Runtime/Metasound/Source/MetasoundEditor/Private/MetasoundFactory.cpp
+++ b/Engine/Plugins/Runtime/Metasound/Source/MetasoundEditor/Private/MetasoundFactory.cpp
@@ -31,7 +31,7 @@ UObject* UMetasoundFactory::FactoryCreateNew(UClass* InClass, UObject* InParent,
 
 	NewMetasound->SetMetadata(Metadata);
 
-	NewMetasound->ConformDocumentToArchetype();
+	NewMetasound->ConformDocumentToMetasoundArchetype();
 
 	return NewMetasound;
 }
@@ -60,7 +60,7 @@ UObject* UMetasoundSourceFactory::FactoryCreateNew(UClass* InClass, UObject* InP
 
 	MetasoundSource->SetMetadata(Metadata);
 
-	MetasoundSource->ConformDocumentToArchetype();
+	MetasoundSource->ConformDocumentToMetasoundArchetype();
 
 	return MetasoundSource;
 }
diff --git a/Engine/Plugins/Runtime/Metasound/Source/MetasoundEngine/Private/Metasound.cpp b/Engine/Plugins/Runtime/Metasound/Source/MetasoundEngine/Private/Metasound.cpp
--- a/Engine/Plugins/Runtime/Metasound/Source/MetasoundEngine/Private/Metasound.cpp
+++ b/Engine/Plugins/Runtime/Metasound/Source/MetasoundEngine/Private/Metasound.cpp
@@ -53,20 +53,20 @@ Metasound::Frontend::TAccessPtr<const FMetasoundFrontendDocument> UMetasound::Ge
 	return Metasound::Frontend::MakeAccessPtr<const FMetasoundFrontendDocument>(MetasoundDocument.AccessPoint, MetasoundDocument);
 }
 
-const TArray<FMetasoundFrontendArchetype>& UMetasound::GetPreferredArchetypes() const
+const TArray<FMetasoundFrontendArchetype>& UMetasound::GetPreferredMetasoundArchetypes() const
 {
 	// Not preferred archetypes for a basic UMetasound.
 	static const TArray<FMetasoundFrontendArchetype> Preferred;
 	return Preferred;
 }
 
-bool UMetasound::IsArchetypeSupported(const FMetasoundFrontendArchetype& InArchetype) const
+bool UMetasound::IsMetasoundArchetypeSupported(const FMetasoundFrontendArchetype& InArchetype) const
 {
 	// All archetypes are supported.
 	return true;
 }
 
-const FMetasoundFrontendArchetype& UMetasound::GetPreferredArchetype(const FMetasoundFrontendDocument& InDocument) const
+const FMetasoundFrontendArchetype& UMetasound::GetPreferredMetasoundArchetype(const FMetasoundFrontendDocument& InDocument) const
 {
 	// Prefer to keep original archetype.
 	return InDocument.Archetype;
@@ -82,7 +82,7 @@ void UMetasound::SetMetadata(FMetasoundFrontendClassMetadata& InMetadata)
 
 void UMetasound::PostLoad()
 {
-	ConformDocumentToArchetype();
+	ConformDocumentToMetasoundArchetype();
 	Super::PostLoad();
 }
 
diff --git a/Engine/Plugins/Runtime/Metasound/Source/MetasoundEngine/Private/MetasoundSource.cpp b/Engine/Plugins/Runtime/Metasound/Source/MetasoundEngine/Private/MetasoundSource.cpp
--- a/Engine/Plugins/Runtime/Metasound/Source/MetasoundEngine/Private/MetasoundSource.cpp
+++ b/Engine/Plugins/Runtime/Metasound/Source/MetasoundEngine/Private/MetasoundSource.cpp
@@ -9,9 +9,14 @@
 #include "MetasoundAudioFormats.h"
 #include "MetasoundBop.h"
 #include "MetasoundEngineEnvironment.h"
+#include "MetasoundFrontendController.h"
+#include "MetasoundFrontendQuery.h"
+#include "MetasoundFrontendQuerySteps.h"
 #include "MetasoundGenerator.h"
+#include "MetasoundInstanceTransmitter.h"
 #include "MetasoundLog.h"
 #include "MetasoundPrimitives.h"
+#include "MetasoundReceiveNode.h"
 
 
 #if WITH_EDITORONLY_DATA
@@ -54,7 +59,7 @@ void UMetasoundSource::PostEditChangeProperty(FPropertyChangedEvent& InEvent)
 			case EMetasoundSourceAudioFormat::Stereo:
 
 				NumChannels = 2;
-				ensure(SetArchetype(GetStereoSourceArchetype()));
+				ensure(SetMetasoundArchetype(GetStereoSourceArchetype()));
 
 				break;
 
@@ -62,7 +67,7 @@ void UMetasoundSource::PostEditChangeProperty(FPropertyChangedEvent& InEvent)
 			default:
 
 				NumChannels = 1;
-				ensure(SetArchetype(GetMonoSourceArchetype()));
+				ensure(SetMetasoundArchetype(GetMonoSourceArchetype()));
 
 				break;
 		}
@@ -97,8 +102,7 @@ ISoundGeneratorPtr UMetasoundSource::CreateSoundGenerator(const FSoundGeneratorI
 	VirtualizationMode = EVirtualizationMode::PlayWhenSilent;
 
 	SampleRate = InParams.SampleRate;
-	const float BlockRate = 100.f; // Metasound graph gets evaluated 100 times per second.
-	FOperatorSettings InSettings(InParams.SampleRate, BlockRate);
+	FOperatorSettings InSettings = GetOperatorSettings(SampleRate);
 	FMetasoundEnvironment Environment;
 
 	// Add audio device ID to environment.
@@ -116,12 +120,27 @@ ISoundGeneratorPtr UMetasoundSource::CreateSoundGenerator(const FSoundGeneratorI
 	}
 	Environment.SetValue<FAudioDeviceHandle>(GetAudioDeviceHandleVariableName(), DeviceHandle);
 
-	// Create the operator from the graph handle.
-	TArray<IOperatorBuilder::FBuildErrorPtr> BuildErrors;
+	const FMetasoundFrontendDocument* OriginalDoc = GetDocument().Get();
+	if (nullptr == OriginalDoc)
+	{
+		UE_LOG(LogMetasound, Error, TEXT("Cannot create sound generator. Null Metasound document in UMetasoundSource [Name:%s]"), *GetName());
+		return ISoundGeneratorPtr(nullptr);
+	}
+
+	// Inject receive nodes for unused transmittable inputs. Perform edits on a copy
+	// of the document to avoid altering document.
+	// TODO: Use a light wrapper of the document instead of a copy.
+	FMetasoundFrontendDocument DocumentWithInjectedReceives;
+	ensure(CopyDocumentAndInjectReceiveNodes(InParams.InstanceID, *OriginalDoc, DocumentWithInjectedReceives));
+
+	// Create handles for new root graph
+	Frontend::FConstDocumentHandle NewDocumentHandle = Frontend::IDocumentController::CreateDocumentHandle(Frontend::MakeAccessPtr<const FMetasoundFrontendDocument>(DocumentWithInjectedReceives.AccessPoint, DocumentWithInjectedReceives));
+	Frontend::FConstGraphHandle RootGraph = NewDocumentHandle->GetRootGraph();
 
-	Frontend::FGraphHandle RootGraph = GetRootGraphHandle();
 	ensureAlways(RootGraph->IsValid());
 
+	// Create the operator from the graph handle.
+	TArray<IOperatorBuilder::FBuildErrorPtr> BuildErrors;
 	TUniquePtr<IOperator> Operator = RootGraph->BuildOperator(InSettings, Environment, BuildErrors);
 
 	if (!Operator.IsValid())
@@ -179,14 +198,236 @@ ISoundGeneratorPtr UMetasoundSource::CreateSoundGenerator(const FSoundGeneratorI
 	return ISoundGeneratorPtr(nullptr);
 }
 
+TUniquePtr<IAudioInstanceTransmitter> UMetasoundSource::CreateInstanceTransmitter(const FAudioInstanceTransmitterInitParams& InParams) const
+{
+	Metasound::FMetasoundInstanceTransmitter::FInitParams InitParams(GetOperatorSettings(InParams.SampleRate), InParams.InstanceID);
+
+	for (const FSendInfoAndVertexName& InfoAndName : GetSendInfos(InParams.InstanceID))
+	{
+		InitParams.Infos.Add(InfoAndName.SendInfo);
+	}
+
+	return MakeUnique<Metasound::FMetasoundInstanceTransmitter>(InitParams);
+}
+
 void UMetasoundSource::PostLoad()
 {
 	Super::PostLoad();
 
-	ConformDocumentToArchetype();
+	ConformDocumentToMetasoundArchetype();
+}
+
+bool UMetasoundSource::GetReceiveNodeMetadataForDataType(const FName& InTypeName, FMetasoundFrontendClassMetadata& OutMetadata) const
+{
+	using namespace Metasound;
+
+	FFrontendQuery Query;
+
+	Query.AddStep<FGenerateAllAvailableNodeClasses>()
+		.AddStep<FFilterClassesByClassName>(FReceiveNodeNames::GetClassNameForDataType(InTypeName));
+	
+	FFrontendQuerySelectionView Result = Query.ExecuteQuery();
+
+	TArrayView<const FFrontendQueryEntry* const> Selection = Result.GetSelection();
+
+	if (Selection.Num() > 0)
+	{
+		check(Selection[0]->Value.IsType<FMetasoundFrontendClass>());
+		OutMetadata = Selection[0]->Value.Get<FMetasoundFrontendClass>().Metadata;
+		return true;
+	}
+
+	return false;
+}
+
+Metasound::Frontend::FNodeHandle UMetasoundSource::AddInputPinForSendAddress(const Metasound::FMetasoundInstanceTransmitter::FSendInfo& InSendInfo, Metasound::Frontend::FGraphHandle InGraph) const
+{
+	FMetasoundFrontendClassInput Description;
+	FGuid PointID = InGraph->GetNewPointID();
+
+	Description.Name = InSendInfo.Address.ChannelName.ToString();
+	Description.TypeName = Metasound::GetMetasoundDataTypeName<Metasound::FSendAddress>();
+	Description.Metadata.Description = FText::GetEmpty();
+	Description.PointIDs.Add(PointID);
+
+	FMetasoundFrontendVertexLiteral DefaultValue;
+	DefaultValue.PointID = PointID;
+	DefaultValue.Value.Set(InSendInfo.Address.ChannelName.ToString());
+	Description.Defaults.Add(DefaultValue);
+	
+	return InGraph->AddInputVertex(Description);
+}
+
+bool UMetasoundSource::CopyDocumentAndInjectReceiveNodes(uint64 InInstanceID, const FMetasoundFrontendDocument& InSourceDoc, FMetasoundFrontendDocument& OutDestDoc) const
+{
+	using namespace Metasound;
+
+	OutDestDoc = InSourceDoc;
+
+	Frontend::FDocumentHandle Document = Frontend::IDocumentController::CreateDocumentHandle(Frontend::MakeAccessPtr(OutDestDoc.AccessPoint, OutDestDoc));
+	Frontend::FGraphHandle RootGraph = Document->GetRootGraph();
+
+	TArray<FSendInfoAndVertexName> SendInfoAndVertexes = GetSendInfos(InInstanceID);
+
+	// Inject receive nodes for each transmittable input
+	for (const FSendInfoAndVertexName& InfoAndVertexName : SendInfoAndVertexes)
+	{
+
+		// Add receive node to graph
+		FMetasoundFrontendClassMetadata ReceiveNodeMetadata;
+		bool bSuccess = GetReceiveNodeMetadataForDataType(InfoAndVertexName.SendInfo.TypeName, ReceiveNodeMetadata);
+		if (!bSuccess)
+		{
+			// TODO: log warning
+			continue;
+		}
+		Frontend::FNodeHandle ReceiveNode = RootGraph->AddNode(ReceiveNodeMetadata);
+
+		// Add receive node address to graph
+		Frontend::FNodeHandle AddressNode = UMetasoundSource::AddInputPinForSendAddress(InfoAndVertexName.SendInfo, RootGraph);
+		TArray<Frontend::FOutputHandle> AddressNodeOutputs = AddressNode->GetOutputs();
+		if (AddressNodeOutputs.Num() != 1)
+		{
+			// TODO: log warning
+			continue;
+		}
+
+		Frontend::FOutputHandle AddressOutput = AddressNodeOutputs[0];
+		TArray<Frontend::FInputHandle> ReceiveAddressInput = ReceiveNode->GetInputsWithVertexName(Metasound::FReceiveNodeNames::GetAddressInputName());
+		if (ReceiveAddressInput.Num() != 1)
+		{
+			// TODO: log error
+			continue;
+		}
+
+		ensure(ReceiveAddressInput[0]->Connect(*AddressOutput));
+
+
+
+		// Swap input node connections with receive node connections
+		Frontend::FNodeHandle InputNode = RootGraph->GetInputNodeWithName(InfoAndVertexName.VertexName);
+		if (!ensure(InputNode->GetOutputs().Num() == 1))
+		{
+			// TODO: handle input node with varying number of outputs or varying output types.
+			continue;
+		}
+
+		Frontend::FOutputHandle InputNodeOutput = InputNode->GetOutputs()[0];
+
+		if (ensure(ReceiveNode->IsValid()))
+		{
+			TArray<Frontend::FOutputHandle> ReceiveNodeOutputs = ReceiveNode->GetOutputs();
+			if (!ensure(ReceiveNodeOutputs.Num() == 1))
+			{
+				// TODO: handle array outputs and receive nodes of varying formats.
+				continue;
+			}
+
+			TArray<Frontend::FInputHandle> ReceiveDefaultInputs = ReceiveNode->GetInputsWithVertexName(Metasound::FReceiveNodeNames::GetDefaultDataInputName());
+			if (ensure(ReceiveDefaultInputs.Num() == 1))
+			{
+				Frontend::FOutputHandle ReceiverNodeOutput = ReceiveNodeOutputs[0];
+				for (Frontend::FInputHandle NodeInput : InputNodeOutput->GetCurrentlyConnectedInputs())
+				{
+					// Swap connections to receiver node
+					ensure(InputNodeOutput->Disconnect(*NodeInput));
+					ensure(ReceiverNodeOutput->Connect(*NodeInput));
+				}
+
+				ReceiveDefaultInputs[0]->Connect(*InputNodeOutput);
+			}
+		}
+	}
+
+	return true;
+}
+
+TArray<FString> UMetasoundSource::GetTransmittableInputVertexNames() const
+{
+	using namespace Metasound;
+
+	// Unused inputs are all input vertices that are not in the archetype.
+	TArray<FString> ArchetypeInputVertexNames;
+	for (const FMetasoundFrontendClassVertex& Vertex : GetMetasoundArchetype().Interface.Inputs)
+	{
+		ArchetypeInputVertexNames.Add(Vertex.Name);
+	}
+
+	Frontend::FConstGraphHandle RootGraph = GetRootGraphHandle();
+	TArray<FString> GraphInputVertexNames = RootGraph->GetInputVertexNames();
+
+	// Filter graph inputs by archetype inputs.
+	GraphInputVertexNames = GraphInputVertexNames.FilterByPredicate([&](const FString& InName) { return !ArchetypeInputVertexNames.Contains(InName); });
+
+	auto IsDataTypeTransmittable = [&](const FString& InVertexName)
+	{
+		Frontend::FConstClassInputAccessPtr ClassInputPtr = RootGraph->FindClassInputWithName(InVertexName);
+		if (const FMetasoundFrontendClassInput* ClassInput = ClassInputPtr.Get())
+		{
+			FDataTypeRegistryInfo TypeInfo;
+			if (Frontend::GetTraitsForDataType(ClassInput->TypeName, TypeInfo))
+			{
+				return TypeInfo.bIsTransmittable;
+			}
+		}
+		return false;
+	};
+
+	GraphInputVertexNames = GraphInputVertexNames.FilterByPredicate(IsDataTypeTransmittable);
+
+	return GraphInputVertexNames;
+}
+
+Metasound::FOperatorSettings UMetasoundSource::GetOperatorSettings(float InSampleRate) const
+{
+	constexpr float BlockRate = 100.f; // Metasound graph gets evaluated 100 times per second.
+	return Metasound::FOperatorSettings(InSampleRate, BlockRate);
+}
+
+Metasound::FSendAddress UMetasoundSource::CreateSendAddress(uint64 InInstanceID, const FString& InVertexName) const
+{
+	using namespace Metasound;
+
+	FSendAddress Address;
+
+	Address.Subsystem = GetSubsystemNameForSendScope(ETransmissionScope::Global);
+	Address.ChannelName = FName(FString::Printf(TEXT("%d:%s"), InInstanceID, *InVertexName));
+
+	return Address;
+}
+
+TArray<UMetasoundSource::FSendInfoAndVertexName> UMetasoundSource::GetSendInfos(uint64 InInstanceID) const
+{
+	using FSendInfo = Metasound::FMetasoundInstanceTransmitter::FSendInfo;
+	using namespace Metasound::Frontend;
+
+	TArray<FSendInfoAndVertexName> SendInfos;
+
+	FConstGraphHandle RootGraph = GetRootGraphHandle();
+
+	TArray<FString> SendVertices = GetTransmittableInputVertexNames();
+	for (const FString& VertexName : SendVertices)
+	{
+		FConstNodeHandle InputNode = RootGraph->GetInputNodeWithName(VertexName);
+		for (FConstOutputHandle InputHandle : InputNode->GetConstOutputs())
+		{
+			FSendInfoAndVertexName Info;
+
+			// TODO: incorporate PointID into address. But need to ensure that PointID
+			// will be maintained after injecting Receive nodes. 
+			Info.SendInfo.Address = CreateSendAddress(InInstanceID, InputHandle->GetName());
+			Info.SendInfo.ParameterName = FName(InputHandle->GetDisplayName().ToString()); // TODO: display name hack. Need to have nameing consistent in editor for inputs
+			Info.SendInfo.TypeName = InputHandle->GetDataType();
+			Info.VertexName = VertexName;
+			
+			SendInfos.Add(Info);
+		}
+	}
+
+	return SendInfos;
 }
 
-const TArray<FMetasoundFrontendArchetype>& UMetasoundSource::GetPreferredArchetypes() const 
+const TArray<FMetasoundFrontendArchetype>& UMetasoundSource::GetPreferredMetasoundArchetypes() const 
 {
 	static const TArray<FMetasoundFrontendArchetype> Preferred({GetMonoSourceArchetype(), GetStereoSourceArchetype()});
 
diff --git a/Engine/Plugins/Runtime/Metasound/Source/MetasoundEngine/Public/Metasound.h b/Engine/Plugins/Runtime/Metasound/Source/MetasoundEngine/Public/Metasound.h
--- a/Engine/Plugins/Runtime/Metasound/Source/MetasoundEngine/Public/Metasound.h
+++ b/Engine/Plugins/Runtime/Metasound/Source/MetasoundEngine/Public/Metasound.h
@@ -75,11 +75,11 @@ public:
 		return const_cast<UObject*>(CastChecked<const UObject>(this));
 	}
 
-	const TArray<FMetasoundFrontendArchetype>& GetPreferredArchetypes() const override;
+	const TArray<FMetasoundFrontendArchetype>& GetPreferredMetasoundArchetypes() const override;
 
-	bool IsArchetypeSupported(const FMetasoundFrontendArchetype& InArchetype) const override;
+	bool IsMetasoundArchetypeSupported(const FMetasoundFrontendArchetype& InArchetype) const override;
 
-	const FMetasoundFrontendArchetype& GetPreferredArchetype(const FMetasoundFrontendDocument& InDocument) const override;
+	const FMetasoundFrontendArchetype& GetPreferredMetasoundArchetype(const FMetasoundFrontendDocument& InDocument) const override;
 
 	// Updates the Metasound's metadata (name, author, etc).
 	// @param InMetadata Metadata containing corrections to the class metadata.
diff --git a/Engine/Plugins/Runtime/Metasound/Source/MetasoundEngine/Public/MetasoundSource.h b/Engine/Plugins/Runtime/Metasound/Source/MetasoundEngine/Public/MetasoundSource.h
--- a/Engine/Plugins/Runtime/Metasound/Source/MetasoundEngine/Public/MetasoundSource.h
+++ b/Engine/Plugins/Runtime/Metasound/Source/MetasoundEngine/Public/MetasoundSource.h
@@ -2,9 +2,11 @@
 #pragma once
 
 #include "CoreMinimal.h"
+#include "MetasoundAssetBase.h"
 #include "MetasoundFrontend.h"
 #include "MetasoundFrontendDocument.h"
-#include "MetasoundAssetBase.h"
+#include "MetasoundInstanceTransmitter.h"
+#include "MetasoundRouter.h"
 
 #include "Sound/SoundWaveProcedural.h"
 
@@ -111,10 +113,11 @@ public:
 	bool SupportsSubtitles() const override;
 	float GetDuration() override;
 	ISoundGeneratorPtr CreateSoundGenerator(const FSoundGeneratorInitParams& InParams) override;
+	TUniquePtr<IAudioInstanceTransmitter> CreateInstanceTransmitter(const FAudioInstanceTransmitterInitParams& InParams) const override;
 	void PostLoad() override;
 
 	// Get the most up to date archetype for metasound sources.
-	const TArray<FMetasoundFrontendArchetype>& GetPreferredArchetypes() const override;
+	const TArray<FMetasoundFrontendArchetype>& GetPreferredMetasoundArchetypes() const override;
 
 protected:
 
@@ -133,6 +136,23 @@ protected:
 	}
 
 private:
+	struct FSendInfoAndVertexName
+	{
+		Metasound::FMetasoundInstanceTransmitter::FSendInfo SendInfo;
+		FString VertexName;
+	};
+
+	bool GetReceiveNodeMetadataForDataType(const FName& InTypeName, FMetasoundFrontendClassMetadata& OutMetadata) const;
+	Metasound::Frontend::FNodeHandle AddInputPinForSendAddress(const Metasound::FMetasoundInstanceTransmitter::FSendInfo& InSendInfo, Metasound::Frontend::FGraphHandle InGraph) const;
+	bool CopyDocumentAndInjectReceiveNodes(uint64 InInstanceID, const FMetasoundFrontendDocument& InSourceDoc, FMetasoundFrontendDocument& OutDestDoc) const;	
+	TArray<FString> GetTransmittableInputVertexNames() const;
+	Metasound::FOperatorSettings GetOperatorSettings(float InSampleRate) const;
+	Metasound::FSendAddress CreateSendAddress(uint64 InInstanceID, const FString& InVertexName) const;
+	
+
+
+	TArray<FSendInfoAndVertexName> GetSendInfos(uint64 InInstanceID) const;
+
 
 	static const FString& GetOnPlayInputName();
 	static const FString& GetAudioOutputName();
diff --git a/Engine/Plugins/Runtime/Metasound/Source/MetasoundEngine/Public/MetasoundUObjectRegistry.h b/Engine/Plugins/Runtime/Metasound/Source/MetasoundEngine/Public/MetasoundUObjectRegistry.h
--- a/Engine/Plugins/Runtime/Metasound/Source/MetasoundEngine/Public/MetasoundUObjectRegistry.h
+++ b/Engine/Plugins/Runtime/Metasound/Source/MetasoundEngine/Public/MetasoundUObjectRegistry.h
@@ -145,7 +145,7 @@ namespace Metasound
 			{
 				static_assert(std::is_base_of<FMetasoundAssetBase, UClassType>::value, "UClass must be derived from FMetasoundAssetBase");
 
-				const TArray<FMetasoundFrontendArchetype>& PreferredArchetypes = GetDefault<UClassType>()->GetPreferredArchetypes();
+				const TArray<FMetasoundFrontendArchetype>& PreferredArchetypes = GetDefault<UClassType>()->GetPreferredMetasoundArchetypes();
 
 				for (const FMetasoundFrontendArchetype& Arch : PreferredArchetypes)
 				{
diff --git a/Engine/Plugins/Runtime/Metasound/Source/MetasoundFrontend/Private/MetasoundAssetBase.cpp b/Engine/Plugins/Runtime/Metasound/Source/MetasoundFrontend/Private/MetasoundAssetBase.cpp
--- a/Engine/Plugins/Runtime/Metasound/Source/MetasoundFrontend/Private/MetasoundAssetBase.cpp
+++ b/Engine/Plugins/Runtime/Metasound/Source/MetasoundFrontend/Private/MetasoundAssetBase.cpp
@@ -35,17 +35,17 @@ void FMetasoundAssetBase::SetMetadata(FMetasoundFrontendClassMetadata& InMetadat
 	}
 }
 
-const FMetasoundFrontendArchetype& FMetasoundAssetBase::GetArchetype() const
+const FMetasoundFrontendArchetype& FMetasoundAssetBase::GetMetasoundArchetype() const
 {
 	return Archetype;
 }
 
-bool FMetasoundAssetBase::SetArchetype(const FMetasoundFrontendArchetype& InArchetype)
+bool FMetasoundAssetBase::SetMetasoundArchetype(const FMetasoundFrontendArchetype& InArchetype)
 {
-	if (IsArchetypeSupported(InArchetype))
+	if (IsMetasoundArchetypeSupported(InArchetype))
 	{
 		Archetype = InArchetype;
-		ConformDocumentToArchetype();
+		ConformDocumentToMetasoundArchetype();
 
 		return true;
 	}
@@ -54,26 +54,26 @@ bool FMetasoundAssetBase::SetArchetype(const FMetasoundFrontendArchetype& InArch
 	return false;
 }
 
-bool FMetasoundAssetBase::IsArchetypeSupported(const FMetasoundFrontendArchetype& InArchetype) const
+bool FMetasoundAssetBase::IsMetasoundArchetypeSupported(const FMetasoundFrontendArchetype& InArchetype) const
 {
 	auto IsEqualArchetype = [&](const FMetasoundFrontendArchetype& SupportedArchetype)
 	{
 		return Metasound::Frontend::IsEqualArchetype(InArchetype, SupportedArchetype);
 	};
 
-	return Algo::AnyOf(GetPreferredArchetypes(), IsEqualArchetype);
+	return Algo::AnyOf(GetPreferredMetasoundArchetypes(), IsEqualArchetype);
 }
 
-const FMetasoundFrontendArchetype& FMetasoundAssetBase::GetPreferredArchetype(const FMetasoundFrontendDocument& InDocument) const
+const FMetasoundFrontendArchetype& FMetasoundAssetBase::GetPreferredMetasoundArchetype(const FMetasoundFrontendDocument& InDocument) const
 {
 	// Default to archetype provided in case it is supported. 
-	if (IsArchetypeSupported(InDocument.Archetype))
+	if (IsMetasoundArchetypeSupported(InDocument.Archetype))
 	{
 		return InDocument.Archetype;
 	}
 
 	// If existing archetype is not supported, get the most similar that still supports the documents environment.
-	const FMetasoundFrontendArchetype* SimilarArchetype = Metasound::Frontend::FindMostSimilarArchetypeSupportingEnvironment(InDocument, GetPreferredArchetypes());
+	const FMetasoundFrontendArchetype* SimilarArchetype = Metasound::Frontend::FindMostSimilarArchetypeSupportingEnvironment(InDocument, GetPreferredMetasoundArchetypes());
 
 	if (nullptr != SimilarArchetype)
 	{
@@ -81,7 +81,7 @@ const FMetasoundFrontendArchetype& FMetasoundAssetBase::GetPreferredArchetype(co
 	}
 
 	// Nothing found. Return the existing archetype for the FMetasoundAssetBase.
-	return GetArchetype();
+	return GetMetasoundArchetype();
 }
 
 void FMetasoundAssetBase::SetDocument(const FMetasoundFrontendDocument& InDocument, bool bForceUpdateArchetype)
@@ -91,18 +91,18 @@ void FMetasoundAssetBase::SetDocument(const FMetasoundFrontendDocument& InDocume
 
 	FMetasoundFrontendArchetype NewArch = InDocument.Archetype;
 
-	if (bForceUpdateArchetype || (!IsArchetypeSupported(NewArch)))
+	if (bForceUpdateArchetype || (!IsMetasoundArchetypeSupported(NewArch)))
 	{
-		NewArch = GetPreferredArchetype(Document);
+		NewArch = GetPreferredMetasoundArchetype(Document);
 	}
 
-	ensure(SetArchetype(NewArch));
+	ensure(SetMetasoundArchetype(NewArch));
 }
 
-void FMetasoundAssetBase::ConformDocumentToArchetype()
+void FMetasoundAssetBase::ConformDocumentToMetasoundArchetype()
 {
 	FMetasoundFrontendDocument& Document = GetDocumentChecked();
-	Document.Archetype = GetArchetype();
+	Document.Archetype = GetMetasoundArchetype();
 
 	Metasound::Frontend::FMatchRootGraphToArchetype Transform;
 	Transform.Transform(GetDocumentHandle());
@@ -168,3 +168,4 @@ const FMetasoundFrontendDocument& FMetasoundAssetBase::GetDocumentChecked() cons
 	check(DocAccessPtr.IsValid());
 	return *DocAccessPtr;
 }
+
diff --git a/Engine/Plugins/Runtime/Metasound/Source/MetasoundFrontend/Private/MetasoundFrontendInvalidController.h b/Engine/Plugins/Runtime/Metasound/Source/MetasoundFrontend/Private/MetasoundFrontendInvalidController.h
--- a/Engine/Plugins/Runtime/Metasound/Source/MetasoundFrontend/Private/MetasoundFrontendInvalidController.h
+++ b/Engine/Plugins/Runtime/Metasound/Source/MetasoundFrontend/Private/MetasoundFrontendInvalidController.h
@@ -53,6 +53,9 @@ namespace Metasound
 			TSharedRef<INodeController> GetOwningNode() override;
 			TSharedRef<const INodeController> GetOwningNode() const override;
 
+			TArray<FInputHandle> GetCurrentlyConnectedInputs() override { return TArray<FInputHandle>(); }
+			TArray<FConstInputHandle> GetCurrentlyConnectedInputs() const override { return TArray<FConstInputHandle>(); }
+			bool Disconnect() override { return false; }
 
 			FConnectability CanConnectTo(const IInputController& InController) const override { return FConnectability(); }
 			bool Connect(IInputController& InController) override { return false; }
@@ -95,12 +98,12 @@ namespace Metasound
 
 			virtual TSharedRef<IOutputController> GetCurrentlyConnectedOutput() override { return FInvalidOutputController::GetInvalid(); }
 			virtual TSharedRef<const IOutputController> GetCurrentlyConnectedOutput() const override { return FInvalidOutputController::GetInvalid(); }
+			bool Disconnect() override { return false; }
 
 			virtual FConnectability CanConnectTo(const IOutputController& InController) const override { return FConnectability(); }
 			virtual bool Connect(IOutputController& InController) override { return false; }
 			virtual bool ConnectWithConverterNode(IOutputController& InController, const FConverterNodeInfo& InNodeClassName) override { return false; }
 			virtual bool Disconnect(IOutputController& InController) override { return false; }
-			virtual bool Disconnect() override { return false; }
 		protected:
 			FDocumentAccess ShareAccess() override { return FDocumentAccess(); }
 			FConstDocumentAccess ShareAccess() const override { return FConstDocumentAccess(); }
@@ -255,6 +258,8 @@ namespace Metasound
 
 			TSharedRef<INodeController> AddNode(const FNodeClassInfo& InNodeClass) override { return FInvalidNodeController::GetInvalid(); }
 			TSharedRef<INodeController> AddNode(const FNodeRegistryKey& InNodeClass) override { return FInvalidNodeController::GetInvalid(); }
+			TSharedRef<INodeController> AddNode(const FMetasoundFrontendClassMetadata& InNodeClass) override { return FInvalidNodeController::GetInvalid(); }
+
 
 			// Remove the node corresponding to this node handle.
 			// On success, invalidates the received node handle.
diff --git a/Engine/Plugins/Runtime/Metasound/Source/MetasoundFrontend/Private/MetasoundFrontendQuerySteps.cpp b/Engine/Plugins/Runtime/Metasound/Source/MetasoundFrontend/Private/MetasoundFrontendQuerySteps.cpp
--- a/Engine/Plugins/Runtime/Metasound/Source/MetasoundFrontend/Private/MetasoundFrontendQuerySteps.cpp
+++ b/Engine/Plugins/Runtime/Metasound/Source/MetasoundFrontend/Private/MetasoundFrontendQuerySteps.cpp
@@ -50,5 +50,14 @@ namespace Metasound
 		);
 	}
 
+	FFilterClassesByClassName::FFilterClassesByClassName(const FString& InClassName)
+	: ClassName(InClassName)
+	{
+	}
+
+	bool FFilterClassesByClassName::Filter(const FFrontendQueryEntry& InEntry) const 
+	{
+		return InEntry.Value.Get<FMetasoundFrontendClass>().Metadata.Name.Name == ClassName;
+	}
 }
 
diff --git a/Engine/Plugins/Runtime/Metasound/Source/MetasoundFrontend/Private/MetasoundFrontendStandardController.cpp b/Engine/Plugins/Runtime/Metasound/Source/MetasoundFrontend/Private/MetasoundFrontendStandardController.cpp
--- a/Engine/Plugins/Runtime/Metasound/Source/MetasoundFrontend/Private/MetasoundFrontendStandardController.cpp
+++ b/Engine/Plugins/Runtime/Metasound/Source/MetasoundFrontend/Private/MetasoundFrontendStandardController.cpp
@@ -347,13 +347,14 @@ namespace Metasound
 		FBaseOutputController::FBaseOutputController(const FBaseOutputController::FInitParams& InParams)
 		: ID(InParams.ID)
 		, NodeVertexPtr(InParams.NodeVertexPtr)
+		, GraphPtr(InParams.GraphPtr)
 		, OwningNode(InParams.OwningNode)
 		{
 		}
 
 		bool FBaseOutputController::IsValid() const
 		{
-			return OwningNode->IsValid() && NodeVertexPtr.IsValid();
+			return OwningNode->IsValid() && NodeVertexPtr.IsValid() && GraphPtr.IsValid();
 		}
 
 		FGuid FBaseOutputController::GetID() const
@@ -394,6 +395,61 @@ namespace Metasound
 			return OwningNode;
 		}
 
+		TArray<FInputHandle> FBaseOutputController::GetCurrentlyConnectedInputs() 
+		{
+			TArray<FInputHandle> Inputs;
+
+			// Create output handle from output node.
+			FGraphHandle Graph = OwningNode->GetOwningGraph();
+
+			for (const FMetasoundFrontendEdge& Edge : FindEdges())
+			{
+				FNodeHandle InputNode = Graph->GetNodeWithID(Edge.ToNodeID);
+
+				FInputHandle Input = InputNode->GetInputWithID(Edge.ToPointID);
+				if (Input->IsValid())
+				{
+					Inputs.Add(Input);
+				}
+			}
+
+			return Inputs;
+		}
+
+		TArray<FConstInputHandle> FBaseOutputController::GetCurrentlyConnectedInputs() const 
+		{
+			TArray<FConstInputHandle> Inputs;
+
+			// Create output handle from output node.
+			FConstGraphHandle Graph = OwningNode->GetOwningGraph();
+
+			for (const FMetasoundFrontendEdge& Edge : FindEdges())
+			{
+				FConstNodeHandle InputNode = Graph->GetNodeWithID(Edge.ToNodeID);
+
+				FConstInputHandle Input = InputNode->GetInputWithID(Edge.ToPointID);
+				if (Input->IsValid())
+				{
+					Inputs.Add(Input);
+				}
+			}
+
+			return Inputs;
+		}
+
+		bool FBaseOutputController::Disconnect() 
+		{
+			bool bSuccess = true;
+			for (FInputHandle Input : GetCurrentlyConnectedInputs())
+			{
+				if (Input->IsValid())
+				{
+					bSuccess &= Disconnect(*Input);
+				}
+			}
+			return bSuccess;
+		}
+
 		FConnectability FBaseOutputController::CanConnectTo(const IInputController& InController) const
 		{
 			return InController.CanConnectTo(*this);
@@ -414,11 +470,32 @@ namespace Metasound
 			return InController.Disconnect(*this);
 		}
 
+		TArray<FMetasoundFrontendEdge> FBaseOutputController::FindEdges() const
+		{
+			if (GraphPtr.IsValid())
+			{
+				const FGuid NodeID = GetOwningNodeID();
+				FGuid PointID = GetID();
+
+				auto EdgeHasMatchingSource = [&](const FMetasoundFrontendEdge& Edge)
+				{
+					return (Edge.FromNodeID == NodeID) && (Edge.FromPointID == PointID);
+				};
+
+				return GraphPtr->Edges.FilterByPredicate(EdgeHasMatchingSource);
+			}
+
+			return TArray<FMetasoundFrontendEdge>();
+		}
+
+
 		FDocumentAccess FBaseOutputController::ShareAccess() 
 		{
 			FDocumentAccess Access;
 
 			Access.ConstVertex = NodeVertexPtr;
+			Access.Graph = GraphPtr;
+			Access.ConstGraph = GraphPtr;
 
 			return Access;
 		}
@@ -428,6 +505,7 @@ namespace Metasound
 			FConstDocumentAccess Access;
 
 			Access.ConstVertex = NodeVertexPtr;
+			Access.ConstGraph = GraphPtr;
 
 			return Access;
 		}
@@ -436,7 +514,7 @@ namespace Metasound
 		// FNodeOutputController
 		//
 		FNodeOutputController::FNodeOutputController(const FNodeOutputController::FInitParams& InParams)
-		: FBaseOutputController({InParams.ID, InParams.NodeVertexPtr, InParams.OwningNode})
+		: FBaseOutputController({InParams.ID, InParams.NodeVertexPtr, InParams.GraphPtr, InParams.OwningNode})
 		, ClassOutputPtr(InParams.ClassOutputPtr)
 		{
 		}
@@ -488,7 +566,7 @@ namespace Metasound
 		// FInputNodeOutputController
 		// 
 		FInputNodeOutputController::FInputNodeOutputController(const FInputNodeOutputController::FInitParams& InParams)
-		: FBaseOutputController({InParams.ID, InParams.NodeVertexPtr, InParams.OwningNode})
+		: FBaseOutputController({InParams.ID, InParams.NodeVertexPtr, InParams.GraphPtr, InParams.OwningNode})
 		, OwningGraphClassInputPtr(InParams.OwningGraphClassInputPtr)
 		{
 		}
@@ -1554,7 +1632,7 @@ namespace Metasound
 
 		FOutputHandle FNodeController::CreateOutputController(FGuid InPointID, FConstVertexAccessPtr InNodeVertexPtr, FConstClassOutputAccessPtr InClassOutputPtr, FNodeHandle InOwningNode) const
 		{
-			return MakeShared<FNodeOutputController>(FNodeOutputController::FInitParams{InPointID, InNodeVertexPtr, InClassOutputPtr, InOwningNode});
+			return MakeShared<FNodeOutputController>(FNodeOutputController::FInitParams{InPointID, InNodeVertexPtr, InClassOutputPtr, GraphPtr, InOwningNode});
 		}
 
 		FDocumentAccess FNodeController::ShareAccess() 
@@ -1742,7 +1820,7 @@ namespace Metasound
 
 		FOutputHandle FInputNodeController::CreateOutputController(FGuid InPointID, FConstVertexAccessPtr InNodeVertexPtr, FConstClassOutputAccessPtr InClassOutputPtr, FNodeHandle InOwningNode) const
 		{
-			return MakeShared<FInputNodeOutputController>(FInputNodeOutputController::FInitParams{InPointID, InNodeVertexPtr, OwningGraphClassInputPtr, InOwningNode});
+			return MakeShared<FInputNodeOutputController>(FInputNodeOutputController::FInitParams{InPointID, InNodeVertexPtr, OwningGraphClassInputPtr, GraphPtr, InOwningNode});
 		}
 
 		FDocumentAccess FInputNodeController::ShareAccess() 
@@ -2301,6 +2379,12 @@ namespace Metasound
 			return AddNode(ClassInfo);
 		}
 
+		FNodeHandle FGraphController::AddNode(const FMetasoundFrontendClassMetadata& InClassMetadata)
+		{
+			return AddNode(FMetasoundFrontendRegistryContainer
+::GetRegistryKey(InClassMetadata));
+		}
+
 		// Remove the node corresponding to this node handle.
 		// On success, invalidates the received node handle.
 		bool FGraphController::RemoveNode(INodeController& InNode)
diff --git a/Engine/Plugins/Runtime/Metasound/Source/MetasoundFrontend/Private/MetasoundFrontendStandardController.h b/Engine/Plugins/Runtime/Metasound/Source/MetasoundFrontend/Private/MetasoundFrontendStandardController.h
--- a/Engine/Plugins/Runtime/Metasound/Source/MetasoundFrontend/Private/MetasoundFrontendStandardController.h
+++ b/Engine/Plugins/Runtime/Metasound/Source/MetasoundFrontend/Private/MetasoundFrontendStandardController.h
@@ -26,6 +26,7 @@ namespace Metasound
 				FGuid ID;
 
 				FConstVertexAccessPtr NodeVertexPtr;
+				FGraphAccessPtr GraphPtr; 
 
 				/* Node handle which owns this output. */
 				FNodeHandle OwningNode;
@@ -47,6 +48,10 @@ namespace Metasound
 			FNodeHandle GetOwningNode() override;
 			FConstNodeHandle GetOwningNode() const override;
 
+			TArray<FInputHandle> GetCurrentlyConnectedInputs() override;
+			TArray<FConstInputHandle> GetCurrentlyConnectedInputs() const override;
+			bool Disconnect() override;
+
 			// Connection logic.
 			FConnectability CanConnectTo(const IInputController& InController) const override;
 			bool Connect(IInputController& InController) override;
@@ -57,9 +62,15 @@ namespace Metasound
 			FDocumentAccess ShareAccess() override;
 			FConstDocumentAccess ShareAccess() const override;
 
+
 			FGuid ID;
 			FConstVertexAccessPtr NodeVertexPtr;	
+			FGraphAccessPtr GraphPtr; 
 			FNodeHandle OwningNode;
+
+		private:
+
+			TArray<FMetasoundFrontendEdge> FindEdges() const;
 		};
 
 		/** FNodeOutputController represents a single output of a single node. */
@@ -73,6 +84,7 @@ namespace Metasound
 				FConstVertexAccessPtr NodeVertexPtr;
 			
 				FConstClassOutputAccessPtr ClassOutputPtr;
+				FGraphAccessPtr GraphPtr; 
 
 				/* Node handle which owns this output. */
 				FNodeHandle OwningNode;
@@ -112,6 +124,7 @@ namespace Metasound
 				FConstVertexAccessPtr NodeVertexPtr;
 			
 				FConstClassInputAccessPtr OwningGraphClassInputPtr;
+				FGraphAccessPtr GraphPtr; 
 
 				/* Node handle which owns this output. */
 				FNodeHandle OwningNode;
@@ -700,6 +713,7 @@ namespace Metasound
 
 			FNodeHandle AddNode(const FNodeClassInfo& InNodeClass) override;
 			FNodeHandle AddNode(const FNodeRegistryKey& InNodeClass) override;
+			FNodeHandle AddNode(const FMetasoundFrontendClassMetadata& InClassMetadata) override;
 
 			// Remove the node corresponding to this node handle.
 			// On success, invalidates the received node handle.
diff --git a/Engine/Plugins/Runtime/Metasound/Source/MetasoundFrontend/Private/MetasoundInstanceTransmitter.cpp b/Engine/Plugins/Runtime/Metasound/Source/MetasoundFrontend/Private/MetasoundInstanceTransmitter.cpp
new file mode 100644
--- /dev/null
+++ b/Engine/Plugins/Runtime/Metasound/Source/MetasoundFrontend/Private/MetasoundInstanceTransmitter.cpp
@@ -0,0 +1,53 @@
+// Copyright Epic Games, Inc. All Rights Reserved.
+
+#include "MetasoundInstanceTransmitter.h"
+
+namespace Metasound
+{
+	FMetasoundInstanceTransmitter::FMetasoundInstanceTransmitter(const FMetasoundInstanceTransmitter::FInitParams& InInitParams)
+	: SendInfos(InInitParams.Infos)
+	, OperatorSettings(InInitParams.OperatorSettings)
+	, InstanceID(InInitParams.InstanceID)
+	{
+	}
+
+	uint64 FMetasoundInstanceTransmitter::GetInstanceID() const
+	{
+		return InstanceID;
+	}
+
+	bool FMetasoundInstanceTransmitter::SetFloatParameter(const FName& InParameterName, float InValue)
+	{
+		return SendValue<float>(InParameterName, InValue);
+	}
+
+	bool FMetasoundInstanceTransmitter::SetIntParameter(const FName& InParameterName, int32 InValue) 
+	{
+		return SendValue<int32>(InParameterName, InValue);
+	}
+
+	bool FMetasoundInstanceTransmitter::SetBoolParameter(const FName& InParameterName, bool InValue) 
+	{
+		return SendValue<bool>(InParameterName, InValue);
+	}
+
+	TUniquePtr<IAudioInstanceTransmitter> FMetasoundInstanceTransmitter::Clone() const
+	{
+		return MakeUnique<FMetasoundInstanceTransmitter>(FMetasoundInstanceTransmitter::FInitParams(OperatorSettings, InstanceID, SendInfos));
+	}
+
+	const FMetasoundInstanceTransmitter::FSendInfo* FMetasoundInstanceTransmitter::FindSendInfo(const FName& InParameterName) const
+	{
+		return SendInfos.FindByPredicate([&](const FSendInfo& Info) { return Info.ParameterName == InParameterName; });
+	}
+
+	ISender* FMetasoundInstanceTransmitter::FindSender(const FName& InParameterName)
+	{
+		TUniquePtr<ISender>* SenderPtrPtr = InputSends.Find(InParameterName);
+		if (SenderPtrPtr)
+		{
+			return SenderPtrPtr->Get();
+		}
+		return nullptr;
+	}
+}
diff --git a/Engine/Plugins/Runtime/Metasound/Source/MetasoundFrontend/Public/MetasoundAssetBase.h b/Engine/Plugins/Runtime/Metasound/Source/MetasoundFrontend/Public/MetasoundAssetBase.h
--- a/Engine/Plugins/Runtime/Metasound/Source/MetasoundFrontend/Public/MetasoundAssetBase.h
+++ b/Engine/Plugins/Runtime/Metasound/Source/MetasoundFrontend/Public/MetasoundAssetBase.h
@@ -7,8 +7,11 @@
 #include "MetasoundFrontendBaseClasses.h"
 #include "MetasoundFrontendController.h"
 #include "MetasoundGraph.h"
+#include "MetasoundLog.h"
 #include "UObject/WeakObjectPtrTemplates.h"
 
+
+
 class UEdGraph;
 
 /** FMetasoundAssetBase is intended to be a mix-in subclass for UObjects which utilize
@@ -69,22 +72,22 @@ public:
 	virtual void SetMetadata(FMetasoundFrontendClassMetadata& InMetadata);
 
 	// Returns  a description of the required inputs and outputs for this metasound UClass.
-	const FMetasoundFrontendArchetype& GetArchetype() const;
+	const FMetasoundFrontendArchetype& GetMetasoundArchetype() const;
 
 	// Sets/overwrites the document archetype.
 	// @param InArchetype - The desired archetype .
 	//
 	// @return True on success. False if archetype is not supported by this object. 
-	bool SetArchetype(const FMetasoundFrontendArchetype& InArchetype);
+	bool SetMetasoundArchetype(const FMetasoundFrontendArchetype& InArchetype);
 
 	// Returns an array of archetypes preferred for this class.
-	virtual const TArray<FMetasoundFrontendArchetype>& GetPreferredArchetypes() const = 0;
+	virtual const TArray<FMetasoundFrontendArchetype>& GetPreferredMetasoundArchetypes() const = 0;
 
 	// Returns true if the archetype is supported by this object.
-	virtual bool IsArchetypeSupported(const FMetasoundFrontendArchetype& InArchetype) const;
+	virtual bool IsMetasoundArchetypeSupported(const FMetasoundFrontendArchetype& InArchetype) const;
 
 	// Returns the preferred archetype for the given document.
-	virtual const FMetasoundFrontendArchetype& GetPreferredArchetype(const FMetasoundFrontendDocument& InDocument) const;
+	virtual const FMetasoundFrontendArchetype& GetPreferredMetasoundArchetype(const FMetasoundFrontendDocument& InDocument) const;
 
 	// Returns the root class metadata
 	FMetasoundFrontendClassMetadata GetMetadata();
@@ -115,7 +118,7 @@ public:
 	const FMetasoundFrontendDocument& GetDocumentChecked() const;
 
 	// This must be called on UObject::PostLoad, as well as in this asset's UFactory, to fix up the root document based on the most recent version of the archetype.
-	void ConformDocumentToArchetype();
+	void ConformDocumentToMetasoundArchetype();
 
 protected:
 
diff --git a/Engine/Plugins/Runtime/Metasound/Source/MetasoundFrontend/Public/MetasoundDataTypeRegistrationMacro.h b/Engine/Plugins/Runtime/Metasound/Source/MetasoundFrontend/Public/MetasoundDataTypeRegistrationMacro.h
--- a/Engine/Plugins/Runtime/Metasound/Source/MetasoundFrontend/Public/MetasoundDataTypeRegistrationMacro.h
+++ b/Engine/Plugins/Runtime/Metasound/Source/MetasoundFrontend/Public/MetasoundDataTypeRegistrationMacro.h
@@ -49,8 +49,16 @@ namespace Metasound
 	struct TIsTransmittable
 	{
 	private:
-		static constexpr bool bCanBeTransmitted =
-			std::is_copy_constructible<TDataType>::value|| TIsDerivedFrom<TDataType, IAudioDataType>::Value;
+		static constexpr bool bIsCopyConstructible = std::is_copy_constructible<TDataType>::value;
+		static constexpr bool bIsCopyAssignable = std::is_copy_assignable<TDataType>::value;
+
+		// TODO: audio types were intended to be send/receive nodes but they require 
+		// template specialization.  TIsTransmittable should ask the Send and Receive nodes
+		// if they are transmittable rather than attempting to do the logic for all
+		// types here. 
+		//static constexpr bool bIsAudioDataType = TIsDerivedFrom<TDataType, IAudioDataType>::Value;
+
+		static constexpr bool bCanBeTransmitted = bIsCopyConstructible && bIsCopyAssignable;
 
 	public:
 
@@ -214,6 +222,8 @@ namespace Metasound
 		RegistryInfo.bIsProxyParsable = TIsParsable<TDataType, const Audio::IProxyDataPtr&>::Value;
 		RegistryInfo.bIsProxyArrayParsable = TIsParsable<TDataType, const TArray<Audio::IProxyDataPtr>& >::Value;
 		RegistryInfo.bIsDefaultParsable = TIsParsable<TDataType>::Value;
+
+		RegistryInfo.bIsTransmittable = TIsTransmittable<TDataType>::Value;
 		
 		RegistryInfo.ProxyGeneratorClass = UClassToUse::StaticClass();
 
diff --git a/Engine/Plugins/Runtime/Metasound/Source/MetasoundFrontend/Public/MetasoundFrontendController.h b/Engine/Plugins/Runtime/Metasound/Source/MetasoundFrontend/Public/MetasoundFrontendController.h
--- a/Engine/Plugins/Runtime/Metasound/Source/MetasoundFrontend/Public/MetasoundFrontendController.h
+++ b/Engine/Plugins/Runtime/Metasound/Source/MetasoundFrontend/Public/MetasoundFrontendController.h
@@ -186,12 +186,15 @@ namespace Metasound
 			/** Returns a FConstNodeHandle to the node which owns this output. */
 			virtual FConstNodeHandle GetOwningNode() const = 0;
 
+			/** Returns the currently connected output. If this input is not
+			 * connected, the returned handle will be invalid. */
+			virtual TArray<FInputHandle> GetCurrentlyConnectedInputs() = 0;
+
+			/** Returns the currently connected output. If this input is not
+			 * connected, the returned handle will be invalid. */
+			virtual TArray<FConstInputHandle> GetCurrentlyConnectedInputs() const = 0;
 
-			// TODO: with the current spec, we'd have to scan all nodes in the graph for this output to solve this.
-			// Is this worth it, or should we just require folks to use FInputController::GetCurrentlyConnectedOutput()?
-			//IInputController* GetCurrentlyConnectedInput();
-			//const IInputController* GetCurrentlyConnectedInput() const;
-			//void Disconnect();
+			virtual bool Disconnect() = 0;
 
 			/** Returns information describing connectability between this output and the supplied input. */
 			virtual FConnectability CanConnectTo(const IInputController& InController) const = 0;
@@ -608,6 +611,14 @@ namespace Metasound
 			 */
 			virtual FNodeHandle AddNode(const FNodeRegistryKey& InNodeClass) = 0;
 
+			/** Add a new node to this graph.
+			 *
+			 * @param InClassMetadat - Info for node class.
+			 * 
+			 * @return Node handle for class. On error, an invalid handle is returned. 
+			 */
+			virtual FNodeHandle AddNode(const FMetasoundFrontendClassMetadata& InClassMetadata) = 0;
+
 			/** Remove the node corresponding to this node handle.
 			 *
 			 * @return True on success, false on failure. 
diff --git a/Engine/Plugins/Runtime/Metasound/Source/MetasoundFrontend/Public/MetasoundFrontendQuerySteps.h b/Engine/Plugins/Runtime/Metasound/Source/MetasoundFrontend/Public/MetasoundFrontendQuerySteps.h
--- a/Engine/Plugins/Runtime/Metasound/Source/MetasoundFrontend/Public/MetasoundFrontendQuerySteps.h
+++ b/Engine/Plugins/Runtime/Metasound/Source/MetasoundFrontend/Public/MetasoundFrontendQuerySteps.h
@@ -49,5 +49,16 @@ namespace Metasound
 	private:
 		FName OutputVertexTypeName;
 	};
+
+	class METASOUNDFRONTEND_API FFilterClassesByClassName : public IFrontendQueryFilterStep
+	{
+	public: 
+		FFilterClassesByClassName(const FString& InClassName);
+
+		bool Filter(const FFrontendQueryEntry& InEntry) const override;
+
+	private:
+		FString ClassName;
+	};
 }
 
diff --git a/Engine/Plugins/Runtime/Metasound/Source/MetasoundFrontend/Public/MetasoundFrontendRegistries.h b/Engine/Plugins/Runtime/Metasound/Source/MetasoundFrontend/Public/MetasoundFrontendRegistries.h
--- a/Engine/Plugins/Runtime/Metasound/Source/MetasoundFrontend/Public/MetasoundFrontendRegistries.h
+++ b/Engine/Plugins/Runtime/Metasound/Source/MetasoundFrontend/Public/MetasoundFrontendRegistries.h
@@ -44,6 +44,9 @@ namespace Metasound
 		bool bIsProxyParsable = false;
 		bool bIsProxyArrayParsable = false;
 
+		// Determines whether the type can be used with send/receive transmitters
+		bool bIsTransmittable = false;
+
 		// If this datatype was registered with a specific UClass to use to filter with, that will be used here:
 		UClass* ProxyGeneratorClass = nullptr;
 
diff --git a/Engine/Plugins/Runtime/Metasound/Source/MetasoundFrontend/Public/MetasoundInstanceTransmitter.h b/Engine/Plugins/Runtime/Metasound/Source/MetasoundFrontend/Public/MetasoundInstanceTransmitter.h
new file mode 100644
--- /dev/null
+++ b/Engine/Plugins/Runtime/Metasound/Source/MetasoundFrontend/Public/MetasoundInstanceTransmitter.h
@@ -0,0 +1,126 @@
+// Copyright Epic Games, Inc. All Rights Reserved.
+
+#pragma once
+
+#include "IAudioExtensionPlugin.h"
+#include "MetasoundDataReference.h"
+#include "MetasoundLog.h"
+#include "MetasoundOperatorSettings.h"
+#include "MetasoundPrimitives.h"
+#include "MetasoundRouter.h"
+
+namespace Metasound
+{
+	class METASOUNDFRONTEND_API FMetasoundInstanceTransmitter : public IAudioInstanceTransmitter
+	{
+		FMetasoundInstanceTransmitter(const FMetasoundInstanceTransmitter&) = delete;
+		FMetasoundInstanceTransmitter& operator=(const FMetasoundInstanceTransmitter&) = delete;
+	public:
+		struct FSendInfo
+		{
+			FSendAddress Address;
+			FName ParameterName;
+			FName TypeName;
+		};
+
+		struct FInitParams
+		{
+			FOperatorSettings OperatorSettings;
+			uint64 InstanceID;
+			TArray<FSendInfo> Infos;
+
+			FInitParams(const FOperatorSettings& InSettings, uint64 InInstanceID, const TArray<FSendInfo>& InInfos=TArray<FSendInfo>())
+			: OperatorSettings(InSettings)
+			, InstanceID(InInstanceID)
+			, Infos(InInfos)
+			{
+			}
+
+		};
+
+		FMetasoundInstanceTransmitter(const FMetasoundInstanceTransmitter::FInitParams& InInitParams);
+		virtual ~FMetasoundInstanceTransmitter() = default;
+
+		uint64 GetInstanceID() const override;
+
+		bool SetFloatParameter(const FName& InParameterName, float InValue) override;
+		bool SetIntParameter(const FName& InParameterName, int32 InValue) override;
+		bool SetBoolParameter(const FName& InParameterName, bool InValue) override;
+
+		TUniquePtr<IAudioInstanceTransmitter> Clone() const override;
+
+		template<typename DataType>
+		bool SendValue(const FName& InParameterName, const DataType& InValue)
+		{
+			if (TSender<DataType>* Sender = FindSenderOfDataType<DataType>(InParameterName))
+			{
+				return Sender->Push(InValue);
+			}
+			return false;
+		}
+
+	private:
+
+		const FSendInfo* FindSendInfo(const FName& InParameterName) const;
+		ISender* FindSender(const FName& InParameterName);
+
+		template<typename DataType> 
+		TSender<DataType>* FindSenderOfDataType(const FName& InParameterName)
+		{
+			using FSenderType = TSender<DataType>;
+
+			if (ISender* ExistingSender = FindSender(InParameterName))
+			{
+				if (ExistingSender->CheckType<FSenderType>())
+				{
+					FSenderType& DerivedSender = ExistingSender->GetAs<FSenderType>();
+					return &DerivedSender;
+				}
+				else
+				{
+					UE_LOG(LogMetasound, Warning, TEXT("Could not send to Metasound input. Mismatched data type [Recieved:%s, Expected: %s]"), *GetMetasoundDataTypeString<DataType>(), *ExistingSender->GetDataType().ToString());
+				}
+			}
+			else if (const FSendInfo* SendInfo = FindSendInfo(InParameterName))
+			{
+				TUniquePtr<TSender<DataType>> NewSender = CreateSender<DataType>(*SendInfo);
+				if (NewSender.IsValid())
+				{
+					FSenderType* DerivedSender = NewSender.Get();
+
+					InputSends.Add(InParameterName, MoveTemp(NewSender));
+
+					return DerivedSender;
+				}
+			}
+
+			return nullptr;
+		}
+
+		template<typename DataType>
+		TUniquePtr<TSender<DataType>> CreateSender(const FSendInfo& InInfo)
+		{
+			const float DelayTimeInSeconds = 10.0f; // TODO: likely want to remove this and opt for different protocols having different behaviors.
+			FSenderInitParams InitParams = {OperatorSettings, DelayTimeInSeconds};
+
+			if (InInfo.TypeName == GetMetasoundDataTypeName<DataType>())
+			{
+				return FDataTransmissionCenter::Get().RegisterNewSend<DataType>(InInfo.Address, InitParams);
+			}
+			else
+			{
+				UE_LOG(LogMetasound, Error, TEXT("Cannot create sender. Requested data type [TypeName:%s] does not match receive data type [TypeName:%s]"), *GetMetasoundDataTypeName<DataType>().ToString(), *InInfo.TypeName.ToString())
+			}
+
+			return TUniquePtr<TSender<DataType>>(nullptr);
+		}
+
+
+		TArray<FSendInfo> SendInfos;
+		FOperatorSettings OperatorSettings;
+		uint64 InstanceID;
+
+		TMap<FName, TUniquePtr<ISender>> InputSends;
+	};
+}
+
diff --git a/Engine/Plugins/Runtime/Metasound/Source/MetasoundFrontend/Public/MetasoundReceiveNode.h b/Engine/Plugins/Runtime/Metasound/Source/MetasoundFrontend/Public/MetasoundReceiveNode.h
--- a/Engine/Plugins/Runtime/Metasound/Source/MetasoundFrontend/Public/MetasoundReceiveNode.h
+++ b/Engine/Plugins/Runtime/Metasound/Source/MetasoundFrontend/Public/MetasoundReceiveNode.h
@@ -18,31 +18,46 @@
 
 namespace Metasound
 {
-	template<typename TDataType>
-	class TReceiveNode : public FNode
+	struct FReceiveNodeNames
 	{
-
-	public:
 		static const FString& GetAddressInputName()
 		{
 			static const FString InputName = FString(TEXT("Address"));
 			return InputName;
 		}
 
-		static const FString& GetSendOutputName()
+		static const FString& GetDefaultDataInputName()
+		{
+			static const FString DefaultDataName = FString(TEXT("Default"));
+			return DefaultDataName;
+		}
+
+		static const FString& GetOutputName()
 		{
-			static const FString SendInput = GetMetasoundDataTypeString<TDataType>();
-			return SendInput;
+			static const FString OutputName = FString(TEXT("Out"));
+			return OutputName;
 		}
 
+		static FString GetClassNameForDataType(const FName& InDataTypeName)
+		{
+			return FString::Printf(TEXT("Receive.%s"), *InDataTypeName.ToString());
+		}
+	};
+
+	template<typename TDataType>
+	class TReceiveNode : public FNode
+	{
+	public:
+
 		static FVertexInterface DeclareVertexInterface()
 		{
 			return FVertexInterface(
 				FInputVertexInterface(
-					TInputDataVertexModel<FSendAddress>(GetAddressInputName(), FText::GetEmpty())
+					TInputDataVertexModel<FSendAddress>(FReceiveNodeNames::GetAddressInputName(), FText::GetEmpty()),
+					TInputDataVertexModel<TDataType>(FReceiveNodeNames::GetDefaultDataInputName(), FText::GetEmpty())
 				),
 				FOutputVertexInterface(
-					TOutputDataVertexModel<TDataType>(GetSendOutputName(), FText::GetEmpty())
+					TOutputDataVertexModel<TDataType>(FReceiveNodeNames::GetOutputName(), FText::GetEmpty())
 				)
 			);
 		}
@@ -52,7 +67,7 @@ namespace Metasound
 			auto InitNodeInfo = []() -> FNodeInfo
 			{
 				FNodeInfo Info;
-				Info.ClassName = FName(*FString::Printf(TEXT("Receive %s"), *GetSendOutputName()));
+				Info.ClassName = FName(FReceiveNodeNames::GetClassNameForDataType(GetMetasoundDataTypeName<TDataType>()));
 				Info.MajorVersion = 1;
 				Info.MinorVersion = 0;
 				Info.Description = LOCTEXT("Metasound_ReceiveNodeDescription", "Receives data from a send node with the same name.");
@@ -76,8 +91,10 @@ namespace Metasound
 				TReceiverOperator() = delete;
 			public:
 
-				TReceiverOperator(const TDataWriteReference<TDataType>& Data, TDataReadReference<FSendAddress> InSendAddress, const FOperatorSettings& InOperatorSettings)
-					: OutputData(Data)
+				TReceiverOperator(TDataReadReference<TDataType> InInitDataRef, TDataWriteReference<TDataType> InOutDataRef, TDataReadReference<FSendAddress> InSendAddress, const FOperatorSettings& InOperatorSettings)
+					: bHasNotReceivedData(true)
+					, DefaultData(InInitDataRef)
+					, OutputData(InOutDataRef)
 					, SendAddress(InSendAddress)
 					, CachedSendAddress(*InSendAddress)
 					, CachedReceiverParams({InOperatorSettings})
@@ -89,13 +106,18 @@ namespace Metasound
 
 				virtual FDataReferenceCollection GetInputs() const override
 				{
-					return {};
+					FDataReferenceCollection Inputs;
+
+					Inputs.AddDataReadReference<TDataType>(FReceiveNodeNames::GetDefaultDataInputName(), DefaultData);
+					Inputs.AddDataReadReference<FSendAddress>(FReceiveNodeNames::GetAddressInputName(), SendAddress);
+
+					return Inputs;
 				}
 
 				virtual FDataReferenceCollection GetOutputs() const override
 				{
 					FDataReferenceCollection Outputs;
-					Outputs.AddDataReadReference<TDataType>(GetSendOutputName(), TDataReadReference<TDataType>(OutputData));
+					Outputs.AddDataReadReference<TDataType>(FReceiveNodeNames::GetOutputName(), TDataReadReference<TDataType>(OutputData));
 					return Outputs;
 				}
 
@@ -105,18 +127,30 @@ namespace Metasound
 					{
 						CachedSendAddress = *SendAddress;
 						Receiver = FDataTransmissionCenter::Get().RegisterNewReceiver<TDataType>(CachedSendAddress, CachedReceiverParams);
-						check(Receiver.IsValid());
 					}
 
-					if (Receiver->CanPop())
+					if (bHasNotReceivedData)
+					{
+						*OutputData = *DefaultData;
+					}
+
+					if (ensure(Receiver.IsValid()))
 					{
-						Receiver->Pop(*OutputData);
+						if (Receiver->CanPop())
+						{
+							Receiver->Pop(*OutputData);
+							bHasNotReceivedData = false;
+						}
 					}
 				}
 
 			private:
+				bool bHasNotReceivedData;
+
+				TDataReadReference<TDataType> DefaultData;
 				TDataWriteReference<TDataType> OutputData;
 				TDataReadReference<FSendAddress> SendAddress;
+
 				FSendAddress CachedSendAddress;
 				FReceiverInitParams CachedReceiverParams;
 
@@ -130,9 +164,17 @@ namespace Metasound
 
 				virtual TUniquePtr<IOperator> CreateOperator(const FCreateOperatorParams& InParams, FBuildErrorArray& OutErrors) override
 				{
+					TDataReadReference<TDataType> DefaultReadRef = TDataReadReferenceFactory<TDataType>::CreateAny(InParams.OperatorSettings);
+
+					if (InParams.InputDataReferences.ContainsDataReadReference<TDataType>(FReceiveNodeNames::GetDefaultDataInputName()))
+					{
+						DefaultReadRef = InParams.InputDataReferences.GetDataReadReference<TDataType>(FReceiveNodeNames::GetDefaultDataInputName());
+					}
+
 					return MakeUnique<TReceiverOperator>(
-						TDataWriteReferenceFactory<TDataType>::CreateAny(InParams.OperatorSettings),
-						InParams.InputDataReferences.GetDataReadReference<FSendAddress>(GetAddressInputName()),
+						DefaultReadRef,
+						TDataWriteReferenceFactory<TDataType>::CreateAny(InParams.OperatorSettings, *DefaultReadRef),
+						InParams.InputDataReferences.GetDataReadReferenceOrConstruct<FSendAddress>(FReceiveNodeNames::GetAddressInputName()),
 						InParams.OperatorSettings
 						);
 				}
diff --git a/Engine/Plugins/Runtime/Metasound/Source/MetasoundFrontend/Public/MetasoundRouter.h b/Engine/Plugins/Runtime/Metasound/Source/MetasoundFrontend/Public/MetasoundRouter.h
--- a/Engine/Plugins/Runtime/Metasound/Source/MetasoundFrontend/Public/MetasoundRouter.h
+++ b/Engine/Plugins/Runtime/Metasound/Source/MetasoundFrontend/Public/MetasoundRouter.h
@@ -366,6 +366,7 @@ namespace Metasound
 			{
 				SenderBuffer.SetCapacity(NumElementsToDelayBy * BufferResizeThreshold);
 			}
+			return true;
 		}
 
 		static FName GetDataTypeName()
@@ -378,7 +379,7 @@ namespace Metasound
 			: ISender(InDataChannel, GetDataTypeName())
 			, Params(InitParams)
 		{
-
+			SetDelay(InitParams.DelayTimeInSeconds);
 		}
 
 		TDataType RetrievePayload()
diff --git a/Engine/Plugins/Runtime/Metasound/Source/MetasoundGraphCore/Public/MetasoundDataReference.h b/Engine/Plugins/Runtime/Metasound/Source/MetasoundGraphCore/Public/MetasoundDataReference.h
--- a/Engine/Plugins/Runtime/Metasound/Source/MetasoundGraphCore/Public/MetasoundDataReference.h
+++ b/Engine/Plugins/Runtime/Metasound/Source/MetasoundGraphCore/Public/MetasoundDataReference.h
@@ -58,7 +58,7 @@ namespace Metasound
 	{
 		// This static assert is triggered if TDataReferenceTypeInfo is used 
 		// without specialization.
-		static_assert(TSpecializationHelper<DataType>::Value, "TDataReferenceTypeInfo must be specialized.  Use macro DECLARE_METASOUND_DATA_REFERENCE_TYPES");
+		static_assert(TSpecializationHelper<DataType>::Value, "TDataReferenceTypeInfo is not specialized.  Use macro DECLARE_METASOUND_DATA_REFERENCE_TYPES to declare a new type, or ensure that an existing DECLARE_METASOUND_DATA_REFERENCE_TYPES exists in the include path.");
 	};
 
 
diff --git a/Engine/Source/Runtime/AudioExtensions/Public/IAudioExtensionPlugin.h b/Engine/Source/Runtime/AudioExtensions/Public/IAudioExtensionPlugin.h
--- a/Engine/Source/Runtime/AudioExtensions/Public/IAudioExtensionPlugin.h
+++ b/Engine/Source/Runtime/AudioExtensions/Public/IAudioExtensionPlugin.h
@@ -605,4 +605,38 @@ public:
 	virtual void OnListenerShutdown(FAudioDevice* AudioDevice)
 	{
 	}
-};
\ No newline at end of file
+};
+
+/** Parameters passed to USoundBase::CreateInstanceTransmitter. */
+struct FAudioInstanceTransmitterInitParams
+{
+	// Unique ID for this audio instance.
+	uint64 InstanceID;
+	// Audio sample rate.
+	float SampleRate;
+};
+
+/** Interface for a audio instance transmitter.
+ *
+ * An audio instance transmitter ushers control parameters to a single audio object instance.
+ */
+class IAudioInstanceTransmitter
+{
+	public:
+		virtual ~IAudioInstanceTransmitter() = default;
+
+		// Return the instance ID
+		virtual uint64 GetInstanceID() const = 0;
+
+		// Set a float parameter by name. Return true on success, false on failure.
+		virtual bool SetFloatParameter(const FName& InParameterName, float InValue) = 0;
+
+		// Set an int parameter by name. Return true on success, false on failure.
+		virtual bool SetIntParameter(const FName& InParameterName, int32 InValue) = 0;
+
+		// Set a bool parameter by name. Return true on success, false on failure.
+		virtual bool SetBoolParameter(const FName& InParameterName, bool InValue) = 0;
+
+		// Create a copy of the instance transmitter.
+		virtual TUniquePtr<IAudioInstanceTransmitter> Clone() const = 0;
+};
diff --git a/Engine/Source/Runtime/AudioMixer/Private/AudioMixerSource.cpp b/Engine/Source/Runtime/AudioMixer/Private/AudioMixerSource.cpp
--- a/Engine/Source/Runtime/AudioMixer/Private/AudioMixerSource.cpp
+++ b/Engine/Source/Runtime/AudioMixer/Private/AudioMixerSource.cpp
@@ -773,7 +773,15 @@ namespace Audio
 		}
 
 		check(!MixerSourceBuffer.IsValid());
-		MixerSourceBuffer = FMixerSourceBuffer::Create(AudioDevice->GetSampleRate(), *MixerBuffer, SoundWave, InWaveInstance->LoopingMode, bIsSeeking);
+
+		// Active sound instance ID is the audio component ID of active sound.
+		uint64 InstanceID = 0;
+		if (WaveInstance->ActiveSound)
+		{
+			InstanceID = WaveInstance->ActiveSound->GetAudioComponentID();
+		}
+
+		MixerSourceBuffer = FMixerSourceBuffer::Create(InstanceID, AudioDevice->GetSampleRate(), *MixerBuffer, SoundWave, InWaveInstance->LoopingMode, bIsSeeking);
 		
 		if (!MixerSourceBuffer.IsValid())
 		{
diff --git a/Engine/Source/Runtime/AudioMixer/Private/AudioMixerSourceBuffer.cpp b/Engine/Source/Runtime/AudioMixer/Private/AudioMixerSourceBuffer.cpp
--- a/Engine/Source/Runtime/AudioMixer/Private/AudioMixerSourceBuffer.cpp
+++ b/Engine/Source/Runtime/AudioMixer/Private/AudioMixerSourceBuffer.cpp
@@ -56,16 +56,16 @@ namespace Audio
 		return CurrentSample >= NumSamples;
 	}
 
-	TSharedPtr<FMixerSourceBuffer, ESPMode::ThreadSafe> FMixerSourceBuffer::Create(int32 InSampleRate, FMixerBuffer& InBuffer, USoundWave& InWave, ELoopingMode InLoopingMode, bool bInIsSeeking, bool bInForceSyncDecode)
+	TSharedPtr<FMixerSourceBuffer, ESPMode::ThreadSafe> FMixerSourceBuffer::Create(uint64 InInstanceID, int32 InSampleRate, FMixerBuffer& InBuffer, USoundWave& InWave, ELoopingMode InLoopingMode, bool bInIsSeeking, bool bInForceSyncDecode)
 	{
 		LLM_SCOPE(ELLMTag::AudioMixer);
 
-		TSharedPtr<FMixerSourceBuffer, ESPMode::ThreadSafe> NewSourceBuffer = MakeShareable(new FMixerSourceBuffer(InSampleRate, InBuffer, InWave, InLoopingMode, bInIsSeeking, bInForceSyncDecode));
+		TSharedPtr<FMixerSourceBuffer, ESPMode::ThreadSafe> NewSourceBuffer = MakeShareable(new FMixerSourceBuffer(InInstanceID, InSampleRate, InBuffer, InWave, InLoopingMode, bInIsSeeking, bInForceSyncDecode));
 
 		return NewSourceBuffer;
 	}
 
-	FMixerSourceBuffer::FMixerSourceBuffer(int32 InSampleRate, FMixerBuffer& InBuffer, USoundWave& InWave, ELoopingMode InLoopingMode, bool bInIsSeeking, bool bInForceSyncDecode)
+	FMixerSourceBuffer::FMixerSourceBuffer(uint64 InInstanceID, int32 InSampleRate, FMixerBuffer& InBuffer, USoundWave& InWave, ELoopingMode InLoopingMode, bool bInIsSeeking, bool bInForceSyncDecode)
 		: NumBuffersQeueued(0)
 		, CurrentBuffer(0)
 		, SoundWave(&InWave)
@@ -95,6 +95,7 @@ namespace Audio
 			InitParams.SampleRate = InSampleRate;
 			InitParams.NumChannels = NumChannels;
 			InitParams.NumFramesPerCallback = MONO_PCM_BUFFER_SAMPLES;
+			InitParams.InstanceID = InInstanceID;
 			SoundGenerator = InWave.CreateSoundGenerator(InitParams);
 		}
 
diff --git a/Engine/Source/Runtime/AudioMixer/Private/AudioMixerSourceBuffer.h b/Engine/Source/Runtime/AudioMixer/Private/AudioMixerSourceBuffer.h
--- a/Engine/Source/Runtime/AudioMixer/Private/AudioMixerSourceBuffer.h
+++ b/Engine/Source/Runtime/AudioMixer/Private/AudioMixerSourceBuffer.h
@@ -50,7 +50,7 @@ namespace Audio
 	class FMixerSourceBuffer : public ISoundWaveClient
 	{
 	public:
-		static FMixerSourceBufferPtr Create(int32 InSampleRate, FMixerBuffer& InBuffer, USoundWave& InWave, ELoopingMode InLoopingMode, bool bInIsSeeking, bool bInForceSyncDecode = false);
+		static FMixerSourceBufferPtr Create(uint64 InInstanceID, int32 InSampleRate, FMixerBuffer& InBuffer, USoundWave& InWave, ELoopingMode InLoopingMode, bool bInIsSeeking, bool bInForceSyncDecode = false);
 
 		~FMixerSourceBuffer();
 
@@ -101,7 +101,7 @@ namespace Audio
 		bool IsGeneratorFinished() const;
 
 	private:
-		FMixerSourceBuffer(int32 InSampleRate, FMixerBuffer& InBuffer, USoundWave& InWave, ELoopingMode InLoopingMode, bool bInIsSeeking, bool bInForceSyncDecode = false);
+		FMixerSourceBuffer(uint64 InInstanceID, int32 InSampleRate, FMixerBuffer& InBuffer, USoundWave& InWave, ELoopingMode InLoopingMode, bool bInIsSeeking, bool bInForceSyncDecode = false);
 
 		void SubmitInitialPCMBuffers();
 		void SubmitInitialRealtimeBuffers();
diff --git a/Engine/Source/Runtime/AudioMixer/Private/SoundWaveDecoder.cpp b/Engine/Source/Runtime/AudioMixer/Private/SoundWaveDecoder.cpp
--- a/Engine/Source/Runtime/AudioMixer/Private/SoundWaveDecoder.cpp
+++ b/Engine/Source/Runtime/AudioMixer/Private/SoundWaveDecoder.cpp
@@ -57,7 +57,8 @@ namespace Audio
 		bool bIsValid = false;
 		{
 			FScopeLock Lock(&MixerSourceBufferCritSec);
-			MixerSourceBuffer = FMixerSourceBuffer::Create(InSampleRate, *MixerBuffer, *SoundWave, LoopingMode, bIsSeeking, bForceSyncDecode);
+			uint64 InstanceID = 0; // Unused.
+			MixerSourceBuffer = FMixerSourceBuffer::Create(InstanceID, InSampleRate, *MixerBuffer, *SoundWave, LoopingMode, bIsSeeking, bForceSyncDecode);
 			bIsValid = MixerSourceBuffer.IsValid();
 		}
 
diff --git a/Engine/Source/Runtime/Engine/Classes/Sound/SoundBase.h b/Engine/Source/Runtime/Engine/Classes/Sound/SoundBase.h
--- a/Engine/Source/Runtime/Engine/Classes/Sound/SoundBase.h
+++ b/Engine/Source/Runtime/Engine/Classes/Sound/SoundBase.h
@@ -276,5 +276,7 @@ public:
 	/** Creates a sound generator instance from this sound base. Return true if this is being implemented by a subclass. Sound generators procedurally generate audio in the audio render thread. */
 	virtual ISoundGeneratorPtr CreateSoundGenerator(const FSoundGeneratorInitParams& InParams) { return nullptr; }
 
+	/** Creates a instance transmitter for communicating with active sound instances. */
+	virtual TUniquePtr<IAudioInstanceTransmitter> CreateInstanceTransmitter(const FAudioInstanceTransmitterInitParams& InParams) const { return nullptr; }
 };
 
diff --git a/Engine/Source/Runtime/Engine/Classes/Sound/SoundGenerator.h b/Engine/Source/Runtime/Engine/Classes/Sound/SoundGenerator.h
--- a/Engine/Source/Runtime/Engine/Classes/Sound/SoundGenerator.h
+++ b/Engine/Source/Runtime/Engine/Classes/Sound/SoundGenerator.h
@@ -13,11 +13,13 @@ struct FSoundGeneratorInitParams
 	float SampleRate;
 	int32 NumChannels;
 	int32 NumFramesPerCallback;
+	uint64 InstanceID;
 
 	FSoundGeneratorInitParams()
 		: SampleRate(0.0f)
 		, NumChannels(0)
 		, NumFramesPerCallback(0)
+		, InstanceID(0)
 	{}
 };
 
@@ -71,4 +73,4 @@ public:
 	}
 };
 
-typedef TSharedPtr<ISoundGenerator, ESPMode::ThreadSafe> ISoundGeneratorPtr;
\ No newline at end of file
+typedef TSharedPtr<ISoundGenerator, ESPMode::ThreadSafe> ISoundGeneratorPtr;
diff --git a/Engine/Source/Runtime/Engine/Private/ActiveSound.cpp b/Engine/Source/Runtime/Engine/Private/ActiveSound.cpp
--- a/Engine/Source/Runtime/Engine/Private/ActiveSound.cpp
+++ b/Engine/Source/Runtime/Engine/Private/ActiveSound.cpp
@@ -1280,21 +1280,32 @@ void FActiveSound::UpdateConcurrencyVolumeScalars(const float DeltaTime)
 void FActiveSound::SetFloatParameter(const FName InName, const float InFloat)
 {
 	if (InName != NAME_None)
-	{
-		// First see if an entry for this name already exists
-		for (FAudioComponentParam& P : InstanceParameters)
+	{		
+		if (InstanceTransmitter.IsValid())
 		{
-			if (P.ParamName == InName)
+			bool bSuccess = InstanceTransmitter->SetFloatParameter(InName, InFloat);
+			if (!bSuccess)
 			{
-				P.FloatParam = InFloat;
-				return;
+				UE_LOG(LogAudio, Warning, TEXT("Failed to set instance float parameter [InstanceID:%d, ParameterName:%s]"), InstanceTransmitter->GetInstanceID(), *InName.ToString());
 			}
 		}
+		else
+		{
+			// First see if an entry for this name already exists
+			for (FAudioComponentParam& P : InstanceParameters)
+			{
+				if (P.ParamName == InName)
+				{
+					P.FloatParam = InFloat;
+					return;
+				}
+			}
 
-		// We didn't find one, so create a new one.
-		const int32 NewParamIndex = InstanceParameters.AddDefaulted();
-		InstanceParameters[NewParamIndex].ParamName = InName;
-		InstanceParameters[NewParamIndex].FloatParam = InFloat;
+			// We didn't find one, so create a new one.
+			const int32 NewParamIndex = InstanceParameters.AddDefaulted();
+			InstanceParameters[NewParamIndex].ParamName = InName;
+			InstanceParameters[NewParamIndex].FloatParam = InFloat;
+		}
 	}
 }
 
@@ -1359,20 +1370,31 @@ void FActiveSound::SetBoolParameter(const FName InName, const bool InBool)
 {
 	if (InName != NAME_None)
 	{
-		// First see if an entry for this name already exists
-		for (FAudioComponentParam& P : InstanceParameters)
+		if (InstanceTransmitter.IsValid())
 		{
-			if (P.ParamName == InName)
+			bool bSuccess = InstanceTransmitter->SetBoolParameter(InName, InBool);
+			if (!bSuccess)
 			{
-				P.BoolParam = InBool;
-				return;
+				UE_LOG(LogAudio, Warning, TEXT("Failed to set instance bool parameter [InstanceID:%d, ParameterName:%s]"), InstanceTransmitter->GetInstanceID(), *InName.ToString());
 			}
 		}
+		else
+		{
+			// First see if an entry for this name already exists
+			for (FAudioComponentParam& P : InstanceParameters)
+			{
+				if (P.ParamName == InName)
+				{
+					P.BoolParam = InBool;
+					return;
+				}
+			}
 
-		// We didn't find one, so create a new one.
-		const int32 NewParamIndex = InstanceParameters.AddDefaulted();
-		InstanceParameters[NewParamIndex].ParamName = InName;
-		InstanceParameters[NewParamIndex].BoolParam = InBool;
+			// We didn't find one, so create a new one.
+			const int32 NewParamIndex = InstanceParameters.AddDefaulted();
+			InstanceParameters[NewParamIndex].ParamName = InName;
+			InstanceParameters[NewParamIndex].BoolParam = InBool;
+		}
 	}
 }
 
@@ -1398,20 +1420,31 @@ void FActiveSound::SetIntParameter(const FName InName, const int32 InInt)
 {
 	if (InName != NAME_None)
 	{
-		// First see if an entry for this name already exists
-		for (FAudioComponentParam& P : InstanceParameters)
+		if (InstanceTransmitter.IsValid())
 		{
-			if (P.ParamName == InName)
+			bool bSuccess = InstanceTransmitter->SetIntParameter(InName, InInt);
+			if (!bSuccess)
 			{
-				P.IntParam = InInt;
-				return;
+				UE_LOG(LogAudio, Warning, TEXT("Failed to set instance int parameter [InstanceID:%d, ParameterName:%s]"), InstanceTransmitter->GetInstanceID(), *InName.ToString());
 			}
 		}
+		else
+		{
+			// First see if an entry for this name already exists
+			for (FAudioComponentParam& P : InstanceParameters)
+			{
+				if (P.ParamName == InName)
+				{
+					P.IntParam = InInt;
+					return;
+				}
+			}
 
-		// We didn't find one, so create a new one.
-		const int32 NewParamIndex = InstanceParameters.AddDefaulted();
-		InstanceParameters[NewParamIndex].ParamName = InName;
-		InstanceParameters[NewParamIndex].IntParam = InInt;
+			// We didn't find one, so create a new one.
+			const int32 NewParamIndex = InstanceParameters.AddDefaulted();
+			InstanceParameters[NewParamIndex].ParamName = InName;
+			InstanceParameters[NewParamIndex].IntParam = InInt;
+		}
 	}
 }
 
diff --git a/Engine/Source/Runtime/Engine/Private/AudioDevice.cpp b/Engine/Source/Runtime/Engine/Private/AudioDevice.cpp
--- a/Engine/Source/Runtime/Engine/Private/AudioDevice.cpp
+++ b/Engine/Source/Runtime/Engine/Private/AudioDevice.cpp
@@ -4757,6 +4757,13 @@ void FAudioDevice::AddNewActiveSoundInternal(const FActiveSound& NewActiveSound,
 		VirtualActiveSound.ClearAudioComponent();
 	}
 
+	FAudioInstanceTransmitterInitParams TransmitterInitParams
+	{
+		ActiveSound->GetAudioComponentID(),
+		GetSampleRate()
+	};
+	ActiveSound->InstanceTransmitter = Sound->CreateInstanceTransmitter(TransmitterInitParams);
+
 	ActiveSounds.Add(ActiveSound);
 	if (ActiveSound->GetAudioComponentID() > 0)
 	{
@@ -6689,4 +6696,4 @@ int32 FAudioDevice::GetNumPrecacheFrames() const
 	}
 	// Otherwise, use the default value or value set in ini file
 	return NumPrecacheFrames;
-}
\ No newline at end of file
+}
diff --git a/Engine/Source/Runtime/Engine/Public/ActiveSound.h b/Engine/Source/Runtime/Engine/Public/ActiveSound.h
--- a/Engine/Source/Runtime/Engine/Public/ActiveSound.h
+++ b/Engine/Source/Runtime/Engine/Public/ActiveSound.h
@@ -354,6 +354,63 @@ private:
 
 	TMap<UPTRINT, FWaveInstance*> WaveInstances;
 
+	class FInstanceTransmitter
+	{
+	public:
+		FInstanceTransmitter() = default;
+		FInstanceTransmitter(FInstanceTransmitter&&) = default;
+		FInstanceTransmitter& operator=(FInstanceTransmitter&&) = default;
+
+		FInstanceTransmitter(TUniquePtr<IAudioInstanceTransmitter>&& InTransmitterImpl)
+		:	TransmitterImpl(MoveTemp(InTransmitterImpl))
+		{
+		}
+
+		FInstanceTransmitter(const FInstanceTransmitter& InOther)
+		{
+			if (InOther.IsValid())
+			{
+				TransmitterImpl = InOther->Clone();
+			}
+		}
+
+		FInstanceTransmitter& operator=(const FInstanceTransmitter& InOther)
+		{
+			TransmitterImpl.Reset();
+			if (InOther.IsValid())
+			{
+				TransmitterImpl = InOther->Clone();
+			}
+			return *this;
+		}
+
+		FInstanceTransmitter& operator=(TUniquePtr<IAudioInstanceTransmitter>&& InTransmitterImpl)
+		{
+			TransmitterImpl = MoveTemp(InTransmitterImpl);
+			return *this;
+		}
+
+		bool IsValid() const
+		{
+			return TransmitterImpl.IsValid();
+		}
+
+		IAudioInstanceTransmitter* operator->()
+		{
+			return TransmitterImpl.Get();
+		}
+
+		const IAudioInstanceTransmitter* operator->() const
+		{
+			return TransmitterImpl.Get();
+		}
+
+	private:
+		TUniquePtr<IAudioInstanceTransmitter> TransmitterImpl;
+	};
+
+	FInstanceTransmitter InstanceTransmitter;
+
 public:
 	enum class EFadeOut : uint8
 	{
@@ -769,4 +826,4 @@ private:
 
 	/** Helper function which retrieves attenuation frequency value for HPF and LPF distance-based filtering. */
 	float GetAttenuationFrequency(const FSoundAttenuationSettings* InSettings, const FAttenuationListenerData& ListenerData, const FVector2D& FrequencyRange, const FRuntimeFloatCurve& CustomCurve);
-};
\ No newline at end of file
+};
