Fix stale buffer data is OscNode when toggling on and off

diff --git a/Engine/Plugins/Runtime/Metasound/Source/MetasoundStandardNodes/Private/MetasoundOscNode.cpp b/Engine/Plugins/Runtime/Metasound/Source/MetasoundStandardNodes/Private/MetasoundOscNode.cpp
--- a/Engine/Plugins/Runtime/Metasound/Source/MetasoundStandardNodes/Private/MetasoundOscNode.cpp
+++ b/Engine/Plugins/Runtime/Metasound/Source/MetasoundStandardNodes/Private/MetasoundOscNode.cpp
@@ -55,7 +55,6 @@ namespace Metasound
 		// The current phase of oscillator (between 0.0 and 1.0)
 		float Phase;
 		float OneOverSampleRate;
-		int32 BlockSize;		
 		float Nyquist;
 
 		FFrequencyReadRef Frequency;
@@ -70,7 +69,6 @@ namespace Metasound
 	FOscOperator::FOscOperator(const FOperatorSettings& InSettings, const FFrequencyReadRef& InFrequency, const FBoolReadRef& InEnabled)
 		: Phase(0.f)
 		, OneOverSampleRate(1.f / InSettings.GetSampleRate())
-		, BlockSize(InSettings.GetNumFramesPerBlock())
 		, Nyquist(InSettings.GetSampleRate() / 2.0f)
 		, Frequency(InFrequency)
 		, Enabled(InEnabled)
@@ -98,20 +96,19 @@ namespace Metasound
 	void FOscOperator::Execute()
 	{
 		// Clamp frequencies into Nyquist range
-		float Freq = FMath::Clamp(Frequency->GetHertz(), -Nyquist, Nyquist);
-		
-		float PhaseInc = Freq * OneOverSampleRate;
-		
+		const float Freq = FMath::Clamp(Frequency->GetHertz(), -Nyquist, Nyquist);
+		const float PhaseInc = Freq * OneOverSampleRate;
 		float* Data = AudioBuffer->GetData();
-		FMemory::Memzero(Data, BlockSize);
+
+		FMemory::Memzero(Data, AudioBuffer->Num() * sizeof(float));
 
 		if (*Enabled)
 		{
-			for (int32 i = 0; i < BlockSize; i++)
+			for (int32 i = 0; i < AudioBuffer->Num(); i++)
 			{
-				// This is borrowed from the FOsc class with the intention to eventually recreate that functionality here. 
+				// This is borrowed from the FOsc class with the intention to eventually recreate that functionality here.
 				// We don't wish to use it directly as it has a virtual call for each sample. Phase is in cents with the intent
-				// to support other oscilator types in time.
+				// to support other oscillator types in time.
 
 				const float Radians = (Phase * TwoPi) - PI;
 				Data[i] = Audio::FastSin3(-1.0f * Radians);
