Literal Input Nodes
GainMultiply fix-ups (fade fast to avoid stair stepping & missing input declaration)
- Various EdGraphBuilder clean-up

diff --git a/Engine/Plugins/Runtime/Metasound/Source/MetasoundEditor/Private/MetasoundDetailCustomization.cpp b/Engine/Plugins/Runtime/Metasound/Source/MetasoundEditor/Private/MetasoundDetailCustomization.cpp
--- a/Engine/Plugins/Runtime/Metasound/Source/MetasoundEditor/Private/MetasoundDetailCustomization.cpp
+++ b/Engine/Plugins/Runtime/Metasound/Source/MetasoundEditor/Private/MetasoundDetailCustomization.cpp
@@ -11,6 +11,7 @@
 #include "IDetailGroup.h"
 #include "MetasoundAssetBase.h"
 #include "MetasoundFrontend.h"
+#include "MetasoundFrontendController.h"
 #include "MetasoundUObjectRegistry.h"
 #include "PropertyEditorDelegates.h"
 #include "PropertyHandle.h"
@@ -26,6 +27,14 @@
 
 #define LOCTEXT_NAMESPACE "MetasoundEditor"
 
+static int32 ShowLiteralMetasoundInputsInEditorCVar = 0;
+FAutoConsoleVariableRef CVarShowLiteralMetasoundInputsInEditor(
+	TEXT("au.Debug.Editor.Metasounds.ShowLiteralInputs"),
+	ShowLiteralMetasoundInputsInEditorCVar,
+	TEXT("Show literal inputs in the Metasound Editor.\n")
+	TEXT("0: Disabled (default), !0: Enabled"),
+	ECVF_Default);
+
 namespace Metasound
 {
 	namespace Editor
@@ -40,6 +49,34 @@ namespace Metasound
 			return FName(InBasePath.ToString() + TEXT(".") + InPropertyName.ToString());
 		}
 
+		TSet<FString> GetLiteralInputs(IDetailLayoutBuilder& InDetailLayout)
+		{
+			TSet<FString> LiteralInputs;
+			TArray<TWeakObjectPtr<UObject>> Objects;
+			InDetailLayout.GetObjectsBeingCustomized(Objects);
+
+			if (Objects.IsEmpty() || !Objects[0].IsValid())
+			{
+				return LiteralInputs;
+			}
+
+			UObject* Metasound = Objects[0].Get();
+			if (FMetasoundAssetBase* MetasoundAsset = IMetasoundUObjectRegistry::Get().GetObjectAsAssetBase(Metasound))
+			{
+				Frontend::FGraphHandle GraphHandle = MetasoundAsset->GetRootGraphHandle();
+				TArray<Frontend::FNodeHandle> InputNodes = GraphHandle->GetInputNodes();
+				for (Frontend::FNodeHandle& NodeHandle : InputNodes)
+				{
+					if (NodeHandle->GetNodeStyle().Display.Visibility == EMetasoundFrontendNodeStyleDisplayVisibility::Hidden)
+					{
+						LiteralInputs.Add(NodeHandle->GetNodeName());
+					}
+				}
+			}
+
+			return LiteralInputs;
+		}
+
 		template <typename T>
 		void BuildIOFixedArray(IDetailLayoutBuilder& InDetailLayout, FName InCategoryName, FName InPropertyName, const TSet<FString>& InRequiredValues)
 		{
@@ -49,6 +86,12 @@ namespace Metasound
 			TSharedPtr<IPropertyHandle> ParentProperty = InDetailLayout.GetProperty(InPropertyName);
 			TSharedPtr<IPropertyHandleArray> ArrayHandle = ParentProperty->AsArray();
 
+			TSet<FString> LiteralInputs;
+			if (bIsInput && !ShowLiteralMetasoundInputsInEditorCVar)
+			{
+				LiteralInputs = GetLiteralInputs(InDetailLayout);
+			}
+
 			uint32 NumElements = 0;
 			ArrayHandle->GetNumElements(NumElements);
 			for (int32 i = 0; i < static_cast<int32>(NumElements); ++i)
@@ -70,6 +113,12 @@ namespace Metasound
 				const bool bNameFound = NameProperty->GetValue(Name) == FPropertyAccess::Success;
 				const bool bIsRequired = bNameFound && InRequiredValues.Contains(Name);
 
+				// Hide literal members
+				if (LiteralInputs.Contains(Name))
+				{
+					continue;
+				}
+
 				FText DisplayName;
 				DisplayNameProperty->GetValue(DisplayName);
 				DisplayNameProperty->SetInstanceMetaData("LastValidName", DisplayName.ToString());
diff --git a/Engine/Plugins/Runtime/Metasound/Source/MetasoundEditor/Private/MetasoundEditor.cpp b/Engine/Plugins/Runtime/Metasound/Source/MetasoundEditor/Private/MetasoundEditor.cpp
--- a/Engine/Plugins/Runtime/Metasound/Source/MetasoundEditor/Private/MetasoundEditor.cpp
+++ b/Engine/Plugins/Runtime/Metasound/Source/MetasoundEditor/Private/MetasoundEditor.cpp
@@ -36,7 +36,6 @@
 
 #define LOCTEXT_NAMESPACE "MetasoundEditor"
 
-
 namespace Metasound
 {
 	namespace Editor
@@ -108,42 +107,12 @@ namespace Metasound
 				Graph->ParentMetasound = &InMetasound;
 				Graph->Schema = UMetasoundEditorGraphSchema::StaticClass();
 				MetasoundAsset->SetGraph(Graph);
-				FGraphBuilder::RebuildGraph(InMetasound);
+				FGraphBuilder::ConstructGraph(InMetasound);
 			}
 
 			return MetasoundAsset->GetRootGraphHandle();
 		}
 
-		bool FEditor::RebuildGraph() const
-		{
-			UObject* MetasoundObj = GetMetasoundObject();
-			if (nullptr != MetasoundObj)
-			{
-				if (IMetasoundUObjectRegistry::Get().IsRegisteredClass(MetasoundObj))
-				{
-					FGraphBuilder::RebuildGraph(*MetasoundObj);
-					return true;
-				}
-			}
-
-			return false;
-		}
-
-		bool FEditor::SynchronizeGraph() const
-		{
-			UObject* MetasoundObj = GetMetasoundObject();
-			if (nullptr != MetasoundObj)
-			{
-				if (IMetasoundUObjectRegistry::Get().IsRegisteredClass(MetasoundObj))
-				{
-					FGraphBuilder::SynchronizeGraph(*MetasoundObj);
-					return true;
-				}
-			}
-
-			return false;
-		}
-
 		void FEditor::InitMetasoundEditor(const EToolkitMode::Type Mode, const TSharedPtr<IToolkitHost>& InitToolkitHost, UObject* ObjectToEdit)
 		{
 			check(ObjectToEdit);
@@ -265,9 +234,15 @@ namespace Metasound
 				// If a property change event occurs outside of the metasound UEdGraph and results in the metasound document changing, 
 				// then the document and the UEdGraph need to be synchronized. There may be a better trigger for this call to reduce
 				// the number of times the graph is synchronized. 
-				SynchronizeGraph();
-
-				MetasoundGraphEditor->NotifyGraphChanged();
+				UObject* MetasoundObj = GetMetasoundObject();
+				if (nullptr != MetasoundObj)
+				{
+					if (IMetasoundUObjectRegistry::Get().IsRegisteredClass(MetasoundObj))
+					{
+						FGraphBuilder::SynchronizeGraph(*MetasoundObj);
+						MetasoundGraphEditor->NotifyGraphChanged();
+					}
+				}
 			}
 		}
 
@@ -684,6 +659,8 @@ namespace Metasound
 
 		void FEditor::DeleteSelectedNodes()
 		{
+			using namespace Metasound::Frontend;
+
 			const FScopedTransaction Transaction(LOCTEXT("MetasoundEditorDeleteSelectedNode", "Delete Selected Metasound Node(s)"));
 
 			UMetasoundEditorGraph* Graph = CastChecked<UMetasoundEditorGraph>(MetasoundGraphEditor->GetCurrentGraph());
@@ -692,64 +669,14 @@ namespace Metasound
 			const FGraphPanelSelectionSet SelectedNodes = MetasoundGraphEditor->GetSelectedNodes();
 			MetasoundGraphEditor->ClearSelectionSet();
 
-			UObject* ParentMetasoundObject = Graph->ParentMetasound;
-			check(ParentMetasoundObject);
-			if (FMetasoundAssetBase* MetasoundAsset = IMetasoundUObjectRegistry::Get().GetObjectAsAssetBase(ParentMetasoundObject))
+			for (FGraphPanelSelectionSet::TConstIterator NodeIt(SelectedNodes); NodeIt; ++NodeIt)
 			{
-				Frontend::FConstDocumentHandle DocumentHandle = MetasoundAsset->GetDocumentHandle();
-				for (FGraphPanelSelectionSet::TConstIterator NodeIt(SelectedNodes); NodeIt; ++NodeIt)
+				// Some nodes may not be metasound nodes (ex. comments and perhaps aliases eventually), but can be safely deleted.
+				if (UEdGraphNode* Node = Cast<UEdGraphNode>(*NodeIt))
 				{
-					Frontend::FNodeHandle NodeHandle = Frontend::INodeController::GetInvalidHandle();
-
-					if (UMetasoundEditorGraphNode* Node = Cast<UMetasoundEditorGraphNode>(*NodeIt))
-					{
-						NodeHandle = Node->GetNodeHandle();
-
-						if (NodeHandle->GetClassType() == EMetasoundFrontendClassType::Input)
-						{
-							auto IsRequiredInput = [&](const Frontend::FConstInputHandle& InputHandle) 
-							{ 
-								return DocumentHandle->IsRequiredInput(InputHandle->GetName()); 
-							};
-							TArray<Frontend::FConstInputHandle> NodeInputs = NodeHandle->GetConstInputs();
-
-							if (Frontend::FConstInputHandle* InputHandle = NodeInputs.FindByPredicate(IsRequiredInput))
-							{
-								FNotificationInfo Info(FText::Format(LOCTEXT("Metasounds_CannotDeleteRequiredInput",
-									"'Required Input '{0}' cannot be deleted."), (*InputHandle)->GetDisplayName()));
-								Info.bFireAndForget = true;
-								Info.ExpireDuration = 2.0f;
-								Info.bUseThrobber = true;
-								FSlateNotificationManager::Get().AddNotification(Info);
-								continue;
-							}
-						}
-
-						if (NodeHandle->GetClassType() == EMetasoundFrontendClassType::Output)
-						{
-							auto IsRequiredOutput = [&](const Frontend::FConstOutputHandle& OutputHandle) 
-							{ 
-								return DocumentHandle->IsRequiredOutput(OutputHandle->GetName()); 
-							};
-							TArray<Frontend::FConstOutputHandle> NodeOutputs = NodeHandle->GetConstOutputs();
-
-							if (Frontend::FConstOutputHandle* OutputHandle = NodeOutputs.FindByPredicate(IsRequiredOutput))
-							{
-								FNotificationInfo Info(FText::Format(LOCTEXT("Metasounds_CannotDeleteRequiredOutput",
-									"'Required Output '{0}' cannot be deleted."), (*OutputHandle)->GetDisplayName()));
-								Info.bFireAndForget = true;
-								Info.ExpireDuration = 2.0f;
-								Info.bUseThrobber = true;
-								FSlateNotificationManager::Get().AddNotification(Info);
-								continue;
-							}
-						}
-					}
-
-					// Some nodes may not be metasound nodes (ex. comments and perhaps aliases eventually), but can be safely deleted.
-					if (UEdGraphNode* Node = Cast<UEdGraphNode>(*NodeIt))
+					if (!FGraphBuilder::DeleteNode(*Node, true /* bInRecordTransaction */))
 					{
-						FGraphBuilder::DeleteNode(*Node, NodeHandle, true /* bInRecordTransaction */);
+						MetasoundGraphEditor->SetNodeSelection(Node, true /* bSelect */);
 					}
 				}
 			}
diff --git a/Engine/Plugins/Runtime/Metasound/Source/MetasoundEditor/Private/MetasoundEditor.h b/Engine/Plugins/Runtime/Metasound/Source/MetasoundEditor/Private/MetasoundEditor.h
--- a/Engine/Plugins/Runtime/Metasound/Source/MetasoundEditor/Private/MetasoundEditor.h
+++ b/Engine/Plugins/Runtime/Metasound/Source/MetasoundEditor/Private/MetasoundEditor.h
@@ -48,8 +48,6 @@ namespace Metasound
 			virtual UObject* GetMetasoundObject() const override;
 			virtual void SetSelection(const TArray<UObject*>& SelectedObjects) override;
 			virtual bool GetBoundsForSelectedNodes(FSlateRect& Rect, float Padding) override;
-			bool RebuildGraph() const override;
-			bool SynchronizeGraph() const override;
 
 			/** IToolkit interface */
 			virtual FName GetToolkitFName() const override;
diff --git a/Engine/Plugins/Runtime/Metasound/Source/MetasoundEditor/Private/MetasoundEditorGraphBuilder.cpp b/Engine/Plugins/Runtime/Metasound/Source/MetasoundEditor/Private/MetasoundEditorGraphBuilder.cpp
--- a/Engine/Plugins/Runtime/Metasound/Source/MetasoundEditor/Private/MetasoundEditorGraphBuilder.cpp
+++ b/Engine/Plugins/Runtime/Metasound/Source/MetasoundEditor/Private/MetasoundEditorGraphBuilder.cpp
@@ -4,18 +4,19 @@
 #include "EdGraph/EdGraph.h"
 #include "EdGraph/EdGraphNode.h"
 #include "EdGraph/EdGraphPin.h"
+#include "Framework/Notifications/NotificationManager.h"
 #include "Metasound.h"
 #include "MetasoundAssetBase.h"
 #include "MetasoundEditorGraph.h"
 #include "MetasoundEditorGraphNode.h"
 #include "MetasoundEditorModule.h"
-#include "MetasoundFrontend.h"
-#include "MetasoundFrontendController.h"
 #include "MetasoundFrontendRegistries.h"
+#include "MetasoundLiteral.h"
 #include "MetasoundUObjectRegistry.h"
 #include "Modules/ModuleManager.h"
 #include "ScopedTransaction.h"
 #include "Templates/Tuple.h"
+#include "Widgets/Notifications/SNotificationList.h"
 
 
 #define LOCTEXT_NAMESPACE "MetasoundEditor"
@@ -25,18 +26,104 @@ namespace Metasound
 {
 	namespace Editor
 	{
-		const FName FGraphBuilder::PinAudioFormat = "Format";
-		const FName FGraphBuilder::PinAudioNumeric = "Numeric";
-		const FName FGraphBuilder::PinPrimitiveBoolean = "Boolean";
-		const FName FGraphBuilder::PinPrimitiveFloat = "Float";
-		const FName FGraphBuilder::PinPrimitiveInt32 = "Int32";
-		const FName FGraphBuilder::PinPrimitiveInt64 = "Int64";
-		const FName FGraphBuilder::PinPrimitiveString = "String";
-		const FName FGraphBuilder::PinPrimitiveTrigger = "Trigger";
-		const FName FGraphBuilder::PinPrimitiveUObject = "UObject";
-		const FName FGraphBuilder::PinPrimitiveUObjectArray = "UObjectArray";
-
-		UEdGraphNode* FGraphBuilder::AddNode(UObject& InMetasound, const FVector2D& Location, Frontend::FNodeHandle& InNodeHandle, bool bInSelectNewNode)
+		// Category names must match those found in UEdGraphSchema_K2:PC_<type>
+		// so default selectors function the same way by default
+		const FName FGraphBuilder::PinCategoryBoolean = "bool";
+		const FName FGraphBuilder::PinCategoryDouble = "double";
+		const FName FGraphBuilder::PinCategoryExec = "exec";
+		const FName FGraphBuilder::PinCategoryFloat = "float";
+		const FName FGraphBuilder::PinCategoryInt32 = "int";
+		const FName FGraphBuilder::PinCategoryInt64 = "int64";
+		const FName FGraphBuilder::PinCategoryObject = "object";
+		const FName FGraphBuilder::PinCategoryString = "string";
+
+		const FName FGraphBuilder::PinSubCategoryAudioFormat = "Format";
+		const FName FGraphBuilder::PinSubCategoryAudioNumeric = "Numeric";
+		const FName FGraphBuilder::PinSubCategoryObjectArray = "UObjectArray";
+
+		namespace GraphBuilderPrivate
+		{
+			void DeleteNode(UObject& InMetasound, Frontend::FNodeHandle InNodeHandle)
+			{
+				if (InNodeHandle->IsValid())
+				{
+					Frontend::FGraphHandle GraphHandle = InNodeHandle->GetOwningGraph();
+					if (GraphHandle->IsValid())
+					{
+						GraphHandle->RemoveNode(*InNodeHandle);
+					}
+				}
+
+				InMetasound.MarkPackageDirty();
+			}
+
+			void SetInputLiteral(Frontend::FNodeHandle InInputNodeHandle, FGuid PointID, const FName InTypeName, const FLiteral& InDefaultValue)
+			{
+				Frontend::FGraphHandle GraphHandle = InInputNodeHandle->GetOwningGraph();
+
+				const FString Name = InInputNodeHandle->GetNodeName();
+				switch (InDefaultValue.GetType())
+				{
+					case ELiteralType::Boolean:
+					{
+						GraphHandle->SetDefaultInputToLiteral(Name, PointID, InDefaultValue.Value.Get<bool>());
+					}
+					break;
+
+					case ELiteralType::Float:
+					{
+						GraphHandle->SetDefaultInputToLiteral(Name, PointID, InDefaultValue.Value.Get<float>());
+					}
+					break;
+
+					case ELiteralType::Integer:
+					{
+						GraphHandle->SetDefaultInputToLiteral(Name, PointID, InDefaultValue.Value.Get<int32>());
+					}
+					break;
+
+					case ELiteralType::String:
+					{
+						GraphHandle->SetDefaultInputToLiteral(Name, PointID, InDefaultValue.Value.Get<FString>());
+					}
+					break;
+
+					case ELiteralType::UObjectProxy:
+					{
+						// TODO: Support default UObject value on node
+						UClass* ClassToUse = FMetasoundFrontendRegistryContainer::Get()->GetLiteralUClassForDataType(InTypeName);
+						if (ClassToUse)
+						{
+							GraphHandle->SetDefaultInputToLiteral(Name, PointID, ClassToUse->ClassDefaultObject);
+						}
+					}
+					break;
+
+					case ELiteralType::UObjectProxyArray:
+					{
+						// TODO: Support default UObject array value on node
+						UClass* ClassToUse = FMetasoundFrontendRegistryContainer::Get()->GetLiteralUClassForDataType(InTypeName);
+						if (ClassToUse)
+						{
+							TArray<UObject*> ObjectArray;
+							ObjectArray.Add(ClassToUse->ClassDefaultObject);
+							GraphHandle->SetDefaultInputToLiteral(Name, PointID, ObjectArray);
+						}
+					}
+					break;
+
+					case ELiteralType::Invalid:
+					case ELiteralType::None:
+					default:
+					{
+						static_assert(static_cast<int32>(ELiteralType::Invalid) == 7, "Possible missing ELiteralType case coverage");
+					}
+					break;
+				}
+			}
+		}
+
+		UEdGraphNode* FGraphBuilder::AddNode(UObject& InMetasound, Frontend::FNodeHandle& InNodeHandle, bool bInSelectNewNode)
 		{
 			const FScopedTransaction Transaction(LOCTEXT("AddMetasoundGraphNode", "Add Metasound Node"));
 
@@ -48,10 +135,12 @@ namespace Metasound
 			UMetasoundEditorGraphNode* NewGraphNode = NodeCreator.CreateNode(bInSelectNewNode);
 			NodeCreator.Finalize();
 
+			const FMetasoundFrontendNodeStyle& Style = InNodeHandle->GetNodeStyle();
+
 			NewGraphNode->SetNodeID(InNodeHandle->GetID());
 			NewGraphNode->CreateNewGuid();
-			NewGraphNode->NodePosX = Location.X;
-			NewGraphNode->NodePosY = Location.Y;
+			NewGraphNode->NodePosX = Style.Display.Location.X;
+			NewGraphNode->NodePosY = Style.Display.Location.Y;
 
 			RebuildNodePins(*NewGraphNode, InNodeHandle);
 
@@ -61,18 +150,20 @@ namespace Metasound
 			return NewGraphNode;
 		}
 
-		UEdGraphNode* FGraphBuilder::AddNode(UObject& InMetasound, const FVector2D& Location, const Frontend::FNodeClassInfo& InClassInfo, bool bInSelectNewNode)
+		UEdGraphNode* FGraphBuilder::AddNode(UObject& InMetasound, const Frontend::FNodeClassInfo& InClassInfo, const FMetasoundFrontendNodeStyle& InNodeStyle, bool bInSelectNewNode)
 		{
-			Frontend::FNodeHandle NodeHandle = AddNodeHandle(InMetasound, InClassInfo);
-			return AddNode(InMetasound, Location, NodeHandle, bInSelectNewNode);
+			Frontend::FNodeHandle NodeHandle = AddNodeHandle(InMetasound, InClassInfo, InNodeStyle);
+			return AddNode(InMetasound, NodeHandle, bInSelectNewNode);
 		}
 
-		Frontend::FNodeHandle FGraphBuilder::AddNodeHandle(UObject& InMetasound, const Frontend::FNodeClassInfo& InClassInfo)
+		Frontend::FNodeHandle FGraphBuilder::AddNodeHandle(UObject& InMetasound, const Frontend::FNodeClassInfo& InClassInfo, const FMetasoundFrontendNodeStyle& InNodeStyle)
 		{
 			FMetasoundAssetBase* MetasoundAsset = IMetasoundUObjectRegistry::Get().GetObjectAsAssetBase(&InMetasound);
 			check(MetasoundAsset);
 
-			return MetasoundAsset->GetRootGraphHandle()->AddNode(InClassInfo);
+			Frontend::FNodeHandle NewNode = MetasoundAsset->GetRootGraphHandle()->AddNode(InClassInfo);
+			NewNode->SetNodeStyle(InNodeStyle);
+			return NewNode;
 		}
 
 		FString FGraphBuilder::GetDataTypeDisplayName(const FName& InDataTypeName)
@@ -136,123 +227,195 @@ namespace Metasound
 			return NewOutputName;
 		}
 
-		UEdGraphNode* FGraphBuilder::AddInput(UObject& InMetasound, const FVector2D& Location, const FString& InName, const FName InTypeName, const FText& InToolTip, bool bInSelectNewNode)
+		UEdGraphNode* FGraphBuilder::AddInput(UObject& InMetasound, const FString& InName, const FName InTypeName, const FMetasoundFrontendNodeStyle& InNodeStyle, const FText& InToolTip, bool bInSelectNewNode)
 		{
-			Frontend::FNodeHandle NodeHandle = AddInputNodeHandle(InMetasound, InName, InTypeName, InToolTip);
-			return AddNode(InMetasound, Location, NodeHandle, bInSelectNewNode);
+			Frontend::FNodeHandle NodeHandle = AddInputNodeHandle(InMetasound, InName, InTypeName, InNodeStyle, &InToolTip);
+			return AddNode(InMetasound, NodeHandle, bInSelectNewNode);
 		}
 
-		Frontend::FNodeHandle FGraphBuilder::AddInputNodeHandle(UObject& InMetasound, const FString& InName, const FName InTypeName, const FText& InToolTip)
+		void FGraphBuilder::AddOrUpdateLiteralInput(UObject& InMetasound, Frontend::FNodeHandle InNodeHandle, const UEdGraphPin& InInputPin)
 		{
+			using namespace Metasound::Frontend;
 
-			FMetasoundAssetBase* MetasoundAsset = IMetasoundUObjectRegistry::Get().GetObjectAsAssetBase(&InMetasound);
-			check(MetasoundAsset);
-
-			Frontend::FGraphHandle GraphHandle = MetasoundAsset->GetRootGraphHandle();
+			const FScopedTransaction Transaction(LOCTEXT("SetMetasoundGraphNode", "Set Metasound Literal Input"));
 
-			FMetasoundFrontendClassInput Description;
-
-			Description.Name = InName;
-			Description.TypeName = InTypeName;
-			Description.Metadata.Description = InToolTip;
-			FGuid PointID = GraphHandle->GetNewPointID();
-
-			Description.PointIDs.Add(PointID);
-			FMetasoundFrontendVertexLiteral DefaultValue;
-			DefaultValue.PointID = PointID;
-			Description.Defaults.Add(DefaultValue);
-
-			Frontend::FNodeHandle NodeHandle = GraphHandle->AddInputVertex(Description);
-			if (!NodeHandle->IsValid())
+			if (!ensureAlways(InNodeHandle->GetClassType() == EMetasoundFrontendClassType::External))
 			{
-				return NodeHandle;
+				return;
 			}
 
-			IMetasoundEditorModule& EditorModule = FModuleManager::GetModuleChecked<IMetasoundEditorModule>("MetasoundEditor");
-			FEditorDataType DataType = EditorModule.FindDataType(InTypeName);
+			const FString& InInputName = InInputPin.GetName();
+			const FString& InStringValue = InInputPin.DefaultValue;
+			const UObject* InObjectValue = InInputPin.DefaultObject;
 
-			Metasound::FLiteral LiteralParam = Frontend::GetDefaultParamForDataType(InTypeName);
-			if (!ensureAlways(LiteralParam.IsValid()))
+			TArray<FInputHandle> InputHandles = InNodeHandle->GetInputsWithVertexName(InInputName);
+			if (!ensureAlways(InputHandles.Num() == 1))
 			{
-				return NodeHandle;
+				return;
 			}
 
-			switch (LiteralParam.GetType())
+			FInputHandle InputHandle = InputHandles[0];
+			IMetasoundEditorModule& EditorModule = FModuleManager::GetModuleChecked<IMetasoundEditorModule>("MetasoundEditor");
+
+			FLiteral Literal;
+			const FName TypeName = InputHandle->GetDataType();
+			const FEditorDataType DataType = EditorModule.FindDataType(TypeName);
+			switch (DataType.RegistryInfo.PreferredLiteralType)
 			{
 				case ELiteralType::Boolean:
 				{
-					GraphHandle->SetDefaultInputToLiteral(InName, PointID, false);
+					Literal.Value.Set<bool>(FCString::ToBool(*InStringValue));
 				}
 				break;
 
 				case ELiteralType::Float:
 				{
-					GraphHandle->SetDefaultInputToLiteral(InName, PointID, 0.0f);
+					Literal.Value.Set<float>(FCString::Atof(*InStringValue));
 				}
 				break;
 
 				case ELiteralType::Integer:
 				{
-					GraphHandle->SetDefaultInputToLiteral(InName, PointID, 0);
+					Literal.Value.Set<int32>(FCString::Atoi(*InStringValue));
 				}
 				break;
 
 				case ELiteralType::String:
 				{
-					GraphHandle->SetDefaultInputToLiteral(InName, PointID, FString(TEXT("")));
+					Literal.Value.Set<FString>(InStringValue);
 				}
 				break;
 
+				// TODO: Support UObjects/UObject arrays for default literals
 				case ELiteralType::UObjectProxy:
+				case ELiteralType::UObjectProxyArray:
 				{
-					UClass* ClassToUse = FMetasoundFrontendRegistryContainer::Get()->GetLiteralUClassForDataType(InTypeName);
-					if (ClassToUse)
-					{
-						GraphHandle->SetDefaultInputToLiteral(InName, PointID, ClassToUse->ClassDefaultObject);
-					}
+					return;
 				}
 				break;
 
-				case ELiteralType::UObjectProxyArray:
+				// If no literal supported, no need for private input
+				case ELiteralType::None:
 				{
-					UClass* ClassToUse = FMetasoundFrontendRegistryContainer::Get()->GetLiteralUClassForDataType(InTypeName);
-					if (ClassToUse)
-					{
-						TArray<UObject*> ObjectArray;
-						ObjectArray.Add(ClassToUse->ClassDefaultObject);
-						GraphHandle->SetDefaultInputToLiteral(InName, PointID, ObjectArray);
-					}
-					GraphHandle->SetDefaultInputToLiteral(InName, PointID, FString(TEXT("")));
+					return;
 				}
 				break;
 
 				case ELiteralType::Invalid:
-				case ELiteralType::None:
 				default:
 				{
-					static_assert(static_cast<int32>(ELiteralType::Invalid) == 7, "Possible missing ELiteralType case coverage");
+					ensureAlwaysMsgf(false, TEXT("Failed to set input node default: Literal type not supported"));
+					return;
 				}
 				break;
 			}
 
+			bool bNewValueSet = false;
+			FOutputHandle OutputHandle = InputHandle->GetCurrentlyConnectedOutput();
+			if (!OutputHandle->IsValid())
+			{
+				FMetasoundFrontendNodeStyle Style;
+				Style.Display.Visibility = EMetasoundFrontendNodeStyleDisplayVisibility::Hidden;
+
+				Frontend::FNodeHandle NewInputNode = AddInputNodeHandle(InMetasound, GenerateUniqueInputName(InMetasound, "LiteralInput"), TypeName, Style, nullptr /* ToolTip */, &Literal);
+				NewInputNode->SetNodeStyle(Style);
+
+				bNewValueSet = true;
+				TArray<Metasound::Frontend::FOutputHandle> OutputHandles = NewInputNode->GetOutputs();
+				if (ensureAlways(OutputHandles.Num() == 1))
+				{
+					OutputHandle = OutputHandles[0];
+				}
+
+				ensureAlways(InputHandle->Connect(*OutputHandle));
+			}
+
+			if (OutputHandle->IsValid())
+			{
+				FNodeHandle InputNode = OutputHandle->GetOwningNode();
+				const FMetasoundFrontendNodeStyle& Style = InputNode->GetNodeStyle();
+				if (InputNode->IsValid() && Style.Display.Visibility == EMetasoundFrontendNodeStyleDisplayVisibility::Hidden)
+				{
+					if (!bNewValueSet)
+					{
+						// TODO: Support multiple literal point ids for arrays
+						Frontend::FGraphHandle GraphHandle = InputNode->GetOwningGraph();
+						TArray<FGuid> PointIDs = GraphHandle->GetDefaultIDsForInputVertex(InputNode->GetNodeName());
+						if (ensureAlways(PointIDs.Num() == 1))
+						{
+							GraphBuilderPrivate::SetInputLiteral(InputNode, PointIDs[0], TypeName, Literal);
+						}
+					}
+					InMetasound.PostEditChange();
+					InMetasound.MarkPackageDirty();
+				}
+			}
+		}
+
+		Frontend::FNodeHandle FGraphBuilder::AddInputNodeHandle(UObject& InMetasound, const FString& InName, const FName InTypeName, const FMetasoundFrontendNodeStyle& InNodeStyle, const FText* InToolTip, const FLiteral* InDefaultValue)
+		{
+			FMetasoundAssetBase* MetasoundAsset = IMetasoundUObjectRegistry::Get().GetObjectAsAssetBase(&InMetasound);
+			check(MetasoundAsset);
+
+			Frontend::FGraphHandle GraphHandle = MetasoundAsset->GetRootGraphHandle();
+
+			FMetasoundFrontendClassInput Description;
+
+			Description.Name = InName;
+			Description.TypeName = InTypeName;
+
+			if (InToolTip)
+			{
+				Description.Metadata.Description = *InToolTip;
+			}
+
+			FGuid PointID = GraphHandle->GetNewPointID();
+
+			Description.PointIDs.Add(PointID);
+			FMetasoundFrontendVertexLiteral DefaultVertex;
+			DefaultVertex.PointID = PointID;
+			Description.Defaults.Add(DefaultVertex);
+
+			Frontend::FNodeHandle NodeHandle = GraphHandle->AddInputVertex(Description);
+			if (!ensureAlways(NodeHandle->IsValid()))
+			{
+				return NodeHandle;
+			}
+
 			GraphHandle->SetInputDisplayName(InName, FText::FromString(InName));
+
+			Metasound::FLiteral LiteralParam = Frontend::GetDefaultParamForDataType(InTypeName);
+			if (!LiteralParam.IsValid())
+			{
+				return NodeHandle;
+			}
+
+			if (InDefaultValue)
+			{
+				if (!ensureAlways(InDefaultValue->GetType() == LiteralParam.GetType()))
+				{
+					return NodeHandle;
+				}
+			};
+
+			GraphBuilderPrivate::SetInputLiteral(NodeHandle, PointID, InTypeName, InDefaultValue ? *InDefaultValue : LiteralParam);
+
 			return NodeHandle;
 		}
 
-		UEdGraphNode* FGraphBuilder::AddOutput(UObject& InMetasound, const FVector2D& Location, const FString& InName, const FName InTypeName, const FText& InToolTip, bool bInSelectNewNode)
+		UEdGraphNode* FGraphBuilder::AddOutput(UObject& InMetasound, const FString& InName, const FName InTypeName, const FMetasoundFrontendNodeStyle& InNodeStyle, const FText& InToolTip, bool bInSelectNewNode)
 		{
-			Frontend::FNodeHandle NodeHandle = AddOutputNodeHandle(InMetasound, InName, InTypeName, InToolTip);
-			return AddNode(InMetasound, Location, NodeHandle, bInSelectNewNode);
+			Frontend::FNodeHandle NodeHandle = AddOutputNodeHandle(InMetasound, InName, InTypeName, InNodeStyle, InToolTip);
+			return AddNode(InMetasound, NodeHandle, bInSelectNewNode);
 		}
 
-		Frontend::FNodeHandle FGraphBuilder::AddOutputNodeHandle(UObject& InMetasound, const FString& InName, const FName InTypeName, const FText& InToolTip)
+		Frontend::FNodeHandle FGraphBuilder::AddOutputNodeHandle(UObject& InMetasound, const FString& InName, const FName InTypeName, const FMetasoundFrontendNodeStyle& InNodeStyle, const FText& InToolTip)
 		{
 			FMetasoundAssetBase* MetasoundAsset = IMetasoundUObjectRegistry::Get().GetObjectAsAssetBase(&InMetasound);
 			check(MetasoundAsset);
 
 			Frontend::FGraphHandle GraphHandle = MetasoundAsset->GetRootGraphHandle();
 
-
 			FMetasoundFrontendClassOutput Description;
 			Description.Name = InName;
 			Description.TypeName = InTypeName;
@@ -260,40 +423,55 @@ namespace Metasound
 			Description.PointIDs.Add(GraphHandle->GetNewPointID());
 
 			Frontend::FNodeHandle NodeHandle = GraphHandle->AddOutputVertex(Description);
+			NodeHandle->SetNodeStyle(InNodeStyle);
 
 			GraphHandle->SetOutputDisplayName(InName, FText::FromString(InName));
 
 			return NodeHandle;
 		}
 
-		void FGraphBuilder::DeleteNode(UEdGraphNode& InNode, Frontend::FNodeHandle InNodeHandle, bool bInRecordTransaction)
+		bool FGraphBuilder::ConnectNodes(UEdGraphPin& InInputPin, UEdGraphPin& InOutputPin)
 		{
-			const FScopedTransaction Transaction(LOCTEXT("DeleteMetasoundGraphNode", "Delete Metasound Node"), bInRecordTransaction);
+			using namespace Metasound::Frontend;
 
-			UMetasoundEditorGraph* Graph = CastChecked<UMetasoundEditorGraph>(InNode.GetGraph());
-			if (InNode.CanUserDeleteNode() && Graph->RemoveNode(&InNode))
+			UMetasoundEditorGraphNode* InputGraphNode = CastChecked<UMetasoundEditorGraphNode>(InInputPin.GetOwningNode());
+			FNodeHandle InputNodeHandle = InputGraphNode->GetNodeHandle();
+			TArray<FInputHandle> InputHandles = InputNodeHandle->GetInputsWithVertexName(InInputPin.GetName());
+
+			UMetasoundEditorGraphNode* OutputGraphNode = CastChecked<UMetasoundEditorGraphNode>(InOutputPin.GetOwningNode());
+			FNodeHandle OutputNodeHandle = OutputGraphNode->GetNodeHandle();
+			TArray<FOutputHandle> OutputHandles = OutputNodeHandle->GetOutputsWithVertexName(InOutputPin.GetName());
+
+			if (!ensure(InputHandles.Num() == 1 && OutputHandles.Num() == 1))
 			{
-				Graph->PostEditChange();
-				Graph->MarkPackageDirty();
+				InInputPin.BreakLinkTo(&InOutputPin);
+				return false;
 			}
 
-			if (InNodeHandle->IsValid())
-			{
-				FMetasoundAssetBase* MetasoundAsset = IMetasoundUObjectRegistry::Get().GetObjectAsAssetBase(Graph->GetMetasound());
-				check(MetasoundAsset);
+			FInputHandle InputHandle = InputHandles[0];
+			FOutputHandle OutputHandle = OutputHandles[0];
 
-				Frontend::FGraphHandle GraphHandle = MetasoundAsset->GetRootGraphHandle();
-				if (GraphHandle->IsValid())
+			FOutputHandle ExistingOutput = InputHandle->GetCurrentlyConnectedOutput();
+			if (ExistingOutput->IsValid())
+			{
+				FNodeHandle NodeHandle = ExistingOutput->GetOwningNode();
+				const FMetasoundFrontendNodeStyle& NodeStyle = NodeHandle->GetNodeStyle();
+				if (NodeHandle->IsValid() && NodeStyle.Display.Visibility == EMetasoundFrontendNodeStyleDisplayVisibility::Hidden)
 				{
-					GraphHandle->RemoveNode(*InNodeHandle);
+					GraphBuilderPrivate::DeleteNode(InputGraphNode->GetMetasoundChecked(), NodeHandle);
 				}
 			}
 
-			InNode.PostEditChange();
-			InNode.MarkPackageDirty();
+			if (!ensure(InputHandle->Connect(*OutputHandle)))
+			{
+				InInputPin.BreakLinkTo(&InOutputPin);
+				return false;
+			}
+
+			return true;
 		}
 
-		void FGraphBuilder::RebuildGraph(UObject& InMetasound)
+		void FGraphBuilder::ConstructGraph(UObject& InMetasound)
 		{
 			using namespace Frontend;
 
@@ -321,23 +499,25 @@ namespace Metasound
 			for (FNodeHandle& NodeHandle : NodeHandles)
 			{
 				const EMetasoundFrontendClassType NodeType = NodeHandle->GetClassType();
-				FVector2D Location;
+				FMetasoundFrontendNodeStyle Style;
 				if (NodeType == EMetasoundFrontendClassType::Input)
 				{
-					Location = InputNodeLocation;
+					Style.Display.Location = InputNodeLocation;
 					InputNodeLocation.Y += 100.0f;
 				}
 				else if (NodeType == EMetasoundFrontendClassType::Output)
 				{
-					Location = OutputNodeLocation;
+					Style.Display.Location = OutputNodeLocation;
 					OutputNodeLocation.Y += 100.0f;
 				}
 				else
 				{
-					Location = OpNodeLocation;
+					Style.Display.Location = OpNodeLocation;
 					OpNodeLocation.Y += 100.0f;
 				}
-				UEdGraphNode* NewNode = AddNode(InMetasound, Location, NodeHandle, false /* bInSelectNewNode */);
+				NodeHandle->SetNodeStyle(Style);
+
+				UEdGraphNode* NewNode = AddNode(InMetasound, NodeHandle, false /* bInSelectNewNode */);
 				NewIdNodeMap.Add(NodeHandle->GetID(), FNodePair { NodeHandle, NewNode });
 			}
 
@@ -381,12 +561,120 @@ namespace Metasound
 			InMetasound.MarkPackageDirty();
 		}
 
-		void FGraphBuilder::RebuildNodePins(UMetasoundEditorGraphNode& InGraphNode, Frontend::FNodeHandle InNodeHandle, bool bInRecordTransaction)
+		void FGraphBuilder::DeleteLiteralInputs(UEdGraphNode& InNode)
+		{
+			using namespace Metasound::Editor;
+			using namespace Metasound::Frontend;
+
+			FNodeHandle NodeHandle = Frontend::INodeController::GetInvalidHandle();
+			if (UMetasoundEditorGraphNode* Node = Cast<UMetasoundEditorGraphNode>(&InNode))
+			{
+				NodeHandle = Node->GetNodeHandle();
+
+				if (NodeHandle->GetClassType() == EMetasoundFrontendClassType::External)
+				{
+					TArray<FInputHandle> Inputs = NodeHandle->GetInputs();
+					for (FInputHandle& Input : Inputs)
+					{
+						FOutputHandle Output = Input->GetCurrentlyConnectedOutput();
+						if (Output->IsValid())
+						{
+							FNodeHandle InputHandle = Output->GetOwningNode();
+							if (InputHandle->GetNodeStyle().Display.Visibility == EMetasoundFrontendNodeStyleDisplayVisibility::Hidden)
+							{
+								FGraphHandle Graph = InputHandle->GetOwningGraph();
+								Graph->RemoveNode(*InputHandle);
+							}
+						}
+					}
+				}
+			}
+		}
+
+		bool FGraphBuilder::DeleteNode(UEdGraphNode& InNode, bool bInRecordTransaction)
 		{
-			const FScopedTransaction Transaction(LOCTEXT("RebuildMetasoundGraphNodePins", "Rebuild Metasound Pins"), bInRecordTransaction);
+			using namespace Metasound::Frontend;
+
+			const FScopedTransaction Transaction(LOCTEXT("DeleteMetasoundGraphNode", "Delete Metasound Node"), bInRecordTransaction);
+
+			FNodeHandle NodeHandle = Frontend::INodeController::GetInvalidHandle();
+			if (UMetasoundEditorGraphNode* Node = Cast<UMetasoundEditorGraphNode>(&InNode))
+			{
+				NodeHandle = Node->GetNodeHandle();
+
+				if (NodeHandle->GetClassType() == EMetasoundFrontendClassType::Input)
+				{
+					FConstDocumentHandle DocumentHandle = NodeHandle->GetOwningGraph()->GetOwningDocument();
+					auto IsRequiredInput = [&](const Frontend::FConstInputHandle& InputHandle)
+					{
+						return DocumentHandle->IsRequiredInput(InputHandle->GetName());
+					};
+					TArray<Frontend::FConstInputHandle> NodeInputs = NodeHandle->GetConstInputs();
+
+					if (Frontend::FConstInputHandle* InputHandle = NodeInputs.FindByPredicate(IsRequiredInput))
+					{
+						FNotificationInfo Info(FText::Format(LOCTEXT("Metasounds_CannotDeleteRequiredInput",
+							"'Required Input '{0}' cannot be deleted."), (*InputHandle)->GetDisplayName()));
+						Info.bFireAndForget = true;
+						Info.ExpireDuration = 2.0f;
+						Info.bUseThrobber = true;
+						FSlateNotificationManager::Get().AddNotification(Info);
+						return false;
+					}
+				}
+
+				if (NodeHandle->GetClassType() == EMetasoundFrontendClassType::Output)
+				{
+					FConstDocumentHandle DocumentHandle = NodeHandle->GetOwningGraph()->GetOwningDocument();
+					auto IsRequiredOutput = [&](const Frontend::FConstOutputHandle& OutputHandle)
+					{
+						return DocumentHandle->IsRequiredOutput(OutputHandle->GetName());
+					};
+					TArray<Frontend::FConstOutputHandle> NodeOutputs = NodeHandle->GetConstOutputs();
+
+					if (Frontend::FConstOutputHandle* OutputHandle = NodeOutputs.FindByPredicate(IsRequiredOutput))
+					{
+						FNotificationInfo Info(FText::Format(LOCTEXT("Metasounds_CannotDeleteRequiredOutput",
+							"'Required Output '{0}' cannot be deleted."), (*OutputHandle)->GetDisplayName()));
+						Info.bFireAndForget = true;
+						Info.ExpireDuration = 2.0f;
+						Info.bUseThrobber = true;
+						FSlateNotificationManager::Get().AddNotification(Info);
+						return false;
+					}
+				}
+			}
+
+			DeleteLiteralInputs(InNode);
+
+			UMetasoundEditorGraph* Graph = CastChecked<UMetasoundEditorGraph>(InNode.GetGraph());
+			if (InNode.CanUserDeleteNode() && Graph->RemoveNode(&InNode))
+			{
+				Graph->PostEditChange();
+				Graph->MarkPackageDirty();
+			}
+
+			if (NodeHandle->IsValid())
+			{
+				Frontend::FGraphHandle GraphHandle = NodeHandle->GetOwningGraph();
+				if (GraphHandle->IsValid())
+				{
+					GraphHandle->RemoveNode(*NodeHandle);
+				}
+			}
+
+			InNode.PostEditChange();
+			InNode.MarkPackageDirty();
+			return true;
+		}
+
+		void FGraphBuilder::RebuildNodePins(UMetasoundEditorGraphNode& InGraphNode, Frontend::FNodeHandle InNodeHandle)
+		{
+			const FScopedTransaction Transaction(LOCTEXT("RebuildMetasoundGraphNodePins", "Rebuild Metasound Pins"));
 
 			for (int32 i = InGraphNode.Pins.Num() - 1; i >= 0; i--)
 			{
+				DeleteLiteralInputs(InGraphNode);
 				InGraphNode.RemovePin(InGraphNode.Pins[i]);
 			}
 
@@ -423,9 +711,7 @@ namespace Metasound
 					if (InEditorPin.PinType == EditorModule.FindDataType(InInputHandle->GetDataType()).PinType)
 					{
 						UEdGraphNode* EditorNode = InEditorPin.GetOwningNode();
-						UMetasoundEditorGraphNode* MetasoundEditorNode = Cast<UMetasoundEditorGraphNode>(EditorNode);
-
-						if (nullptr != MetasoundEditorNode)
+						if (UMetasoundEditorGraphNode* MetasoundEditorNode = Cast<UMetasoundEditorGraphNode>(EditorNode))
 						{
 							if  (MetasoundEditorNode->GetNodeHandle()->GetID() == InInputHandle->GetOwningNodeID())
 							{
@@ -466,6 +752,9 @@ namespace Metasound
 
 		UEdGraphPin* FGraphBuilder::AddPinToNode(UMetasoundEditorGraphNode& InEditorNode, Frontend::FInputHandle InInputHandle)
 		{
+			using namespace Metasound::Editor;
+			using namespace Metasound::Frontend;
+
 			IMetasoundEditorModule& EditorModule = FModuleManager::GetModuleChecked<IMetasoundEditorModule>("MetasoundEditor");
 
 			FEdGraphPinType PinType("MetasoundNode", NAME_None, nullptr, EPinContainerType::None, false, FEdGraphTerminalType());
@@ -475,6 +764,12 @@ namespace Metasound
 			{
 				NewPin->PinToolTip = InInputHandle->GetTooltip().ToString();
 				NewPin->PinType = EditorModule.FindDataType(InInputHandle->GetDataType()).PinType;
+
+				FNodeHandle NodeHandle = InInputHandle->GetOwningNode();
+				if (NodeHandle->GetClassType() == EMetasoundFrontendClassType::External)
+				{
+					FGraphBuilder::AddOrUpdateLiteralInput(InEditorNode.GetMetasoundChecked(), NodeHandle, *NewPin);
+				}
 			}
 
 			return NewPin;
@@ -506,34 +801,28 @@ namespace Metasound
 			FMetasoundAssetBase* MetasoundAsset = IMetasoundUObjectRegistry::Get().GetObjectAsAssetBase(&InMetasound);
 			check(MetasoundAsset);
 
-			// Get all nodes from metasound graph
+			// Get all nodes from frontend graph
 			FGraphHandle GraphHandle = MetasoundAsset->GetRootGraphHandle();
-			TArray<FNodeHandle> Nodes = GraphHandle->GetNodes();
+			TArray<FNodeHandle> FrontendNodes = GraphHandle->GetNodes();
 
-			// Get all nodes from metasound editor graph
+			// Get all editor nodes from editor graph (some nodes on graph may *NOT* be metasound ed nodes,
+			// such as comment boxes, etc, so just get nodes of class UMetasoundEditorGraph).
 			UMetasoundEditorGraph* EditorGraph = CastChecked<UMetasoundEditorGraph>(MetasoundAsset->GetGraph());
 			TArray<UMetasoundEditorGraphNode*> EditorNodes;
-			for (UEdGraphNode* EdGraphNode : EditorGraph->Nodes)
-			{
-				UMetasoundEditorGraphNode* EditorNode = Cast<UMetasoundEditorGraphNode>(EdGraphNode);
-				if (ensure(nullptr != EditorNode))
-				{
-					EditorNodes.Add(EditorNode);
-				}
-			}
+			EditorGraph->GetNodesOfClass(EditorNodes);
 
 			// Find existing pairs of nodes and editor nodes
 			struct FNodePair
 			{
 				UMetasoundEditorGraphNode* EditorNode = nullptr;
-				FNodeHandle Node = INodeController::GetInvalidHandle();;
+				FNodeHandle Node = INodeController::GetInvalidHandle();
 			};
 			TMap<FGuid, FNodePair> PairedNodes;
 
 			// Reverse iterate so paired nodes can safely be removed from the array.
-			for (int32 i = Nodes.Num() - 1; i >= 0; i--)
+			for (int32 i = FrontendNodes.Num() - 1; i >= 0; i--)
 			{
-				FNodeHandle Node = Nodes[i];
+				FNodeHandle Node = FrontendNodes[i];
 				auto IsEditorNodeWithSameNodeID = [&](const UMetasoundEditorGraphNode* InEditorNode)
 				{
 					return InEditorNode->GetNodeID() == Node->GetID();
@@ -545,16 +834,16 @@ namespace Metasound
 					EditorNode = *PointerEditorNode;
 				}
 
-				if (nullptr != EditorNode)
+				if (EditorNode)
 				{
 					PairedNodes.Add(Node->GetID(), { EditorNode, Node });
 					EditorNodes.RemoveSwap(EditorNode);
-					Nodes.RemoveAtSwap(i);
+					FrontendNodes.RemoveAtSwap(i);
 				}
 			}
 
-			// Nodes contains nodes which need to be added.
-			// EditorNodes contains nodes that need to be removed. 
+			// FrontendNodes contains nodes which need to be added.
+			// EditorNodes contains nodes that need to be removed.
 			// PairedNodes contains pairs which we have to check have synchronized pins
 
 			// Add and remove nodes first in order to make sure correct editor nodes
@@ -564,12 +853,16 @@ namespace Metasound
 				bIsEditorGraphDirty |= EditorGraph->RemoveNode(EditorNode);
 			}
 
-			// Add missing nodes. 
-			bIsEditorGraphDirty |= (Nodes.Num() > 0);
-			for (FNodeHandle Node : Nodes)
+			// Add missing editor nodes marked as visible.
+			bIsEditorGraphDirty |= (FrontendNodes.Num() > 0);
+			for (FNodeHandle Node : FrontendNodes)
 			{
-				UEdGraphNode* NewNode = AddNode(InMetasound, EditorGraph->GetGoodPlaceForNewNode(), Node, false /* bInSelectNewNode */);
-				PairedNodes.Add(Node->GetID(), {Cast<UMetasoundEditorGraphNode>(NewNode), Node});
+				const FMetasoundFrontendNodeStyle& Style = Node->GetNodeStyle();
+				if (Style.Display.Visibility == EMetasoundFrontendNodeStyleDisplayVisibility::Visible)
+				{
+					UEdGraphNode* NewNode = AddNode(InMetasound, Node, false /* bInSelectNewNode */);
+					PairedNodes.Add(Node->GetID(), { Cast<UMetasoundEditorGraphNode>(NewNode), Node });
+				}
 			}
 
 			// Synchronize pins on node pairs.
@@ -685,8 +978,7 @@ namespace Metasound
 			TMap<FGuid, UMetasoundEditorGraphNode*> EditorNodesByID;
 			for (UEdGraphNode* EdGraphNode : EditorGraph->Nodes)
 			{
-				UMetasoundEditorGraphNode* MetasoundEditorNode = Cast<UMetasoundEditorGraphNode>(EdGraphNode);
-				if (nullptr != MetasoundEditorNode)
+				if (UMetasoundEditorGraphNode* MetasoundEditorNode = Cast<UMetasoundEditorGraphNode>(EdGraphNode))
 				{
 					EditorNodes.Add(MetasoundEditorNode);
 					EditorNodesByID.Add(MetasoundEditorNode->GetNodeID(), MetasoundEditorNode);
@@ -698,7 +990,14 @@ namespace Metasound
 			{
 				bool bIsNodeDirty = false;
 
+				// Get all editor nodes from editor graph (some nodes on graph may *NOT* be metasound ed nodes,
+				// such as comment boxes, etc, so just get nodes of class UMetasoundEditorGraph).
 				UMetasoundEditorGraphNode* MetasoundEditorNode = Cast<UMetasoundEditorGraphNode>(EdGraphNode);
+				if (!MetasoundEditorNode)
+				{
+					continue;
+				}
+
 				Frontend::FNodeHandle Node = MetasoundEditorNode->GetNodeHandle();
 
 				TArray<UEdGraphPin*> Pins = MetasoundEditorNode->GetAllPins();
@@ -717,17 +1016,30 @@ namespace Metasound
 						MatchingPin = *DoublePointer;
 					}
 
-					if (ensure(nullptr != MatchingPin))
+					if (ensure(MatchingPin))
 					{
-						// Remove pin so it isn't used twice. 
+						// Remove pin so it isn't used twice.
 						Pins.Remove(MatchingPin);
 
+						bool bShowConnectionInEditor = true;
+
 						Frontend::FOutputHandle OutputHandle = NodeInput->GetCurrentlyConnectedOutput();
 						if (OutputHandle->IsValid())
+						{
+							Frontend::FNodeHandle InputNodeHandle = OutputHandle->GetOwningNode();
+							const FMetasoundFrontendNodeStyle& InputNodeStyle = InputNodeHandle->GetNodeStyle();
+							bShowConnectionInEditor = InputNodeStyle.Display.Visibility == EMetasoundFrontendNodeStyleDisplayVisibility::Visible;
+						}
+						else
+						{
+							bShowConnectionInEditor = false;
+						}
+						
+						if (bShowConnectionInEditor)
 						{
 							bool bAddLink = false;
 
-							if (MatchingPin->LinkedTo.Num() < 1)
+							if (MatchingPin->LinkedTo.IsEmpty())
 							{
 								// No link currently exists. Add the appropriate link.
 								bAddLink = true;
@@ -741,15 +1053,16 @@ namespace Metasound
 
 							if (bAddLink)
 							{
-								UMetasoundEditorGraphNode* OutputEditorNode = EditorNodesByID[OutputHandle->GetOwningNodeID()];
+								const FGuid InputNodeID = OutputHandle->GetOwningNodeID();
+								UMetasoundEditorGraphNode* OutputEditorNode = EditorNodesByID[InputNodeID];
 								UEdGraphPin* OutputPin = OutputEditorNode->FindPinChecked(OutputHandle->GetName(), EEdGraphPinDirection::EGPD_Output);
 								MatchingPin->MakeLinkTo(OutputPin);
 								bIsNodeDirty = true;
 							}
 						}
-						else if (MatchingPin->LinkedTo.Num() > 0)
+						// No link should exist.
+						else
 						{
-							// No link should exist.
 							MatchingPin->BreakAllPinLinks();
 							bIsNodeDirty = true;
 						}
diff --git a/Engine/Plugins/Runtime/Metasound/Source/MetasoundEditor/Private/MetasoundEditorGraphBuilder.h b/Engine/Plugins/Runtime/Metasound/Source/MetasoundEditor/Private/MetasoundEditorGraphBuilder.h
--- a/Engine/Plugins/Runtime/Metasound/Source/MetasoundEditor/Private/MetasoundEditorGraphBuilder.h
+++ b/Engine/Plugins/Runtime/Metasound/Source/MetasoundEditor/Private/MetasoundEditorGraphBuilder.h
@@ -7,6 +7,7 @@
 #include "EdGraph/EdGraphPin.h"
 #include "MetasoundFrontend.h"
 #include "MetasoundFrontendController.h"
+#include "MetasoundFrontendDocument.h"
 
 
 // Forward Declarations
@@ -15,33 +16,44 @@ class UEdGraphPin;
 class UMetasound;
 class UMetasoundEditorGraphNode;
 
+struct FMetasoundFrontendNodeStyle;
+
+
 namespace Metasound
 {
+	// Forward Declarations
+	struct FLiteral;
+
 	namespace Editor
 	{
 		class FGraphBuilder
 		{
 		public:
+			static const FName PinCategoryBoolean;
+			static const FName PinCategoryDouble;
+			static const FName PinCategoryExec;
+			static const FName PinCategoryFloat;
+			static const FName PinCategoryInt32;
+			static const FName PinCategoryInt64;
+			static const FName PinCategoryObject;
+			static const FName PinCategoryString;
+
 			// Custom pin-related styles (ex. wire color, pin heads, etc.)
-			static const FName PinAudioFormat;				// Audio formats (ex. Buffers, Mono, Stereo)
-			static const FName PinAudioNumeric;				// Audio numeric types (ex. Frequency, Time)
-			static const FName PinPrimitiveBoolean;
-			static const FName PinPrimitiveFloat;
-			static const FName PinPrimitiveInt32;
-			static const FName PinPrimitiveInt64;
-			static const FName PinPrimitiveString;
-			static const FName PinPrimitiveTrigger;
-			static const FName PinPrimitiveUObject;
-			static const FName PinPrimitiveUObjectArray;
+			static const FName PinSubCategoryAudioFormat;				// Audio formats (ex. Buffers, Mono, Stereo)
+			static const FName PinSubCategoryAudioNumeric;				// Audio numeric types (ex. Frequency, Time)
+			static const FName PinSubCategoryObjectArray;
 
 			// Adds a node to the editor graph that corresponds to the provided node handle.
-			static UEdGraphNode* AddNode(UObject& InMetasound, const FVector2D& Location, Frontend::FNodeHandle& InNodeHandle, bool bInSelectNewNode = true);
+			static UEdGraphNode* AddNode(UObject& InMetasound, Frontend::FNodeHandle& InNodeHandle, bool bInSelectNewNode = true);
 
 			// Adds a node with the given class info to both the editor and document graphs
-			static UEdGraphNode* AddNode(UObject& InMetasound, const FVector2D& Location, const Frontend::FNodeClassInfo& InClassInfo, bool bInSelectNewNode = true);
+			static UEdGraphNode* AddNode(UObject& InMetasound, const Frontend::FNodeClassInfo& InClassInfo, const FMetasoundFrontendNodeStyle& InNodeStyle, bool bInSelectNewNode = true);
 
-			// Adds a node handle with the given class info
-			static Frontend::FNodeHandle AddNodeHandle(UObject& InMetasound, const Frontend::FNodeClassInfo& InClassInfo);
+			// Adds a node handle with the given class and style info
+			static Frontend::FNodeHandle AddNodeHandle(UObject& InMetasound, const Frontend::FNodeClassInfo& InClassInfo, const FMetasoundFrontendNodeStyle& InNodeStyle);
+
+			// Attempts to connect graph nodes together.  Returns true if succeeded, breaks pin link and returns false if failed.
+			static bool ConnectNodes(UEdGraphPin& InInputPin, UEdGraphPin& InOutputPin);
 
 			static FString GenerateUniqueInputName(const UObject& InMetasound, const FName InBaseName);
 
@@ -51,21 +63,28 @@ namespace Metasound
 
 			static TArray<FString> GetDataTypeNameCategories(const FName& InDataTypeName);
 
-			static UEdGraphNode* AddInput(UObject& InMetasound, const FVector2D& Location, const FString& InName, const FName InTypeName, const FText& InToolTip, bool bInSelectNewNode = true);
+			static UEdGraphNode* AddInput(UObject& InMetasound, const FString& InName, const FName InTypeName, const FMetasoundFrontendNodeStyle& InNodeStyle, const FText& InToolTip, bool bInSelectNewNode = true);
+
+			static void AddOrUpdateLiteralInput(UObject& InMetasound, Frontend::FNodeHandle InNodeHandle, const UEdGraphPin& InInputPin);
 
 			// Adds an input node handle with the given class info
-			static Frontend::FNodeHandle AddInputNodeHandle(UObject& InMetasound, const FString& InName, const FName InTypeName, const FText& InToolTip);
+			static Frontend::FNodeHandle AddInputNodeHandle(UObject& InMetasound, const FString& InName, const FName InTypeName, const FMetasoundFrontendNodeStyle& InNodeStyle, const FText* InToolTip = nullptr, const FLiteral* InDefaultValue = nullptr);
 
-			static UEdGraphNode* AddOutput(UObject& InMetasound, const FVector2D& Location, const FString& InName, const FName InTypeName, const FText& InToolTip, bool bInSelectNewNode = true);
+			static UEdGraphNode* AddOutput(UObject& InMetasound, const FString& InName, const FName InTypeName, const FMetasoundFrontendNodeStyle& InNodeStyle, const FText& InToolTip, bool bInSelectNewNode = true);
 
 			// Adds an output node handle with the given class info
-			static Frontend::FNodeHandle AddOutputNodeHandle(UObject& InMetasound, const FString& InName, const FName InTypeName, const FText& InToolTip);
+			static Frontend::FNodeHandle AddOutputNodeHandle(UObject& InMetasound, const FString& InName, const FName InTypeName, const FMetasoundFrontendNodeStyle& InNodeStyle, const FText& InToolTip);
+
+			// Constructs graph with default inputs & outputs.
+			static void ConstructGraph(UObject& InMetasound);
 
-			static void DeleteNode(UEdGraphNode& InNode, Frontend::FNodeHandle InNodeHandle, bool bInRecordTransaction = true);
+			static void RebuildNodePins(UMetasoundEditorGraphNode& InGraphNode, Frontend::FNodeHandle InNodeHandle);
 
-			static void RebuildGraph(UObject& InMetasound);
+			// Removes all literal inputs connected to the given node
+			static void DeleteLiteralInputs(UEdGraphNode& InNode);
 
-			static void RebuildNodePins(UMetasoundEditorGraphNode& InGraphNode, Frontend::FNodeHandle InNodeHandle, bool bInRecordTransaction = true);
+			// Deletes both the editor graph & frontend nodes from respective graphs
+			static bool DeleteNode(UEdGraphNode& InNode, bool bInRecordTransaction);
 
 			// Adds an Input UEdGraphPin to a UMetasoundEditorGraphNode
 			static UEdGraphPin* AddPinToNode(UMetasoundEditorGraphNode& InEditorNode, Frontend::FInputHandle InInputHandle);
diff --git a/Engine/Plugins/Runtime/Metasound/Source/MetasoundEditor/Private/MetasoundEditorGraphNode.cpp b/Engine/Plugins/Runtime/Metasound/Source/MetasoundEditor/Private/MetasoundEditorGraphNode.cpp
--- a/Engine/Plugins/Runtime/Metasound/Source/MetasoundEditor/Private/MetasoundEditorGraphNode.cpp
+++ b/Engine/Plugins/Runtime/Metasound/Source/MetasoundEditor/Private/MetasoundEditorGraphNode.cpp
@@ -8,11 +8,12 @@
 #include "Framework/Commands/GenericCommands.h"
 #include "GraphEditorActions.h"
 #include "Metasound.h"
-#include "MetasoundFrontend.h"
 #include "MetasoundEditorGraph.h"
 #include "MetasoundEditorGraphBuilder.h"
 #include "MetasoundEditorGraphSchema.h"
 #include "MetasoundEditorCommands.h"
+#include "MetasoundEditorModule.h"
+#include "MetasoundFrontend.h"
 #include "MetasoundUObjectRegistry.h"
 #include "ScopedTransaction.h"
 #include "ToolMenus.h"
@@ -122,12 +123,13 @@ void UMetasoundEditorGraphNode::AllocateDefaultPins()
 	using namespace Metasound;
 
 	ensureAlways(Pins.Num() == 0);
-	Metasound::Editor::FGraphBuilder::RebuildNodePins(*this, GetNodeHandle(), false /* bInRecordTransaction */);
+	Editor::FGraphBuilder::RebuildNodePins(*this, GetNodeHandle());
 }
 
 void UMetasoundEditorGraphNode::ReconstructNode()
 {
 	using namespace Metasound;
+
 	Editor::FGraphBuilder::RebuildNodePins(*this, GetNodeHandle());
 }
 
@@ -255,6 +257,20 @@ FText UMetasoundEditorGraphNode::GetNodeTitle(ENodeTitleType::Type TitleType) co
 	}
 }
 
+void UMetasoundEditorGraphNode::PinDefaultValueChanged(UEdGraphPin* Pin)
+{
+	using namespace Metasound;
+	using namespace Metasound::Editor;
+	using namespace Metasound::Frontend;
+
+	if (Pin && Pin->Direction == EGPD_Input)
+	{
+		FNodeHandle NodeHandle = GetNodeHandle();
+		IMetasoundEditorModule& EditorModule = FModuleManager::GetModuleChecked<IMetasoundEditorModule>("MetasoundEditor");
+		FGraphBuilder::AddOrUpdateLiteralInput(GetMetasoundChecked(), NodeHandle, *Pin);
+	}
+}
+
 void UMetasoundEditorGraphNode::PrepareForCopying()
 {
 }
@@ -263,6 +279,19 @@ void UMetasoundEditorGraphNode::PostEditImport()
 {
 }
 
+void UMetasoundEditorGraphNode::PostEditChangeProperty(FPropertyChangedEvent& InEvent)
+{
+	const FName PropertyName = InEvent.GetPropertyName();
+	if (PropertyName == GET_MEMBER_NAME_CHECKED(UEdGraphNode, NodePosX)
+		|| PropertyName == GET_MEMBER_NAME_CHECKED(UEdGraphNode, NodePosY))
+	{
+		Metasound::Frontend::FNodeHandle NodeHandle = GetNodeHandle();
+		FMetasoundFrontendNodeStyle Style = NodeHandle->GetNodeStyle();
+		Style.Display.Location = FVector2D(NodePosX, NodePosY);
+		GetNodeHandle()->SetNodeStyle(Style);
+	}
+}
+
 void UMetasoundEditorGraphNode::PostDuplicate(bool bDuplicateForPIE)
 {
 	Super::PostDuplicate(bDuplicateForPIE);
@@ -280,7 +309,7 @@ void UMetasoundEditorGraphNode::GetNodeContextMenuActions(UToolMenu* Menu, UGrap
 	if (Context->Pin)
 	{
 		// If on an input that can be deleted, show option
-		if (Context->Pin->Direction == EGPD_Input /*&& SoundNode->ChildNodes.Num() > SoundNode->GetMinChildNodes()*/)
+		if (Context->Pin->Direction == EGPD_Input)
 		{
 			FToolMenuSection& Section = Menu->AddSection("MetasoundEditorGraphDeleteInput");
 			Section.AddMenuEntry(FEditorCommands::Get().DeleteInput);
diff --git a/Engine/Plugins/Runtime/Metasound/Source/MetasoundEditor/Private/MetasoundEditorGraphSchema.cpp b/Engine/Plugins/Runtime/Metasound/Source/MetasoundEditor/Private/MetasoundEditorGraphSchema.cpp
--- a/Engine/Plugins/Runtime/Metasound/Source/MetasoundEditor/Private/MetasoundEditorGraphSchema.cpp
+++ b/Engine/Plugins/Runtime/Metasound/Source/MetasoundEditor/Private/MetasoundEditorGraphSchema.cpp
@@ -4,6 +4,7 @@
 
 #include "EdGraph/EdGraph.h"
 #include "EdGraphNode_Comment.h"
+#include "EdGraphSchema_K2.h"
 #include "Editor.h"
 #include "Engine/Selection.h"
 #include "Framework/MultiBox/MultiBoxBuilder.h"
@@ -13,6 +14,7 @@
 #include "Layout/SlateRect.h"
 #include "Metasound.h"
 #include "MetasoundAssetBase.h"
+#include "MetasoundDataReference.h"
 #include "MetasoundEditor.h"
 #include "MetasoundEditorGraph.h"
 #include "MetasoundEditorGraphBuilder.h"
@@ -50,50 +52,56 @@ namespace Metasound
 			const UGraphEditorSettings* Settings = GetDefault<UGraphEditorSettings>();
 			check(Settings);
 
-			if (PinType.PinCategory == FGraphBuilder::PinPrimitiveTrigger)
+			if (PinType.PinSubCategory == FGraphBuilder::PinSubCategoryAudioNumeric)
+			{
+				return Settings->SoftObjectPinTypeColor;
+			}
+
+			if (PinType.PinSubCategory == FGraphBuilder::PinSubCategoryAudioFormat)
+			{
+				return Settings->SoftClassPinTypeColor;
+			}
+
+			if (PinType.PinCategory == FGraphBuilder::PinCategoryExec)
 			{
 				return Settings->ExecutionPinTypeColor;
 			}
 
-			if (PinType.PinCategory == FGraphBuilder::PinPrimitiveBoolean)
+			if (PinType.PinCategory == FGraphBuilder::PinCategoryBoolean)
 			{
 				return Settings->BooleanPinTypeColor;
 			}
 
-			if (PinType.PinCategory == FGraphBuilder::PinPrimitiveFloat)
+			if (PinType.PinCategory == FGraphBuilder::PinCategoryFloat)
 			{
 				return Settings->FloatPinTypeColor;
 			}
 
-			if (PinType.PinCategory == FGraphBuilder::PinPrimitiveInt32)
+			if (PinType.PinCategory == FGraphBuilder::PinCategoryInt32)
 			{
 				return Settings->IntPinTypeColor;
 			}
 
-			if (PinType.PinCategory == FGraphBuilder::PinPrimitiveInt64)
+			if (PinType.PinCategory == FGraphBuilder::PinCategoryInt64)
 			{
 				return Settings->Int64PinTypeColor;
 			}
 
-			if (PinType.PinCategory == FGraphBuilder::PinPrimitiveString)
+			if (PinType.PinCategory == FGraphBuilder::PinCategoryString)
 			{
 				return Settings->StringPinTypeColor;
 			}
 
-			if (PinType.PinCategory == FGraphBuilder::PinPrimitiveUObject)
+			if (PinType.PinCategory == FGraphBuilder::PinCategoryDouble)
 			{
-				return Settings->ClassPinTypeColor;
+				return Settings->DoublePinTypeColor;
 			}
 
-			if (PinType.PinCategory == FGraphBuilder::PinAudioNumeric)
+			if (PinType.PinCategory == FGraphBuilder::PinCategoryObject)
 			{
-				return Settings->SoftObjectPinTypeColor;
+				return Settings->ClassPinTypeColor;
 			}
 
-			if (PinType.PinCategory == FGraphBuilder::PinAudioFormat)
-			{
-				return Settings->SoftClassPinTypeColor;
-			}
 
 			return Settings->StructPinTypeColor;
 		}
@@ -134,11 +142,6 @@ namespace Metasound
 			OutParams.AssociatedPin1 = InputPin;
 			OutParams.AssociatedPin2 = OutputPin;
 
-			// Get the schema and grab the default color from it
-			const UEdGraphSchema* Schema = GraphObj->GetSchema();
-			FNodeHandle InputNodeHandle = CastChecked<UMetasoundEditorGraphNode>(InputPin->GetOwningNode())->GetNodeHandle();
-			FNodeHandle OutputNodeHandle = CastChecked<UMetasoundEditorGraphNode>(OutputPin->GetOwningNode())->GetNodeHandle();
-
 			OutParams.WireColor = GetPinCategoryColor(OutputPin->PinType);
 			bool bExecuted = false;
 
@@ -156,11 +159,11 @@ namespace Metasound
 
 			if (!bExecuted)
 			{
-				if (InputPin->PinType.PinCategory == FGraphBuilder::PinPrimitiveTrigger)
+				if (InputPin->PinType.PinCategory == FGraphBuilder::PinCategoryExec)
 				{
 					OutParams.WireThickness = Settings->DefaultExecutionWireThickness;
 				}
-				else if (InputPin->PinType.PinCategory == FGraphBuilder::PinAudioFormat)
+				else if (InputPin->PinType.PinCategory == FGraphBuilder::PinSubCategoryAudioFormat)
 				{
 					OutParams.bDrawBubbles = true;
 				}
@@ -185,7 +188,10 @@ UEdGraphNode* FMetasoundGraphSchemaAction_NewNode::PerformAction(UEdGraph* Paren
 	UObject& ParentMetasound = CastChecked<UMetasoundEditorGraph>(ParentGraph)->GetMetasoundChecked();
 	ParentMetasound.Modify();
 
-	UEdGraphNode* NewGraphNode = FGraphBuilder::AddNode(ParentMetasound, Location, NodeClassInfo);
+	FMetasoundFrontendNodeStyle Style;
+	Style.Display.Location = Location;
+
+	UEdGraphNode* NewGraphNode = FGraphBuilder::AddNode(ParentMetasound, NodeClassInfo, Style);
 
 	return NewGraphNode;
 }
@@ -208,7 +214,11 @@ UEdGraphNode* FMetasoundGraphSchemaAction_NewInput::PerformAction(UEdGraph* Pare
 	ParentMetasound.Modify();
 
 	FString NewNodeName = FGraphBuilder::GenerateUniqueInputName(ParentMetasound, NodeTypeName);
-	UEdGraphNode* NewGraphNode = FGraphBuilder::AddInput(ParentMetasound, Location, NewNodeName, NodeTypeName, FText::GetEmpty());
+
+	FMetasoundFrontendNodeStyle Style;
+	Style.Display.Location = Location;
+
+	UEdGraphNode* NewGraphNode = FGraphBuilder::AddInput(ParentMetasound, NewNodeName, NodeTypeName, Style, FText::GetEmpty());
 	return NewGraphNode;
 }
 
@@ -231,7 +241,10 @@ UEdGraphNode* FMetasoundGraphSchemaAction_NewOutput::PerformAction(UEdGraph* Par
 
 	FString NewNodeName = FGraphBuilder::GenerateUniqueInputName(ParentMetasound, NodeTypeName);
 
-	UEdGraphNode* NewGraphNode = FGraphBuilder::AddOutput(ParentMetasound, Location, NewNodeName, NodeTypeName, FText::GetEmpty());
+	FMetasoundFrontendNodeStyle Style;
+	Style.Display.Location = Location;
+
+	UEdGraphNode* NewGraphNode = FGraphBuilder::AddOutput(ParentMetasound, NewNodeName, NodeTypeName, Style, FText::GetEmpty());
 	return NewGraphNode;
 }
 
@@ -429,6 +442,13 @@ const FPinConnectionResponse UMetasoundEditorGraphSchema::CanCreateConnection(co
 
 bool UMetasoundEditorGraphSchema::TryCreateConnection(UEdGraphPin* PinA, UEdGraphPin* PinB) const
 {
+	using namespace Metasound::Frontend;
+
+	if (!ensureAlways(PinA && PinB))
+	{
+		return false;
+	}
+
 	UEdGraphPin* InputPin = nullptr;
 	UEdGraphPin* OutputPin = nullptr;
 	if (!CategorizePinsByDirection(PinA, PinB, InputPin, OutputPin))
@@ -436,30 +456,19 @@ bool UMetasoundEditorGraphSchema::TryCreateConnection(UEdGraphPin* PinA, UEdGrap
 		return false;
 	}
 
-	const bool bModified = UEdGraphSchema::TryCreateConnection(PinA, PinB);
-	if (bModified)
+	if (!ensureAlways(InputPin && OutputPin))
 	{
-		UMetasoundEditorGraphNode* InputGraphNode = CastChecked<UMetasoundEditorGraphNode>(InputPin->GetOwningNode());
-		Metasound::Frontend::FNodeHandle InputNodeHandle = InputGraphNode->GetNodeHandle();
-		TArray<Metasound::Frontend::FInputHandle> InputHandles = InputNodeHandle->GetInputsWithVertexName(InputPin->GetName());
-
-		UMetasoundEditorGraphNode* OutputGraphNode = CastChecked<UMetasoundEditorGraphNode>(OutputPin->GetOwningNode());
-		Metasound::Frontend::FNodeHandle OutputNodeHandle = OutputGraphNode->GetNodeHandle();
-		TArray<Metasound::Frontend::FOutputHandle> OutputHandles = OutputNodeHandle->GetOutputsWithVertexName(OutputPin->GetName());
+		return false;
+	}
 
-		if (ensure((InputHandles.Num() == 1) && (OutputHandles.Num() == 1)))
-		{
-			Metasound::Frontend::FInputHandle InputHandle = InputHandles[0];
-			Metasound::Frontend::FOutputHandle OutputHandle = OutputHandles[0];
+	// TODO: Implement YesWithConverterNode with selected conversion option
 
-			// TODO: Implement YesWithConverterNode with selected conversion option
-			if (!ensure(InputHandle->Connect(*OutputHandle)))
-			{
-				InputPin->BreakLinkTo(PinB);
-				return false;
-			}
-		}
+	bool bModified = UEdGraphSchema::TryCreateConnection(PinA, PinB);
+	if (bModified)
+	{
+		bModified = Metasound::Editor::FGraphBuilder::ConnectNodes(*InputPin, *OutputPin);
 	}
+
 	return bModified;
 }
 
@@ -528,44 +537,34 @@ void UMetasoundEditorGraphSchema::BreakNodeLinks(UEdGraphNode& TargetNode) const
 {
 	using namespace Metasound::Frontend;
 
-	const FScopedTransaction Transaction(NSLOCTEXT("UnrealEd", "GraphEd_BreakPinLinks", "Break Pin Links"));
+	const FScopedTransaction Transaction(NSLOCTEXT("UnrealEd", "GraphEd_BreakNodeLinks", "Break Node Links"));
 
-	Super::BreakNodeLinks(TargetNode);
-
-	FNodeHandle NodeHandle = CastChecked<UMetasoundEditorGraphNode>(&TargetNode)->GetNodeHandle();
-	const FGuid NodeID = NodeHandle->GetID();
-
-	FGraphHandle GraphHandle = CastChecked<UMetasoundEditorGraphNode>(&TargetNode)->GetRootGraphHandle();
-	TArray<FNodeHandle> AllNodes = GraphHandle->GetNodes();
-	for (Metasound::Frontend::FNodeHandle& IterNode : AllNodes)
+	TArray<UEdGraphPin*> Pins = TargetNode.GetAllPins();
+	for (UEdGraphPin* Pin : Pins)
 	{
-		if (NodeID != IterNode->GetID())
-		{
-			TArray<FInputHandle> Inputs = IterNode->GetInputs();
-			for (FInputHandle& Input : Inputs)
-			{
-				FOutputHandle Output = Input->GetCurrentlyConnectedOutput();
-				if (Output->GetOwningNodeID() == NodeID)
-				{
-					Input->Disconnect(*Output);
-				}
-			}
-		}
+		BreakPinLinks(*Pin, false);
 	}
+
+	Super::BreakNodeLinks(TargetNode);
 }
 
 void UMetasoundEditorGraphSchema::BreakPinLinks(UEdGraphPin& TargetPin, bool bSendsNodeNotifcation) const
 {
+	using namespace Metasound::Editor;
 	using namespace Metasound::Frontend;
 
 	const FScopedTransaction Transaction(NSLOCTEXT("UnrealEd", "GraphEd_BreakPinLinks", "Break Pin Links"));
 
 	TArray<FInputHandle> InputHandles;
+	TArray<UEdGraphPin*> InputPins;
+
+	UObject& Metasound = CastChecked<UMetasoundEditorGraphNode>(TargetPin.GetOwningNode())->GetMetasoundChecked();
 
 	if (TargetPin.Direction == EGPD_Input)
 	{
 		FNodeHandle NodeHandle = CastChecked<UMetasoundEditorGraphNode>(TargetPin.GetOwningNode())->GetNodeHandle();
 		InputHandles = NodeHandle->GetInputsWithVertexName(TargetPin.GetName());
+		InputPins.Add(&TargetPin);
 	}
 	else
 	{
@@ -574,12 +573,21 @@ void UMetasoundEditorGraphSchema::BreakPinLinks(UEdGraphPin& TargetPin, bool bSe
 		{
 			FNodeHandle NodeHandle = CastChecked<UMetasoundEditorGraphNode>(Pin->GetOwningNode())->GetNodeHandle();
 			InputHandles.Append(NodeHandle->GetInputsWithVertexName(Pin->GetName()));
+			InputPins.Add(Pin);
 		}
 	}
 
-	for (FInputHandle Handle : InputHandles)
+	IMetasoundEditorModule& EditorModule = FModuleManager::GetModuleChecked<IMetasoundEditorModule>("MetasoundEditor");
+	for (int32 i = 0; i < InputHandles.Num(); ++i)
 	{
+		FInputHandle Handle = InputHandles[i];
 		Handle->Disconnect();
+
+		FNodeHandle NodeHandle = Handle->GetOwningNode();
+		if (NodeHandle->GetClassType() == EMetasoundFrontendClassType::External)
+		{
+			FGraphBuilder::AddOrUpdateLiteralInput(Metasound, NodeHandle, *InputPins[i]);
+		}
 	}
 
 	Super::BreakPinLinks(TargetPin, bSendsNodeNotifcation);
@@ -616,7 +624,7 @@ void UMetasoundEditorGraphSchema::GetAllMetasoundActions(FGraphActionMenuBuilder
 	EditorModule.IterateDataTypes([MenuJoinFormat, InputMenuName, OutputMenuName, InMenuBuilder = &ActionMenuBuilder](const FEditorDataType& DataType)
 	{
 		const FName DataTypeName = DataType.RegistryInfo.DataTypeName;
-		const FText DataTypeDisplayName = FText::FromString(FGraphBuilder::GetDataTypeDisplayName(DataTypeName).Replace(TEXT("Primitive:"), TEXT("")));
+		const FText DataTypeDisplayName = FText::FromString(FGraphBuilder::GetDataTypeDisplayName(DataTypeName));
 		const FText DataTypeTextName = FText::FromName(DataTypeName);
 
 		const TArray<FString> Categories = FGraphBuilder::GetDataTypeNameCategories(DataTypeName);
@@ -656,7 +664,7 @@ void UMetasoundEditorGraphSchema::GetAllMetasoundActions(FGraphActionMenuBuilder
 			? Metadata.Description
 			: FText::Format(LOCTEXT("MetasoundTooltipAuthorFormat", "{0}\nAuthor: {1}"), Metadata.Description, Metadata.Author);
 
-		const FString DisplayName = ClassInfo.NodeName.Replace(TEXT("Primitive:"), TEXT(""));
+		FString DisplayName = ClassInfo.NodeName;
 
 		const FText* MenuName = &NodeMenuName;
 		FText CategoriesText = FText::GetEmpty();
diff --git a/Engine/Plugins/Runtime/Metasound/Source/MetasoundEditor/Private/MetasoundEditorModule.cpp b/Engine/Plugins/Runtime/Metasound/Source/MetasoundEditor/Private/MetasoundEditorModule.cpp
--- a/Engine/Plugins/Runtime/Metasound/Source/MetasoundEditor/Private/MetasoundEditorModule.cpp
+++ b/Engine/Plugins/Runtime/Metasound/Source/MetasoundEditor/Private/MetasoundEditorModule.cpp
@@ -28,6 +28,7 @@
 #include "Styling/SlateStyleRegistry.h"
 #include "Styling/SlateTypes.h"
 #include "Templates/SharedPointer.h"
+#include "../../BlueprintGraph/Classes/EdGraphSchema_K2.h"
 
 DEFINE_LOG_CATEGORY(LogMetasoundEditor);
 
@@ -111,9 +112,9 @@ namespace Metasound
 				}
 			}
 
-			virtual void RegisterDataType(FName InPinCategoryName, const FDataTypeRegistryInfo& InRegistryInfo) override
+			virtual void RegisterDataType(FName InPinCategoryName, FName InPinSubCategoryName, const FDataTypeRegistryInfo& InRegistryInfo) override
 			{
-				DataTypeInfo.Add(InRegistryInfo.DataTypeName, FEditorDataType(InPinCategoryName, InRegistryInfo));
+				DataTypeInfo.Add(InRegistryInfo.DataTypeName, FEditorDataType(InPinCategoryName, InPinSubCategoryName, InRegistryInfo));
 			}
 
 			void RegisterCoreDataTypes()
@@ -124,42 +125,21 @@ namespace Metasound
 					FDataTypeRegistryInfo RegistryInfo;
 					Frontend::GetTraitsForDataType(DataTypeName, RegistryInfo);
 
-					FName PinType = DataTypeName;
-
-					// Special Flowers
+					FName PinCategory = DataTypeName;
+					FName PinSubCategory;
 
 					// Execution path triggers are specialized
-					if (DataTypeName == "Primitive:Trigger")
+					if (DataTypeName == "Trigger")
 					{
-						PinType = FGraphBuilder::PinPrimitiveTrigger;
+						PinCategory = FGraphBuilder::PinCategoryExec;
 					}
 
 					// GraphEditor by default designates specialized connection
-					// specification for Int64, so use it
-					else if (DataTypeName == "Primitive:Int64")
-					{
-						PinType = FGraphBuilder::PinPrimitiveInt64;
-					}
-
-					// Differentiate stronger numeric types associated with audio
-					else if (DataTypeName == "Primitive:Frequency"
-						|| DataTypeName == "Primitive:Time"
-						|| DataTypeName == "Primitive:Time:HighResolution"
-						|| DataTypeName == "Primitive:Time:SampleResolution"
-						)
+					// specification for Int64, so use it even though literal is
+					// boiled down to int32
+					else if (DataTypeName == "Int64")
 					{
-						PinType = FGraphBuilder::PinAudioNumeric;
-					}
-
-					// Audio types are ubiquitous, so specialize
-					else if (DataTypeName == "Audio:Buffer"
-						|| DataTypeName == "Audio:Unformatted"
-						|| DataTypeName == "Audio:Mono"
-						|| DataTypeName == "Audio:Stereo"
-						|| DataTypeName == "Audio:Multichannel"
-						)
-					{
-						PinType = FGraphBuilder::PinAudioFormat;
+						PinCategory = FGraphBuilder::PinCategoryInt64;
 					}
 
 					// Primitives
@@ -169,37 +149,54 @@ namespace Metasound
 						{
 							case ELiteralType::Boolean:
 							{
-								PinType = FGraphBuilder::PinPrimitiveBoolean;
+								PinCategory = FGraphBuilder::PinCategoryBoolean;
 							}
 							break;
 
 							case ELiteralType::Float:
 							{
-								PinType = FGraphBuilder::PinPrimitiveFloat;
+								PinCategory = FGraphBuilder::PinCategoryFloat;
+
+								// Doubles use the same preferred literal
+								// but different colorization
+								if (DataTypeName == "Double")
+								{
+									PinCategory = FGraphBuilder::PinCategoryFloat;
+								}
+
+								// Differentiate stronger numeric types associated with audio
+								if (DataTypeName == "Frequency"
+									|| DataTypeName == "Time"
+									|| DataTypeName == "Time:HighResolution"
+									|| DataTypeName == "Time:SampleResolution"
+								)
+								{
+									PinSubCategory = FGraphBuilder::PinSubCategoryAudioNumeric;
+								}
 							}
 							break;
 
 							case ELiteralType::Integer:
 							{
-								PinType = FGraphBuilder::PinPrimitiveInt32;
+								PinCategory = FGraphBuilder::PinCategoryInt32;
 							}
 							break;
 
 							case ELiteralType::String:
 							{
-								PinType = FGraphBuilder::PinPrimitiveString;
+								PinCategory = FGraphBuilder::PinCategoryString;
 							}
 							break;
 
 							case ELiteralType::UObjectProxy:
 							{
-								PinType = FGraphBuilder::PinPrimitiveUObject;
+								PinCategory = FGraphBuilder::PinCategoryObject;
 							}
 							break;
 
 							case ELiteralType::UObjectProxyArray:
 							{
-								PinType = FGraphBuilder::PinPrimitiveUObjectArray;
+								PinCategory = FGraphBuilder::PinSubCategoryObjectArray;
 							}
 							break;
 
@@ -208,13 +205,24 @@ namespace Metasound
 							case ELiteralType::None:
 							case ELiteralType::Invalid:
 							{
+
+								// Audio types are ubiquitous, so specialize
+								if (DataTypeName == "Audio:Buffer"
+									|| DataTypeName == "Audio:Unformatted"
+									|| DataTypeName == "Audio:Mono"
+									|| DataTypeName == "Audio:Stereo"
+									|| DataTypeName == "Audio:Multichannel"
+								)
+								{
+									PinSubCategory = FGraphBuilder::PinSubCategoryAudioFormat;
+								}
 								static_assert(static_cast<int32>(ELiteralType::Invalid) == 7, "Possible missing binding of pin category to primitive type");
 							}
 							break;
 						}
 					}
 
-					DataTypeInfo.Add(DataTypeName, FEditorDataType(PinType, RegistryInfo));
+					DataTypeInfo.Add(DataTypeName, FEditorDataType(PinCategory, PinSubCategory, RegistryInfo));
 				}
 			}
 
diff --git a/Engine/Plugins/Runtime/Metasound/Source/MetasoundEditor/Private/MetasoundEditorModule.h b/Engine/Plugins/Runtime/Metasound/Source/MetasoundEditor/Private/MetasoundEditorModule.h
--- a/Engine/Plugins/Runtime/Metasound/Source/MetasoundEditor/Private/MetasoundEditorModule.h
+++ b/Engine/Plugins/Runtime/Metasound/Source/MetasoundEditor/Private/MetasoundEditorModule.h
@@ -20,8 +20,8 @@ namespace Metasound
 			FEdGraphPinType PinType;
 			FDataTypeRegistryInfo RegistryInfo;
 
-			FEditorDataType(FName InGraphDataType, const FDataTypeRegistryInfo& InRegistryInfo)
-				: PinType(InGraphDataType, NAME_None, nullptr, EPinContainerType::None, false, FEdGraphTerminalType())
+			FEditorDataType(FName InPinCategory, FName InPinSubCategory, const FDataTypeRegistryInfo& InRegistryInfo)
+				: PinType(InPinCategory, InPinSubCategory, nullptr, EPinContainerType::None, false, FEdGraphTerminalType())
 				, RegistryInfo(InRegistryInfo)
 			{
 			}
@@ -30,7 +30,7 @@ namespace Metasound
 		class METASOUNDEDITOR_API IMetasoundEditorModule : public IModuleInterface
 		{
 		public:
-			virtual void RegisterDataType(FName InPinCategoryName, const FDataTypeRegistryInfo& InRegistryInfo) = 0;
+			virtual void RegisterDataType(FName InPinCategoryName, FName InPinSubCategoryName, const FDataTypeRegistryInfo& InRegistryInfo) = 0;
 			virtual const FEditorDataType& FindDataType(FName InDataTypeName) const = 0;
 			virtual void IterateDataTypes(TUniqueFunction<void(const FEditorDataType&)> InDataTypeFunction) const = 0;
 		};
diff --git a/Engine/Plugins/Runtime/Metasound/Source/MetasoundEditor/Private/SMetasoundGraphNode.cpp b/Engine/Plugins/Runtime/Metasound/Source/MetasoundEditor/Private/SMetasoundGraphNode.cpp
--- a/Engine/Plugins/Runtime/Metasound/Source/MetasoundEditor/Private/SMetasoundGraphNode.cpp
+++ b/Engine/Plugins/Runtime/Metasound/Source/MetasoundEditor/Private/SMetasoundGraphNode.cpp
@@ -69,14 +69,14 @@ void SMetasoundGraphNode::CreateStandardPinWidget(UEdGraphPin* CurPin)
 		Metasound::Frontend::FNodeHandle NodeHandle = GetMetasoundNode().GetNodeHandle();
 		if (CurPin->Direction == EGPD_Input)
 		{
-			if (!NodeHandle->GetClassDisplayInfo().bShowInputName)
+			if (!NodeHandle->GetClassStyle().Display.bShowInputNames)
 			{
 				NewPin->SetShowLabel(false);
 			}
 		}
 		else if (CurPin->Direction == EGPD_Output)
 		{
-			if (!NodeHandle->GetClassDisplayInfo().bShowOutputName)
+			if (!NodeHandle->GetClassStyle().Display.bShowOutputNames)
 			{
 				NewPin->SetShowLabel(false);
 			}
@@ -89,7 +89,7 @@ void SMetasoundGraphNode::CreateStandardPinWidget(UEdGraphPin* CurPin)
 TSharedRef<SWidget> SMetasoundGraphNode::CreateTitleWidget(TSharedPtr<SNodeTitle> NodeTitle)
 {
 	Metasound::Frontend::FNodeHandle NodeHandle = GetMetasoundNode().GetNodeHandle();
-	if (!NodeHandle->GetClassDisplayInfo().bShowName)
+	if (!NodeHandle->GetClassStyle().Display.bShowName)
 	{
 		return SNullWidget::NullWidget;
 	}
@@ -102,7 +102,7 @@ void SMetasoundGraphNode::SetDefaultTitleAreaWidget(TSharedRef<SOverlay> Default
 	SGraphNode::SetDefaultTitleAreaWidget(DefaultTitleAreaWidget);
 
 	Metasound::Frontend::FNodeHandle NodeHandle = GetMetasoundNode().GetNodeHandle();
-	if (NodeHandle->GetClassDisplayInfo().bShowName)
+	if (NodeHandle->GetClassStyle().Display.bShowName)
 	{
 		DefaultTitleAreaWidget->ClearChildren();
 		TSharedPtr<SNodeTitle> NodeTitle = SNew(SNodeTitle, GraphNode);
@@ -213,11 +213,11 @@ TSharedRef<SWidget> SMetasoundGraphNode::CreateNodeContentArea()
 	using namespace Metasound::Frontend;
 
 	FNodeHandle NodeHandle = GetMetasoundNode().GetNodeHandle();
-	FMetasoundFrontendClassDisplayInfo DisplayInfo = NodeHandle->GetClassDisplayInfo();
+	const FMetasoundFrontendClassStyleDisplay& StyleDisplay = NodeHandle->GetClassStyle().Display;
 
 	TSharedRef<SHorizontalBox> ContentBox = SNew(SHorizontalBox);
 
-	if (DisplayInfo.ImageName.IsNone())
+	if (StyleDisplay.ImageName.IsNone())
 	{
 		ContentBox->AddSlot()
 			.HAlign(HAlign_Left)
@@ -238,7 +238,7 @@ TSharedRef<SWidget> SMetasoundGraphNode::CreateNodeContentArea()
 
 		if (const ISlateStyle* MetasoundStyle = FSlateStyleRegistry::FindSlateStyle("MetasoundStyle"))
 		{
-			const FSlateBrush* ImageBrush = MetasoundStyle->GetBrush(DisplayInfo.ImageName);
+			const FSlateBrush* ImageBrush = MetasoundStyle->GetBrush(StyleDisplay.ImageName);
 			ContentBox->AddSlot()
 				.AutoWidth()
 				.HAlign(HAlign_Center)
diff --git a/Engine/Plugins/Runtime/Metasound/Source/MetasoundEditor/Public/IMetasoundEditor.h b/Engine/Plugins/Runtime/Metasound/Source/MetasoundEditor/Public/IMetasoundEditor.h
--- a/Engine/Plugins/Runtime/Metasound/Source/MetasoundEditor/Public/IMetasoundEditor.h
+++ b/Engine/Plugins/Runtime/Metasound/Source/MetasoundEditor/Public/IMetasoundEditor.h
@@ -14,10 +14,4 @@ class IMetasoundEditor : public FAssetEditorToolkit
 	virtual UObject* GetMetasoundObject() const = 0;
 	virtual void SetSelection(const TArray<UObject*>& SelectedObjects) = 0;
 	virtual bool GetBoundsForSelectedNodes(FSlateRect& Rect, float Padding) = 0;
-
-	// Rebuilds all nodes and connections from root FMetasoundFrontendDocument. Results in loss of node layout.
-	virtual bool RebuildGraph() const = 0;
-
-	// Synchronizes the nodes and connections of the graph to the root FMetasoundFrontendDocument.
-	virtual bool SynchronizeGraph() const = 0;
 };
diff --git a/Engine/Plugins/Runtime/Metasound/Source/MetasoundEditor/Public/MetasoundEditorGraphNode.h b/Engine/Plugins/Runtime/Metasound/Source/MetasoundEditor/Public/MetasoundEditorGraphNode.h
--- a/Engine/Plugins/Runtime/Metasound/Source/MetasoundEditor/Public/MetasoundEditorGraphNode.h
+++ b/Engine/Plugins/Runtime/Metasound/Source/MetasoundEditor/Public/MetasoundEditorGraphNode.h
@@ -42,6 +42,7 @@ public:
 	virtual FString GetDocumentationLink() const override;
 	virtual void GetNodeContextMenuActions(class UToolMenu* Menu, class UGraphNodeContextMenuContext* Context) const override;
 	virtual FText GetTooltipText() const override;
+	virtual void PinDefaultValueChanged(UEdGraphPin* Pin) override;
 	virtual void PrepareForCopying() override;
 	virtual void ReconstructNode() override;
 	// End of UEdGraphNode interface
@@ -49,6 +50,10 @@ public:
 	// UObject interface
 	virtual void PostLoad() override;
 
+#if WITH_EDITOR
+	virtual void PostEditChangeProperty(FPropertyChangedEvent& InEvent) override;
+#endif // WITH_EDITOR
+
 	virtual void PostEditImport() override;
 	virtual void PostDuplicate(bool bDuplicateForPIE) override;
 	// End of UObject interface
diff --git a/Engine/Plugins/Runtime/Metasound/Source/MetasoundEngine/Private/MetasoundEngineModule.cpp b/Engine/Plugins/Runtime/Metasound/Source/MetasoundEngine/Private/MetasoundEngineModule.cpp
--- a/Engine/Plugins/Runtime/Metasound/Source/MetasoundEngine/Private/MetasoundEngineModule.cpp
+++ b/Engine/Plugins/Runtime/Metasound/Source/MetasoundEngine/Private/MetasoundEngineModule.cpp
@@ -11,7 +11,7 @@
 
 DEFINE_LOG_CATEGORY(LogMetasoundEngine);
 
-REGISTER_METASOUND_DATATYPE(Metasound::FWaveAsset, "Primitive:WaveAsset", Metasound::ELiteralType::UObjectProxy, USoundWave);
+REGISTER_METASOUND_DATATYPE(Metasound::FWaveAsset, "WaveAsset", Metasound::ELiteralType::UObjectProxy, USoundWave);
 
 class FMetasoundEngineModule : public IMetasoundEngineModule
 {
diff --git a/Engine/Plugins/Runtime/Metasound/Source/MetasoundFrontend/Private/MetasoundFrontend.cpp b/Engine/Plugins/Runtime/Metasound/Source/MetasoundFrontend/Private/MetasoundFrontend.cpp
--- a/Engine/Plugins/Runtime/Metasound/Source/MetasoundFrontend/Private/MetasoundFrontend.cpp
+++ b/Engine/Plugins/Runtime/Metasound/Source/MetasoundFrontend/Private/MetasoundFrontend.cpp
@@ -64,32 +64,28 @@ namespace Metasound
 				TEXT("") // Variant
 			};
 
-			FMetasoundFrontendVersionNumber ClassVersion
-			{
-				InNodeMetadata.MajorVersion,
-				InNodeMetadata.MinorVersion
-			};
-
 			FMetasoundFrontendClass ClassDescription;
 
-			// Set metadata for class description.
 			ClassDescription.Metadata = FMetasoundFrontendClassMetadata
 			{
 				ClassName,
-				ClassVersion,
+				FMetasoundFrontendVersionNumber
+				{
+					InNodeMetadata.MajorVersion,
+					InNodeMetadata.MinorVersion
+				},
 				ClassType,
 				InNodeMetadata.Description,
 				InNodeMetadata.PromptIfMissing,
 				InNodeMetadata.Author,
 				InNodeMetadata.Keywords,
 				InNodeMetadata.CategoryHierarchy,
-				FMetasoundFrontendClassDisplayInfo
-				{
-					InNodeMetadata.ImageName,
-					InNodeMetadata.bShowName,
-					InNodeMetadata.bShowInputNames,
-					InNodeMetadata.bShowOutputNames
-				},
+			};
+
+			FMetasoundFrontendClassStyleDisplay DisplayStyle(InNodeMetadata.DisplayStyle);
+			ClassDescription.Style = FMetasoundFrontendClassStyle
+			{
+				DisplayStyle
 			};
 
 			// External metasounds aren't dependent on any other nodes by definition, so all we need to do
@@ -238,12 +234,12 @@ class FMetasoundFrontendModule : public IModuleInterface
 };
 
 
-REGISTER_METASOUND_DATATYPE(bool, "Primitive:Bool", ::Metasound::ELiteralType::Boolean)
-REGISTER_METASOUND_DATATYPE(int32, "Primitive:Int32", ::Metasound::ELiteralType::Integer)
-REGISTER_METASOUND_DATATYPE(int64, "Primitive:Int64", ::Metasound::ELiteralType::Integer)
-REGISTER_METASOUND_DATATYPE(float, "Primitive:Float", ::Metasound::ELiteralType::Float)
-REGISTER_METASOUND_DATATYPE(double, "Primitive:Double", ::Metasound::ELiteralType::Float)
-REGISTER_METASOUND_DATATYPE(FString, "Primitive:String", ::Metasound::ELiteralType::String)
+REGISTER_METASOUND_DATATYPE(bool, "Bool", ::Metasound::ELiteralType::Boolean)
+REGISTER_METASOUND_DATATYPE(int32, "Int32", ::Metasound::ELiteralType::Integer)
+REGISTER_METASOUND_DATATYPE(int64, "Int64", ::Metasound::ELiteralType::Integer)
+REGISTER_METASOUND_DATATYPE(float, "Float", ::Metasound::ELiteralType::Float)
+REGISTER_METASOUND_DATATYPE(double, "Double", ::Metasound::ELiteralType::Float)
+REGISTER_METASOUND_DATATYPE(FString, "String", ::Metasound::ELiteralType::String)
 
 REGISTER_METASOUND_DATATYPE(Metasound::FAudioBuffer, "Audio:Buffer")
 REGISTER_METASOUND_DATATYPE(Metasound::FSendAddress, "Transmission:Address")
diff --git a/Engine/Plugins/Runtime/Metasound/Source/MetasoundFrontend/Private/MetasoundFrontendInvalidController.h b/Engine/Plugins/Runtime/Metasound/Source/MetasoundFrontend/Private/MetasoundFrontendInvalidController.h
--- a/Engine/Plugins/Runtime/Metasound/Source/MetasoundFrontend/Private/MetasoundFrontendInvalidController.h
+++ b/Engine/Plugins/Runtime/Metasound/Source/MetasoundFrontend/Private/MetasoundFrontendInvalidController.h
@@ -142,6 +142,9 @@ namespace Metasound
 			TSharedRef<const IInputController> GetInputWithID(FGuid InPointID) const override { return FInvalidInputController::GetInvalid(); }
 			TSharedRef<const IOutputController> GetOutputWithID(FGuid InPointID) const override { return FInvalidOutputController::GetInvalid(); }
 
+			const FMetasoundFrontendNodeStyle& GetNodeStyle() const override { static const FMetasoundFrontendNodeStyle Invalid; return Invalid; }
+			void SetNodeStyle(const FMetasoundFrontendNodeStyle& InNodeStyle) { };
+
 			bool CanAddInput(const FString& InVertexName) const override { return false; }
 			FInputHandle AddInput(const FString& InVertexName, const FMetasoundFrontendLiteral* InDefault) override { return FInvalidInputController::GetInvalid(); }
 			bool RemoveInput(FGuid InPointID) override { return false; }
@@ -154,7 +157,7 @@ namespace Metasound
 			const FString& GetClassName() const override { return MetasoundFrontendInvalidControllerPrivate::GetInvalid<FString>(); }
 			FMetasoundFrontendVersionNumber GetClassVersionNumber() const override { return MetasoundFrontendInvalidControllerPrivate::GetInvalid<FMetasoundFrontendVersionNumber>(); }
 			const FText& GetClassDescription() const override { return MetasoundFrontendInvalidControllerPrivate::GetInvalid<FText>(); }
-			FMetasoundFrontendClassDisplayInfo GetClassDisplayInfo() const { return FMetasoundFrontendClassDisplayInfo(); }
+			const FMetasoundFrontendClassStyle& GetClassStyle() const override { static const FMetasoundFrontendClassStyle Invalid; return Invalid; }
 
 			TSharedRef<IGraphController> AsGraph() override;
 			TSharedRef<const IGraphController> AsGraph() const override;
diff --git a/Engine/Plugins/Runtime/Metasound/Source/MetasoundFrontend/Private/MetasoundFrontendStandardController.cpp b/Engine/Plugins/Runtime/Metasound/Source/MetasoundFrontend/Private/MetasoundFrontendStandardController.cpp
--- a/Engine/Plugins/Runtime/Metasound/Source/MetasoundFrontend/Private/MetasoundFrontendStandardController.cpp
+++ b/Engine/Plugins/Runtime/Metasound/Source/MetasoundFrontend/Private/MetasoundFrontendStandardController.cpp
@@ -1063,14 +1063,34 @@ namespace Metasound
 			return FMetasoundFrontendVersionNumber();
 		}
 
-		FMetasoundFrontendClassDisplayInfo FBaseNodeController::GetClassDisplayInfo() const
+		const FMetasoundFrontendClassStyle& FBaseNodeController::GetClassStyle() const
 		{
 			if (ClassPtr.IsValid())
 			{
-				return ClassPtr->Metadata.DisplayInfo;
+				return ClassPtr->Style;
 			}
 
-			return FMetasoundFrontendClassDisplayInfo();
+			static const FMetasoundFrontendClassStyle Invalid;
+			return Invalid;
+		}
+
+		const FMetasoundFrontendNodeStyle& FBaseNodeController::GetNodeStyle() const
+		{
+			if (NodePtr.IsValid())
+			{
+				return NodePtr->Style;
+			}
+
+			static const FMetasoundFrontendNodeStyle Invalid;
+			return Invalid;
+		}
+
+		void FBaseNodeController::SetNodeStyle(const FMetasoundFrontendNodeStyle& InStyle)
+		{
+			if (ClassPtr.IsValid())
+			{
+				NodePtr->Style = InStyle;
+			}
 		}
 
 		const FText& FBaseNodeController::GetClassDescription() const
diff --git a/Engine/Plugins/Runtime/Metasound/Source/MetasoundFrontend/Private/MetasoundFrontendStandardController.h b/Engine/Plugins/Runtime/Metasound/Source/MetasoundFrontend/Private/MetasoundFrontendStandardController.h
--- a/Engine/Plugins/Runtime/Metasound/Source/MetasoundFrontend/Private/MetasoundFrontendStandardController.h
+++ b/Engine/Plugins/Runtime/Metasound/Source/MetasoundFrontend/Private/MetasoundFrontendStandardController.h
@@ -312,7 +312,10 @@ namespace Metasound
 			FGuid GetClassID() const override;
 			FMetasoundFrontendVersionNumber GetClassVersionNumber() const override;
 			const FText& GetClassDescription() const override;
-			FMetasoundFrontendClassDisplayInfo GetClassDisplayInfo() const override;
+			const FMetasoundFrontendClassStyle& GetClassStyle() const override;
+
+			const FMetasoundFrontendNodeStyle& GetNodeStyle() const override;
+			void SetNodeStyle(const FMetasoundFrontendNodeStyle& InStyle) override;
 
 			const FString& GetNodeName() const override;
 			EMetasoundFrontendClassType GetClassType() const override;
diff --git a/Engine/Plugins/Runtime/Metasound/Source/MetasoundFrontend/Public/MetasoundAutoConverterNode.h b/Engine/Plugins/Runtime/Metasound/Source/MetasoundFrontend/Public/MetasoundAutoConverterNode.h
--- a/Engine/Plugins/Runtime/Metasound/Source/MetasoundFrontend/Public/MetasoundAutoConverterNode.h
+++ b/Engine/Plugins/Runtime/Metasound/Source/MetasoundFrontend/Public/MetasoundAutoConverterNode.h
@@ -10,6 +10,7 @@
 #include "MetasoundDataFactory.h"
 #include "MetasoundDataReference.h"
 #include "MetasoundExecutableOperator.h"
+#include "MetasoundFrontend.h"
 
 #include <type_traits>
 
@@ -52,17 +53,26 @@ namespace Metasound
 		{
 			auto InitNodeInfo = []() -> FNodeInfo
 			{
+				const FString& InputDisplayName = GetInputName();
+				const FString& OutputDisplayName = GetOutputName();
+
+				FNodeDisplayStyle DisplayStyle;
+				DisplayStyle.bShowName = false;
+				DisplayStyle.bShowInputNames = false;
+				DisplayStyle.bShowOutputNames = false;
+
 				FNodeInfo Info;
-				Info.ClassName = FName(*FString::Printf(TEXT("%s To %s"), *GetInputName(), *GetOutputName()));
+				Info.ClassName = FName(*FString::Printf(TEXT("%s To %s"), *InputDisplayName, *OutputDisplayName));
 				Info.MajorVersion = 1;
 				Info.MinorVersion = 0;
 				Info.Description = LOCTEXT("Metasound_ConverterNodeDescription", "Converts between two different data types.");
 				Info.Author = PluginAuthor;
+				Info.DisplayStyle = DisplayStyle;
 				Info.PromptIfMissing = PluginNodeMissingPrompt;
 				Info.DefaultInterface = DeclareVertexInterface();
 
 				Info.CategoryHierarchy = { LOCTEXT("Metasound_ConvertNodeCategory", "Conversions") };
-				Info.Keywords = {TEXT("convert"), GetMetasoundDataTypeName<FromDataType>(), GetMetasoundDataTypeName<ToDataType>()};
+				Info.Keywords = { "Convert", GetMetasoundDataTypeName<FromDataType>(), GetMetasoundDataTypeName<ToDataType>()};
 
 				return Info;
 			};
@@ -91,6 +101,7 @@ namespace Metasound
 				virtual FDataReferenceCollection GetInputs() const override
 				{
 					FDataReferenceCollection Inputs;
+					Inputs.AddDataReadReference<FromDataType>(GetInputName(), FromData);
 					return Inputs;
 				}
 
diff --git a/Engine/Plugins/Runtime/Metasound/Source/MetasoundFrontend/Public/MetasoundFrontendController.h b/Engine/Plugins/Runtime/Metasound/Source/MetasoundFrontend/Public/MetasoundFrontendController.h
--- a/Engine/Plugins/Runtime/Metasound/Source/MetasoundFrontend/Public/MetasoundFrontendController.h
+++ b/Engine/Plugins/Runtime/Metasound/Source/MetasoundFrontend/Public/MetasoundFrontendController.h
@@ -365,7 +365,7 @@ namespace Metasound
 			virtual const FString& GetClassName() const = 0;
 			virtual FMetasoundFrontendVersionNumber GetClassVersionNumber() const = 0;
 			virtual const FText& GetClassDescription() const = 0;
-			virtual FMetasoundFrontendClassDisplayInfo GetClassDisplayInfo() const = 0;
+			virtual const FMetasoundFrontendClassStyle& GetClassStyle() const = 0;
 
 			/** If the node is also a graph, this returns a graph handle.
 			 * If the node is not also a graph, it will return an invalid handle.
@@ -380,6 +380,9 @@ namespace Metasound
 			/** Returns the name of this node. */
 			virtual const FString& GetNodeName() const = 0;
 
+			virtual const FMetasoundFrontendNodeStyle& GetNodeStyle() const = 0;
+			virtual void SetNodeStyle(const FMetasoundFrontendNodeStyle& InStyle) = 0;
+
 			/** Returns the ID associated with this node. */
 			virtual FGuid GetID() const = 0;
 
diff --git a/Engine/Plugins/Runtime/Metasound/Source/MetasoundFrontend/Public/MetasoundFrontendDocument.h b/Engine/Plugins/Runtime/Metasound/Source/MetasoundFrontend/Public/MetasoundFrontendDocument.h
--- a/Engine/Plugins/Runtime/Metasound/Source/MetasoundFrontend/Public/MetasoundFrontendDocument.h
+++ b/Engine/Plugins/Runtime/Metasound/Source/MetasoundFrontend/Public/MetasoundFrontendDocument.h
@@ -5,6 +5,7 @@
 #include "CoreMinimal.h"
 #include "Internationalization/Text.h"
 #include "MetasoundAccessPtr.h"
+#include "MetasoundNodeInterface.h"
 #include "Misc/Guid.h"
 
 #include "MetasoundFrontendDocument.generated.h"
@@ -180,10 +181,40 @@ struct METASOUNDFRONTEND_API FMetasoundFrontendNodeInterface
 	TArray<FMetasoundFrontendVertex> Environment;
 };
 
+UENUM()
+enum class EMetasoundFrontendNodeStyleDisplayVisibility : uint8
+{
+	Visible,
+	Hidden
+};
+
+USTRUCT()
+struct METASOUNDFRONTEND_API FMetasoundFrontendNodeStyleDisplay
+{
+	GENERATED_BODY()
+
+	// Visibility state of node
+	UPROPERTY()
+	EMetasoundFrontendNodeStyleDisplayVisibility Visibility = EMetasoundFrontendNodeStyleDisplayVisibility::Visible;
+
+	// Position of node
+	UPROPERTY()
+	FVector2D Location = FVector2D::ZeroVector;
+};
+
+USTRUCT()
+struct METASOUNDFRONTEND_API FMetasoundFrontendNodeStyle
+{
+	GENERATED_BODY()
+
+	// Display style of a node
+	UPROPERTY()
+	FMetasoundFrontendNodeStyleDisplay Display;
+};
 
 // An FMetasoundFrontendNode represents a single instance of a FMetasoundFrontendClass
 USTRUCT()
-struct METASOUNDFRONTEND_API FMetasoundFrontendNode 
+struct METASOUNDFRONTEND_API FMetasoundFrontendNode
 {
 	GENERATED_BODY()
 
@@ -212,7 +243,9 @@ struct METASOUNDFRONTEND_API FMetasoundFrontendNode
 	UPROPERTY()
 	TArray<FMetasoundFrontendVertexLiteral> InputLiterals;
 
-	// TODO: Add node init data if it is used. It may be dropped as a concept.
+	// Style info related to a node.
+	UPROPERTY()
+	FMetasoundFrontendNodeStyle Style;
 };
 
 
@@ -398,10 +431,20 @@ struct METASOUNDFRONTEND_API FMetasoundFrontendClassVertex : public FMetasoundFr
 
 // Information regarding how to display a node class
 USTRUCT()
-struct METASOUNDFRONTEND_API FMetasoundFrontendClassDisplayInfo
+struct METASOUNDFRONTEND_API FMetasoundFrontendClassStyleDisplay
 {
 	GENERATED_BODY()
 
+	FMetasoundFrontendClassStyleDisplay() = default;
+
+	FMetasoundFrontendClassStyleDisplay(const Metasound::FNodeDisplayStyle& InDisplayStyle)
+	:	ImageName(InDisplayStyle.ImageName)
+	,	bShowName(InDisplayStyle.bShowName)
+	,	bShowInputNames(InDisplayStyle.bShowInputNames)
+	,	bShowOutputNames(InDisplayStyle.bShowOutputNames)
+	{
+	}
+
 	UPROPERTY()
 	FName ImageName;
 
@@ -409,10 +452,10 @@ struct METASOUNDFRONTEND_API FMetasoundFrontendClassDisplayInfo
 	bool bShowName = true;
 
 	UPROPERTY()
-	bool bShowInputName = true;
+	bool bShowInputNames = true;
 
 	UPROPERTY()
-	bool bShowOutputName = true;
+	bool bShowOutputNames = true;
 };
 
 // Contains info for input vertex of a Metasound class.
@@ -584,30 +627,17 @@ struct FMetasoundFrontendClassMetadata
 
 	UPROPERTY()
 	TArray<FText> CategoryHierarchy;
-
-	UPROPERTY()
-	FMetasoundFrontendClassDisplayInfo DisplayInfo;
 };
 
-UENUM()
-enum class EMetasoundFrontendStyleNodeDisplay : uint8
-{
-	Default,
-	Inherited,
-	Minimized
-};
-
-USTRUCT() 
+USTRUCT()
 struct FMetasoundFrontendClassStyle
 {
 	GENERATED_BODY()
 
 	UPROPERTY()
-	EMetasoundFrontendStyleNodeDisplay NodeDisplay = EMetasoundFrontendStyleNodeDisplay::Inherited;
+	FMetasoundFrontendClassStyleDisplay Display;
 };
 
-
-
 USTRUCT()
 struct FMetasoundFrontendEditorData
 {
@@ -625,9 +655,6 @@ struct METASOUNDFRONTEND_API FMetasoundFrontendClass
 {
 	GENERATED_BODY()
 
-	//Metasound::Frontend::TAccessPoint<FMetasoundFrontendClass> AccessPoint;
-	//Metasound::Frontend::FAccessPoint AccessPoint;
-
 	virtual ~FMetasoundFrontendClass() = default;
 
 	UPROPERTY()
diff --git a/Engine/Plugins/Runtime/Metasound/Source/MetasoundFrontend/Public/MetasoundSendNode.h b/Engine/Plugins/Runtime/Metasound/Source/MetasoundFrontend/Public/MetasoundSendNode.h
--- a/Engine/Plugins/Runtime/Metasound/Source/MetasoundFrontend/Public/MetasoundSendNode.h
+++ b/Engine/Plugins/Runtime/Metasound/Source/MetasoundFrontend/Public/MetasoundSendNode.h
@@ -23,13 +23,13 @@ namespace Metasound
 	public:
 		static const FString& GetAddressInputName()
 		{
-			static const FString InputName = FString(TEXT("Address"));
+			static const FString InputName = TEXT("Address");
 			return InputName;
 		}
 
 		static const FString& GetSendInputName()
 		{
-			static const FString SendInput = GetMetasoundDataTypeString<TDataType>();
+			static const FString& SendInput = GetMetasoundDataTypeString<TDataType>();
 			return SendInput;
 		}
 
@@ -49,9 +49,10 @@ namespace Metasound
 		{
 			auto InitNodeInfo = []() -> FNodeInfo
 			{
+				const FString& InputName = GetSendInputName();
 				FNodeInfo Info;
 
-				Info.ClassName = FName(*FString::Printf(TEXT("Send %s"), *GetSendInputName()));
+				Info.ClassName = FName(*FString::Printf(TEXT("Send %s"), *InputName));
 				Info.MajorVersion = 1;
 				Info.MinorVersion = 0;
 				Info.Description = LOCTEXT("Metasound_SendNodeDescription", "Sends data from a send node with the same name.");
@@ -59,7 +60,7 @@ namespace Metasound
 				Info.PromptIfMissing = PluginNodeMissingPrompt;
 				Info.DefaultInterface = DeclareVertexInterface();
 				Info.CategoryHierarchy = { LOCTEXT("Metasound_TransmissionNodeCategory", "Transmission") };
-				Info.Keywords = {TEXT("send"), GetMetasoundDataTypeName<TDataType>()};
+				Info.Keywords = { "Send", GetMetasoundDataTypeName<TDataType>()};
 
 				return Info;
 			};
diff --git a/Engine/Plugins/Runtime/Metasound/Source/MetasoundGraphCore/Public/MetasoundNodeInterface.h b/Engine/Plugins/Runtime/Metasound/Source/MetasoundGraphCore/Public/MetasoundNodeInterface.h
--- a/Engine/Plugins/Runtime/Metasound/Source/MetasoundGraphCore/Public/MetasoundNodeInterface.h
+++ b/Engine/Plugins/Runtime/Metasound/Source/MetasoundGraphCore/Public/MetasoundNodeInterface.h
@@ -33,6 +33,24 @@ namespace Metasound
 		}
 	};
 
+	/**
+	 * Node style data
+	 */
+	struct FNodeDisplayStyle
+	{
+		/** Icon name identifier associated with node. */
+		FName ImageName;
+
+		/** Whether or not to show name in visual layout. */
+		bool bShowName = true;
+
+		/** Whether or not to show input names in visual layout. */
+		bool bShowInputNames = true;
+
+		/** Whether or not to show output names in visual layout. */
+		bool bShowOutputNames = true;
+	};
+
 	// TODO: update FNodeInfo with this as class name.
 	struct FNodeClassName
 	{
@@ -70,17 +88,12 @@ namespace Metasound
 		/** Hierarchy of categories for displaying node. */
 		TArray<FText> CategoryHierarchy;
 
-		FName ImageName;
-
-		bool bShowName = true;
-
-		bool bShowInputNames = true;
-
-		bool bShowOutputNames = true;
-
 		/** List of keywords for contextual node searching. */
 		TArray<FName> Keywords;
 
+		/** Display style for node when visualized. */
+		FNodeDisplayStyle DisplayStyle;
+
 		/** Returns an empty FNodeInfo object. */
 		static const FNodeInfo& GetEmpty()
 		{
diff --git a/Engine/Plugins/Runtime/Metasound/Source/MetasoundStandardNodes/Private/MetasoundAudioFormatNodes.cpp b/Engine/Plugins/Runtime/Metasound/Source/MetasoundStandardNodes/Private/MetasoundAudioFormatNodes.cpp
--- a/Engine/Plugins/Runtime/Metasound/Source/MetasoundStandardNodes/Private/MetasoundAudioFormatNodes.cpp
+++ b/Engine/Plugins/Runtime/Metasound/Source/MetasoundStandardNodes/Private/MetasoundAudioFormatNodes.cpp
@@ -67,18 +67,22 @@ namespace Metasound
 		{
 			auto InitNodeInfo = []() -> FNodeInfo
 			{
-				FNodeInfo Info;
+				FNodeDisplayStyle DisplayStyle;
+				DisplayStyle.bShowName = false;
 
-				Info.ClassName = FName(TEXT("MonoAudioFormat"));
+				FNodeInfo Info;
+				Info.ClassName = "Buffer to Mono";
+				Info.DisplayStyle = DisplayStyle;
 				Info.MajorVersion = 1;
 				Info.MinorVersion = 0;
-				Info.Description = LOCTEXT("MonoAudioFormat_NodeDescription", "Mono Audio Format");
+				Info.Description = LOCTEXT("MonoAudioFormat_NodeDescription", "Converts Audio:Buffer To Mono Format");
+				Info.CategoryHierarchy = { LOCTEXT("Metasound_ConvertNodeCategory", "Conversions") };
 				Info.Author = PluginAuthor;
 				Info.PromptIfMissing = PluginNodeMissingPrompt;
 
 				Info.DefaultInterface = FVertexInterface(
 					FInputVertexInterface(
-						TInputDataVertexModel<FAudioBuffer>(GetCenterName(), LOCTEXT("MonoAudioFormat_CenterChannel", "Audio"))
+						TInputDataVertexModel<FAudioBuffer>(GetCenterName(), LOCTEXT("MonoAudioFormat_CenterChannel", "Center"))
 					),
 					FOutputVertexInterface(
 						TOutputDataVertexModel<FMonoAudioFormat>(GetMonoName(), LOCTEXT("MonoAudioFormat_MonoOut", "Mono"))
@@ -169,19 +173,23 @@ namespace Metasound
 		{
 			auto InitNodeInfo = []() -> FNodeInfo
 			{
-				FNodeInfo Info;
+				FNodeDisplayStyle DisplayStyle;
+				DisplayStyle.bShowName = false;
 
-				Info.ClassName = FName(TEXT("StereoAudioFormat"));
+				FNodeInfo Info;
+				Info.ClassName = "Buffer to Stereo";
+				Info.DisplayStyle = DisplayStyle;
 				Info.MajorVersion = 1;
 				Info.MinorVersion = 0;
-				Info.Description = LOCTEXT("StereoAudioFormat_NodeDescription", "Stereo Audio Format");
+				Info.Description = LOCTEXT("StereoAudioFormat_NodeDescription", "Converts Audio:Buffer to Stereo Format");
+				Info.CategoryHierarchy = { LOCTEXT("Metasound_ConvertNodeCategory", "Conversions") };
 				Info.Author = PluginAuthor;
 				Info.PromptIfMissing = PluginNodeMissingPrompt;
 
 				Info.DefaultInterface = FVertexInterface(
 					FInputVertexInterface(
-						TInputDataVertexModel<FAudioBuffer>(GetLeftName(), LOCTEXT("StereoAudioFormat_LeftChannel", "Left Audio")),
-						TInputDataVertexModel<FAudioBuffer>(GetRightName(), LOCTEXT("StereoAudioFormat_RightChannel", "Right Audio"))
+						TInputDataVertexModel<FAudioBuffer>(GetLeftName(), LOCTEXT("StereoAudioFormat_LeftChannel", "Left")),
+						TInputDataVertexModel<FAudioBuffer>(GetRightName(), LOCTEXT("StereoAudioFormat_RightChannel", "Right"))
 					),
 					FOutputVertexInterface(
 						TOutputDataVertexModel<FStereoAudioFormat>(GetStereoName(), LOCTEXT("StereoAudioFormat_StereoOut", "Stereo"))
diff --git a/Engine/Plugins/Runtime/Metasound/Source/MetasoundStandardNodes/Private/MetasoundGainMultiplyNode.cpp b/Engine/Plugins/Runtime/Metasound/Source/MetasoundStandardNodes/Private/MetasoundGainMultiplyNode.cpp
--- a/Engine/Plugins/Runtime/Metasound/Source/MetasoundStandardNodes/Private/MetasoundGainMultiplyNode.cpp
+++ b/Engine/Plugins/Runtime/Metasound/Source/MetasoundStandardNodes/Private/MetasoundGainMultiplyNode.cpp
@@ -31,11 +31,14 @@ namespace Metasound
 			{
 				check(OutputBuffer->Num() == BlockSize);
 				check(Buffer->Num() == BlockSize);
+				LastGain = *Gain;
 			}
 
 			virtual FDataReferenceCollection GetInputs() const override
 			{
 				FDataReferenceCollection InputDataReferences;
+				InputDataReferences.AddDataReadReference(TEXT("In"), FAudioBufferReadRef(Buffer));
+				InputDataReferences.AddDataReadReference(TEXT("Gain"), FFloatReadRef(Gain));
 				return InputDataReferences;
 			}
 
@@ -48,14 +51,36 @@ namespace Metasound
 
 			void Execute()
 			{
+				if (!bInit)
+				{
+					bInit = true;
+					LastGain = *Gain;
+				}
+
 				FMemory::Memcpy(OutputBuffer->GetData(), Buffer->GetData(), sizeof(float) * BlockSize);
-				Audio::MultiplyBufferByConstantInPlace(*OutputBuffer, *Gain);
+
+				const int32 SIMDRemainder = OutputBuffer->Num() % AUDIO_SIMD_FLOAT_ALIGNMENT;
+				const int32 SIMDCount = OutputBuffer->Num() - SIMDRemainder;
+
+				Audio::FadeBufferFast(OutputBuffer->GetData(), SIMDCount, LastGain, *Gain);
+
+				for (int32 i = SIMDCount; i < OutputBuffer->Num(); ++i)
+				{
+					OutputBuffer->GetData()[i] *= *Gain;
+				}
+
+				LastGain = *Gain;
 			}
 
 		private:
 			FAudioBufferReadRef Buffer;
 			FFloatReadRef Gain;
 			FAudioBufferWriteRef OutputBuffer;
+
+			float LastGain = 1.0f;
+
+			bool bInit = false;
+
 			int32 BlockSize;
 	};
 
@@ -73,11 +98,11 @@ namespace Metasound
 	{
 		static const FVertexInterface Interface(
 			FInputVertexInterface(
-				TInputDataVertexModel<FAudioBuffer>(TEXT("Buffer"), LOCTEXT("InputBuffer1Tooltip", "Input buffer.")),
+				TInputDataVertexModel<FAudioBuffer>(TEXT("Buffer"), LOCTEXT("InputBufferTooltip", "Input buffer.")),
 				TInputDataVertexModel<float>(TEXT("Gain"), LOCTEXT("InputGainTooltip", "Input gain."))
 			),
 			FOutputVertexInterface(
-				TOutputDataVertexModel<FAudioBuffer>(TEXT("Out"), LOCTEXT("OutpuBufferTooltip", "The output audio."))
+				TOutputDataVertexModel<FAudioBuffer>(TEXT("Out"), LOCTEXT("OutpuBufferTooltip", "Output buffer with gain multiplier applied."))
 			)
 		);
 
@@ -99,11 +124,19 @@ namespace Metasound
 	{
 		auto InitNodeInfo = []() -> FNodeInfo
 		{
+			FNodeDisplayStyle DisplayStyle;
+			DisplayStyle.ImageName = "MetasoundEditor.Graph.Node.Math.Multiply";
+			DisplayStyle.bShowName = false;
+			DisplayStyle.bShowInputNames = false;
+			DisplayStyle.bShowOutputNames = false;
+
 			FNodeInfo Info;
-			Info.ClassName = FName(TEXT("Gain Multiply"));
+			Info.ClassName = "Multiply AudioBuffer by Gain";
+			Info.DisplayStyle = DisplayStyle;
 			Info.MajorVersion = 1;
 			Info.MinorVersion = 0;
-			Info.Description = LOCTEXT("GainMultiply_NodeDescription", "Multiply audio buffer by float."),
+			Info.Description = LOCTEXT("GainMultiply_NodeDescription", "Multiply AudioBuffer by float.");
+			Info.CategoryHierarchy = { LOCTEXT("Metasound_MathCategory", "Math") };
 			Info.Author = PluginAuthor;
 			Info.PromptIfMissing = PluginNodeMissingPrompt;
 			Info.DefaultInterface = DeclareVertexInterface();
diff --git a/Engine/Plugins/Runtime/Metasound/Source/MetasoundStandardNodes/Private/MetasoundStandardNodesModule.cpp b/Engine/Plugins/Runtime/Metasound/Source/MetasoundStandardNodes/Private/MetasoundStandardNodesModule.cpp
--- a/Engine/Plugins/Runtime/Metasound/Source/MetasoundStandardNodes/Private/MetasoundStandardNodesModule.cpp
+++ b/Engine/Plugins/Runtime/Metasound/Source/MetasoundStandardNodes/Private/MetasoundStandardNodesModule.cpp
@@ -10,11 +10,11 @@
 #include "MetasoundTime.h"
 #include "MetasoundFrequency.h"
 
-REGISTER_METASOUND_DATATYPE(Metasound::FTrigger, "Primitive:Trigger", ::Metasound::ELiteralType::Boolean)
-REGISTER_METASOUND_DATATYPE(Metasound::FFloatTime, "Primitive:Time", ::Metasound::ELiteralType::Float)
-REGISTER_METASOUND_DATATYPE(Metasound::FDoubleTime, "Primitive:Time:HighResolution", ::Metasound::ELiteralType::Float)
-REGISTER_METASOUND_DATATYPE(Metasound::FSampleTime, "Primitive:Time:SampleResolution", ::Metasound::ELiteralType::Float)
-REGISTER_METASOUND_DATATYPE(Metasound::FFrequency, "Primitive:Frequency", ::Metasound::ELiteralType::Float)
+REGISTER_METASOUND_DATATYPE(Metasound::FTrigger, "Trigger", ::Metasound::ELiteralType::Boolean)
+REGISTER_METASOUND_DATATYPE(Metasound::FFloatTime, "Time", ::Metasound::ELiteralType::Float)
+REGISTER_METASOUND_DATATYPE(Metasound::FDoubleTime, "Time:HighResolution", ::Metasound::ELiteralType::Float)
+REGISTER_METASOUND_DATATYPE(Metasound::FSampleTime, "Time:SampleResolution", ::Metasound::ELiteralType::Float)
+REGISTER_METASOUND_DATATYPE(Metasound::FFrequency, "Frequency", ::Metasound::ELiteralType::Float)
 
 namespace Metasound 
 {
diff --git a/Engine/Plugins/Runtime/Metasound/Source/MetasoundStandardNodes/Public/MetasoundMathNodes.h b/Engine/Plugins/Runtime/Metasound/Source/MetasoundStandardNodes/Public/MetasoundMathNodes.h
--- a/Engine/Plugins/Runtime/Metasound/Source/MetasoundStandardNodes/Public/MetasoundMathNodes.h
+++ b/Engine/Plugins/Runtime/Metasound/Source/MetasoundStandardNodes/Public/MetasoundMathNodes.h
@@ -67,6 +67,12 @@ namespace Metasound
 			{
 				auto InitNodeInfo = []() -> FNodeInfo
 				{
+					FNodeDisplayStyle DisplayStyle;
+					DisplayStyle.ImageName = TDerivedClass::GetImageName();
+					DisplayStyle.bShowName = false;
+					DisplayStyle.bShowInputNames = false;
+					DisplayStyle.bShowOutputNames = false;
+
 					FNodeInfo Info;
 					Info.ClassName = TDerivedClass::GetClassName();
 					Info.MajorVersion = 1;
@@ -75,10 +81,7 @@ namespace Metasound
 					Info.Author = PluginAuthor;
 					Info.PromptIfMissing = PluginNodeMissingPrompt;
 					Info.DefaultInterface = DeclareVertexInterface();
-					Info.ImageName = TDerivedClass::GetImageName();
-					Info.bShowName = false;
-					Info.bShowInputNames = false;
-					Info.bShowOutputNames = false;
+					Info.DisplayStyle = DisplayStyle;
 					Info.CategoryHierarchy = { LOCTEXT("Metasound_MathCategory", "Math") };
 
 					return Info;
