Metasound wave selector node

diff --git a/Engine/Plugins/Runtime/Metasound/Source/MetasoundEngine/Private/MetasoundWaveSelectorNode.cpp b/Engine/Plugins/Runtime/Metasound/Source/MetasoundEngine/Private/MetasoundWaveSelectorNode.cpp
new file mode 100644
--- /dev/null
+++ b/Engine/Plugins/Runtime/Metasound/Source/MetasoundEngine/Private/MetasoundWaveSelectorNode.cpp
@@ -0,0 +1,248 @@
+// Copyright Epic Games, Inc. All Rights Reserved.
+
+#include "MetasoundWaveSelectorNode.h"
+#include "MetasoundExecutableOperator.h"
+#include "MetasoundNodeRegistrationMacro.h"
+#include "MetasoundPrimitives.h"
+#include "MetasoundWave.h"
+#include "MetasoundNodeRegistrationMacro.h"
+#include "MetasoundBuildError.h"
+#include "MetasoundBop.h"
+
+#define LOCTEXT_NAMESPACE "MetasoundWaveNode"
+
+namespace Metasound
+{
+	class FWaveSelectorOperator : public TExecutableOperator<FWaveSelectorOperator>
+	{
+	public:
+		FWaveSelectorOperator(
+			const FOperatorSettings& InSettings,
+			const FWaveAssetReadRef& InWave0, const FWaveAssetReadRef& InWave1, const FWaveAssetReadRef& InWave2,
+			const FWaveAssetReadRef& InWave3, const FWaveAssetReadRef& InWave4,
+			const FInt32ReadRef& InIndex,
+			const FTriggerReadRef& InTrigger
+			)
+			: OperatorSettings(InSettings)
+			, Wave0(InWave0), Wave1(InWave1), Wave2(InWave2), Wave3(InWave3), Wave4(InWave4)
+			, Index(InIndex)
+			, TrigIn(InTrigger)
+			, SelectedWave(FWaveAssetWriteRef::CreateNew())
+			, TriggerNewWaveSelected(FTriggerWriteRef::CreateNew(InSettings))
+
+		{
+		}
+
+		virtual FDataReferenceCollection GetInputs() const override
+		{
+			FDataReferenceCollection InputDataReferences;
+			InputDataReferences.AddDataReadReference(TEXT("Wave0"), FWaveAssetReadRef(Wave0));
+			InputDataReferences.AddDataReadReference(TEXT("Wave1"), FWaveAssetReadRef(Wave1));
+			InputDataReferences.AddDataReadReference(TEXT("Wave2"), FWaveAssetReadRef(Wave2));
+			InputDataReferences.AddDataReadReference(TEXT("Wave3"), FWaveAssetReadRef(Wave3));
+			InputDataReferences.AddDataReadReference(TEXT("Wave4"), FWaveAssetReadRef(Wave4));
+
+			InputDataReferences.AddDataReadReference(TEXT("Index"), FInt32ReadRef(Index));
+
+			InputDataReferences.AddDataReadReference(TEXT("TrigIn"), FBopReadRef(TrigIn));
+			return InputDataReferences;
+		}
+
+		virtual FDataReferenceCollection GetOutputs() const override
+		{
+			FDataReferenceCollection OutputDataReferences;
+			OutputDataReferences.AddDataReadReference(TEXT("SelectedWave"), FWaveAssetWriteRef(SelectedWave));
+			OutputDataReferences.AddDataReadReference(TEXT("New Wave Selected"), FTriggerWriteRef(TriggerNewWaveSelected));
+			return OutputDataReferences;
+		}
+
+
+		void Execute()
+		{
+			TriggerNewWaveSelected->AdvanceBlock();
+
+			TrigIn->ExecuteBlock(
+				// OnPreBop
+				[&](int32 StartFrame, int32 EndFrame)
+				{
+				},
+				// OnBop
+				[&](int32 StartFrame, int32 EndFrame)
+				{
+					UpdateWave(*Index);
+					TriggerNewWaveSelected->BopFrame(StartFrame);
+				}
+			);
+		}
+
+		void UpdateWave(int32 InIndex)
+		{
+			UpdateWaveArray();
+			int32 NumValidWaves = WavePtrs.Num();
+			if (!NumValidWaves)
+			{
+				return;
+			}
+
+			InIndex %= WavePtrs.Num();
+			*SelectedWave = **WavePtrs[InIndex];
+		}
+
+		void UpdateWaveArray()
+		{
+			WavePtrs.Reset();
+			if (Wave0->IsSoundWaveValid())
+			{
+				WavePtrs.Add(&Wave0);
+			}
+
+			if (Wave1->IsSoundWaveValid())
+			{
+				WavePtrs.Add(&Wave1);
+			}
+
+			if (Wave2->IsSoundWaveValid())
+			{
+				WavePtrs.Add(&Wave2);
+			}
+
+			if (Wave3->IsSoundWaveValid())
+			{
+				WavePtrs.Add(&Wave3);
+			}
+
+			if (Wave4->IsSoundWaveValid())
+			{
+				WavePtrs.Add(&Wave4);
+			}
+		}
+
+	private:
+		const FOperatorSettings OperatorSettings;
+
+		// Inputs:
+		FWaveAssetReadRef Wave0;
+		FWaveAssetReadRef Wave1;
+		FWaveAssetReadRef Wave2;
+		FWaveAssetReadRef Wave3;
+		FWaveAssetReadRef Wave4;
+		FInt32ReadRef Index;
+		FTriggerReadRef TrigIn;
+
+		// Outputs
+		FWaveAssetWriteRef SelectedWave;
+		FBopWriteRef TriggerNewWaveSelected;
+
+		TArray<FWaveAssetReadRef*> WavePtrs;
+
+	};
+
+
+	TUniquePtr<IOperator> FWaveSelectorNode::FOperatorFactory::CreateOperator(
+		const FCreateOperatorParams& InParams, 
+		FBuildErrorArray& OutErrors) 
+	{
+		using namespace Audio;
+
+		const FWaveSelectorNode& WaveNode = static_cast<const FWaveSelectorNode&>(InParams.Node);
+
+		const FDataReferenceCollection& InputDataRefs = InParams.InputDataReferences;
+
+		// initialize inputs
+		FWaveAssetReadRef Wave0 = InputDataRefs.GetDataReadReferenceOrConstruct<FWaveAsset>(TEXT("Wave0"));
+		FWaveAssetReadRef Wave1 = InputDataRefs.GetDataReadReferenceOrConstruct<FWaveAsset>(TEXT("Wave1"));
+		FWaveAssetReadRef Wave2 = InputDataRefs.GetDataReadReferenceOrConstruct<FWaveAsset>(TEXT("Wave2"));
+		FWaveAssetReadRef Wave3 = InputDataRefs.GetDataReadReferenceOrConstruct<FWaveAsset>(TEXT("Wave3"));
+		FWaveAssetReadRef Wave4 = InputDataRefs.GetDataReadReferenceOrConstruct<FWaveAsset>(TEXT("Wave4"));
+
+		FInt32ReadRef Index = InputDataRefs.GetDataReadReferenceOrConstruct<int32>(TEXT("Index"), 0);
+
+		FTriggerReadRef TrigIn = InputDataRefs.GetDataReadReferenceOrConstruct<FTrigger>(TEXT("TrigIn"), InParams.OperatorSettings);
+
+
+		return MakeUnique<FWaveSelectorOperator>(
+			  InParams.OperatorSettings
+			, Wave0, Wave1, Wave2, Wave3, Wave4
+			, Index
+			, TrigIn
+			);
+	}
+
+	FVertexInterface FWaveSelectorNode::DeclareVertexInterface()
+	{
+		return FVertexInterface(
+			FInputVertexInterface(
+				TInputDataVertexModel<FTrigger>(TEXT("TrigIn"), LOCTEXT("TrigInTooltip", "Trigger the playing of the input wave.")),
+				TInputDataVertexModel<FWaveAsset>(TEXT("Wave0"), LOCTEXT("WaveTooltip", "One of the wave options")),
+				TInputDataVertexModel<FWaveAsset>(TEXT("Wave1"), LOCTEXT("WaveTooltip", "One of the wave options")),
+				TInputDataVertexModel<FWaveAsset>(TEXT("Wave2"), LOCTEXT("WaveTooltip", "One of the wave options")),
+				TInputDataVertexModel<FWaveAsset>(TEXT("Wave3"), LOCTEXT("WaveTooltip", "One of the wave options")),
+				TInputDataVertexModel<FWaveAsset>(TEXT("Wave4"), LOCTEXT("WaveTooltip", "One of the wave options")),
+				TInputDataVertexModel<int32>(TEXT("Index"), LOCTEXT("IndexTooltip", "Which wave to choose"))
+
+			),
+			FOutputVertexInterface(
+				TOutputDataVertexModel<FWaveAsset>(TEXT("SelectedWave"), LOCTEXT("AudioTooltip", "The output wave. updated when triggered")),
+				TOutputDataVertexModel<FTrigger>(TEXT("New Wave Selected"), LOCTEXT("AudioTooltip", "pass-through of input trigger"))
+			)
+		);
+	}
+
+	const FNodeInfo& FWaveSelectorNode::GetNodeInfo()
+	{
+		auto InitNodeInfo = []() -> FNodeInfo
+		{
+			FNodeInfo Info;
+			Info.ClassName = FName(TEXT("Wave Selector"));
+			Info.MajorVersion = 1;
+			Info.MinorVersion = 0;
+			Info.Description = LOCTEXT("Metasound_WaveSelectorNodeDescription", "Allows selection of a Wave Asset from a list");
+			Info.Author = PluginAuthor;
+			Info.PromptIfMissing = PluginNodeMissingPrompt;
+			Info.DefaultInterface = DeclareVertexInterface();
+
+			return Info;
+		};
+
+		static const FNodeInfo Info = InitNodeInfo();
+
+		return Info;
+	}
+
+	FWaveSelectorNode::FWaveSelectorNode(const FString& InName)
+		:	FNode(InName, GetNodeInfo())
+		,	Factory(MakeOperatorFactoryRef<FWaveSelectorNode::FOperatorFactory>())
+		,	Interface(DeclareVertexInterface())
+	{
+	}
+
+	FWaveSelectorNode::FWaveSelectorNode(const FNodeInitData& InInitData)
+		: FWaveSelectorNode(InInitData.InstanceName)
+	{
+	}
+
+	FOperatorFactorySharedRef FWaveSelectorNode::GetDefaultOperatorFactory() const 
+	{
+		return Factory;
+	}
+
+
+
+	const FVertexInterface& FWaveSelectorNode::GetVertexInterface() const
+	{
+		return Interface;
+	}
+
+	bool FWaveSelectorNode::SetVertexInterface(const FVertexInterface& InInterface)
+	{
+		return InInterface == Interface;
+	}
+
+	bool FWaveSelectorNode::IsVertexInterfaceSupported(const FVertexInterface& InInterface) const
+	{
+		return InInterface == Interface;
+	}
+
+	METASOUND_REGISTER_NODE(FWaveSelectorNode)
+} // namespace Metasound
+#undef LOCTEXT_NAMESPACE // MetasoundWaveNode
diff --git a/Engine/Plugins/Runtime/Metasound/Source/MetasoundEngine/Public/MetasoundWaveSelectorNode.h b/Engine/Plugins/Runtime/Metasound/Source/MetasoundEngine/Public/MetasoundWaveSelectorNode.h
new file mode 100644
--- /dev/null
+++ b/Engine/Plugins/Runtime/Metasound/Source/MetasoundEngine/Public/MetasoundWaveSelectorNode.h
@@ -0,0 +1,60 @@
+// Copyright Epic Games, Inc. All Rights Reserved.
+
+#pragma once
+
+#include "CoreMinimal.h"
+#include "MetasoundNode.h"
+#include "MetasoundOperatorInterface.h"
+#include "MetasoundBuilderInterface.h"
+#include "MetasoundDataReferenceCollection.h"
+
+namespace Metasound
+{
+	class METASOUNDENGINE_API FWaveSelectorNode : public FNode
+	{
+		class FOperatorFactory : public IOperatorFactory
+		{
+			virtual TUniquePtr<IOperator> CreateOperator(const FCreateOperatorParams& InParams, FBuildErrorArray& OutErrors) override;
+		};
+
+		
+
+	public:
+		static FVertexInterface DeclareVertexInterface();
+		static const FNodeInfo& GetNodeInfo();
+
+		FWaveSelectorNode(const FString& InName);
+
+		// constructor used by the Metasound Frontend.
+		FWaveSelectorNode(const FNodeInitData& InInitData);
+
+		virtual ~FWaveSelectorNode() = default;
+
+		virtual FOperatorFactorySharedRef GetDefaultOperatorFactory() const override;
+
+		/** Return the current vertex interface. */
+		virtual const FVertexInterface& GetVertexInterface() const override;
+
+		/** Set the vertex interface. If the vertex was successfully changed, returns true. 
+		 *
+		 * @param InInterface - New interface for node. 
+		 *
+		 * @return True on success, false otherwise.
+		 */
+		virtual bool SetVertexInterface(const FVertexInterface& InInterface) override;
+
+		/** Expresses whether a specific vertex interface is supported.
+		 *
+		 * @param InInterface - New interface. 
+		 *
+		 * @return True if the interface is supported, false otherwise. 
+		 */
+		virtual bool IsVertexInterfaceSupported(const FVertexInterface& InInterface) const override;
+
+	private:
+
+		FOperatorFactorySharedRef Factory;
+
+		FVertexInterface Interface;
+	};
+}
