SRC distortion fix in MetaSound Wave Player Node

diff --git a/Engine/Plugins/Runtime/Metasound/Source/MetasoundEngine/Private/MetasoundWavePlayerNode.cpp b/Engine/Plugins/Runtime/Metasound/Source/MetasoundEngine/Private/MetasoundWavePlayerNode.cpp
--- a/Engine/Plugins/Runtime/Metasound/Source/MetasoundEngine/Private/MetasoundWavePlayerNode.cpp
+++ b/Engine/Plugins/Runtime/Metasound/Source/MetasoundEngine/Private/MetasoundWavePlayerNode.cpp
@@ -12,6 +12,20 @@
 
 #define LOCTEXT_NAMESPACE "MetasoundWaveNode"
 
+// static const int32 SMOOTH = -1;
+// static int32 NoDiscontinuities(const float* start, int32 numframes, const float thresh = 0.3f)
+// {
+// 	for (int32 i = 0; i < numframes - 1; ++i)
+// 	{
+// 		float delta = FMath::Abs(start[i] - start[i + 1]);
+// 		if (delta > thresh)
+// 		{
+// 			return i;
+// 		}
+// 	}
+// 	return -1; // all good!
+// }
+
 namespace Metasound
 {
 	// WavePlayer custom error 
@@ -101,7 +115,7 @@ namespace Metasound
 				InputSampleRate = Wave->SoundWaveProxy->GetSampleRate();
 
 				NumChannels = Wave->SoundWaveProxy->GetNumChannels();
-				CircularDecoderOutputBuffer.SetCapacity(OutputBlockSizeInFrames * NumChannels * 2);
+				CircularDecoderOutputBuffer.SetCapacity(OutputBlockSizeInFrames * NumChannels * 4);
 
 				CurrentSoundWaveName = NewSoundWaveName;
 			}
@@ -155,14 +169,12 @@ namespace Metasound
 			{
 				++NumSamplesToDecode;
 			}
-			
 
-			const bool bNeedsSRC = (InputSampleRate != OutputSampleRate);
+			const bool bNeedsSRC = !FMath::IsNearlyEqual(InputSampleRate, OutputSampleRate);
 			const bool bNeedsUpmix = (NumInputChannels == 1);
 			const bool bNeedsDeinterleave = !bNeedsUpmix;
 
 
-
 			// Decode audio if we need to (see if we are done)
 			while (bIsPlaying && !bDecoderIsDone && (CircularDecoderOutputBuffer.Num() < NumSamplesToDecode))
 			{
@@ -175,7 +187,8 @@ namespace Metasound
 				const int32 NumSamplesDecoded = DecoderTrio.Output->PopAudio(MakeArrayView(TempBufferA.GetData(), NumSamplesInDecodeBlock), Details);
 
 				// push that (interleaved) audio to the (interleaved) circular buffer
-				CircularDecoderOutputBuffer.Push(TempBufferA.GetData(), NumSamplesDecoded);
+				int32 NumPushed = CircularDecoderOutputBuffer.Push(TempBufferA.GetData(), NumSamplesDecoded);
+				ensure(NumPushed == NumSamplesDecoded); // there will be a discontinuity in the output because our CircularDecoderOutputBuffer was not large enough!
 			}
 
 			// now that we have enough audio decoded, pop off the circular buffer into an interleaved, pre-src temp buffer
@@ -275,7 +288,6 @@ namespace Metasound
 
 		bool bIsPlaying{ false };
 		bool bDecoderIsDone{ false };
-
 	};
 
 	TUniquePtr<IOperator> FWavePlayerNode::FOperatorFactory::CreateOperator(
