Add pitch control to Wave Player Node

diff --git a/Engine/Plugins/Runtime/Metasound/Source/MetasoundEngine/Private/MetasoundWave.cpp b/Engine/Plugins/Runtime/Metasound/Source/MetasoundEngine/Private/MetasoundWave.cpp
--- a/Engine/Plugins/Runtime/Metasound/Source/MetasoundEngine/Private/MetasoundWave.cpp
+++ b/Engine/Plugins/Runtime/Metasound/Source/MetasoundEngine/Private/MetasoundWave.cpp
@@ -10,61 +10,210 @@
 #include "AudioDevice.h"
 
 
-namespace Metasound
-{	
-	FWaveAsset::FDecoderInputPtr FWaveAsset::CreateDecoderInput() const
+// namespace Metasound
+// {	
+// 	FWaveAsset::FDecoderInputPtr FWaveAsset::CreateDecoderInput() const
+// 	{
+// 		if (SoundWaveProxy.IsValid())
+// 		{
+// 			FName OldFormat = SoundWaveProxy->GetRuntimeFormat();
+// 
+// 			TUniquePtr<Audio::IDecoderInput> Input = Audio::CreateBackCompatDecoderInput(OldFormat, *SoundWaveProxy);
+// 
+// 			return FWaveAsset::FDecoderInputPtr(Input.Release());
+// 		}
+// 
+// 		return nullptr;
+// 	}
+// 
+// 
+// 	FWaveAsset::FDecoderTrio FWaveAsset::CreateDecoderTrio(const float OutputSampleRate, const int32 NumFramesPerBlock) const
+// 	{
+// 		using namespace Audio;
+// 
+// 		if (!IsSoundWaveValid())
+// 		{
+// 			return { };
+// 		}
+// 
+// 		FWaveAsset::FDecoderInputPtr Input;
+// 		if (!(Input = FWaveAsset::CreateDecoderInput()))
+// 		{
+// 			return { };
+// 		}
+// 
+// 		ICodecRegistry::FCodecPtr Codec = ICodecRegistry::Get().FindCodecByParsingInput(Input.Get());
+// 		if (!Codec)
+// 		{
+// 			return { };
+// 		}
+// 
+// 		IDecoderOutput::FRequirements Reqs
+// 		{
+// 			Float32_Interleaved,
+// 			NumFramesPerBlock,
+// 			static_cast<int32>(OutputSampleRate)
+// 		};
+// 
+// 		TUniquePtr<IDecoderOutput> Output = IDecoderOutput::Create(Reqs);
+// 		TUniquePtr<IDecoder> Decoder = Codec->CreateDecoder(Input.Get(), Output.Get());
+// 
+// 		return
+// 		{
+// 			  MoveTemp(Input)
+// 			, MoveTemp(Output)
+// 			, MoveTemp(Decoder)
+// 		};
+// 
+// 	}
+// 
+// }
+
+namespace Audio
+{
+	bool FSimpleDecoderWrapper::Initialize(const InitParams& InInitParams, const FSoundWaveProxy& InWave)
 	{
-		if (SoundWaveProxy.IsValid())
+		// validate data
+
+		// initialize input/output data
+		InputSampleRate = InWave.GetSampleRate();
+		OutputSampleRate = InInitParams.OutputSampleRate;
+		FsInToFsOutRatio = (InputSampleRate / OutputSampleRate);
+		MaxPitchShiftRatio = FMath::Pow(2.f, InInitParams.MaxPitchShiftMagnitudeAllowedInOctaves);
+		MaxPitchShiftCents = InInitParams.MaxPitchShiftMagnitudeAllowedInOctaves * 1200.f;
+
+		NumChannels = InWave.GetNumChannels();
+		DecodeBlockSizeInFrames = InInitParams.OutputBlockSizeInFrames;
+		DecodeBlockSizeInSamples = DecodeBlockSizeInFrames * NumChannels;
+
+		// set Circular Buffer capacity
+		CircularDecoderOutputBuffer.SetCapacity(DecodeBlockSizeInSamples * (1.f + FsInToFsOutRatio * MaxPitchShiftRatio));
+		
+		// try to initialize decoders
+		bool bSuccessful = InitializeDecodersInternal(InWave);
+		bDecoderIsDone = !bSuccessful;
+
+		return bSuccessful;
+	}
+
+	uint32 FSimpleDecoderWrapper::GenerateAudio(float* OutputDest, int32 NumOutputFrames, float PitchShiftInCents)
+	{
+		PitchShiftInCents = FMath::Clamp(PitchShiftInCents, -MaxPitchShiftCents, MaxPitchShiftCents);
+		const float PitchShiftRatio = FMath::Pow(2.f, PitchShiftInCents / 1200.f);
+		const float TotalInSamplesOutSamplesRatio = PitchShiftRatio * FsInToFsOutRatio;
+
+		const uint32 NumOutputSamples = NumOutputFrames * NumChannels;
+		uint32 NumSamplesToDecode = NumChannels * NumOutputFrames * TotalInSamplesOutSamplesRatio;
+
+		if (NumSamplesToDecode % 2)
+		{
+			++NumSamplesToDecode;
+		}
+
+		// zero the output
+		FMemory::Memzero(OutputDest, NumOutputSamples * sizeof(float));
+
+		// TODO: calculate the ratio first using SR & PS, then see if the ratio is 1.0f
+		const bool bNeedsSRC = !(FMath::IsNearlyZero(PitchShiftInCents) && FMath::IsNearlyEqual(InputSampleRate, OutputSampleRate));
+
+		// Decode audio until we have enough to satisfy the output (post SRC)
+		while (!bDecoderIsDone && (CircularDecoderOutputBuffer.Num() < NumSamplesToDecode))
 		{
-			FName OldFormat = SoundWaveProxy->GetRuntimeFormat();
+			// get more audio from the decoder
+			Audio::IDecoderOutput::FPushedAudioDetails Details;
+			bDecoderIsDone = (Decoder->Decode() == Audio::IDecoder::EDecodeResult::Finished);
 
-			TUniquePtr<Audio::IDecoderInput> Input = Audio::CreateBackCompatDecoderInput(OldFormat, *SoundWaveProxy);
+			PreSrcBuffer.Reset(DecodeBlockSizeInSamples);
+			PreSrcBuffer.AddZeroed(DecodeBlockSizeInSamples);
+			const int32 NumSamplesDecoded = Output->PopAudio(MakeArrayView(PreSrcBuffer.GetData(), DecodeBlockSizeInSamples), Details);
 
-			return FWaveAsset::FDecoderInputPtr(Input.Release());
+			// push that (interleaved) audio to the (interleaved) circular buffer
+			int32 NumPushed = CircularDecoderOutputBuffer.Push(PreSrcBuffer.GetData(), NumSamplesDecoded);
+			ensure(NumPushed == NumSamplesDecoded); // there will be a discontinuity in the output because our CircularDecoderOutputBuffer was not large enough!
 		}
 
-		return nullptr;
-	}
+		// now that we have enough audio decoded, pop off the circular buffer into an interleaved, pre-src temp buffer
+		// (It is possible that CircularDecoderOutputBuffer.Num() < NumOutputSamples if the decoder is dry...)
+		const uint32 NumSamplesToPop = FMath::Min(CircularDecoderOutputBuffer.Num(), NumSamplesToDecode);
+		const uint32 NumPostSrcSamples = NumOutputFrames * NumChannels;
 
+		PreSrcBuffer.Reset(NumSamplesToPop);
+		PreSrcBuffer.AddZeroed(NumSamplesToPop);
 
-	FWaveAsset::FDecoderTrio FWaveAsset::CreateDecoderTrio(const float OutputSampleRate, const int32 NumFramesPerBlock) const
-	{
-		using namespace Audio;
+		// TODO: if we don't need SRC, just pop directly into the output buffer and call it a day
+		// pop into pre-src buffer
+		uint32 NumSamplesPopped = CircularDecoderOutputBuffer.Pop(PreSrcBuffer.GetData(), NumSamplesToPop);
+		ensure(NumSamplesPopped == NumSamplesToPop); // Circular buffer is too small! there will be discontinuities in the output
 
-		if (!IsSoundWaveValid())
+		// holds result of sample rate conversion
+		float* PostSrcBufferPtr = PreSrcBuffer.GetData(); // assume no SRC
+
+		int32 NumFramesConverted = NumSamplesToPop / NumChannels;
+
+		if (bNeedsSRC)
+		{
+			// perform SRC
+			Resampler.Init(Audio::EResamplingMethod::Linear, 1.f / TotalInSamplesOutSamplesRatio, NumChannels);
+
+			// something is wrong
+			int32 Error = Resampler.ProcessAudio(PreSrcBuffer.GetData(), NumSamplesToPop / NumChannels, false, OutputDest, NumOutputFrames, NumFramesConverted);
+			ensure(!Error);
+
+			if (Error)
+			{
+				bDecoderIsDone = false;
+				Input.Reset();
+				Output.Reset();
+				Decoder.Reset();
+			}
+		}
+		else
 		{
-			return { };
+			FMemory::Memcpy(OutputDest, PreSrcBuffer.GetData(), NumOutputSamples * sizeof(float));
 		}
 
-		FWaveAsset::FDecoderInputPtr Input;
-		if (!(Input = FWaveAsset::CreateDecoderInput()))
+		if (NumFramesConverted < (int32)NumOutputFrames)
 		{
-			return { };
+			bDecoderIsDone = true;
 		}
 
+		return NumFramesConverted;
+	}
+
+	bool FSimpleDecoderWrapper::InitializeDecodersInternal(const FSoundWaveProxy& Wave)
+	{
+		// Input:
+		FName OldFormat = Wave.GetRuntimeFormat();
+		Input = MakeShareable(Audio::CreateBackCompatDecoderInput(OldFormat, Wave).Release());
+
+		if (!Input)
+		{
+			return false;
+		}
+		
+		// aquire codec:
 		ICodecRegistry::FCodecPtr Codec = ICodecRegistry::Get().FindCodecByParsingInput(Input.Get());
 		if (!Codec)
 		{
-			return { };
+			return false;
 		}
 
+		// specify requirements
 		IDecoderOutput::FRequirements Reqs
 		{
 			Float32_Interleaved,
-			NumFramesPerBlock,
+			DecodeBlockSizeInFrames,
 			static_cast<int32>(OutputSampleRate)
 		};
 
-		TUniquePtr<IDecoderOutput> Output = IDecoderOutput::Create(Reqs);
-		TUniquePtr<IDecoder> Decoder = Codec->CreateDecoder(Input.Get(), Output.Get());
+		// Output:
+		Output = IDecoderOutput::Create(Reqs);
 
-		return
-		{
-			  MoveTemp(Input)
-			, MoveTemp(Output)
-			, MoveTemp(Decoder)
-		};
+		// Decoder:
+		Decoder = Codec->CreateDecoder(Input.Get(), Output.Get());
 
+		// return true if all the components were succesfully create
+		return Input.IsValid() && Output.IsValid() && Decoder.IsValid();
 	}
 
-}
+} // namesapce Audio
diff --git a/Engine/Plugins/Runtime/Metasound/Source/MetasoundEngine/Private/MetasoundWavePlayerNode.cpp b/Engine/Plugins/Runtime/Metasound/Source/MetasoundEngine/Private/MetasoundWavePlayerNode.cpp
--- a/Engine/Plugins/Runtime/Metasound/Source/MetasoundEngine/Private/MetasoundWavePlayerNode.cpp
+++ b/Engine/Plugins/Runtime/Metasound/Source/MetasoundEngine/Private/MetasoundWavePlayerNode.cpp
@@ -52,28 +52,24 @@ namespace Metasound
 			const FOperatorSettings& InSettings,
 			const FWaveAssetReadRef& InWave,
 			const FTriggerReadRef& InTrigger,
-			const float& InGraphSamplerate)
+			const FFloatReadRef InPitchShiftCents)
 			: OperatorSettings(InSettings)
 			, TrigIn(InTrigger)
 			, Wave(InWave)
+			, PitchShiftCents(InPitchShiftCents)
 			, AudioBufferL(FAudioBufferWriteRef::CreateNew(InSettings))
 			, AudioBufferR(FAudioBufferWriteRef::CreateNew(InSettings))
 			, TrigggerOnDone(FTriggerWriteRef::CreateNew(InSettings))
-			, InputSampleRate(InWave->SoundWaveProxy.IsValid()? InWave->SoundWaveProxy->GetSampleRate() : -1)
-			, OutputSampleRate(InGraphSamplerate)
+			, OutputSampleRate(InSettings.GetSampleRate())
 			, OutputBlockSizeInFrames(InSettings.GetNumFramesPerBlock())
 		{
-			check(OutputSampleRate && InputSampleRate); // divide by zero when calculating SrcRatio!
+			check(OutputSampleRate);
 			check(AudioBufferL->Num() == OutputBlockSizeInFrames && AudioBufferR->Num() == OutputBlockSizeInFrames);
 
 			if (Wave->SoundWaveProxy.IsValid())
 			{
+				ResetDecoder();
 				CurrentSoundWaveName = Wave->SoundWaveProxy->GetFName();
-
-				DecoderTrio = Wave->CreateDecoderTrio(OutputSampleRate, OutputBlockSizeInFrames);
-
-				const int32 NumChannels = Wave->SoundWaveProxy->GetNumChannels();
-				CircularDecoderOutputBuffer.SetCapacity(OutputBlockSizeInFrames * NumChannels * 2);
 			}
 
 
@@ -84,6 +80,7 @@ namespace Metasound
 			FDataReferenceCollection InputDataReferences;
 			InputDataReferences.AddDataReadReference(TEXT("Audio"), FWaveAssetReadRef(Wave));
 			InputDataReferences.AddDataReadReference(TEXT("TrigIn"), FBopReadRef(TrigIn));
+			InputDataReferences.AddDataReadReference(TEXT("PitchShiftCents"), FFloatReadRef(PitchShiftCents));
 			return InputDataReferences;
 		}
 
@@ -103,7 +100,6 @@ namespace Metasound
 
 			// see if we have a new soundwave input
 			FName NewSoundWaveName = FName();
-			int32 NumChannels = 1;
 			if (Wave->IsSoundWaveValid())
 			{
 				NewSoundWaveName = Wave->SoundWaveProxy->GetFName();
@@ -111,12 +107,7 @@ namespace Metasound
 
 			if (NewSoundWaveName != CurrentSoundWaveName)
 			{
-				DecoderTrio = Wave->CreateDecoderTrio(OutputSampleRate, OutputBlockSizeInFrames);
-				InputSampleRate = Wave->SoundWaveProxy->GetSampleRate();
-
-				NumChannels = Wave->SoundWaveProxy->GetNumChannels();
-				CircularDecoderOutputBuffer.SetCapacity(OutputBlockSizeInFrames * NumChannels * 4);
-
+				ResetDecoder();
 				CurrentSoundWaveName = NewSoundWaveName;
 			}
 
@@ -124,6 +115,11 @@ namespace Metasound
 			FMemory::Memzero(AudioBufferL->GetData(), OutputBlockSizeInFrames * sizeof(float));
 			FMemory::Memzero(AudioBufferR->GetData(), OutputBlockSizeInFrames * sizeof(float));
 
+			if (!Decoder.CanGenerateAudio())
+			{
+				return;
+			}
+
 			TrigIn->ExecuteBlock(
 				// OnPreBop
 				[&](int32 StartFrame, int32 EndFrame)
@@ -138,11 +134,11 @@ namespace Metasound
 				{
 					if (!bIsPlaying)
 					{
-						bIsPlaying = true;
+						bIsPlaying = Decoder.CanGenerateAudio();
 					}
 					else
 					{
-						DecoderTrio = Wave->CreateDecoderTrio(OutputSampleRate, OutputBlockSizeInFrames);
+						ResetDecoder();
 					}
 
 					ExecuteInternal(StartFrame, EndFrame);
@@ -150,106 +146,63 @@ namespace Metasound
 			);
 		}
 
-
-		void ExecuteInternal(int32 StartFrame, int32 EndFrame)
+		bool ResetDecoder()
 		{
-			float* FinalOutputLeft = AudioBufferL->GetData() + StartFrame;
-			float* FinalOutputRight = AudioBufferR->GetData() + StartFrame;
+			Audio::FSimpleDecoderWrapper::InitParams Params;
+			Params.OutputBlockSizeInFrames = OutputBlockSizeInFrames;
+			Params.OutputSampleRate = OutputSampleRate;
+			Params.MaxPitchShiftMagnitudeAllowedInOctaves = 4.f;
 
-			const float FsInToFsOutRatio = (InputSampleRate / OutputSampleRate); // TODO: account for pitch param
-
-			const uint32 NumInputChannels = Wave->SoundWaveProxy.IsValid() ? Wave->SoundWaveProxy->GetNumChannels() : 1;
-			const uint32 NumOutputChannels = 2; // forcing stereo output
-			const uint32 NumOutputFrames = EndFrame - StartFrame;
-			const uint32 NumOutputSamples = NumOutputFrames * NumOutputChannels;
-			const uint32 NumSamplesInDecodeBlock = OutputBlockSizeInFrames * NumInputChannels;
-			uint32 NumSamplesToDecode = NumInputChannels * NumOutputFrames * FsInToFsOutRatio;
-
-			if (NumSamplesToDecode % 2)
+			if (false == Wave->IsSoundWaveValid())
 			{
-				++NumSamplesToDecode;
+				return false;
 			}
 
-			const bool bNeedsSRC = !FMath::IsNearlyEqual(InputSampleRate, OutputSampleRate);
-			const bool bNeedsUpmix = (NumInputChannels == 1);
-			const bool bNeedsDeinterleave = !bNeedsUpmix;
-
-
-			// Decode audio if we need to (see if we are done)
-			while (bIsPlaying && !bDecoderIsDone && (CircularDecoderOutputBuffer.Num() < NumSamplesToDecode))
-			{
-				// get more audio from the decoder
-				Audio::IDecoderOutput::FPushedAudioDetails Details;
-				bDecoderIsDone = (DecoderTrio.Decoder->Decode() == Audio::IDecoder::EDecodeResult::Finished);
-
-				TempBufferA.Reset(NumSamplesInDecodeBlock);
-				TempBufferA.AddZeroed(NumSamplesInDecodeBlock);
-				const int32 NumSamplesDecoded = DecoderTrio.Output->PopAudio(MakeArrayView(TempBufferA.GetData(), NumSamplesInDecodeBlock), Details);
-
-				// push that (interleaved) audio to the (interleaved) circular buffer
-				int32 NumPushed = CircularDecoderOutputBuffer.Push(TempBufferA.GetData(), NumSamplesDecoded);
-				ensure(NumPushed == NumSamplesDecoded); // there will be a discontinuity in the output because our CircularDecoderOutputBuffer was not large enough!
-			}
+			return Decoder.Initialize(Params, *Wave->SoundWaveProxy);
+		}
 
-			// now that we have enough audio decoded, pop off the circular buffer into an interleaved, pre-src temp buffer
-			// (It is possible that CircularDecoderOutputBuffer.Num() < NumOutputSamples if the decoder is dry...)
-			const uint32 NumSamplesToPop = FMath::Min(CircularDecoderOutputBuffer.Num(), NumSamplesToDecode);
 
-			// (...if that's the case, it means the sound is done)
-			if (NumSamplesToPop < NumOutputChannels)
-			{
-				bIsPlaying = false;
-				TrigggerOnDone->BopFrame(NumSamplesToPop / NumInputChannels);
-			}
+		void ExecuteInternal(int32 StartFrame, int32 EndFrame)
+		{
+			// shouldn't be calling this function if we don't have access to a valid SoundWave
+			ensure(Wave->IsSoundWaveValid());
+			ensure(Wave->SoundWaveProxy->GetNumChannels() <= 2); // only support mono or stereo inputs
+			ensure(Decoder.CanGenerateAudio());
 
-			// TODO: special-cases to reduce temp buffer usage where we don't need them
+			// note: output is hard-coded to stereo (dual-mono)
+			float* FinalOutputLeft = AudioBufferL->GetData() + StartFrame;
+			float* FinalOutputRight = AudioBufferR->GetData() + StartFrame;
 
-			const uint32 NumPostSrcSamples = NumOutputFrames * NumInputChannels;
-			TempBufferA.Reset(NumSamplesToPop); // pre-src buffer 
-			TempBufferA.AddZeroed(NumSamplesToPop);
+			const int32 NumInputChannels = Wave->SoundWaveProxy->GetNumChannels();
+			const bool bNeedsUpmix = (NumInputChannels == 1);
+			const bool bNeedsDeinterleave = !bNeedsUpmix;
 
-			// pop into pre-src buffer
-			CircularDecoderOutputBuffer.Pop(TempBufferA.GetData(), NumSamplesToPop);
+			const int32 NumOutputFrames = (EndFrame - StartFrame);
+			const int32 NumSamplesToGenerate = NumOutputFrames * NumInputChannels; // (stereo output)
 
-			// holds result of sample rate conversion
-			float* PostSrcBufferPtr = TempBufferA.GetData(); // assume no SRC
+			PostSrcBuffer.Reset(NumSamplesToGenerate);
+			PostSrcBuffer.AddZeroed(NumSamplesToGenerate);
+			float* PostSrcBufferPtr = PostSrcBuffer.GetData();
 
-			int32 NumFramesConverted = NumSamplesToPop / NumInputChannels;
+			int32 NumFramesDecoded = Decoder.GenerateAudio(PostSrcBufferPtr, NumOutputFrames, *PitchShiftCents); // TODO: pitch shift
 
-			if (bNeedsSRC)
+			// TODO: handle decoder having completed during it's decode
+			if (!Decoder.CanGenerateAudio() ||  (NumFramesDecoded < NumOutputFrames))
 			{
-				// post-src buffer
-				TempBufferB.Reset(NumPostSrcSamples);
-				TempBufferB.AddZeroed(NumPostSrcSamples);
-				PostSrcBufferPtr = TempBufferB.GetData();
-
-				// perform SRC
-				Resampler.Init(Audio::EResamplingMethod::Linear, 1.f / FsInToFsOutRatio, NumInputChannels);
-
-				int32 Error = Resampler.ProcessAudio(TempBufferA.GetData(), NumSamplesToPop / NumInputChannels, false, PostSrcBufferPtr, NumOutputFrames, NumFramesConverted);
-				if (Error)
-				{
-					bIsPlaying = false;
-					bDecoderIsDone = false;
-					DecoderTrio = {};
-				}
+				bIsPlaying = false;
+				TrigggerOnDone->BopFrame(StartFrame + NumFramesDecoded);
 			}
 
-			// We may not have NumOutputSamples to play, and our output buffers have already been zero-ed out.
-			// at this point we are outputting NumFramesConverted
-
-			// perform channel conversion (output is forced to be stereo)
-			// mono->stereo?
 			if (bNeedsUpmix)
 			{
 				// TODO: attenuate by -3 dB?
 				// copy to Left & Right output buffers
-				FMemory::Memcpy(FinalOutputLeft, PostSrcBufferPtr, sizeof(float) * NumFramesConverted);
-				FMemory::Memcpy(FinalOutputRight, PostSrcBufferPtr, sizeof(float) * NumFramesConverted);
+				FMemory::Memcpy(FinalOutputLeft, PostSrcBufferPtr, sizeof(float) * NumOutputFrames);
+				FMemory::Memcpy(FinalOutputRight, PostSrcBufferPtr, sizeof(float) * NumOutputFrames);
 			}
 			else if (bNeedsDeinterleave)
 			{
-				for (int32 i = 0; i < NumFramesConverted; ++i)
+				for (int32 i = 0; i < NumOutputFrames; ++i)
 				{
 					// de-interleave each stereo frame into output buffers
 					FinalOutputLeft[i] = PostSrcBufferPtr[(i << 1)];
@@ -264,30 +217,22 @@ namespace Metasound
 		// i/o
 		FTriggerReadRef TrigIn;
 		FWaveAssetReadRef Wave;
+		FFloatReadRef PitchShiftCents;
 
 		FAudioBufferWriteRef AudioBufferL;
 		FAudioBufferWriteRef AudioBufferR;
 		FTriggerWriteRef TrigggerOnDone;
 
-		// src
-		Audio::FResampler Resampler;
-
-		// Decoder/IO. 
-		FWaveAsset::FDecoderTrio DecoderTrio;
-
-		Audio::TCircularAudioBuffer<float> CircularDecoderOutputBuffer;
-		TArray<float> TempBufferA;
-		TArray<float> TempBufferB;
+		// source decode
+		TArray<float> PostSrcBuffer;
+		Audio::FSimpleDecoderWrapper Decoder;
 
 		FName CurrentSoundWaveName{ };
-
-		float InputSampleRate{ 0.f };
+			
 		const float OutputSampleRate{ 0.f };
 		const int32 OutputBlockSizeInFrames{ 0 };
 		
-
 		bool bIsPlaying{ false };
-		bool bDecoderIsDone{ false };
 	};
 
 	TUniquePtr<IOperator> FWavePlayerNode::FOperatorFactory::CreateOperator(
@@ -300,11 +245,9 @@ namespace Metasound
 
 		const FDataReferenceCollection& InputDataRefs = InParams.InputDataReferences;
 
-		// Trigger input
 		FTriggerReadRef TriggerPlay = InputDataRefs.GetDataReadReferenceOrConstruct<FTrigger>(TEXT("TrigIn"), InParams.OperatorSettings);
-
-		// Initialize decoder
 		FWaveAssetReadRef Wave = InputDataRefs.GetDataReadReferenceOrConstruct<FWaveAsset>(TEXT("Wave"));
+		FFloatReadRef PitchShiftCents = InputDataRefs.GetDataReadReferenceOrConstruct<float>(TEXT("PitchShiftCents"));
 		
 		if (!Wave->IsSoundWaveValid())
 		{
@@ -320,7 +263,7 @@ namespace Metasound
 			  InParams.OperatorSettings
 			, Wave
 			, TriggerPlay
-			, InParams.OperatorSettings.GetSampleRate()
+			, PitchShiftCents
 			);
 	}
 
@@ -329,7 +272,8 @@ namespace Metasound
 		return FVertexInterface(
 			FInputVertexInterface(
 				TInputDataVertexModel<FWaveAsset>(TEXT("Wave"), LOCTEXT("WaveTooltip", "The Wave to be decoded")),
-				TInputDataVertexModel<FTrigger>(TEXT("TrigIn"), LOCTEXT("TrigInTooltip", "Trigger the playing of the input wave."))
+				TInputDataVertexModel<FTrigger>(TEXT("TrigIn"), LOCTEXT("TrigInTooltip", "Trigger the playing of the input wave.")),
+				TInputDataVertexModel<float>(TEXT("PitchShiftCents"), LOCTEXT("PitchShiftCentsTooltip", "Pitch Shift in cents."))
 			),
 			FOutputVertexInterface(
 				TOutputDataVertexModel<FAudioBuffer>(TEXT("AudioLeft"), LOCTEXT("AudioTooltip", "The output audio")),
diff --git a/Engine/Plugins/Runtime/Metasound/Source/MetasoundEngine/Public/MetasoundWave.h b/Engine/Plugins/Runtime/Metasound/Source/MetasoundEngine/Public/MetasoundWave.h
--- a/Engine/Plugins/Runtime/Metasound/Source/MetasoundEngine/Public/MetasoundWave.h
+++ b/Engine/Plugins/Runtime/Metasound/Source/MetasoundEngine/Public/MetasoundWave.h
@@ -9,20 +9,6 @@
 #include "IAudioCodecRegistry.h"
 #include "IAudioCodec.h"
 
-// Forward declares
-namespace Audio
-{	
-	struct IDecoderInput;
-}
-
-class USoundWave;
-
-namespace Audio
-{
-	class ICodecRegistry;
-	struct IDecoderOutput;
-	struct IDecoder;
-}
 
 namespace Metasound
 {
@@ -35,15 +21,6 @@ namespace Metasound
 	class METASOUNDENGINE_API FWaveAsset
 	{
 	public:
-		using FDecoderInputPtr = TSharedPtr<Audio::IDecoderInput, ESPMode::ThreadSafe>;
-
-		struct FDecoderTrio
-		{
-			FDecoderInputPtr Input;
-			TUniquePtr<Audio::IDecoderOutput> Output;
-			TUniquePtr<Audio::IDecoder> Decoder;
-		};
-
 		TUniquePtr<FSoundWaveProxy> SoundWaveProxy;
 
 		FWaveAsset() = default;
@@ -64,11 +41,6 @@ namespace Metasound
 			}
 		}
 
-		FDecoderInputPtr CreateDecoderInput() const;
-
-
-		FDecoderTrio CreateDecoderTrio(const float OutputSampleRate, const int32 NumFramesPerBlock) const;
-
 		bool IsSoundWaveValid() const
 		{
 			return SoundWaveProxy.IsValid();
@@ -78,3 +50,66 @@ namespace Metasound
 
 	DECLARE_METASOUND_DATA_REFERENCE_TYPES(FWaveAsset, METASOUNDENGINE_API, FWaveAssetTypeInfo, FWaveAssetReadRef, FWaveAssetWriteRef)
 }
+
+namespace Audio
+{
+	// Forward declares
+	class ICodecRegistry;
+	struct IDecoderInput;
+	struct IDecoderOutput;
+	struct IDecoder;
+
+	class FSimpleDecoderWrapper
+	{
+	public:
+		struct InitParams
+		{
+			float OutputSampleRate;
+			uint32 OutputBlockSizeInFrames;
+			float MaxPitchShiftMagnitudeAllowedInOctaves = 4.f;
+		};
+
+		bool CanGenerateAudio() const
+		{
+			return !bDecoderIsDone && Input.IsValid() && Output.IsValid() && Decoder.IsValid();
+		}
+
+		bool Initialize(const InitParams& InInitParams, const FSoundWaveProxy& InWave);
+
+		// returns number of samples written.   
+		uint32 GenerateAudio(float* OutputDest, int32 NumOutputFrames, float PitchShiftInCents = 0.f);
+
+	private:
+		// actual decoder objects
+		TUniquePtr<Audio::IDecoder> Decoder;
+		TUniquePtr<Audio::IDecoderOutput> Output;
+		TSharedPtr<Audio::IDecoderInput, ESPMode::ThreadSafe> Input;
+
+		// init helper for decoders
+		bool InitializeDecodersInternal(const FSoundWaveProxy& Wave);
+
+		// SRC object
+		Audio::FResampler Resampler;
+
+		// buffers
+		TArray<float> PreSrcBuffer;
+		Audio::TCircularAudioBuffer<float> CircularDecoderOutputBuffer;
+
+		// meta data:
+		float InputSampleRate;
+		float OutputSampleRate;
+		float FsInToFsOutRatio;
+		float MaxPitchShiftCents;
+		float MaxPitchShiftRatio;
+
+		uint32 NumChannels;
+		uint32 DecodeBlockSizeInFrames;
+		uint32 DecodeBlockSizeInSamples;
+		
+
+		bool bDecoderIsDone{ true };
+
+	}; // class FSimpleDecoderWrapper
+
+
+} // namespace Audio
\ No newline at end of file
diff --git a/Engine/Source/Runtime/SignalProcessing/Private/BufferVectorOperations.cpp b/Engine/Source/Runtime/SignalProcessing/Private/BufferVectorOperations.cpp
--- a/Engine/Source/Runtime/SignalProcessing/Private/BufferVectorOperations.cpp
+++ b/Engine/Source/Runtime/SignalProcessing/Private/BufferVectorOperations.cpp
@@ -660,7 +660,7 @@ namespace Audio
 
 		for (int32 i = 0; i < NumSamples; i += 4)
 		{
-			VectorRegister Input = VectorAbs(VectorLoadAligned(&Buffer[i]));
+			VectorRegister Input = VectorLoadAligned(&Buffer[i]);
 			Sum = VectorAdd(Sum, Input);
 		}
 
