CIS fix (adding static cast)

diff --git a/Engine/Plugins/Runtime/Metasound/Source/MetasoundEngine/Private/MetasoundWave.cpp b/Engine/Plugins/Runtime/Metasound/Source/MetasoundEngine/Private/MetasoundWave.cpp
--- a/Engine/Plugins/Runtime/Metasound/Source/MetasoundEngine/Private/MetasoundWave.cpp
+++ b/Engine/Plugins/Runtime/Metasound/Source/MetasoundEngine/Private/MetasoundWave.cpp
@@ -202,7 +202,7 @@ namespace Audio
 		IDecoderOutput::FRequirements Reqs
 		{
 			Float32_Interleaved,
-			DecodeBlockSizeInFrames,
+			static_cast<int32>(DecodeBlockSizeInFrames),
 			static_cast<int32>(OutputSampleRate)
 		};
 
