Mixer Node for Metasounds

diff --git a/Engine/Plugins/Runtime/Metasound/Source/MetasoundStandardNodes/Private/MetasoundMixerNode.cpp b/Engine/Plugins/Runtime/Metasound/Source/MetasoundStandardNodes/Private/MetasoundMixerNode.cpp
new file mode 100644
--- /dev/null
+++ b/Engine/Plugins/Runtime/Metasound/Source/MetasoundStandardNodes/Private/MetasoundMixerNode.cpp
@@ -0,0 +1,187 @@
+// Copyright Epic Games, Inc. All Rights Reserved.
+
+#include "MetasoundMixerNode.h"
+
+#include "MetasoundAudioBuffer.h"
+#include "MetasoundExecutableOperator.h"
+#include "MetasoundFacade.h"
+#include "MetasoundFrequency.h"
+#include "MetasoundNodeRegistrationMacro.h"
+#include "MetasoundPrimitives.h"
+#include "MetasoundVertex.h"
+#include "MetasoundBop.h"
+#include "MetasoundGain.h"
+#include "MetasoundOperatorSettings.h"
+
+#include "DSP/Dsp.h"
+
+#define LOCTEXT_NAMESPACE "MetasoundStandardNodes"
+
+namespace Metasound
+{
+	class FMixerOperator : public TExecutableOperator<FMixerOperator>
+	{
+	public:
+		static TUniquePtr<IOperator> CreateOperator(const FCreateOperatorParams& InParams, FBuildErrorArray& OutErrors);
+		static const FNodeInfo& GetNodeInfo();
+		static FVertexInterface DeclareVertexInterface();	
+
+		FMixerOperator(const FOperatorSettings& InSettings, TArray<FAudioBufferReadRef>&& InBuffers, TArray<FGainReadRef>&& InGains);
+
+		virtual FDataReferenceCollection GetInputs() const override;
+		virtual FDataReferenceCollection GetOutputs() const override;
+
+		void Execute();
+
+	private:		
+		// TODO: Remove when we get dynamic pin support. (if you update this, update DeclareVertexInterface also).
+		static constexpr const int32 NumHardcodedInputs = 4;
+
+		static FString GetInputAudioPinName(int32 InPinIndex) 
+		{
+			return FString::Printf(TEXT("In %d"), InPinIndex + 1);
+		}
+		static FString GetInputGainPinName(int32 InPinIndex)
+		{
+			return FString::Printf(TEXT("Gain %d"), InPinIndex + 1);
+		}
+
+		TArray<FAudioBufferReadRef> AudioBuffersIn;
+		TArray<FGainReadRef> BufferMixGains;
+		FAudioBufferWriteRef AudioBufferOut;
+		
+		static constexpr const TCHAR* AudioOutPinName = TEXT("Out");
+	};
+
+	FMixerOperator::FMixerOperator(const FOperatorSettings& InSettings, TArray<FAudioBufferReadRef>&& InBuffers, TArray<FGainReadRef>&& InGains)
+		: AudioBuffersIn(MoveTemp(InBuffers))
+		, BufferMixGains(MoveTemp(InGains))
+		, AudioBufferOut(FAudioBufferWriteRef::CreateNew(InSettings))
+	{
+		check(AudioBufferOut->Num() == InSettings.GetNumFramesPerBlock());
+		check(BufferMixGains.Num() == AudioBuffersIn.Num());
+	}
+
+	FDataReferenceCollection FMixerOperator::GetInputs() const
+	{
+		FDataReferenceCollection InputDataReferences;
+		for (int32 i = 0; i < AudioBuffersIn.Num(); ++i)
+		{
+			InputDataReferences.AddDataReadReference(GetInputAudioPinName(i), AudioBuffersIn[i]);
+			InputDataReferences.AddDataReadReference(GetInputGainPinName(i), BufferMixGains[i]);
+		}
+		return InputDataReferences;
+	}
+
+	FDataReferenceCollection FMixerOperator::GetOutputs() const
+	{
+		FDataReferenceCollection OutputDataReferences;
+		OutputDataReferences.AddDataReadReference(AudioOutPinName, FAudioBufferReadRef(AudioBufferOut));
+		return OutputDataReferences;
+	}
+
+	void FMixerOperator::Execute()
+	{	
+		AudioBufferOut->Zero();
+		for (int32 i = 0; i < AudioBuffersIn.Num(); ++i)
+		{
+			Audio::MixInBufferFast(AudioBuffersIn[i]->GetData(), AudioBufferOut->GetData(), AudioBufferOut->Num(), *BufferMixGains[i]);
+		}
+	}
+
+	FVertexInterface FMixerOperator::DeclareVertexInterface()
+	{
+		static const FVertexInterface Interface(
+			FInputVertexInterface(
+				// 1
+				TInputDataVertexModel<FAudioBuffer>(GetInputAudioPinName(0), LOCTEXT("MixerInputDescription1", "Audio Input 1 of the Mixer")),		// Can't have dynamic pins yet, so just hard code 4.
+				TInputDataVertexModel<FGain>(GetInputGainPinName(0), LOCTEXT("MixerGainDescription1", "")),
+
+				// 2
+				TInputDataVertexModel<FAudioBuffer>(GetInputAudioPinName(1), LOCTEXT("MixerInputDescription2", "Audio Input 2 of the Mixer")),	
+				TInputDataVertexModel<FGain>(GetInputGainPinName(1), LOCTEXT("MixerGainDescription2", "")),
+
+				// 3
+				TInputDataVertexModel<FAudioBuffer>(GetInputAudioPinName(2), LOCTEXT("MixerInputDescription3", "Audio Input 3 of the Mixer")),	
+				TInputDataVertexModel<FGain>(GetInputGainPinName(2), LOCTEXT("MixerGainDescription3", "")),
+
+				// 4
+				TInputDataVertexModel<FAudioBuffer>(GetInputAudioPinName(3), LOCTEXT("MixerInputDescription3", "Audio Input 4 of the Mixer")),	
+				TInputDataVertexModel<FGain>(GetInputGainPinName(3), LOCTEXT("MixerGainDescription4", ""))
+			),
+			FOutputVertexInterface(
+				TOutputDataVertexModel<FAudioBuffer>(AudioOutPinName, LOCTEXT("AudioOutTooltip", "Audio Ouput from the mixer"))
+			)
+		);
+
+		return Interface;
+	}
+
+	const FNodeInfo& FMixerOperator::GetNodeInfo()
+	{
+		auto InitNodeInfo = []() -> FNodeInfo
+		{
+			FNodeInfo Info;
+			Info.ClassName = FName(TEXT("Mixer"));
+			Info.MajorVersion = 1;
+			Info.MinorVersion = 0;
+			Info.Description = LOCTEXT("Metasound_MixerNodeDescription", "Mixes 1 or more input signals together");
+			Info.Author = PluginAuthor;
+			Info.PromptIfMissing = PluginNodeMissingPrompt;
+			Info.DefaultInterface = DeclareVertexInterface();
+
+			return Info;
+		};
+
+		static const FNodeInfo Info = InitNodeInfo();
+
+		return Info;
+	}
+
+	FMixerNode::FMixerNode(const FString& InName, float InDefaultMixGainLinear)
+	:	FNodeFacade(InName, TFacadeOperatorClass<FMixerOperator>())
+		, DefaultMixGainLinear(InDefaultMixGainLinear)
+	{
+	}
+
+	FMixerNode::FMixerNode(const FNodeInitData& InInitData)
+		: FMixerNode(InInitData.InstanceName, 1.0f)
+	{
+	}
+
+	float FMixerNode::GetDefaultMixGainLinear() const
+	{
+		return DefaultMixGainLinear;
+	}
+
+	TUniquePtr<IOperator> FMixerOperator::CreateOperator(const FCreateOperatorParams& InParams, FBuildErrorArray& OutErrors)
+	{
+		const FMixerNode& Node = static_cast<const FMixerNode&>(InParams.Node);
+		const FDataReferenceCollection& InputCol = InParams.InputDataReferences;
+		const FOperatorSettings& Settings = InParams.OperatorSettings;
+
+		// TODO: Remove this string manipulation when we get dynamic pin support.
+		TArray<FAudioBufferReadRef> InputBuffers;
+		TArray<FGainReadRef> InputGains;
+		for (int32 i = 0; i < NumHardcodedInputs; ++i)
+		{			
+			FString AudioPinName = GetInputAudioPinName(i);
+			FString GainPinName = GetInputGainPinName(i);
+			if (InputCol.ContainsDataReadReference<FAudioBuffer>(AudioPinName))
+			{
+				// Only Create buffers if there's something connected to it.
+				InputBuffers.Emplace(InputCol.GetDataReadReference<FAudioBuffer>(AudioPinName));
+				
+				// Make sure for every valid/connected pin, we have a corresponding gain, even if its defaulted.
+				InputGains.Emplace(InputCol.GetDataReadReferenceOrConstruct<FGain>(GainPinName, Node.GetDefaultMixGainLinear()));
+			}
+		}
+		
+		// Make a node even if we don't have any inputs to make sure we don't fail graph compilation.
+		return MakeUnique<FMixerOperator>(InParams.OperatorSettings, MoveTemp(InputBuffers), MoveTemp(InputGains));
+	}
+
+	METASOUND_REGISTER_NODE(FMixerNode);
+}
+#undef LOCTEXT_NAMESPACE //MetasoundMixerNode
+
diff --git a/Engine/Plugins/Runtime/Metasound/Source/MetasoundStandardNodes/Private/MetasoundStandardNodesModule.cpp b/Engine/Plugins/Runtime/Metasound/Source/MetasoundStandardNodes/Private/MetasoundStandardNodesModule.cpp
--- a/Engine/Plugins/Runtime/Metasound/Source/MetasoundStandardNodes/Private/MetasoundStandardNodesModule.cpp
+++ b/Engine/Plugins/Runtime/Metasound/Source/MetasoundStandardNodes/Private/MetasoundStandardNodesModule.cpp
@@ -9,12 +9,14 @@
 #include "MetasoundPrimitives.h"
 #include "MetasoundTime.h"
 #include "MetasoundFrequency.h"
+#include "MetasoundGain.h"
 
 REGISTER_METASOUND_DATATYPE(Metasound::FTrigger, "Trigger", ::Metasound::ELiteralType::Boolean)
 REGISTER_METASOUND_DATATYPE(Metasound::FFloatTime, "Time", ::Metasound::ELiteralType::Float)
 REGISTER_METASOUND_DATATYPE(Metasound::FDoubleTime, "Time:HighResolution", ::Metasound::ELiteralType::Float)
 REGISTER_METASOUND_DATATYPE(Metasound::FSampleTime, "Time:SampleResolution", ::Metasound::ELiteralType::Float)
 REGISTER_METASOUND_DATATYPE(Metasound::FFrequency, "Frequency", ::Metasound::ELiteralType::Float)
+REGISTER_METASOUND_DATATYPE(Metasound::FGain, "Gain", ::Metasound::ELiteralType::Float)
 
 namespace Metasound 
 {
diff --git a/Engine/Plugins/Runtime/Metasound/Source/MetasoundStandardNodes/Public/MetasoundGain.h b/Engine/Plugins/Runtime/Metasound/Source/MetasoundStandardNodes/Public/MetasoundGain.h
new file mode 100644
--- /dev/null
+++ b/Engine/Plugins/Runtime/Metasound/Source/MetasoundStandardNodes/Public/MetasoundGain.h
@@ -0,0 +1,124 @@
+// Copyright Epic Games, Inc. All Rights Reserved.
+
+#pragma once
+
+#include "MetasoundDataReference.h"
+#include "MetasoundDataTypeRegistrationMacro.h"
+
+namespace Metasound
+{
+	/** Representation of a gain */
+	enum class EGainRepresentation  : uint8
+	{
+		/** Gain in linear scale  */
+		Linear,
+
+		/** Gain is in decibels */
+		Decibels
+	};
+
+	/** FGain represents a gain value. It provides clearly defined 
+	 * getters and setters as well as several convenience functions for 
+	 * converting gain to decibels.
+	 */
+	class FGain
+	{
+		public:
+			static constexpr const float DefaultGain = 1.0f;
+
+			/** Construct gain of 1  */
+			FGain()
+			:	LinearGain(DefaultGain)
+			{
+			}
+
+			/** Construct a Gain with a given value and representation.
+			 *
+			 * @param InValue - The value of the gain at the given 
+			 * 					representation.
+			 * @param InRep - The representation of the given value.
+			 */
+			FGain(float InValue, EGainRepresentation InRep)
+			:	LinearGain(DefaultGain)
+			{
+				switch (InRep)
+				{
+					case EGainRepresentation::Decibels:
+						SetDecibels(InValue);
+						break;
+					case EGainRepresentation::Linear:
+						SetLinear(InValue);
+						break;
+					default:
+						checkNoEntry();
+						break;
+				}
+			}
+
+			/**
+			 * FGain constructor used to pass in float literals from the metasound frontend.
+			 */
+			FGain(float InValue)
+				: FGain(InValue, EGainRepresentation::Linear)
+			{}
+
+			/** Implicit operator used for math operations. */
+			operator float() const
+			{
+				return LinearGain;
+			}
+
+			FGain& operator+=(const FGain& InOther)
+			{
+				LinearGain += InOther.GetLinear();
+				return *this;
+			}
+
+			FGain& operator-=(const FGain& InOther)
+			{
+				LinearGain -= InOther.GetLinear();
+				return *this;
+			}
+
+			FGain& operator*=(const FGain& InOther)
+			{
+				LinearGain *= InOther.GetLinear();
+				return *this;
+			}
+
+			FGain& operator/=(const FGain& InOther)
+			{
+				LinearGain /= InOther.GetLinear();
+				return *this;
+			}
+
+			/** Set the gain in decibels */
+			void SetDecibels(float InDecibels)
+			{
+				LinearGain = Audio::ConvertToLinear(InDecibels);
+			}
+
+			/** Get the gain in decibels */
+			float GetDecibels() const
+			{
+				return Audio::ConvertToLinear(LinearGain);
+			}
+
+			/** Set the gain in LinearScale */
+			void SetLinear(float InLinearScale)
+			{
+				LinearGain = InLinearScale;
+			}
+
+			/** Return the gain in LinearScale. */
+			float GetLinear() const
+			{
+				return LinearGain;
+			}
+		  
+		private:
+			float LinearGain = DefaultGain;
+	};
+
+	DECLARE_METASOUND_DATA_REFERENCE_TYPES(FGain, METASOUNDSTANDARDNODES_API, FGainTypeInfo, FGainReadRef, FGainWriteRef)
+}
diff --git a/Engine/Plugins/Runtime/Metasound/Source/MetasoundStandardNodes/Public/MetasoundMixerNode.h b/Engine/Plugins/Runtime/Metasound/Source/MetasoundStandardNodes/Public/MetasoundMixerNode.h
new file mode 100644
--- /dev/null
+++ b/Engine/Plugins/Runtime/Metasound/Source/MetasoundStandardNodes/Public/MetasoundMixerNode.h
@@ -0,0 +1,27 @@
+// Copyright Epic Games, Inc. All Rights Reserved.
+
+#pragma once
+
+#include "CoreMinimal.h"
+
+#include "MetasoundAudioBuffer.h"
+#include "MetasoundBuilderInterface.h"
+#include "MetasoundDataReferenceCollection.h"
+#include "MetasoundExecutableOperator.h"
+#include "MetasoundFacade.h"
+#include "MetasoundNode.h"
+#include "MetasoundOperatorInterface.h"
+
+namespace Metasound
+{
+	class METASOUNDSTANDARDNODES_API FMixerNode : public FNodeFacade
+	{
+	public:
+		FMixerNode(const FString& InName, float InDefaultMixGainCents);
+		FMixerNode(const FNodeInitData& InInitData);
+
+		float GetDefaultMixGainLinear() const;
+	private:
+		float DefaultMixGainLinear = 1.0f;
+	};
+}
