Enabling ForceInline UsoundWave path for non-metasound rendering.

diff --git a/Engine/Plugins/Runtime/Metasound/Source/MetasoundEngine/Private/MetasoundWave.cpp b/Engine/Plugins/Runtime/Metasound/Source/MetasoundEngine/Private/MetasoundWave.cpp
--- a/Engine/Plugins/Runtime/Metasound/Source/MetasoundEngine/Private/MetasoundWave.cpp
+++ b/Engine/Plugins/Runtime/Metasound/Source/MetasoundEngine/Private/MetasoundWave.cpp
@@ -74,6 +74,17 @@ namespace Audio
 	bool FSimpleDecoderWrapper::Initialize(const InitParams& InInitParams, const FSoundWaveProxy& InWave)
 	{
 		// validate data
+		if (false == InWave.IsStreaming())
+		{
+			ensureAlwaysMsgf(InWave.IsStreaming(), TEXT("Metasounds does not support Force Inline (sound must be streaming)"));
+
+			Input.Reset();
+			Output.Reset();
+			Decoder.Reset();
+			bDecoderIsDone = false;
+
+			return false;
+		}
 
 		// initialize input/output data
 		InputSampleRate = InWave.GetSampleRate();
diff --git a/Engine/Source/Runtime/AudioMixer/Private/AudioMixerBuffer.cpp b/Engine/Source/Runtime/AudioMixer/Private/AudioMixerBuffer.cpp
--- a/Engine/Source/Runtime/AudioMixer/Private/AudioMixerBuffer.cpp
+++ b/Engine/Source/Runtime/AudioMixer/Private/AudioMixerBuffer.cpp
@@ -372,7 +372,15 @@ namespace Audio
 			InWave->InitAudioResource(AudioDevice->GetRuntimeFormat(InWave));
 		}
 
-		Buffer->DecompressionState = AudioDevice->CreateCompressedAudioInfo(InWave->GetThisAsProxy());
+		// stream caching uses the proxy
+		if (InWave->IsStreaming())
+		{
+			Buffer->DecompressionState = AudioDevice->CreateCompressedAudioInfo(InWave->GetThisAsProxy());
+		}
+		else
+		{
+			Buffer->DecompressionState = AudioDevice->CreateCompressedAudioInfo(InWave);
+		}
 
 		if (Buffer->DecompressionState)
 		{
diff --git a/Engine/Source/Runtime/AudioMixer/Private/AudioMixerDevice.cpp b/Engine/Source/Runtime/AudioMixer/Private/AudioMixerDevice.cpp
--- a/Engine/Source/Runtime/AudioMixer/Private/AudioMixerDevice.cpp
+++ b/Engine/Source/Runtime/AudioMixer/Private/AudioMixerDevice.cpp
@@ -577,7 +577,7 @@ namespace Audio
 
 	class ICompressedAudioInfo* FMixerDevice::CreateCompressedAudioInfo(USoundWave* InSoundWave)
 	{
-//		check(InSoundWave);
+		check(InSoundWave);
 		check(AudioMixerPlatform);
 		return AudioMixerPlatform->CreateCompressedAudioInfo(InSoundWave);
 	}
