Adding metasound buffer add float node

diff --git a/Engine/Plugins/Runtime/Metasound/Source/MetasoundStandardNodes/Private/MetasoundBufferAddFloatNode.cpp b/Engine/Plugins/Runtime/Metasound/Source/MetasoundStandardNodes/Private/MetasoundBufferAddFloatNode.cpp
new file mode 100644
--- /dev/null
+++ b/Engine/Plugins/Runtime/Metasound/Source/MetasoundStandardNodes/Private/MetasoundBufferAddFloatNode.cpp
@@ -0,0 +1,137 @@
+// Copyright Epic Games, Inc. All Rights Reserved.
+
+#include "MetasoundBufferAddFloatNode.h"
+
+#include "DSP/BufferVectorOperations.h"
+#include "MetasoundAudioBuffer.h"
+#include "MetasoundExecutableOperator.h"
+#include "MetasoundFacade.h"
+#include "MetasoundNodeRegistrationMacro.h"
+#include "MetasoundOperatorSettings.h"
+#include "MetasoundPrimitives.h"
+
+#define LOCTEXT_NAMESPACE "MetasoundBufferAddFloatNode"
+
+namespace Metasound
+{
+	class FBufferAddFloatOperator : public TExecutableOperator<FBufferAddFloatOperator>
+	{
+		public:
+			static const FNodeInfo& GetNodeInfo();
+
+			static FVertexInterface DeclareVertexInterface();
+
+			static TUniquePtr<IOperator> CreateOperator(const FCreateOperatorParams& InParams, FBuildErrorArray& OutErrors);
+
+			FBufferAddFloatOperator(const FOperatorSettings& InSettings, const FAudioBufferReadRef& InBuffer, const FFloatReadRef& InAddend)
+			:	Buffer(InBuffer)
+			,	Addend(InAddend)
+			,	OutputBuffer(FAudioBufferWriteRef::CreateNew(InSettings))
+			{
+			}
+
+			virtual FDataReferenceCollection GetInputs() const override
+			{
+				FDataReferenceCollection InputDataReferences;
+				InputDataReferences.AddDataReadReference(TEXT("In"), FAudioBufferReadRef(Buffer));
+				InputDataReferences.AddDataReadReference(TEXT("Addend"), FFloatReadRef(Addend));
+				return InputDataReferences;
+			}
+
+			virtual FDataReferenceCollection GetOutputs() const override
+			{
+				FDataReferenceCollection OutputDataReferences;
+				OutputDataReferences.AddDataReadReference(TEXT("Out"), FAudioBufferReadRef(OutputBuffer));
+				return OutputDataReferences;
+			}
+
+			void Execute()
+			{
+				FMemory::Memcpy(OutputBuffer->GetData(), Buffer->GetData(), sizeof(float) * OutputBuffer->Num());
+
+				const int32 SIMDRemainder = OutputBuffer->Num() % AUDIO_SIMD_FLOAT_ALIGNMENT;
+				const int32 SIMDCount = OutputBuffer->Num() - SIMDRemainder;
+
+				Audio::AddConstantToBufferInplace(OutputBuffer->GetData(), SIMDCount, *Addend);
+
+				for (int32 i = SIMDCount; i < OutputBuffer->Num(); ++i)
+				{
+					OutputBuffer->GetData()[i] += *Addend;
+				}
+			}
+
+		private:
+			FAudioBufferReadRef Buffer;
+			FFloatReadRef Addend;
+			FAudioBufferWriteRef OutputBuffer;
+	};
+
+	FBufferAddFloatNode::FBufferAddFloatNode(const FString& InInstanceName)
+	:	FNodeFacade(InInstanceName, TFacadeOperatorClass<FBufferAddFloatOperator>())
+	{
+	}
+
+	FBufferAddFloatNode::FBufferAddFloatNode(const FNodeInitData& InInitData)
+	:	FBufferAddFloatNode(InInitData.InstanceName)
+	{
+	}
+
+	FVertexInterface FBufferAddFloatOperator::DeclareVertexInterface()
+	{
+		static const FVertexInterface Interface(
+			FInputVertexInterface(
+				TInputDataVertexModel<FAudioBuffer>(TEXT("Buffer"), LOCTEXT("InputBufferTooltip", "Input buffer.")),
+				TInputDataVertexModel<float>(TEXT("Addend"), LOCTEXT("InputAddendTooltip", "Input addend."))
+			),
+			FOutputVertexInterface(
+				TOutputDataVertexModel<FAudioBuffer>(TEXT("Out"), LOCTEXT("OutpuBufferTooltip", "Add float value to buffer."))
+			)
+		);
+
+		return Interface;
+	}
+
+	TUniquePtr<IOperator> FBufferAddFloatOperator::CreateOperator(const FCreateOperatorParams& InParams, FBuildErrorArray& OutErrors)
+	{
+		const FBufferAddFloatNode& BufferAddFloatNode = static_cast<const FBufferAddFloatNode&>(InParams.Node);
+		const FDataReferenceCollection& InputCol = InParams.InputDataReferences;
+
+		FAudioBufferReadRef InputBuffer = InputCol.GetDataReadReferenceOrConstruct<FAudioBuffer>(TEXT("Buffer"), FAudioBuffer(InParams.OperatorSettings.GetNumFramesPerBlock()));
+		FFloatReadRef InputGain = InputCol.GetDataReadReferenceOrConstruct<float>(TEXT("Addend"), 0.0f);
+
+		return MakeUnique<FBufferAddFloatOperator>(InParams.OperatorSettings, InputBuffer, InputGain);
+	}
+
+	const FNodeInfo& FBufferAddFloatOperator::GetNodeInfo()
+	{
+		auto InitNodeInfo = []() -> FNodeInfo
+		{
+			FNodeDisplayStyle DisplayStyle;
+			DisplayStyle.ImageName = "MetasoundEditor.Graph.Node.Math.Add";
+			DisplayStyle.bShowName = false;
+			DisplayStyle.bShowInputNames = false;
+			DisplayStyle.bShowOutputNames = false;
+
+			FNodeInfo Info;
+			Info.ClassName = "Add Float to AudioBuffer";
+			Info.DisplayStyle = DisplayStyle;
+			Info.MajorVersion = 1;
+			Info.MinorVersion = 0;
+			Info.Description = LOCTEXT("BufferAddFloat_NodeDescription", "Add Float to AudioBuffer.");
+			Info.CategoryHierarchy = { LOCTEXT("Metasound_MathCategory", "Math") };
+			Info.Author = PluginAuthor;
+			Info.PromptIfMissing = PluginNodeMissingPrompt;
+			Info.DefaultInterface = DeclareVertexInterface();
+
+			return Info;
+		};
+
+		static const FNodeInfo Info = InitNodeInfo();
+
+		return Info;
+	};
+
+	METASOUND_REGISTER_NODE(FBufferAddFloatNode);
+}
+
+#undef LOCTEXT_NAMESPACE //MetasoundBufferAddFloatNode
diff --git a/Engine/Plugins/Runtime/Metasound/Source/MetasoundStandardNodes/Public/MetasoundBufferAddFloatNode.h b/Engine/Plugins/Runtime/Metasound/Source/MetasoundStandardNodes/Public/MetasoundBufferAddFloatNode.h
new file mode 100644
--- /dev/null
+++ b/Engine/Plugins/Runtime/Metasound/Source/MetasoundStandardNodes/Public/MetasoundBufferAddFloatNode.h
@@ -0,0 +1,18 @@
+// Copyright Epic Games, Inc. All Rights Reserved.
+#pragma once
+
+#include "CoreMinimal.h"
+#include "Internationalization/Text.h"
+
+#include "MetasoundFacade.h"
+
+
+namespace Metasound
+{
+	class METASOUNDSTANDARDNODES_API FBufferAddFloatNode : public FNodeFacade
+	{
+		public:
+			FBufferAddFloatNode(const FString& InInstanceName);
+			FBufferAddFloatNode(const FNodeInitData& InInitData);
+	};
+} // namespace Metasound
