- Metasound Graph Pin Clean-up
- Add independent pin color defaults & move Metasound Settings from shared settings to plugin and setup as per-user settings
- Fix ensure hitting when outputs are created and corrisponding literals are constructed

diff --git a/Engine/Plugins/Runtime/Metasound/Source/MetasoundEditor/Private/MetasoundEditorGraphBuilder.cpp b/Engine/Plugins/Runtime/Metasound/Source/MetasoundEditor/Private/MetasoundEditorGraphBuilder.cpp
--- a/Engine/Plugins/Runtime/Metasound/Source/MetasoundEditor/Private/MetasoundEditorGraphBuilder.cpp
+++ b/Engine/Plugins/Runtime/Metasound/Source/MetasoundEditor/Private/MetasoundEditorGraphBuilder.cpp
@@ -26,20 +26,17 @@ namespace Metasound
 {
 	namespace Editor
 	{
-		// Category names must match those found in UEdGraphSchema_K2:PC_<type>
-		// so default selectors function the same way by default
+		const FName FGraphBuilder::PinCategoryAudioFormat = "audio_format";
 		const FName FGraphBuilder::PinCategoryBoolean = "bool";
 		const FName FGraphBuilder::PinCategoryDouble = "double";
-		const FName FGraphBuilder::PinCategoryExec = "exec";
 		const FName FGraphBuilder::PinCategoryFloat = "float";
 		const FName FGraphBuilder::PinCategoryInt32 = "int";
 		const FName FGraphBuilder::PinCategoryInt64 = "int64";
 		const FName FGraphBuilder::PinCategoryObject = "object";
 		const FName FGraphBuilder::PinCategoryString = "string";
+		const FName FGraphBuilder::PinCategoryTrigger = "trigger";
 
-		const FName FGraphBuilder::PinSubCategoryAudioFormat = "Format";
-		const FName FGraphBuilder::PinSubCategoryAudioNumeric = "Numeric";
-		const FName FGraphBuilder::PinSubCategoryObjectArray = "UObjectArray";
+		const FName FGraphBuilder::PinSubCategoryTime = "time";
 
 		namespace GraphBuilderPrivate
 		{
@@ -90,7 +87,6 @@ namespace Metasound
 
 					case ELiteralType::UObjectProxy:
 					{
-						// TODO: Support default UObject value on node
 						UClass* ClassToUse = FMetasoundFrontendRegistryContainer::Get()->GetLiteralUClassForDataType(InTypeName);
 						if (ClassToUse)
 						{
@@ -239,11 +235,6 @@ namespace Metasound
 
 			const FScopedTransaction Transaction(LOCTEXT("SetMetasoundGraphNode", "Set Metasound Literal Input"));
 
-			if (!ensureAlways(InNodeHandle->GetClassType() == EMetasoundFrontendClassType::External))
-			{
-				return;
-			}
-
 			const FString& InInputName = InInputPin.GetName();
 			const FString& InStringValue = InInputPin.DefaultValue;
 			const UObject* InObjectValue = InInputPin.DefaultObject;
@@ -541,9 +532,13 @@ namespace Metasound
 							FOutputHandle OutputHandle = NodeInputs[InputIndex]->GetCurrentlyConnectedOutput();
 							if (OutputHandle->IsValid())
 							{
-								UEdGraphNode* OutputGraphNode = NewIdNodeMap.FindChecked(OutputHandle->GetOwningNodeID()).GraphNode;
-								UEdGraphPin* OutputPin = OutputGraphNode->FindPinChecked(OutputHandle->GetName(), EEdGraphPinDirection::EGPD_Output);
-								Pin->MakeLinkTo(OutputPin);
+								const FMetasoundFrontendNodeStyle& Style = OutputHandle->GetOwningNode()->GetNodeStyle();
+								if (Style.Display.Visibility == EMetasoundFrontendNodeStyleDisplayVisibility::Visible)
+								{
+									UEdGraphNode* OutputGraphNode = NewIdNodeMap.FindChecked(OutputHandle->GetOwningNodeID()).GraphNode;
+									UEdGraphPin* OutputPin = OutputGraphNode->FindPinChecked(OutputHandle->GetName(), EEdGraphPinDirection::EGPD_Output);
+									Pin->MakeLinkTo(OutputPin);
+								}
 							}
 
 							InputIndex++;
@@ -571,20 +566,17 @@ namespace Metasound
 			{
 				NodeHandle = Node->GetNodeHandle();
 
-				if (NodeHandle->GetClassType() == EMetasoundFrontendClassType::External)
+				TArray<FInputHandle> Inputs = NodeHandle->GetInputs();
+				for (FInputHandle& Input : Inputs)
 				{
-					TArray<FInputHandle> Inputs = NodeHandle->GetInputs();
-					for (FInputHandle& Input : Inputs)
+					FOutputHandle Output = Input->GetCurrentlyConnectedOutput();
+					if (Output->IsValid())
 					{
-						FOutputHandle Output = Input->GetCurrentlyConnectedOutput();
-						if (Output->IsValid())
+						FNodeHandle InputHandle = Output->GetOwningNode();
+						if (InputHandle->GetNodeStyle().Display.Visibility == EMetasoundFrontendNodeStyleDisplayVisibility::Hidden)
 						{
-							FNodeHandle InputHandle = Output->GetOwningNode();
-							if (InputHandle->GetNodeStyle().Display.Visibility == EMetasoundFrontendNodeStyleDisplayVisibility::Hidden)
-							{
-								FGraphHandle Graph = InputHandle->GetOwningGraph();
-								Graph->RemoveNode(*InputHandle);
-							}
+							FGraphHandle Graph = InputHandle->GetOwningGraph();
+							Graph->RemoveNode(*InputHandle);
 						}
 					}
 				}
@@ -766,10 +758,7 @@ namespace Metasound
 				NewPin->PinType = EditorModule.FindDataType(InInputHandle->GetDataType()).PinType;
 
 				FNodeHandle NodeHandle = InInputHandle->GetOwningNode();
-				if (NodeHandle->GetClassType() == EMetasoundFrontendClassType::External)
-				{
-					FGraphBuilder::AddOrUpdateLiteralInput(InEditorNode.GetMetasoundChecked(), NodeHandle, *NewPin);
-				}
+				FGraphBuilder::AddOrUpdateLiteralInput(InEditorNode.GetMetasoundChecked(), NodeHandle, *NewPin);
 			}
 
 			return NewPin;
diff --git a/Engine/Plugins/Runtime/Metasound/Source/MetasoundEditor/Private/MetasoundEditorGraphBuilder.h b/Engine/Plugins/Runtime/Metasound/Source/MetasoundEditor/Private/MetasoundEditorGraphBuilder.h
--- a/Engine/Plugins/Runtime/Metasound/Source/MetasoundEditor/Private/MetasoundEditorGraphBuilder.h
+++ b/Engine/Plugins/Runtime/Metasound/Source/MetasoundEditor/Private/MetasoundEditorGraphBuilder.h
@@ -29,19 +29,18 @@ namespace Metasound
 		class FGraphBuilder
 		{
 		public:
+			static const FName PinCategoryAudioFormat;  // Audio formats (ex. Buffers, Mono, Stereo)
 			static const FName PinCategoryBoolean;
 			static const FName PinCategoryDouble;
-			static const FName PinCategoryExec;
 			static const FName PinCategoryFloat;
 			static const FName PinCategoryInt32;
 			static const FName PinCategoryInt64;
 			static const FName PinCategoryObject;
 			static const FName PinCategoryString;
+			static const FName PinCategoryTrigger;
 
-			// Custom pin-related styles (ex. wire color, pin heads, etc.)
-			static const FName PinSubCategoryAudioFormat;				// Audio formats (ex. Buffers, Mono, Stereo)
-			static const FName PinSubCategoryAudioNumeric;				// Audio numeric types (ex. Frequency, Time)
-			static const FName PinSubCategoryObjectArray;
+			// Custom pin-related styles for non-literal types (ex. wire color, pin heads, etc.)
+			static const FName PinSubCategoryTime; // Time type
 
 			// Adds a node to the editor graph that corresponds to the provided node handle.
 			static UEdGraphNode* AddNode(UObject& InMetasound, Frontend::FNodeHandle& InNodeHandle, bool bInSelectNewNode = true);
diff --git a/Engine/Plugins/Runtime/Metasound/Source/MetasoundEditor/Private/MetasoundEditorGraphSchema.cpp b/Engine/Plugins/Runtime/Metasound/Source/MetasoundEditor/Private/MetasoundEditorGraphSchema.cpp
--- a/Engine/Plugins/Runtime/Metasound/Source/MetasoundEditor/Private/MetasoundEditorGraphSchema.cpp
+++ b/Engine/Plugins/Runtime/Metasound/Source/MetasoundEditor/Private/MetasoundEditorGraphSchema.cpp
@@ -1,5 +1,4 @@
 // Copyright Epic Games, Inc. All Rights Reserved.
-
 #include "MetasoundEditorGraphSchema.h"
 
 #include "EdGraph/EdGraph.h"
@@ -20,6 +19,7 @@
 #include "MetasoundEditorGraphBuilder.h"
 #include "MetasoundEditorGraphNode.h"
 #include "MetasoundEditorModule.h"
+#include "MetasoundEditorSettings.h"
 #include "MetasoundFrontend.h"
 #include "ScopedTransaction.h"
 #include "Toolkits/ToolkitManager.h"
@@ -49,22 +49,17 @@ namespace Metasound
 
 		FLinearColor GetPinCategoryColor(const FEdGraphPinType& PinType)
 		{
-			const UGraphEditorSettings* Settings = GetDefault<UGraphEditorSettings>();
+			const UMetasoundEditorSettings* Settings = GetDefault<UMetasoundEditorSettings>();
 			check(Settings);
 
-			if (PinType.PinSubCategory == FGraphBuilder::PinSubCategoryAudioNumeric)
+			if (PinType.PinCategory == FGraphBuilder::PinCategoryAudioFormat)
 			{
-				return Settings->SoftObjectPinTypeColor;
+				return Settings->AudioPinTypeColor;
 			}
 
-			if (PinType.PinSubCategory == FGraphBuilder::PinSubCategoryAudioFormat)
+			if (PinType.PinCategory == FGraphBuilder::PinCategoryTrigger)
 			{
-				return Settings->SoftClassPinTypeColor;
-			}
-
-			if (PinType.PinCategory == FGraphBuilder::PinCategoryExec)
-			{
-				return Settings->ExecutionPinTypeColor;
+				return Settings->TriggerPinTypeColor;
 			}
 
 			if (PinType.PinCategory == FGraphBuilder::PinCategoryBoolean)
@@ -74,6 +69,10 @@ namespace Metasound
 
 			if (PinType.PinCategory == FGraphBuilder::PinCategoryFloat)
 			{
+				if (PinType.PinSubCategory == FGraphBuilder::PinSubCategoryTime)
+				{
+					return Settings->TimePinTypeColor;
+				}
 				return Settings->FloatPinTypeColor;
 			}
 
@@ -99,11 +98,10 @@ namespace Metasound
 
 			if (PinType.PinCategory == FGraphBuilder::PinCategoryObject)
 			{
-				return Settings->ClassPinTypeColor;
+				return Settings->ObjectPinTypeColor;
 			}
 
-
-			return Settings->StructPinTypeColor;
+			return Settings->DefaultPinTypeColor;
 		}
 
 		FConnectionDrawingPolicy* FGraphConnectionDrawingPolicyFactory::CreateConnectionPolicy(
@@ -159,11 +157,11 @@ namespace Metasound
 
 			if (!bExecuted)
 			{
-				if (InputPin->PinType.PinCategory == FGraphBuilder::PinCategoryExec)
+				if (InputPin->PinType.PinCategory == FGraphBuilder::PinCategoryTrigger)
 				{
 					OutParams.WireThickness = Settings->DefaultExecutionWireThickness;
 				}
-				else if (InputPin->PinType.PinCategory == FGraphBuilder::PinSubCategoryAudioFormat)
+				else if (InputPin->PinType.PinCategory == FGraphBuilder::PinCategoryAudioFormat)
 				{
 					OutParams.bDrawBubbles = true;
 				}
@@ -584,10 +582,7 @@ void UMetasoundEditorGraphSchema::BreakPinLinks(UEdGraphPin& TargetPin, bool bSe
 		Handle->Disconnect();
 
 		FNodeHandle NodeHandle = Handle->GetOwningNode();
-		if (NodeHandle->GetClassType() == EMetasoundFrontendClassType::External)
-		{
-			FGraphBuilder::AddOrUpdateLiteralInput(Metasound, NodeHandle, *InputPins[i]);
-		}
+		FGraphBuilder::AddOrUpdateLiteralInput(Metasound, NodeHandle, *InputPins[i]);
 	}
 
 	Super::BreakPinLinks(TargetPin, bSendsNodeNotifcation);
diff --git a/Engine/Plugins/Runtime/Metasound/Source/MetasoundEditor/Private/MetasoundEditorModule.cpp b/Engine/Plugins/Runtime/Metasound/Source/MetasoundEditor/Private/MetasoundEditorModule.cpp
--- a/Engine/Plugins/Runtime/Metasound/Source/MetasoundEditor/Private/MetasoundEditorModule.cpp
+++ b/Engine/Plugins/Runtime/Metasound/Source/MetasoundEditor/Private/MetasoundEditorModule.cpp
@@ -1,5 +1,4 @@
 // Copyright Epic Games, Inc. All Rights Reserved.
-
 #include "MetasoundEditorModule.h"
 
 #include "AssetTypeActions_Base.h"
@@ -10,6 +9,7 @@
 #include "EdGraphUtilities.h"
 #include "EditorStyleSet.h"
 #include "IDetailCustomization.h"
+#include "ISettingsModule.h"
 #include "Metasound.h"
 #include "MetasoundSource.h"
 #include "MetasoundAssetTypeActions.h"
@@ -18,6 +18,8 @@
 #include "MetasoundEditorGraphBuilder.h"
 #include "MetasoundEditorGraphNodeFactory.h"
 #include "MetasoundEditorGraphSchema.h"
+#include "MetasoundEditorSettings.h"
+#include "MetasoundFrontendRegistries.h"
 #include "Modules/ModuleInterface.h"
 #include "Modules/ModuleManager.h"
 #include "PropertyEditorDelegates.h"
@@ -28,7 +30,6 @@
 #include "Styling/SlateStyleRegistry.h"
 #include "Styling/SlateTypes.h"
 #include "Templates/SharedPointer.h"
-#include "../../BlueprintGraph/Classes/EdGraphSchema_K2.h"
 
 DEFINE_LOG_CATEGORY(LogMetasoundEditor);
 
@@ -81,8 +82,8 @@ namespace Metasound
 					Set("MetasoundEditor.Graph.Node.Body.Default", new FSlateImageBrush(RootToContentDir(TEXT("/Graph/node_default_body_64x.png")), FVector2D(64.0f, 64.0f)));
 					Set("MetasoundEditor.Graph.Node.Math.Add", new FSlateImageBrush(RootToContentDir(TEXT("/Graph/node_math_add_40x.png")), Icon40x40));
 					Set("MetasoundEditor.Graph.Node.Math.Divide", new FSlateImageBrush(RootToContentDir(TEXT("/Graph/node_math_divide_40x.png")), Icon40x40));
-					Set("MetasoundEditor.Graph.Node.Math.Multiply", new FSlateImageBrush(RootToContentDir(TEXT("/Graph/node_math_multiply_40x.png")), FVector2D(32.0f, 32.0f)));
-					Set("MetasoundEditor.Graph.Node.Math.RandRange", new FSlateImageBrush(RootToContentDir(TEXT("/Graph/node_math_random_40x.png")), FVector2D(26.0f, 40.0f)));
+					Set("MetasoundEditor.Graph.Node.Math.Multiply", new FSlateImageBrush(RootToContentDir(TEXT("/Graph/node_math_multiply_40x.png")), Icon40x40));
+					Set("MetasoundEditor.Graph.Node.Math.RandRange", new FSlateImageBrush(RootToContentDir(TEXT("/Graph/node_math_random_40x.png")), Icon40x40));
 					Set("MetasoundEditor.Graph.Node.Math.Subtract", new FSlateImageBrush(RootToContentDir(TEXT("/Graph/node_math_subtract_40x.png")), Icon40x40));
 
 					// Misc
@@ -114,7 +115,8 @@ namespace Metasound
 
 			virtual void RegisterDataType(FName InPinCategoryName, FName InPinSubCategoryName, const FDataTypeRegistryInfo& InRegistryInfo) override
 			{
-				DataTypeInfo.Add(InRegistryInfo.DataTypeName, FEditorDataType(InPinCategoryName, InPinSubCategoryName, InRegistryInfo));
+				FEdGraphPinType PinType(InPinCategoryName, InPinSubCategoryName, nullptr, EPinContainerType::None, false, FEdGraphTerminalType());
+				DataTypeInfo.Emplace(InRegistryInfo.DataTypeName, Editor::FEditorDataType(MoveTemp(PinType), InRegistryInfo));
 			}
 
 			void RegisterCoreDataTypes()
@@ -131,7 +133,7 @@ namespace Metasound
 					// Execution path triggers are specialized
 					if (DataTypeName == "Trigger")
 					{
-						PinCategory = FGraphBuilder::PinCategoryExec;
+						PinCategory = FGraphBuilder::PinCategoryTrigger;
 					}
 
 					// GraphEditor by default designates specialized connection
@@ -161,7 +163,7 @@ namespace Metasound
 								// but different colorization
 								if (DataTypeName == "Double")
 								{
-									PinCategory = FGraphBuilder::PinCategoryFloat;
+									PinCategory = FGraphBuilder::PinCategoryDouble;
 								}
 
 								// Differentiate stronger numeric types associated with audio
@@ -171,7 +173,7 @@ namespace Metasound
 									|| DataTypeName == "Time:SampleResolution"
 								)
 								{
-									PinSubCategory = FGraphBuilder::PinSubCategoryAudioNumeric;
+									PinSubCategory = FGraphBuilder::PinSubCategoryTime;
 								}
 							}
 							break;
@@ -196,17 +198,16 @@ namespace Metasound
 
 							case ELiteralType::UObjectProxyArray:
 							{
-								PinCategory = FGraphBuilder::PinSubCategoryObjectArray;
+								// TODO: Implement, or will be nuked in favor of general array support for all types
 							}
 							break;
 
-							// Register atypical primitives
 							default:
 							case ELiteralType::None:
 							case ELiteralType::Invalid:
 							{
-
-								// Audio types are ubiquitous, so specialize
+								// Audio types are ubiquitous, so add as subcategory
+								// to be able to stylize connections (i.e. wire color)
 								if (DataTypeName == "Audio:Buffer"
 									|| DataTypeName == "Audio:Unformatted"
 									|| DataTypeName == "Audio:Mono"
@@ -214,7 +215,7 @@ namespace Metasound
 									|| DataTypeName == "Audio:Multichannel"
 								)
 								{
-									PinSubCategory = FGraphBuilder::PinSubCategoryAudioFormat;
+									PinCategory = FGraphBuilder::PinCategoryAudioFormat;
 								}
 								static_assert(static_cast<int32>(ELiteralType::Invalid) == 7, "Possible missing binding of pin category to primitive type");
 							}
@@ -222,7 +223,11 @@ namespace Metasound
 						}
 					}
 
-					DataTypeInfo.Add(DataTypeName, FEditorDataType(PinCategory, PinSubCategory, RegistryInfo));
+					FEdGraphPinType PinType(PinCategory, PinSubCategory, nullptr, EPinContainerType::None, false, FEdGraphTerminalType());
+ 					UClass* ClassToUse = FMetasoundFrontendRegistryContainer::Get()->GetLiteralUClassForDataType(DataTypeName);
+ 					PinType.PinSubCategoryObject = Cast<UObject>(ClassToUse);
+
+					DataTypeInfo.Emplace(DataTypeName, FEditorDataType(MoveTemp(PinType), RegistryInfo));
 				}
 			}
 
@@ -259,10 +264,23 @@ namespace Metasound
 
 				GraphPanelPinFactory = MakeShared<FMetasoundGraphPanelPinFactory>();
 				FEdGraphUtilities::RegisterVisualPinFactory(GraphPanelPinFactory);
+
+				ISettingsModule& SettingsModule = FModuleManager::LoadModuleChecked<ISettingsModule>("Settings");
+
+				SettingsModule.RegisterSettings("Editor", "Audio", "Metasound Editor",
+					NSLOCTEXT("MetasoundsEditor", "MetasoundEditorSettingsName", "Metasound Editor"),
+					NSLOCTEXT("MetasoundsEditor", "MetasoundEditorSettingsDescription", "Customize Metasound Editor."),
+					GetMutableDefault<UMetasoundEditorSettings>()
+				);
 			}
 
 			virtual void ShutdownModule() override
 			{
+				if (ISettingsModule* SettingsModule = FModuleManager::GetModulePtr<ISettingsModule>("Settings"))
+				{
+					SettingsModule->UnregisterSettings("Editor", "Audio", "Metasound Editor");
+				}
+
 				if (FModuleManager::Get().IsModuleLoaded(AssetToolName))
 				{
 					IAssetTools& AssetTools = FModuleManager::GetModuleChecked<FAssetToolsModule>(AssetToolName).Get();
diff --git a/Engine/Plugins/Runtime/Metasound/Source/MetasoundEditor/Private/MetasoundEditorModule.h b/Engine/Plugins/Runtime/Metasound/Source/MetasoundEditor/Private/MetasoundEditorModule.h
--- a/Engine/Plugins/Runtime/Metasound/Source/MetasoundEditor/Private/MetasoundEditorModule.h
+++ b/Engine/Plugins/Runtime/Metasound/Source/MetasoundEditor/Private/MetasoundEditorModule.h
@@ -20,8 +20,8 @@ namespace Metasound
 			FEdGraphPinType PinType;
 			FDataTypeRegistryInfo RegistryInfo;
 
-			FEditorDataType(FName InPinCategory, FName InPinSubCategory, const FDataTypeRegistryInfo& InRegistryInfo)
-				: PinType(InPinCategory, InPinSubCategory, nullptr, EPinContainerType::None, false, FEdGraphTerminalType())
+			FEditorDataType(FEdGraphPinType&& InPinType, const FDataTypeRegistryInfo& InRegistryInfo)
+				: PinType(MoveTemp(InPinType))
 				, RegistryInfo(InRegistryInfo)
 			{
 			}
diff --git a/Engine/Plugins/Runtime/Metasound/Source/MetasoundEditor/Private/MetasoundEditorSettings.cpp b/Engine/Plugins/Runtime/Metasound/Source/MetasoundEditor/Private/MetasoundEditorSettings.cpp
new file mode 100644
--- /dev/null
+++ b/Engine/Plugins/Runtime/Metasound/Source/MetasoundEditor/Private/MetasoundEditorSettings.cpp
@@ -0,0 +1,24 @@
+// Copyright Epic Games, Inc. All Rights Reserved.
+#include "MetasoundEditorSettings.h"
+
+#include "UObject/UnrealType.h"
+#include "EditorStyleSet.h"
+
+
+UMetasoundEditorSettings::UMetasoundEditorSettings(const FObjectInitializer& ObjectInitializer)
+	: Super(ObjectInitializer)
+{
+	// pin type colors
+	DefaultPinTypeColor = FLinearColor(0.750000f, 0.6f, 0.4f, 1.0f);			// light brown
+
+	AudioPinTypeColor = FLinearColor(1.0f, 0.3f, 1.0f, 1.0f);					// magenta
+	BooleanPinTypeColor = FLinearColor(0.300000f, 0.0f, 0.0f, 1.0f);			// maroon
+	DoublePinTypeColor = FLinearColor(0.039216f, 0.666667f, 0.0f, 1.0f);		// darker green
+	FloatPinTypeColor = FLinearColor(0.357667f, 1.0f, 0.060000f, 1.0f);			// bright green
+	IntPinTypeColor = FLinearColor(0.013575f, 0.770000f, 0.429609f, 1.0f);		// green-blue
+	Int64PinTypeColor = FLinearColor(0.413575f, 0.770000f, 0.429609f, 1.0f);	// green
+	ObjectPinTypeColor = FLinearColor(0.0f, 0.4f, 0.910000f, 1.0f);				// sharp blue
+	StringPinTypeColor = FLinearColor(1.0f, 0.0f, 0.660537f, 1.0f);				// bright pink
+	TimePinTypeColor = FLinearColor(0.3f, 1.0f, 1.0f, 1.0f);					// cyan
+	TriggerPinTypeColor = FLinearColor(1.0f, 1.0f, 1.0f, 1.0f);					// white
+}
diff --git a/Engine/Plugins/Runtime/Metasound/Source/MetasoundEditor/Private/MetasoundEditorSettings.h b/Engine/Plugins/Runtime/Metasound/Source/MetasoundEditor/Private/MetasoundEditorSettings.h
new file mode 100644
--- /dev/null
+++ b/Engine/Plugins/Runtime/Metasound/Source/MetasoundEditor/Private/MetasoundEditorSettings.h
@@ -0,0 +1,61 @@
+// Copyright Epic Games, Inc. All Rights Reserved.
+#pragma once
+
+#include "CoreMinimal.h"
+#include "UObject/NoExportTypes.h"
+#include "UObject/ObjectMacros.h"
+#include "UObject/Object.h"
+
+#include "MetasoundEditorSettings.generated.h"
+
+
+UCLASS(config=EditorPerProjectUserSettings)
+class METASOUNDEDITOR_API UMetasoundEditorSettings : public UObject
+{
+	GENERATED_UCLASS_BODY()
+
+public:
+	/** Default pin type color */
+	UPROPERTY(EditAnywhere, config, Category=PinColors)
+	FLinearColor DefaultPinTypeColor;
+
+	/** Audio pin type color */
+	UPROPERTY(EditAnywhere, config, Category=PinColors)
+	FLinearColor AudioPinTypeColor;
+
+	/** Boolean pin type color */
+	UPROPERTY(EditAnywhere, config, Category=PinColors)
+	FLinearColor BooleanPinTypeColor;
+
+	/** Double pin type color */
+	UPROPERTY(EditAnywhere, config, Category = PinColors)
+	FLinearColor DoublePinTypeColor;
+
+	/** Floating-point pin type color */
+	UPROPERTY(EditAnywhere, config, Category=PinColors)
+	FLinearColor FloatPinTypeColor;
+
+	/** Integer pin type color */
+	UPROPERTY(EditAnywhere, config, Category=PinColors)
+	FLinearColor IntPinTypeColor;
+
+	/** Integer64 pin type color */
+	UPROPERTY(EditAnywhere, config, Category = PinColors)
+	FLinearColor Int64PinTypeColor;
+
+	/** Object pin type color */
+	UPROPERTY(EditAnywhere, config, Category=PinColors)
+	FLinearColor ObjectPinTypeColor;
+
+	/** String pin type color */
+	UPROPERTY(EditAnywhere, config, Category=PinColors)
+	FLinearColor StringPinTypeColor;
+
+	/** Time pin type color */
+	UPROPERTY(EditAnywhere, config, Category=PinColors)
+	FLinearColor TimePinTypeColor;
+
+	/** Trigger pin type color */
+	UPROPERTY(EditAnywhere, config, Category=PinColors)
+	FLinearColor TriggerPinTypeColor;
+};
diff --git a/Engine/Plugins/Runtime/Metasound/Source/MetasoundEditor/Private/SMetasoundGraphNode.cpp b/Engine/Plugins/Runtime/Metasound/Source/MetasoundEditor/Private/SMetasoundGraphNode.cpp
--- a/Engine/Plugins/Runtime/Metasound/Source/MetasoundEditor/Private/SMetasoundGraphNode.cpp
+++ b/Engine/Plugins/Runtime/Metasound/Source/MetasoundEditor/Private/SMetasoundGraphNode.cpp
@@ -4,8 +4,17 @@
 #include "EditorStyleSet.h"
 #include "GraphEditorSettings.h"
 #include "IDocumentation.h"
+#include "KismetPins/SGraphPinNum.h"
+#include "KismetPins/SGraphPinBool.h"
+#include "KismetPins/SGraphPinExec.h"
+#include "KismetPins/SGraphPinInteger.h"
+#include "KismetPins/SGraphPinObject.h"
+#include "KismetPins/SGraphPinString.h"
+#include "MetasoundEditorGraphBuilder.h"
 #include "MetasoundEditorGraphNode.h"
+#include "MetasoundEditorGraphSchema.h"
 #include "MetasoundEditorModule.h"
+#include "NodeFactory.h"
 #include "ScopedTransaction.h"
 #include "SCommentBubble.h"
 #include "SGraphNode.h"
@@ -58,23 +67,79 @@ const UMetasoundEditorGraphNode& SMetasoundGraphNode::GetMetasoundNode() const
 	return *Cast<UMetasoundEditorGraphNode>(GraphNode);
 }
 
-void SMetasoundGraphNode::CreateStandardPinWidget(UEdGraphPin* CurPin)
+TSharedPtr<SGraphPin> SMetasoundGraphNode::CreatePinWidget(UEdGraphPin* InPin) const
 {
-	const bool bShowPin = ShouldPinBeHidden(CurPin);
+	using namespace Metasound::Editor;
+	using namespace Metasound::Frontend;
+
+	if (const UMetasoundEditorGraphSchema* GraphSchema = Cast<const UMetasoundEditorGraphSchema>(InPin->GetSchema()))
+	{
+		if (InPin->PinType.PinCategory == FGraphBuilder::PinCategoryAudioFormat)
+		{
+			return SNew(SGraphPin, InPin);
+		}
+
+		if (InPin->PinType.PinCategory == FGraphBuilder::PinCategoryBoolean)
+		{
+			return SNew(SGraphPinBool, InPin);
+		}
+
+		if (InPin->PinType.PinCategory == FGraphBuilder::PinCategoryDouble)
+		{
+			return SNew(SGraphPinNum<double>, InPin);
+		}
+
+		if (InPin->PinType.PinCategory == FGraphBuilder::PinCategoryFloat)
+		{
+			return SNew(SGraphPinNum<float>, InPin);
+		}
+
+		if (InPin->PinType.PinCategory == FGraphBuilder::PinCategoryInt32)
+		{
+			return SNew(SGraphPinInteger, InPin);
+		}
+
+		if (InPin->PinType.PinCategory == FGraphBuilder::PinCategoryInt64)
+		{
+			return SNew(SGraphPinNum<int64>, InPin);
+		}
+
+		if (InPin->PinType.PinCategory == FGraphBuilder::PinCategoryObject)
+		{
+			return SNew(SGraphPinObject, InPin);
+		}
+
+		if (InPin->PinType.PinCategory == FGraphBuilder::PinCategoryString)
+		{
+			return SNew(SGraphPinString, InPin);
+		}
+
+		if (InPin->PinType.PinCategory == FGraphBuilder::PinCategoryTrigger)
+		{
+			return SNew(SGraphPinExec, InPin);
+		}
+	}
+
+	return SNew(SGraphPin, InPin);
+}
+
+void SMetasoundGraphNode::CreateStandardPinWidget(UEdGraphPin* InPin)
+{
+	const bool bShowPin = ShouldPinBeHidden(InPin);
 	if (bShowPin)
 	{
-		TSharedPtr<SGraphPin> NewPin = CreatePinWidget(CurPin);
+		TSharedPtr<SGraphPin> NewPin = CreatePinWidget(InPin);
 		check(NewPin.IsValid());
 
 		Metasound::Frontend::FNodeHandle NodeHandle = GetMetasoundNode().GetNodeHandle();
-		if (CurPin->Direction == EGPD_Input)
+		if (InPin->Direction == EGPD_Input)
 		{
 			if (!NodeHandle->GetClassStyle().Display.bShowInputNames)
 			{
 				NewPin->SetShowLabel(false);
 			}
 		}
-		else if (CurPin->Direction == EGPD_Output)
+		else if (InPin->Direction == EGPD_Output)
 		{
 			if (!NodeHandle->GetClassStyle().Display.bShowOutputNames)
 			{
@@ -247,6 +312,7 @@ TSharedRef<SWidget> SMetasoundGraphNode::CreateNodeContentArea()
 					SNew(SImage)
 					.Image(ImageBrush)
 					.ColorAndOpacity(FSlateColor::UseForeground())
+					.DesiredSizeOverride(FVector2D(20, 20))
 				];
 		}
 	}
diff --git a/Engine/Plugins/Runtime/Metasound/Source/MetasoundEditor/Private/SMetasoundGraphNode.h b/Engine/Plugins/Runtime/Metasound/Source/MetasoundEditor/Private/SMetasoundGraphNode.h
--- a/Engine/Plugins/Runtime/Metasound/Source/MetasoundEditor/Private/SMetasoundGraphNode.h
+++ b/Engine/Plugins/Runtime/Metasound/Source/MetasoundEditor/Private/SMetasoundGraphNode.h
@@ -1,7 +1,7 @@
 // Copyright Epic Games, Inc. All Rights Reserved.
 #pragma once
-
 #include "CoreMinimal.h"
+
 #include "Input/Reply.h"
 #include "Layout/Visibility.h"
 #include "SGraphNode.h"
@@ -10,6 +10,7 @@
 
 
 // Forward Declarations
+class SGraphPin;
 class SVerticalBox;
 class UMetasoundEditorGraphNode;
 
@@ -28,7 +29,8 @@ public:
 protected:
 	// SGraphNode Interface
 	virtual void CreateOutputSideAddButton(TSharedPtr<SVerticalBox> OutputBox) override;
-	virtual void CreateStandardPinWidget(UEdGraphPin* CurPin) override;
+	virtual TSharedPtr<SGraphPin> CreatePinWidget(UEdGraphPin* InPin) const override;
+	virtual void CreateStandardPinWidget(UEdGraphPin* InPin) override;
 	virtual TSharedRef<SWidget> CreateNodeContentArea() override;
 	virtual TSharedRef<SWidget> CreateTitleWidget(TSharedPtr<SNodeTitle> NodeTitle) override;
 	virtual const FSlateBrush* GetNodeBodyBrush() const override;
diff --git a/Engine/Plugins/Runtime/Metasound/Source/MetasoundFrontend/Private/MetasoundFrontendRegistries.cpp b/Engine/Plugins/Runtime/Metasound/Source/MetasoundFrontend/Private/MetasoundFrontendRegistries.cpp
--- a/Engine/Plugins/Runtime/Metasound/Source/MetasoundFrontend/Private/MetasoundFrontendRegistries.cpp
+++ b/Engine/Plugins/Runtime/Metasound/Source/MetasoundFrontend/Private/MetasoundFrontendRegistries.cpp
@@ -308,7 +308,10 @@ bool FMetasoundFrontendRegistryContainer::DoesDataTypeSupportLiteralType(FName I
 
 bool FMetasoundFrontendRegistryContainer::RegisterDataType(const ::Metasound::FDataTypeRegistryInfo& InDataInfo, const ::Metasound::FDataTypeConstructorCallbacks& InCallbacks)
 {
-	if (!ensureAlwaysMsgf(!DataTypeRegistry.Contains(InDataInfo.DataTypeName), TEXT("Name collision when trying to register Metasound Data Type %s! Make sure that you created a unique name for your data type, and that REGISTER_METASOUND_DATATYPE isn't called in a public header."), *InDataInfo.DataTypeName.ToString()))
+	if (!ensureAlwaysMsgf(!DataTypeRegistry.Contains(InDataInfo.DataTypeName),
+		TEXT("Name collision when trying to register Metasound Data Type %s! DataType must have "
+			"unique name and REGISTER_METASOUND_DATATYPE cannot be called in a public header."),
+			*InDataInfo.DataTypeName.ToString()))
 	{
 		// todo: capture callstack for previous declaration for non-shipping builds to help clarify who already registered this name for a type.
 		return false;
diff --git a/Engine/Source/Editor/AudioEditor/Classes/AudioEditorSettings.h b/Engine/Source/Editor/AudioEditor/Classes/AudioEditorSettings.h
--- a/Engine/Source/Editor/AudioEditor/Classes/AudioEditorSettings.h
+++ b/Engine/Source/Editor/AudioEditor/Classes/AudioEditorSettings.h
@@ -27,20 +27,6 @@ public:
 	//~ End UObject
 };
 
-UCLASS(config = EditorSettings, defaultconfig, meta = (DisplayName = "Metasound Editor"))
-class AUDIOEDITOR_API UAudioMetasoundEditorSettings : public UDeveloperSettings
-{
-	GENERATED_BODY()
-
-public:
-	//~ Begin UDeveloperSettings
-	virtual FName GetCategoryName() const override { return TEXT("Audio"); }
-	//~ End UDeveloperSettings
-
-	//~ Begin UObject
-	virtual void PostEditChangeProperty(FPropertyChangedEvent& PropertyChangedEvent) override;
-	//~ End UObject
-};
 
 UCLASS(config = EditorSettings, defaultconfig, meta = (DisplayName = "SoundClass Editor"))
 class AUDIOEDITOR_API UAudioSoundClassEditorSettings : public UDeveloperSettings
diff --git a/Engine/Source/Editor/AudioEditor/Private/AudioEditorSettings.cpp b/Engine/Source/Editor/AudioEditor/Private/AudioEditorSettings.cpp
--- a/Engine/Source/Editor/AudioEditor/Private/AudioEditorSettings.cpp
+++ b/Engine/Source/Editor/AudioEditor/Private/AudioEditorSettings.cpp
@@ -7,12 +7,6 @@ void UAudioSubmixEditorSettings::PostEditChangeProperty(FPropertyChangedEvent& P
 	Super::PostEditChangeProperty(PropertyChangedEvent);
 }
 
-
-void UAudioMetasoundEditorSettings::PostEditChangeProperty(FPropertyChangedEvent& PropertyChangedEvent)
-{
-	Super::PostEditChangeProperty(PropertyChangedEvent);
-}
-
 void UAudioSoundClassEditorSettings::PostEditChangeProperty(FPropertyChangedEvent& PropertyChangedEvent)
 {
 	Super::PostEditChangeProperty(PropertyChangedEvent);
