Stoppin' the fun boppin' (Final... hopefully... rename of bop to trigger)

diff --git a/Engine/Plugins/Runtime/Metasound/Source/MetasoundEngine/Private/MetasoundSource.cpp b/Engine/Plugins/Runtime/Metasound/Source/MetasoundEngine/Private/MetasoundSource.cpp
--- a/Engine/Plugins/Runtime/Metasound/Source/MetasoundEngine/Private/MetasoundSource.cpp
+++ b/Engine/Plugins/Runtime/Metasound/Source/MetasoundEngine/Private/MetasoundSource.cpp
@@ -1,5 +1,4 @@
 // Copyright Epic Games, Inc. All Rights Reserved.
-
 #include "MetasoundSource.h"
 
 #include "AssetRegistryModule.h"
@@ -7,7 +6,6 @@
 #include "Internationalization/Text.h"
 #include "MetasoundAssetBase.h"
 #include "MetasoundAudioFormats.h"
-#include "MetasoundBop.h"
 #include "MetasoundEngineEnvironment.h"
 #include "MetasoundFrontendController.h"
 #include "MetasoundFrontendQuery.h"
@@ -17,6 +15,7 @@
 #include "MetasoundLog.h"
 #include "MetasoundPrimitives.h"
 #include "MetasoundReceiveNode.h"
+#include "MetasoundTrigger.h"
 
 
 #if WITH_EDITORONLY_DATA
@@ -180,8 +179,8 @@ ISoundGeneratorPtr UMetasoundSource::CreateSoundGenerator(const FSoundGeneratorI
 
 		// References must be cached before moving the operator to the InitParams
 		FDataReferenceCollection Inputs = Operator->GetInputs();
-		FTriggerWriteRef PlayTrigger = Inputs.GetDataWriteReferenceOrConstruct<FBop>(GetOnPlayInputName(), InSettings, false);
-		FTriggerReadRef FinishTrigger = Outputs.GetDataReadReferenceOrConstruct<FBop>(GetIsFinishedOutputName(), InSettings, false);
+		FTriggerWriteRef PlayTrigger = Inputs.GetDataWriteReferenceOrConstruct<FTrigger>(GetOnPlayInputName(), InSettings, false);
+		FTriggerReadRef FinishTrigger = Outputs.GetDataReadReferenceOrConstruct<FTrigger>(GetIsFinishedOutputName(), InSettings, false);
 
 		// Create the FMetasoundGenerator.
 		FMetasoundGeneratorInitParams InitParams =
@@ -467,7 +466,7 @@ const FMetasoundFrontendArchetype& UMetasoundSource::GetBaseArchetype()
 		FMetasoundFrontendClassVertex OnPlayTrigger;
 		OnPlayTrigger.Name = UMetasoundSource::GetOnPlayInputName();
 		OnPlayTrigger.Metadata.DisplayName = FText::FromString(OnPlayTrigger.Name);
-		OnPlayTrigger.TypeName = Metasound::Frontend::GetDataTypeName<Metasound::FBop>();
+		OnPlayTrigger.TypeName = Metasound::Frontend::GetDataTypeName<Metasound::FTrigger>();
 		OnPlayTrigger.Metadata.Description = LOCTEXT("OnPlayTriggerToolTip", "Trigger executed when this source is played.");
 		OnPlayTrigger.PointIDs.Add(FGuid::NewGuid());
 
@@ -476,8 +475,8 @@ const FMetasoundFrontendArchetype& UMetasoundSource::GetBaseArchetype()
 		FMetasoundFrontendClassVertex OnFinished;
 		OnFinished.Name = UMetasoundSource::GetIsFinishedOutputName();
 		OnFinished.Metadata.DisplayName = FText::FromString(OnFinished.Name);
-		OnFinished.TypeName = Metasound::Frontend::GetDataTypeName<Metasound::FBop>();
-		OnFinished.Metadata.Description = LOCTEXT("OnFinishedToolTip", "Bop executed to initiate stopping the source.");
+		OnFinished.TypeName = Metasound::Frontend::GetDataTypeName<Metasound::FTrigger>();
+		OnFinished.Metadata.Description = LOCTEXT("OnFinishedToolTip", "Trigger executed to initiate stopping the source.");
 		OnFinished.PointIDs.Add(FGuid::NewGuid());
 
 		Archetype.Interface.Outputs.Add(OnFinished);
diff --git a/Engine/Plugins/Runtime/Metasound/Source/MetasoundEngine/Private/MetasoundWavePlayerNode.cpp b/Engine/Plugins/Runtime/Metasound/Source/MetasoundEngine/Private/MetasoundWavePlayerNode.cpp
--- a/Engine/Plugins/Runtime/Metasound/Source/MetasoundEngine/Private/MetasoundWavePlayerNode.cpp
+++ b/Engine/Plugins/Runtime/Metasound/Source/MetasoundEngine/Private/MetasoundWavePlayerNode.cpp
@@ -1,14 +1,14 @@
 // Copyright Epic Games, Inc. All Rights Reserved.
-
 #include "MetasoundWavePlayerNode.h"
+
+#include "AudioResampler.h"
+#include "MetasoundBuildError.h"
 #include "MetasoundExecutableOperator.h"
 #include "MetasoundNodeRegistrationMacro.h"
 #include "MetasoundPrimitives.h"
 #include "MetasoundWave.h"
-#include "MetasoundNodeRegistrationMacro.h"
-#include "MetasoundBuildError.h"
-#include "MetasoundBop.h"
-#include "AudioResampler.h"
+#include "MetasoundTrigger.h"
+
 
 #define LOCTEXT_NAMESPACE "MetasoundWaveNode"
 
@@ -79,7 +79,7 @@ namespace Metasound
 		{
 			FDataReferenceCollection InputDataReferences;
 			InputDataReferences.AddDataReadReference(TEXT("Audio"), FWaveAssetReadRef(Wave));
-			InputDataReferences.AddDataReadReference(TEXT("TrigIn"), FBopReadRef(TrigIn));
+			InputDataReferences.AddDataReadReference(TEXT("TrigIn"), FTriggerReadRef(TrigIn));
 			InputDataReferences.AddDataReadReference(TEXT("PitchShiftCents"), FFloatReadRef(PitchShiftCents));
 			return InputDataReferences;
 		}
@@ -121,7 +121,7 @@ namespace Metasound
 			}
 
 			TrigIn->ExecuteBlock(
-				// OnPreBop
+				// OnPreTrigger
 				[&](int32 StartFrame, int32 EndFrame)
 				{
  					if (bIsPlaying)
@@ -129,7 +129,7 @@ namespace Metasound
 						ExecuteInternal(StartFrame, EndFrame);
 					}
 				},
-				// OnBop
+				// OnTrigger
 				[&](int32 StartFrame, int32 EndFrame)
 				{
 					if (!bIsPlaying)
@@ -190,7 +190,7 @@ namespace Metasound
 			if (!Decoder.CanGenerateAudio() ||  (NumFramesDecoded < NumOutputFrames))
 			{
 				bIsPlaying = false;
-				TrigggerOnDone->BopFrame(StartFrame + NumFramesDecoded);
+				TrigggerOnDone->TriggerFrame(StartFrame + NumFramesDecoded);
 			}
 
 			if (bNeedsUpmix)
diff --git a/Engine/Plugins/Runtime/Metasound/Source/MetasoundEngine/Private/MetasoundWaveSelectorNode.cpp b/Engine/Plugins/Runtime/Metasound/Source/MetasoundEngine/Private/MetasoundWaveSelectorNode.cpp
--- a/Engine/Plugins/Runtime/Metasound/Source/MetasoundEngine/Private/MetasoundWaveSelectorNode.cpp
+++ b/Engine/Plugins/Runtime/Metasound/Source/MetasoundEngine/Private/MetasoundWaveSelectorNode.cpp
@@ -1,16 +1,16 @@
 // Copyright Epic Games, Inc. All Rights Reserved.
-
 #include "MetasoundWaveSelectorNode.h"
+
+#include "MetasoundBuildError.h"
 #include "MetasoundExecutableOperator.h"
 #include "MetasoundNodeRegistrationMacro.h"
 #include "MetasoundPrimitives.h"
+#include "MetasoundTrigger.h"
 #include "MetasoundWave.h"
-#include "MetasoundNodeRegistrationMacro.h"
-#include "MetasoundBuildError.h"
-#include "MetasoundBop.h"
 
 #define LOCTEXT_NAMESPACE "MetasoundWaveNode"
 
+
 namespace Metasound
 {
 	class FWaveSelectorOperator : public TExecutableOperator<FWaveSelectorOperator>
@@ -44,7 +44,7 @@ namespace Metasound
 
 			InputDataReferences.AddDataReadReference(TEXT("Index"), FInt32ReadRef(Index));
 
-			InputDataReferences.AddDataReadReference(TEXT("TrigIn"), FBopReadRef(TrigIn));
+			InputDataReferences.AddDataReadReference(TEXT("TrigIn"), FTriggerReadRef(TrigIn));
 			return InputDataReferences;
 		}
 
@@ -62,15 +62,15 @@ namespace Metasound
 			TriggerNewWaveSelected->AdvanceBlock();
 
 			TrigIn->ExecuteBlock(
-				// OnPreBop
+				// OnPreTrigger
 				[&](int32 StartFrame, int32 EndFrame)
 				{
 				},
-				// OnBop
+				// OnTrigger
 				[&](int32 StartFrame, int32 EndFrame)
 				{
 					UpdateWave(*Index);
-					TriggerNewWaveSelected->BopFrame(StartFrame);
+					TriggerNewWaveSelected->TriggerFrame(StartFrame);
 				}
 			);
 		}
@@ -131,7 +131,7 @@ namespace Metasound
 
 		// Outputs
 		FWaveAssetWriteRef SelectedWave;
-		FBopWriteRef TriggerNewWaveSelected;
+		FTriggerWriteRef TriggerNewWaveSelected;
 
 		TArray<FWaveAssetReadRef*> WavePtrs;
 
diff --git a/Engine/Plugins/Runtime/Metasound/Source/MetasoundGenerator/Private/MetasoundGenerator.cpp b/Engine/Plugins/Runtime/Metasound/Source/MetasoundGenerator/Private/MetasoundGenerator.cpp
--- a/Engine/Plugins/Runtime/Metasound/Source/MetasoundGenerator/Private/MetasoundGenerator.cpp
+++ b/Engine/Plugins/Runtime/Metasound/Source/MetasoundGenerator/Private/MetasoundGenerator.cpp
@@ -2,13 +2,13 @@
 #include "MetasoundGenerator.h"
 
 #include "DSP/Dsp.h"
-#include "MetasoundBop.h"
 #include "MetasoundFrequency.h"
 #include "MetasoundGraph.h"
 #include "MetasoundInputNode.h"
 #include "MetasoundOperatorBuilder.h"
 #include "MetasoundOperatorInterface.h"
 #include "MetasoundOutputNode.h"
+#include "MetasoundTrigger.h"
 
 
 namespace Metasound
@@ -103,7 +103,7 @@ namespace Metasound
 
 		if (!bIsPlaying)
 		{
-			OnPlayTriggerRef->BopFrame(0);
+			OnPlayTriggerRef->TriggerFrame(0);
 			bIsPlaying = true;
 		}
 
diff --git a/Engine/Plugins/Runtime/Metasound/Source/MetasoundGenerator/Public/MetasoundGenerator.h b/Engine/Plugins/Runtime/Metasound/Source/MetasoundGenerator/Public/MetasoundGenerator.h
--- a/Engine/Plugins/Runtime/Metasound/Source/MetasoundGenerator/Public/MetasoundGenerator.h
+++ b/Engine/Plugins/Runtime/Metasound/Source/MetasoundGenerator/Public/MetasoundGenerator.h
@@ -1,14 +1,13 @@
 // Copyright Epic Games, Inc. All Rights Reserved.
-
 #pragma once
 
 #include "CoreMinimal.h"
 #include "DSP/Dsp.h"
-#include "MetasoundBop.h"
 #include "MetasoundOperatorInterface.h"
 #include "MetasoundExecutableOperator.h"
 #include "MetasoundAudioFormats.h"
 #include "MetasoundPrimitives.h"
+#include "MetasoundTrigger.h"
 #include "Sound/SoundGenerator.h"
 
 namespace Metasound
@@ -17,8 +16,8 @@ namespace Metasound
 	{
 		TUniquePtr<Metasound::IOperator> GraphOperator;
 		TArrayView<TDataReadReference<FAudioBuffer>> OutputBuffers;
-		FBopWriteRef TriggerOnPlayRef;
-		FBopReadRef TriggerOnFinishRef;
+		FTriggerWriteRef TriggerOnPlayRef;
+		FTriggerReadRef TriggerOnFinishRef;
 	};
 
 	/** FMetasoundGenerator generates audio from a given metasound IOperator
@@ -129,10 +128,10 @@ namespace Metasound
 		TArray<FAudioBufferReadRef> GraphOutputAudio;
 
 		// Triggered when metasound is played
-		FBopWriteRef OnPlayTriggerRef;
+		FTriggerWriteRef OnPlayTriggerRef;
 
 		// Triggered when metasound is finished
-		FBopReadRef OnFinishedTriggerRef;
+		FTriggerReadRef OnFinishedTriggerRef;
 
 		Audio::AlignedFloatBuffer InterleavedAudioBuffer;
 
diff --git a/Engine/Plugins/Runtime/Metasound/Source/MetasoundStandardNodes/Private/MetasoundADSRNode.cpp b/Engine/Plugins/Runtime/Metasound/Source/MetasoundStandardNodes/Private/MetasoundADSRNode.cpp
--- a/Engine/Plugins/Runtime/Metasound/Source/MetasoundStandardNodes/Private/MetasoundADSRNode.cpp
+++ b/Engine/Plugins/Runtime/Metasound/Source/MetasoundStandardNodes/Private/MetasoundADSRNode.cpp
@@ -5,11 +5,11 @@
 #include "DSP/Envelope.h"
 #include "Internationalization/Text.h"
 #include "MetasoundAudioBuffer.h"
-#include "MetasoundBop.h"
 #include "MetasoundExecutableOperator.h"
-#include "MetasoundTime.h"
 #include "MetasoundNodeRegistrationMacro.h"
 #include "MetasoundPrimitives.h"
+#include "MetasoundTime.h"
+#include "MetasoundTrigger.h"
 
 #define LOCTEXT_NAMESPACE "MetasoundStandardNodes"
 
@@ -32,7 +32,7 @@ namespace Metasound
 			static FVertexInterface DeclareVertexInterface();
 			static TUniquePtr<IOperator> CreateOperator(const FCreateOperatorParams& InParams, FBuildErrorArray& OutErrors);
 
-			FADSROperator(const FOperatorSettings& InSettings, const FBopReadRef& InTriggerAttack, const FBopReadRef& InTriggerRelease, const FADSRReferences& InADSRData);
+			FADSROperator(const FOperatorSettings& InSettings, const FTriggerReadRef& InTriggerAttack, const FTriggerReadRef& InTriggerRelease, const FADSRReferences& InADSRData);
 
 			virtual FDataReferenceCollection GetInputs() const override;
 			virtual FDataReferenceCollection GetOutputs() const override;
@@ -46,12 +46,12 @@ namespace Metasound
 			// TODO: write envelope gen for metasound more suited to this processing structure. 
 			Audio::FEnvelope Envelope;
 
-			FBopReadRef TriggerAttack;
-			FBopReadRef TriggerRelease;
+			FTriggerReadRef TriggerAttack;
+			FTriggerReadRef TriggerRelease;
 
-			FBopWriteRef TriggerAttackComplete;
-			FBopWriteRef TriggerDecayComplete;
-			FBopWriteRef TriggerReleaseComplete;
+			FTriggerWriteRef TriggerAttackComplete;
+			FTriggerWriteRef TriggerDecayComplete;
+			FTriggerWriteRef TriggerReleaseComplete;
 
 			bool bReleased;
 
@@ -61,12 +61,12 @@ namespace Metasound
 			FAudioBufferWriteRef EnvelopeBuffer;
 	};
 
-	FADSROperator::FADSROperator(const FOperatorSettings& InSettings, const FBopReadRef& InTriggerAttack, const FBopReadRef& InTriggerRelease, const FADSRReferences& InADSRData)
+	FADSROperator::FADSROperator(const FOperatorSettings& InSettings, const FTriggerReadRef& InTriggerAttack, const FTriggerReadRef& InTriggerRelease, const FADSRReferences& InADSRData)
 	:	TriggerAttack(InTriggerAttack)
 	,	TriggerRelease(InTriggerRelease)
-	,	TriggerAttackComplete(FBopWriteRef::CreateNew(InSettings))
-	,	TriggerDecayComplete(FBopWriteRef::CreateNew(InSettings))
-	,	TriggerReleaseComplete(FBopWriteRef::CreateNew(InSettings))
+	,	TriggerAttackComplete(FTriggerWriteRef::CreateNew(InSettings))
+	,	TriggerDecayComplete(FTriggerWriteRef::CreateNew(InSettings))
+	,	TriggerReleaseComplete(FTriggerWriteRef::CreateNew(InSettings))
 	,	bReleased(false)
 	,	TimeUntilRelease(0, InSettings.GetSampleRate())
 	,	TimePerBlock(InSettings.GetNumFramesPerBlock(), InSettings.GetSampleRate())
@@ -85,8 +85,8 @@ namespace Metasound
 		InputDataReferences.AddDataReadReference(TEXT("Decay"), FFloatTimeReadRef(ADSRReferences.Decay));
 		InputDataReferences.AddDataReadReference(TEXT("Release"), FFloatTimeReadRef(ADSRReferences.Release));
 		InputDataReferences.AddDataReadReference(TEXT("Sustain Level"), FFloatReadRef(ADSRReferences.SustainLevel));
-		InputDataReferences.AddDataReadReference(TEXT("Trigger Attack"), FBopReadRef(TriggerAttack));
-		InputDataReferences.AddDataReadReference(TEXT("Trigger Release"), FBopReadRef(TriggerRelease));
+		InputDataReferences.AddDataReadReference(TEXT("Trigger Attack"), FTriggerReadRef(TriggerAttack));
+		InputDataReferences.AddDataReadReference(TEXT("Trigger Release"), FTriggerReadRef(TriggerRelease));
 		return InputDataReferences;
 	}
 
@@ -94,9 +94,9 @@ namespace Metasound
 	{
 		FDataReferenceCollection OutputDataReferences;
 		OutputDataReferences.AddDataReadReference(TEXT("Envelope"), FAudioBufferReadRef(EnvelopeBuffer));
-		OutputDataReferences.AddDataReadReference(TEXT("Attack Complete"), FBopReadRef(TriggerAttackComplete));
-		OutputDataReferences.AddDataReadReference(TEXT("Decay Complete"), FBopReadRef(TriggerDecayComplete));
-		OutputDataReferences.AddDataReadReference(TEXT("Release Complete"), FBopReadRef(TriggerReleaseComplete));
+		OutputDataReferences.AddDataReadReference(TEXT("Attack Complete"), FTriggerReadRef(TriggerAttackComplete));
+		OutputDataReferences.AddDataReadReference(TEXT("Decay Complete"), FTriggerReadRef(TriggerDecayComplete));
+		OutputDataReferences.AddDataReadReference(TEXT("Release Complete"), FTriggerReadRef(TriggerReleaseComplete));
 		return OutputDataReferences;
 	}
 
@@ -149,17 +149,17 @@ namespace Metasound
 
 		if (AttackCompleteFrame >= 0)
 		{
-			TriggerAttackComplete->BopFrame(AttackCompleteFrame);
+			TriggerAttackComplete->TriggerFrame(AttackCompleteFrame);
 		}
 
 		if (DecayCompleteFrame >= 0)
 		{
-			TriggerDecayComplete->BopFrame(DecayCompleteFrame);
+			TriggerDecayComplete->TriggerFrame(DecayCompleteFrame);
 		}
 
 		if (Envelope.IsDone() && !bWasDone && ReleaseCompleteFrame >= 0)
 		{
-			TriggerReleaseComplete->BopFrame(ReleaseCompleteFrame);
+			TriggerReleaseComplete->TriggerFrame(ReleaseCompleteFrame);
 		}
 	}
 
@@ -175,11 +175,11 @@ namespace Metasound
 		TriggerReleaseComplete->AdvanceBlock();
 
 		TriggerRelease->ExecuteBlock(
-			// OnPreBop
+			// OnPreTrigger
 			[&](int32 StartFrame, int32 EndFrame)
 			{
 			},
-			// OnBop
+			// OnTrigger
 			[&](int32 StartFrame, int32 EndFrame)
 			{
 				if (!bReleased)
@@ -204,12 +204,12 @@ namespace Metasound
 		);
 
 		TriggerAttack->ExecuteBlock(
-			// OnPreBop
+			// OnPreTrigger
 			[&](int32 StartFrame, int32 EndFrame)
 			{
 				GenerateEnvelope(StartFrame, EndFrame);
 			},
-			// OnBop
+			// OnTrigger
 			[&](int32 StartFrame, int32 EndFrame)
 			{
 				Start();
@@ -238,14 +238,14 @@ namespace Metasound
 				TInputDataVertexModel<FFloatTime>(TEXT("Decay"), LOCTEXT("DecayDurationTooltip", "Decay duration.")),
 				TInputDataVertexModel<FFloatTime>(TEXT("Release"), LOCTEXT("ReleaseDurationTooltip", "Release duration.")),
 				TInputDataVertexModel<float>(TEXT("Sustain Level"), LOCTEXT("SustainLevelTooltip", "Sustain level [0.0f, 1.0f].")),
-				TInputDataVertexModel<FBop>(TEXT("Trigger Attack"), LOCTEXT("BopTooltip", "Trigger the envelope's attack.")),
-				TInputDataVertexModel<FBop>(TEXT("Trigger Release"), LOCTEXT("BopTooltip", "Trigger the envelope's release."))
+				TInputDataVertexModel<FTrigger>(TEXT("Trigger Attack"), LOCTEXT("TriggerTooltip", "Trigger the envelope's attack.")),
+				TInputDataVertexModel<FTrigger>(TEXT("Trigger Release"), LOCTEXT("TriggerTooltip", "Trigger the envelope's release."))
 			),
 			FOutputVertexInterface(
 				TOutputDataVertexModel<FAudioBuffer>(TEXT("Envelope"), LOCTEXT("EnvelopeTooltip", "The output envelope")),
-				TOutputDataVertexModel<FBop>(TEXT("Attack Complete"), LOCTEXT("AttackCompleteTooltip", "Triggered when the envelope attack is complete")),
-				TOutputDataVertexModel<FBop>(TEXT("Decay Complete"), LOCTEXT("DecayCompleteTooltip", "Triggered when the envelope decay is complete")),
-				TOutputDataVertexModel<FBop>(TEXT("Release Complete"), LOCTEXT("ReleaseCompleteTooltip", "Triggered when the envelope is released"))
+				TOutputDataVertexModel<FTrigger>(TEXT("Attack Complete"), LOCTEXT("AttackCompleteTooltip", "Triggered when the envelope attack is complete")),
+				TOutputDataVertexModel<FTrigger>(TEXT("Decay Complete"), LOCTEXT("DecayCompleteTooltip", "Triggered when the envelope decay is complete")),
+				TOutputDataVertexModel<FTrigger>(TEXT("Release Complete"), LOCTEXT("ReleaseCompleteTooltip", "Triggered when the envelope is released"))
 			)
 		);
 
@@ -260,7 +260,7 @@ namespace Metasound
 			Info.ClassName = FName(TEXT("ADSR"));
 			Info.MajorVersion = 1;
 			Info.MinorVersion = 0;
-			Info.Description = LOCTEXT("Metasound_ADSRNodeDescription", "Emits an ADSR (Attack, decay, sustain, & release) envelope when bopped.");
+			Info.Description = LOCTEXT("Metasound_ADSRNodeDescription", "Emits an ADSR (Attack, decay, sustain, & release) envelope when triggered.");
 			Info.Author = PluginAuthor;
 			Info.PromptIfMissing = PluginNodeMissingPrompt;
 			Info.DefaultInterface = DeclareVertexInterface();
@@ -278,9 +278,9 @@ namespace Metasound
 		const FADSRNode& ADSRNode = static_cast<const FADSRNode&>(InParams.Node);
 		const FDataReferenceCollection& InputCollection = InParams.InputDataReferences;
 
-		// TODO: Could no-op this if the bop is not connected.
-		FBopReadRef TriggerAttack = InputCollection.GetDataReadReferenceOrConstruct<FBop>(TEXT("Trigger Attack"), InParams.OperatorSettings);
-		FBopReadRef TriggerRelease = InputCollection.GetDataReadReferenceOrConstruct<FBop>(TEXT("Trigger Release"), InParams.OperatorSettings);
+		// TODO: Could no-op this if the trigger is not connected.
+		FTriggerReadRef TriggerAttack = InputCollection.GetDataReadReferenceOrConstruct<FTrigger>(TEXT("Trigger Attack"), InParams.OperatorSettings);
+		FTriggerReadRef TriggerRelease = InputCollection.GetDataReadReferenceOrConstruct<FTrigger>(TEXT("Trigger Release"), InParams.OperatorSettings);
 
 		// TODO: If none of these are connected, could pre-generate ADSR envelope and return a different operator. 
 		FADSRReferences ADSRReferences = 
diff --git a/Engine/Plugins/Runtime/Metasound/Source/MetasoundStandardNodes/Private/MetasoundBop.cpp b/Engine/Plugins/Runtime/Metasound/Source/MetasoundStandardNodes/Private/MetasoundBop.cpp
deleted file mode 100644
diff --git a/Engine/Plugins/Runtime/Metasound/Source/MetasoundStandardNodes/Private/MetasoundMixerNode.cpp b/Engine/Plugins/Runtime/Metasound/Source/MetasoundStandardNodes/Private/MetasoundMixerNode.cpp
--- a/Engine/Plugins/Runtime/Metasound/Source/MetasoundStandardNodes/Private/MetasoundMixerNode.cpp
+++ b/Engine/Plugins/Runtime/Metasound/Source/MetasoundStandardNodes/Private/MetasoundMixerNode.cpp
@@ -1,22 +1,21 @@
 // Copyright Epic Games, Inc. All Rights Reserved.
-
 #include "MetasoundMixerNode.h"
 
+#include "DSP/Dsp.h"
 #include "MetasoundAudioBuffer.h"
 #include "MetasoundExecutableOperator.h"
 #include "MetasoundFacade.h"
 #include "MetasoundFrequency.h"
+#include "MetasoundGain.h"
 #include "MetasoundNodeRegistrationMacro.h"
+#include "MetasoundOperatorSettings.h"
 #include "MetasoundPrimitives.h"
+#include "MetasoundTrigger.h"
 #include "MetasoundVertex.h"
-#include "MetasoundBop.h"
-#include "MetasoundGain.h"
-#include "MetasoundOperatorSettings.h"
-
-#include "DSP/Dsp.h"
 
 #define LOCTEXT_NAMESPACE "MetasoundStandardNodes"
 
+
 namespace Metasound
 {
 	class FMixerOperator : public TExecutableOperator<FMixerOperator>
diff --git a/Engine/Plugins/Runtime/Metasound/Source/MetasoundStandardNodes/Private/MetasoundOscNode.cpp b/Engine/Plugins/Runtime/Metasound/Source/MetasoundStandardNodes/Private/MetasoundOscNode.cpp
--- a/Engine/Plugins/Runtime/Metasound/Source/MetasoundStandardNodes/Private/MetasoundOscNode.cpp
+++ b/Engine/Plugins/Runtime/Metasound/Source/MetasoundStandardNodes/Private/MetasoundOscNode.cpp
@@ -1,19 +1,19 @@
 // Copyright Epic Games, Inc. All Rights Reserved.
-
 #include "MetasoundOscNode.h"
 
+#include "DSP/Dsp.h"
 #include "MetasoundAudioBuffer.h"
 #include "MetasoundExecutableOperator.h"
 #include "MetasoundFacade.h"
 #include "MetasoundFrequency.h"
 #include "MetasoundNodeRegistrationMacro.h"
 #include "MetasoundPrimitives.h"
+#include "MetasoundTrigger.h"
 #include "MetasoundVertex.h"
-#include "MetasoundBop.h"
-#include "DSP/Dsp.h"
 
 #define LOCTEXT_NAMESPACE "MetasoundStandardNodes"
 
+
 namespace Metasound
 {
 	class FOscOperator : public TExecutableOperator<FOscOperator>
diff --git a/Engine/Plugins/Runtime/Metasound/Source/MetasoundStandardNodes/Private/MetasoundPeriodicTriggerNode.cpp b/Engine/Plugins/Runtime/Metasound/Source/MetasoundStandardNodes/Private/MetasoundPeriodicTriggerNode.cpp
--- a/Engine/Plugins/Runtime/Metasound/Source/MetasoundStandardNodes/Private/MetasoundPeriodicTriggerNode.cpp
+++ b/Engine/Plugins/Runtime/Metasound/Source/MetasoundStandardNodes/Private/MetasoundPeriodicTriggerNode.cpp
@@ -110,7 +110,7 @@ namespace Metasound
 			while ((SampleCountdown - ExecuteDurationInSamples) <= 0.f)
 			{
 				uint32 Frame = FMath::RoundToInt(SampleCountdown);
-				TriggerOut->BopFrame(Frame);
+				TriggerOut->TriggerFrame(Frame);
 				SampleCountdown += PeriodInSamples;
 			}
 
diff --git a/Engine/Plugins/Runtime/Metasound/Source/MetasoundStandardNodes/Private/MetasoundStandardNodesModule.cpp b/Engine/Plugins/Runtime/Metasound/Source/MetasoundStandardNodes/Private/MetasoundStandardNodesModule.cpp
--- a/Engine/Plugins/Runtime/Metasound/Source/MetasoundStandardNodes/Private/MetasoundStandardNodesModule.cpp
+++ b/Engine/Plugins/Runtime/Metasound/Source/MetasoundStandardNodes/Private/MetasoundStandardNodesModule.cpp
@@ -1,15 +1,14 @@
 // Copyright Epic Games, Inc. All Rights Reserved.
 
-#include "MetasoundGraphCoreModule.h"
-
 #include "CoreMinimal.h"
-#include "Modules/ModuleInterface.h"
-#include "Modules/ModuleManager.h"
-#include "MetasoundBop.h"
-#include "MetasoundPrimitives.h"
-#include "MetasoundTime.h"
 #include "MetasoundFrequency.h"
 #include "MetasoundGain.h"
+#include "MetasoundGraphCoreModule.h"
+#include "MetasoundPrimitives.h"
+#include "MetasoundTime.h"
+#include "MetasoundTrigger.h"
+#include "Modules/ModuleInterface.h"
+#include "Modules/ModuleManager.h"
 
 REGISTER_METASOUND_DATATYPE(Metasound::FTrigger, "Trigger", ::Metasound::ELiteralType::Boolean)
 REGISTER_METASOUND_DATATYPE(Metasound::FFloatTime, "Time", ::Metasound::ELiteralType::Float)
diff --git a/Engine/Plugins/Runtime/Metasound/Source/MetasoundStandardNodes/Private/MetasoundTrigger.cpp b/Engine/Plugins/Runtime/Metasound/Source/MetasoundStandardNodes/Private/MetasoundTrigger.cpp
new file mode 100644
--- /dev/null
+++ b/Engine/Plugins/Runtime/Metasound/Source/MetasoundStandardNodes/Private/MetasoundTrigger.cpp
@@ -0,0 +1,121 @@
+// Copyright Epic Games, Inc. All Rights Reserved.
+#include "MetasoundTrigger.h"
+
+#include "MetasoundDataReference.h"
+#include "MetasoundDataTypeRegistrationMacro.h"
+#include "MetasoundOperatorSettings.h"
+
+namespace Metasound
+{
+	FTrigger::FTrigger(const FOperatorSettings& InSettings, bool bShouldTrigger)
+	:	NumFramesPerBlock(InSettings.GetNumFramesPerBlock())
+	,	SampleRate(InSettings.GetSampleRate())
+	{
+		if (bShouldTrigger)
+		{
+			TriggerFrame(0);
+		}
+	}
+
+	FTrigger::FTrigger(const FOperatorSettings& InSettings, int32 InFrameToTrigger)
+	:	NumFramesPerBlock(InSettings.GetNumFramesPerBlock())
+	,	SampleRate(InSettings.GetSampleRate())
+	{
+		TriggerFrame(InFrameToTrigger);
+	}
+
+	FTrigger::FTrigger(const FOperatorSettings& InSettings)
+		: FTrigger(InSettings, false)
+	{}
+
+	void FTrigger::TriggerFrame(int32 InFrameToTrigger)
+	{
+		// Insert trigger frame index into sorted triggered frames
+		TriggeredFrames.Insert(InFrameToTrigger, Algo::LowerBound(TriggeredFrames, InFrameToTrigger));
+
+		if (InFrameToTrigger < NumFramesPerBlock)
+		{
+			// Update last trigger pos. 
+			UpdateLastTriggerIndexInBlock();
+		}
+
+		bHasTrigger = true;
+	}
+
+
+	void FTrigger::AdvanceBlock()
+	{
+		Advance(NumFramesPerBlock);
+	}
+
+	void FTrigger::Advance(int32 InNumFrames)
+	{
+		if (bHasTrigger)
+		{
+			const int32 NumToRemove = Algo::LowerBound(TriggeredFrames, InNumFrames);
+
+			if (NumToRemove)
+			{
+				TriggeredFrames.RemoveAt(0, NumToRemove, false /* bAllowShrinking */);
+			}
+
+			const int32 Num = TriggeredFrames.Num();
+
+			if (Num > 0)
+			{
+				for (int32 i = 0; i < Num; i++)
+				{
+					TriggeredFrames[i] -= InNumFrames;
+				}
+			}
+			else
+			{
+				bHasTrigger = false;
+			}
+
+			UpdateLastTriggerIndexInBlock();
+		}
+	}
+
+	int32 FTrigger::Num() const
+	{
+		return TriggeredFrames.Num();
+	}
+
+	int32 FTrigger::NumTriggeredInBlock() const
+	{
+		return LastTriggerIndexInBlock;
+	}
+
+	int32 FTrigger::operator[](int32 InTriggerIndex) const
+	{
+		return TriggeredFrames[InTriggerIndex];
+	}
+
+	bool FTrigger::IsTriggered() const
+	{
+		return bHasTrigger;
+	}
+
+	bool FTrigger::IsTriggeredInBlock() const
+	{
+		return LastTriggerIndexInBlock > 0;
+	}
+
+	FTrigger::operator bool() const 
+	{
+		return IsTriggeredInBlock();
+	}
+
+	void FTrigger::Reset()
+	{
+		TriggeredFrames.Reset();
+		LastTriggerIndexInBlock = 0;
+		bHasTrigger = false;
+	}
+
+	void FTrigger::UpdateLastTriggerIndexInBlock()
+	{
+		LastTriggerIndexInBlock = Algo::LowerBound(TriggeredFrames, NumFramesPerBlock);
+	}
+}
diff --git a/Engine/Plugins/Runtime/Metasound/Source/MetasoundStandardNodes/Private/MetasoundTriggerAccumulatorNode.cpp b/Engine/Plugins/Runtime/Metasound/Source/MetasoundStandardNodes/Private/MetasoundTriggerAccumulatorNode.cpp
--- a/Engine/Plugins/Runtime/Metasound/Source/MetasoundStandardNodes/Private/MetasoundTriggerAccumulatorNode.cpp
+++ b/Engine/Plugins/Runtime/Metasound/Source/MetasoundStandardNodes/Private/MetasoundTriggerAccumulatorNode.cpp
@@ -1,14 +1,14 @@
 // Copyright Epic Games, Inc. All Rights Reserved.
-
 #include "MetasoundTriggerAccumulatorNode.h"
 
 #include "Internationalization/Text.h"
-#include "MetasoundBop.h"
 #include "MetasoundExecutableOperator.h"
 #include "MetasoundNodeRegistrationMacro.h"
 #include "MetasoundPrimitives.h"
+#include "MetasoundTrigger.h"
 #include "MetasoundVertex.h"
 
+
 #define LOCTEXT_NAMESPACE "MetasoundStandardNodes"
 
 namespace Metasound
@@ -25,7 +25,7 @@ namespace Metasound
 			~FTriggerAccumulatorOperator() = default;
 			
 			FTriggerAccumulatorOperator(
-				const FOperatorSettings& InSettings, 
+				const FOperatorSettings& InSettings,
 				const FTriggerReadRef& InInputTriggerToAccumulate,
 				const FTriggerReadRef& InInputTriggerResetCount,
 				const FInt32ReadRef& InInputTriggerAtCount);
@@ -98,7 +98,7 @@ namespace Metasound
 		if (*CurrentCount >= *TriggerCount)
 		{
 			Reset(InCurrentFrameIndex);
-			CountReached->BopFrame(InCurrentFrameIndex);
+			CountReached->TriggerFrame(InCurrentFrameIndex);
 		}
 	}
 
@@ -108,15 +108,15 @@ namespace Metasound
 		CountReached->AdvanceBlock();
 
 		// Query the number of triggers in the current block.
-		int32 NumResets = ResetTrigger->NumBoppedInBlock();
-		int32 NumAccumulates = TriggerToAccumulate->NumBoppedInBlock();
+		int32 NumResets = ResetTrigger->NumTriggeredInBlock();
+		int32 NumAccumulates = TriggerToAccumulate->NumTriggeredInBlock();
 		int32 MaxTriggers = FMath::Max(NumResets, NumAccumulates);
 
 		// Indices into each trigger array
 		int32 ResetIndex = 0;
 		int32 AccumulateIndex = 0;
 		
-		// Move forwards through the bop array and return next bop frame index, or INDEX_NONE if there's none remaining.
+		// Move forwards through the trigger array and return next trigger frame index, or INDEX_NONE if there's none remaining.
 		static const auto GetNextTriggerFrame = [](const FTriggerReadRef& InTrigger, int32& InOutIndex, const int32 InNumTriggers) 
 			-> int32 { return InOutIndex < InNumTriggers ? (*InTrigger)[InOutIndex++] : INDEX_NONE; };
 		
@@ -175,7 +175,7 @@ namespace Metasound
 				TInputDataVertexModel<int32>(TriggerCountPinName, LOCTEXT("CountToTriggerAtTooltip", "Count to trigger at"))
 			),
 			FOutputVertexInterface(
-				TOutputDataVertexModel<FTrigger>(CountReachedPinName, LOCTEXT("CountReachedTooltop", "Triggered when the accumulated count is reached")),				
+				TOutputDataVertexModel<FTrigger>(CountReachedPinName, LOCTEXT("CountReachedTooltop", "Triggered when the accumulated count is reached")),
 				TOutputDataVertexModel<int32>(CurrentCountPinName, LOCTEXT("CountTooltop", "Current accumulator count"))
 			)
 		);
diff --git a/Engine/Plugins/Runtime/Metasound/Source/MetasoundStandardNodes/Private/MetasoundTriggerDelayNode.cpp b/Engine/Plugins/Runtime/Metasound/Source/MetasoundStandardNodes/Private/MetasoundTriggerDelayNode.cpp
--- a/Engine/Plugins/Runtime/Metasound/Source/MetasoundStandardNodes/Private/MetasoundTriggerDelayNode.cpp
+++ b/Engine/Plugins/Runtime/Metasound/Source/MetasoundStandardNodes/Private/MetasoundTriggerDelayNode.cpp
@@ -89,7 +89,7 @@ namespace Metasound
 			}
 			else
 			{
-				TriggerOut->BopFrame(SamplesRemaining + (int32)FramesPerBlock);
+				TriggerOut->TriggerFrame(SamplesRemaining + (int32)FramesPerBlock);
 				Reset();
 			}
 		}
diff --git a/Engine/Plugins/Runtime/Metasound/Source/MetasoundStandardNodes/Private/MetasoundTriggerPipeNode.cpp b/Engine/Plugins/Runtime/Metasound/Source/MetasoundStandardNodes/Private/MetasoundTriggerPipeNode.cpp
--- a/Engine/Plugins/Runtime/Metasound/Source/MetasoundStandardNodes/Private/MetasoundTriggerPipeNode.cpp
+++ b/Engine/Plugins/Runtime/Metasound/Source/MetasoundStandardNodes/Private/MetasoundTriggerPipeNode.cpp
@@ -2,8 +2,6 @@
 #include "MetasoundTriggerPipeNode.h"
 
 #include "CoreMinimal.h"
-
-#include "MetasoundBop.h"
 #include "MetasoundBuilderInterface.h"
 #include "MetasoundDataReferenceCollection.h"
 #include "MetasoundExecutableOperator.h"
@@ -14,9 +12,11 @@
 #include "MetasoundOperatorInterface.h"
 #include "MetasoundPrimitives.h"
 #include "MetasoundTime.h"
+#include "MetasoundTrigger.h"
 
 #define LOCTEXT_NAMESPACE "MetasoundStandardNodes"
 
+
 namespace Metasound
 {
 	METASOUND_REGISTER_NODE(FTriggerPipeNode)
@@ -88,7 +88,7 @@ namespace Metasound
 			}
 			else
 			{
-				TriggerOut->BopFrame(SamplesRemaining + (int32)FramesPerBlock);
+				TriggerOut->TriggerFrame(SamplesRemaining + (int32)FramesPerBlock);
 				SamplesUntilTrigger.RemoveAtSwap(i);
 			}
 		}
diff --git a/Engine/Plugins/Runtime/Metasound/Source/MetasoundStandardNodes/Public/MetasoundBop.h b/Engine/Plugins/Runtime/Metasound/Source/MetasoundStandardNodes/Public/MetasoundBop.h
deleted file mode 100644
diff --git a/Engine/Plugins/Runtime/Metasound/Source/MetasoundStandardNodes/Public/MetasoundMathNodes.h b/Engine/Plugins/Runtime/Metasound/Source/MetasoundStandardNodes/Public/MetasoundMathNodes.h
--- a/Engine/Plugins/Runtime/Metasound/Source/MetasoundStandardNodes/Public/MetasoundMathNodes.h
+++ b/Engine/Plugins/Runtime/Metasound/Source/MetasoundStandardNodes/Public/MetasoundMathNodes.h
@@ -4,7 +4,6 @@
 #include "CoreMinimal.h"
 #include "DSP/BufferVectorOperations.h"
 #include "MetasoundAudioBuffer.h"
-#include "MetasoundBop.h"
 #include "MetasoundBuilderInterface.h"
 #include "MetasoundDataReference.h"
 #include "MetasoundDataReferenceCollection.h"
@@ -17,6 +16,7 @@
 #include "MetasoundOperatorInterface.h"
 #include "MetasoundPrimitives.h"
 #include "MetasoundTime.h"
+#include "MetasoundTrigger.h"
 
 
 #define LOCTEXT_NAMESPACE "MetasoundMathOpNode"
diff --git a/Engine/Plugins/Runtime/Metasound/Source/MetasoundStandardNodes/Public/MetasoundPeriodicTriggerNode.h b/Engine/Plugins/Runtime/Metasound/Source/MetasoundStandardNodes/Public/MetasoundPeriodicTriggerNode.h
--- a/Engine/Plugins/Runtime/Metasound/Source/MetasoundStandardNodes/Public/MetasoundPeriodicTriggerNode.h
+++ b/Engine/Plugins/Runtime/Metasound/Source/MetasoundStandardNodes/Public/MetasoundPeriodicTriggerNode.h
@@ -5,13 +5,14 @@
 #include "CoreMinimal.h"
 #include "MetasoundDataReferenceCollection.h"
 #include "MetasoundNode.h"
-#include "MetasoundBop.h"
 #include "MetasoundBuilderInterface.h"
 #include "MetasoundExecutableOperator.h"
 #include "MetasoundFacade.h"
 #include "MetasoundOperatorInterface.h"
 #include "MetasoundNodeInterface.h"
 #include "MetasoundTime.h"
+#include "MetasoundTrigger.h"
+
 
 namespace Metasound
 {
diff --git a/Engine/Plugins/Runtime/Metasound/Source/MetasoundStandardNodes/Public/MetasoundTrigger.h b/Engine/Plugins/Runtime/Metasound/Source/MetasoundStandardNodes/Public/MetasoundTrigger.h
new file mode 100644
--- /dev/null
+++ b/Engine/Plugins/Runtime/Metasound/Source/MetasoundStandardNodes/Public/MetasoundTrigger.h
@@ -0,0 +1,173 @@
+// Copyright Epic Games, Inc. All Rights Reserved.
+
+#pragma once
+
+#include "MetasoundDataReference.h"
+#include "MetasoundDataTypeRegistrationMacro.h"
+#include "MetasoundOperatorSettings.h"
+
+namespace Metasound
+{
+	/** FTrigger supports sample accurate triggering, sample accurate internal tracking,
+	 * and a convenient interface for running trigger-aligned audio signal processing
+	 * routines on buffers.. 
+	 *
+	 * FTriggers are triggered using FTrigger::TriggerTime or FTrigger::TriggerFrame.
+	 * FTriggers track time internally by calling FTrigger::Advance.
+	 * Executing audio signal processing on buffers can be performed by calling
+	 * FTrigger::ExecuteBlock or FTrigger::LookAhead.
+	 */
+	class METASOUNDSTANDARDNODES_API FTrigger
+	{
+		public:
+			/** FTrigger constructor. 
+			 *
+			 * @param InSettings - Operator settings.
+			 * @param bShouldTrigger - If true, triggers first sample.
+			 */
+			explicit FTrigger(const FOperatorSettings& InSettings, bool bShouldTrigger);
+
+			/** FTrigger constructor. 
+			 *
+			 * @param InFrameToTrigger - Set specific frame to trigger.
+			 * @param InSettings - Operator settings.
+			 */
+			explicit FTrigger(const FOperatorSettings& InSettings, int32 InFrameToTrigger);
+
+			/** FTrigger constructor. By default it is not triggered.
+			 *
+			 * @param InSettings - Operator settings.
+			 */
+			FTrigger(const FOperatorSettings& InSettings);
+
+			/** Trigger a specific frame in the future.
+			 *
+			 * @param InFrameToTrigger - Index of frame to trigger.
+			 */
+			void TriggerFrame(int32 InFrameToTrigger);
+
+			/** Trigger a specific time in the future.
+			 *
+			 * @param InTime - Time to set trigger.
+			 */
+			template<typename TimeType>
+			void TriggerTime(const TimeType& InTime)
+			{
+				TriggerFrame(InTime.GetNumSamples(SampleRate));
+			}
+
+			/** Advance internal frame counters by block size. */
+			void AdvanceBlock();
+
+			/** Advance internal frame counters by specific frame count. */
+			void Advance(int32 InNumFrames);
+
+			/** Number of triggered frames. */
+			int32 Num() const;
+
+			/** Returns true if there is a trigger in the current block of audio */
+			int32 NumTriggeredInBlock() const;
+
+			/** Returns frame index for a given trigger index. 
+			 *
+			 * @param InIndex - Index of trigger. Must be a value between 0 and Num().
+			 *
+			 * @return The frame of the trigger.
+			 */
+			int32 operator[](int32 InTriggerIndex) const;
+
+			/** Returns true if there are any triggered frames. */
+			bool IsTriggered() const;
+
+			/** Returns true there is a trigger in the current block of audio. */
+			bool IsTriggeredInBlock() const;
+
+			/** Implicit conversion of FTrigger into bool by calling IsTriggeredInBlock()
+			 *
+			 * @return Returns true if frame is triggered in current block.
+			 */
+			operator bool() const;
+
+			/** Removes all triggered frames. */
+			void Reset();
+
+			/** Executes one block of frames and calls underlying InPreTrigger and InOnTrigger
+			 * functions with frame indices.
+			 *
+			 * @param InPreTrigger - A function which handles frames before the first 
+			 *                   trigger in the current block. The function must accept 
+			 *                   the arguments (int32 StartFrame, int32 EndFrame).
+			 * @param InOnTrigger - A function which handles frames starting with the 
+			 *                  triggers index and ending the next trigger index or the
+			 *                  number of frames in a block.. The function must
+			 *                  accept the arguments (int32 StartFrame, int32 EndFrame).
+			 */
+			template<typename PreTriggerType, typename OnTriggerType>
+			void ExecuteBlock(PreTriggerType InPreTrigger, OnTriggerType InOnTrigger) const
+			{
+				ExecuteFrames(NumFramesPerBlock, LastTriggerIndexInBlock, InPreTrigger, InOnTrigger);
+			}
+
+			/** Executes a desired number of frames and calls underlying InPreTrigger 
+			 * and InOnTrigger functions with frame indices.
+			 *
+			 * @param InNumFrames - The number of frames to look ahead.
+			 * @param InPreTrigger - A function which handles frames before the first 
+			 *                   trigger in the range of frames. The function must accept 
+			 *                   the arguments (int32 StartFrame, int32 EndFrame).
+			 * @param InOnTrigger - A function which handles frames starting with the 
+			 *                  triggers index and ending the next trigger index or InNumFrames
+			 *             		The function must accept the arguments 
+			 *             		(int32 StartFrame, int32 EndFrame).
+			 */
+			template<typename PreTriggerType, typename OnTriggerType>
+			void LookAhead(int32 InNumFrames, PreTriggerType InPreTrigger, OnTriggerType InOnTrigger) const
+			{
+				if (InNumFrames > 0)
+				{
+					int32 LastTriggerIndexInLookAhead= Algo::LowerBound(TriggeredFrames, InNumFrames);
+					
+					ExecuteFrames(InNumFrames, LastTriggerIndexInLookAhead, InPreTrigger, InOnTrigger);
+				}
+			}
+
+		private:
+			void UpdateLastTriggerIndexInBlock();
+
+			template<typename PreTriggerType, typename OnTriggerType>
+			void ExecuteFrames(int32 InNumFrames, int32 InLastTriggerIndex, PreTriggerType InPreTrigger, OnTriggerType InOnTrigger) const
+			{
+				if (InLastTriggerIndex <= 0)
+				{
+					InPreTrigger(0, InNumFrames);
+					return;
+				}
+
+				const int32* TriggeredFramesData = TriggeredFrames.GetData();
+
+				if (TriggeredFramesData[0] > 0)
+				{
+					InPreTrigger(0, TriggeredFramesData[0]);
+				}
+
+				const int32 NextToLastTriggerIndex = InLastTriggerIndex - 1;
+
+				for (int32 i = 0; i < NextToLastTriggerIndex; i++)
+				{
+					InOnTrigger(TriggeredFramesData[i], TriggeredFramesData[i + 1]);
+				}
+
+				InOnTrigger(TriggeredFramesData[NextToLastTriggerIndex], InNumFrames);
+			}
+
+
+			TArray<int32> TriggeredFrames;
+
+			bool bHasTrigger = false;
+			int32 NumFramesPerBlock = 0;
+			float SampleRate;
+			int32 LastTriggerIndexInBlock = 0;
+	};
+
+	DECLARE_METASOUND_DATA_REFERENCE_TYPES(FTrigger, METASOUNDSTANDARDNODES_API, FTriggerTypeInfo, FTriggerReadRef, FTriggerWriteRef);
+}
diff --git a/Engine/Plugins/Runtime/Metasound/Source/MetasoundStandardNodes/Public/MetasoundTriggerDelayNode.h b/Engine/Plugins/Runtime/Metasound/Source/MetasoundStandardNodes/Public/MetasoundTriggerDelayNode.h
--- a/Engine/Plugins/Runtime/Metasound/Source/MetasoundStandardNodes/Public/MetasoundTriggerDelayNode.h
+++ b/Engine/Plugins/Runtime/Metasound/Source/MetasoundStandardNodes/Public/MetasoundTriggerDelayNode.h
@@ -3,15 +3,16 @@
 #pragma once
 
 #include "CoreMinimal.h"
-#include "MetasoundNode.h"
-#include "MetasoundBop.h"
 #include "MetasoundBuilderInterface.h"
+#include "MetasoundDataReferenceCollection.h"
 #include "MetasoundExecutableOperator.h"
 #include "MetasoundFacade.h"
-#include "MetasoundOperatorInterface.h"
+#include "MetasoundNode.h"
 #include "MetasoundNodeInterface.h"
-#include "MetasoundDataReferenceCollection.h"
+#include "MetasoundOperatorInterface.h"
 #include "MetasoundTime.h"
+#include "MetasoundTrigger.h"
+
 
 namespace Metasound
 {
@@ -24,6 +25,6 @@ namespace Metasound
 			float GetDefaultDelayInSeconds() const;
 
 		private:
-			float DefaultDelay;
+			float DefaultDelay = 1.0f;
 	};
 }
