Fix for sounds crashing editor on Mac

diff --git a/Engine/Plugins/Runtime/Metasound/Source/MetasoundEngine/Private/MetasoundWave.cpp b/Engine/Plugins/Runtime/Metasound/Source/MetasoundEngine/Private/MetasoundWave.cpp
--- a/Engine/Plugins/Runtime/Metasound/Source/MetasoundEngine/Private/MetasoundWave.cpp
+++ b/Engine/Plugins/Runtime/Metasound/Source/MetasoundEngine/Private/MetasoundWave.cpp
@@ -9,66 +9,6 @@
 #include "AudioDeviceManager.h"
 #include "AudioDevice.h"
 
-
-// namespace Metasound
-// {	
-// 	FWaveAsset::FDecoderInputPtr FWaveAsset::CreateDecoderInput() const
-// 	{
-// 		if (SoundWaveProxy.IsValid())
-// 		{
-// 			FName OldFormat = SoundWaveProxy->GetRuntimeFormat();
-// 
-// 			TUniquePtr<Audio::IDecoderInput> Input = Audio::CreateBackCompatDecoderInput(OldFormat, *SoundWaveProxy);
-// 
-// 			return FWaveAsset::FDecoderInputPtr(Input.Release());
-// 		}
-// 
-// 		return nullptr;
-// 	}
-// 
-// 
-// 	FWaveAsset::FDecoderTrio FWaveAsset::CreateDecoderTrio(const float OutputSampleRate, const int32 NumFramesPerBlock) const
-// 	{
-// 		using namespace Audio;
-// 
-// 		if (!IsSoundWaveValid())
-// 		{
-// 			return { };
-// 		}
-// 
-// 		FWaveAsset::FDecoderInputPtr Input;
-// 		if (!(Input = FWaveAsset::CreateDecoderInput()))
-// 		{
-// 			return { };
-// 		}
-// 
-// 		ICodecRegistry::FCodecPtr Codec = ICodecRegistry::Get().FindCodecByParsingInput(Input.Get());
-// 		if (!Codec)
-// 		{
-// 			return { };
-// 		}
-// 
-// 		IDecoderOutput::FRequirements Reqs
-// 		{
-// 			Float32_Interleaved,
-// 			NumFramesPerBlock,
-// 			static_cast<int32>(OutputSampleRate)
-// 		};
-// 
-// 		TUniquePtr<IDecoderOutput> Output = IDecoderOutput::Create(Reqs);
-// 		TUniquePtr<IDecoder> Decoder = Codec->CreateDecoder(Input.Get(), Output.Get());
-// 
-// 		return
-// 		{
-// 			  MoveTemp(Input)
-// 			, MoveTemp(Output)
-// 			, MoveTemp(Decoder)
-// 		};
-// 
-// 	}
-// 
-// }
-
 namespace Audio
 {
 	bool FSimpleDecoderWrapper::Initialize(const InitParams& InInitParams, const FSoundWaveProxy& InWave)
diff --git a/Engine/Source/Runtime/AudioMixer/Private/AudioMixerBuffer.cpp b/Engine/Source/Runtime/AudioMixer/Private/AudioMixerBuffer.cpp
--- a/Engine/Source/Runtime/AudioMixer/Private/AudioMixerBuffer.cpp
+++ b/Engine/Source/Runtime/AudioMixer/Private/AudioMixerBuffer.cpp
@@ -318,11 +318,10 @@ namespace Audio
 		FMixerBuffer* Buffer = new FMixerBuffer(AudioDevice, InWave, EBufferType::Streaming);
 
 		FSoundQualityInfo QualityInfo = { 0 };
-
-		Buffer->DecompressionState = AudioDevice->CreateCompressedAudioInfo(InWave->GetThisAsProxy());
+		Buffer->DecompressionState = AudioDevice->CreateCompressedAudioInfo(InWave);
 
 		// Get the header information of our compressed format
-		if (Buffer->DecompressionState->StreamCompressedInfo(InWave, &QualityInfo))
+		if (Buffer->DecompressionState && Buffer->DecompressionState->StreamCompressedInfo(InWave, &QualityInfo))
 		{
 			// Refresh the wave data
 			InWave->SetSampleRate(QualityInfo.SampleRate);
@@ -372,15 +371,7 @@ namespace Audio
 			InWave->InitAudioResource(AudioDevice->GetRuntimeFormat(InWave));
 		}
 
-		// stream caching uses the proxy
-		if (InWave->IsStreaming())
-		{
-			Buffer->DecompressionState = AudioDevice->CreateCompressedAudioInfo(InWave->GetThisAsProxy());
-		}
-		else
-		{
-			Buffer->DecompressionState = AudioDevice->CreateCompressedAudioInfo(InWave);
-		}
+		Buffer->DecompressionState = AudioDevice->CreateCompressedAudioInfo(InWave);
 
 		if (Buffer->DecompressionState)
 		{
