Adding new audio meter to metasound graph editor

Removing the palette tab

diff --git a/Engine/Plugins/Runtime/Metasound/Metasound.uplugin b/Engine/Plugins/Runtime/Metasound/Metasound.uplugin
--- a/Engine/Plugins/Runtime/Metasound/Metasound.uplugin
+++ b/Engine/Plugins/Runtime/Metasound/Metasound.uplugin
@@ -1,50 +1,61 @@
 {
-	"FileVersion" : 3,
-	"Version" : 1,
-	"VersionName" : "1.0",
-	"FriendlyName" : "Metasound",
-	"Description" : "",
-	"Category" : "Audio",
-	"CreatedBy" : "Epic Games, Inc.",
-	"CreatedByURL" : "http://epicgames.com",
-	"DocsURL" : "",
-	"MarketplaceURL" : "",
-	"SupportURL" : "",
-	"EnabledByDefault" : false,
-	"CanContainContent" : true,
-	"IsBetaVersion" : true,
-	"Installed" : false,
-	"Modules" :
+	"FileVersion": 3,
+	"Version": 1,
+	"VersionName": "1.0",
+	"FriendlyName": "Metasound",
+	"Description": "",
+	"Category": "Audio",
+	"CreatedBy": "Epic Games, Inc.",
+	"CreatedByURL": "http://epicgames.com",
+	"DocsURL": "",
+	"MarketplaceURL": "",
+	"SupportURL": "",
+	"EnabledByDefault": false,
+	"CanContainContent": true,
+	"IsBetaVersion": true,
+	"Installed": false,
+	"Modules": 
 	[
 		{
-			"Name" : "MetasoundGraphCore",
-			"Type" : "Runtime",
-			"LoadingPhase" : "PreDefault"
+			"Name": "MetasoundGraphCore",
+			"Type": "Runtime",
+			"LoadingPhase": "PreDefault"
 		},
 		{
-			"Name" : "MetasoundGenerator",
-			"Type" : "Runtime",
-			"LoadingPhase" : "Default"
+			"Name": "MetasoundGenerator",
+			"Type": "Runtime",
+			"LoadingPhase": "Default"
 		},
 		{
-			"Name" : "MetasoundFrontend",
-			"Type" : "Runtime",
-			"LoadingPhase" : "Default"
+			"Name": "MetasoundFrontend",
+			"Type": "Runtime",
+			"LoadingPhase": "Default"
 		},
 		{
-			"Name" : "MetasoundStandardNodes",
-			"Type" : "Runtime",
-			"LoadingPhase" : "Default"
+			"Name": "MetasoundStandardNodes",
+			"Type": "Runtime",
+			"LoadingPhase": "Default"
 		},
 		{
-			"Name" : "MetasoundEngine",
-			"Type" : "Runtime",
-			"LoadingPhase" : "Default"
+			"Name": "MetasoundEngine",
+			"Type": "Runtime",
+			"LoadingPhase": "Default"
 		},
 		{
 			"Name": "MetasoundEditor",
 			"Type": "Editor",
 			"LoadingPhase": "PostEngineInit"
+		}
+	],
+	"Plugins": 
+	[
+		{
+			"Name": "AudioWidgets",
+			"Enabled": true
 		},
+		{
+			"Name": "AudioSynesthesia",
+			"Enabled": true
+		}
 	]
 }
diff --git a/Engine/Plugins/Runtime/Metasound/Source/MetasoundEditor/MetasoundEditor.Build.cs b/Engine/Plugins/Runtime/Metasound/Source/MetasoundEditor/MetasoundEditor.Build.cs
--- a/Engine/Plugins/Runtime/Metasound/Source/MetasoundEditor/MetasoundEditor.Build.cs
+++ b/Engine/Plugins/Runtime/Metasound/Source/MetasoundEditor/MetasoundEditor.Build.cs
@@ -22,6 +22,8 @@ namespace UnrealBuildTool.Rules
 					"MetasoundFrontend",
 					"MetasoundGraphCore",
 					"ToolMenus",
+					"AudioWidgets",
+					"AudioSynesthesia"
 				}
 			);
 
@@ -46,6 +48,7 @@ namespace UnrealBuildTool.Rules
 					"SlateCore",
 					"EditorFramework",
 					"UnrealEd",
+					"AudioAnalyzer"
 				}
 			);
 		}
diff --git a/Engine/Plugins/Runtime/Metasound/Source/MetasoundEditor/Private/MetasoundEditor.cpp b/Engine/Plugins/Runtime/Metasound/Source/MetasoundEditor/Private/MetasoundEditor.cpp
--- a/Engine/Plugins/Runtime/Metasound/Source/MetasoundEditor/Private/MetasoundEditor.cpp
+++ b/Engine/Plugins/Runtime/Metasound/Source/MetasoundEditor/Private/MetasoundEditor.cpp
@@ -32,7 +32,8 @@
 #include "SNodePanel.h"
 #include "Widgets/Docking/SDockTab.h"
 #include "Widgets/Notifications/SNotificationList.h"
-
+#include "Widgets/SBoxPanel.h"
+#include "AudioMeterStyle.h"
 
 #define LOCTEXT_NAMESPACE "MetasoundEditor"
 
@@ -63,11 +64,11 @@ namespace Metasound
 			.SetGroup(WorkspaceMenuCategoryRef)
 			.SetIcon(FSlateIcon(FEditorStyle::GetStyleSetName(), "LevelEditor.Tabs.Details"));
 
-			InTabManager->RegisterTabSpawner(TabFactory::Names::Palette, FOnSpawnTab::CreateLambda([InPalette = Palette](const FSpawnTabArgs& Args)
+			InTabManager->RegisterTabSpawner(TabFactory::Names::Analyzers, FOnSpawnTab::CreateLambda([InMetasoundMeter = MetasoundMeter](const FSpawnTabArgs& Args)
 			{
-				return TabFactory::CreatePaletteTab(InPalette, Args);
+				return TabFactory::CreateAnalyzersTab(InMetasoundMeter, Args);
 			}))
-			.SetDisplayName(LOCTEXT("PaletteTab", "Palette"))
+			.SetDisplayName(LOCTEXT("AnalyzersTab", "Analyzers"))
 			.SetGroup(WorkspaceMenuCategoryRef)
 			.SetIcon(FSlateIcon(FEditorStyle::GetStyleSetName(), "Kismet.Tabs.Palette"));
 		}
@@ -78,8 +79,8 @@ namespace Metasound
 
 			FAssetEditorToolkit::UnregisterTabSpawners(InTabManager);
 
+			InTabManager->UnregisterTabSpawner(TabFactory::Names::Analyzers);
 			InTabManager->UnregisterTabSpawner(TabFactory::Names::GraphCanvas);
-			InTabManager->UnregisterTabSpawner(TabFactory::Names::Palette);
 			InTabManager->UnregisterTabSpawner(TabFactory::Names::Properties);
 		}
 
@@ -133,7 +134,9 @@ namespace Metasound
 
 			CreateInternalWidgets();
 
-			const TSharedRef<FTabManager::FLayout> StandaloneDefaultLayout = FTabManager::NewLayout("Standalone_MetasoundEditor_Layout_v2")
+			CreateAnalyzers();
+
+			const TSharedRef<FTabManager::FLayout> StandaloneDefaultLayout = FTabManager::NewLayout("Standalone_MetasoundEditor_Layout_v4")
 				->AddArea
 				(
 					FTabManager::NewPrimaryArea()
@@ -160,7 +163,7 @@ namespace Metasound
 							FTabManager::NewStack()
 							->SetSizeCoefficient(0.125f)
 							->SetHideTabWell(true)
-							->AddTab(TabFactory::Names::Palette, ETabState::OpenedTab)
+							->AddTab(TabFactory::Names::Analyzers, ETabState::OpenedTab)
 						)
 					)
 				);
@@ -178,6 +181,11 @@ namespace Metasound
 			return Metasound;
 		}
 
+		UObject* FEditor::GetMetasoundAudioBusObject() const
+		{
+			return MetasoundAudioBus.Get();
+		}
+
 		void FEditor::SetSelection(const TArray<UObject*>& SelectedObjects)
 		{
 			if (MetasoundProperties.IsValid())
@@ -261,6 +269,64 @@ namespace Metasound
 			Palette = SNew(SMetasoundPalette);
 		}
 
+		void FEditor::OnMeterOutput(UMeterAnalyzer* InMeterAnalyzer, int32 ChannelIndex, const FMeterResults& InMeterResults)
+		{
+			if (InMeterAnalyzer && MetasoundMeter.IsValid())
+			{
+				FMeterChannelInfo ChannelInfo;
+				ChannelInfo.MeterValue = InMeterResults.MeterValue;
+				ChannelInfo.PeakValue = InMeterResults.PeakValue;
+
+				check(ChannelIndex < MetasoundChannelInfo.Num());
+				MetasoundChannelInfo[ChannelIndex] = ChannelInfo;
+
+				// If it's the last channel, update the widget
+				if (ChannelIndex == MetasoundChannelInfo.Num() - 1)
+				{
+					MetasoundMeter->SetMeterChannelInfo(MetasoundChannelInfo);
+				}
+			}
+		}
+
+		void FEditor::CreateAnalyzers()
+		{
+			MetasoundAudioBus = TStrongObjectPtr<UAudioBus>(NewObject<UAudioBus>());
+			MetasoundAudioBus->AudioBusChannels = EAudioBusChannels::Stereo;
+
+			MetasoundMeterAnalyzer = TStrongObjectPtr<UMeterAnalyzer>(NewObject<UMeterAnalyzer>());
+
+			MetasoundMeterAnalyzer->OnLatestPerChannelMeterResultsNative.AddRaw(this, &FEditor::OnMeterOutput);
+
+			MetasoundMeterAnalyzerSettings = TStrongObjectPtr<UMeterSettings>(NewObject<UMeterSettings>());
+			MetasoundMeterAnalyzerSettings->PeakHoldTime = 4000.0f;
+
+			MetasoundMeterAnalyzer->Settings = MetasoundMeterAnalyzerSettings.Get();
+
+			UWorld* EditorWorld = GEditor->GetEditorWorldContext().World();
+			MetasoundMeterAnalyzer->StartAnalyzing(EditorWorld, MetasoundAudioBus.Get());
+
+			// Create the audio meter
+			MetasoundMeter = SNew(SAudioMeter);			
+
+			// TODO: THIS IS TEMP, WIP
+			MetasoundMeter->SetOrientation(EOrientation::Orient_Vertical);
+			MetasoundMeter->SetBackgroundColor(FLinearColor(0.0075f, 0.0075f, 0.0075, 1.0f));
+			MetasoundMeter->SetMeterBackgroundColor(FLinearColor(0.031f, 0.031f, 0.031f, 1.0f));
+			MetasoundMeter->SetMeterValueColor(FLinearColor(0.025719f, 0.208333f, 0.069907f, 1.0f));
+			MetasoundMeter->SetMeterPeakColor(FLinearColor(0.24349f, 0.708333f, 0.357002f, 1.0f));
+			MetasoundMeter->SetMeterClippingColor(FLinearColor(1.0f, 0.0f, 0.112334f, 1.0f));
+			MetasoundMeter->SetMeterScaleColor(FLinearColor(0.017642f, 0.017642f, 0.017642f, 1.0f));
+			MetasoundMeter->SetMeterScaleLabelColor(FLinearColor(0.442708f, 0.442708f, 0.442708f, 1.0f));
+
+			FMeterChannelInfo DefaultInfo;
+			DefaultInfo.MeterValue = -60.0f;
+			DefaultInfo.PeakValue = -60.0f;
+			MetasoundChannelInfo.Add(DefaultInfo);
+			MetasoundChannelInfo.Add(DefaultInfo);
+			MetasoundMeter->SetMeterChannelInfo(MetasoundChannelInfo);
+
+		}
+
 		void FEditor::ExtendToolbar()
 		{
 			TSharedPtr<FExtender> ToolbarExtender = MakeShared<FExtender>();
@@ -418,14 +484,21 @@ namespace Metasound
 
 		void FEditor::Play()
 		{
-			// 	TODO: Implement play
 			check(GEditor);
 
-			// TODO: toggle the Play button if Metasound is a USoundBase
 			if (USoundBase* MetasoundToPlay = Cast<USoundBase>(Metasound))
 			{
 				GEditor->PlayPreviewSound(MetasoundToPlay);
 
+				// Set the send to the audio bus that is used for analyzing the metasound output
+				if (UAudioComponent* PreviewComp = GEditor->GetPreviewAudioComponent())
+				{
+					if (MetasoundAudioBus.IsValid())
+					{
+						PreviewComp->SetAudioBusSendPostEffect(MetasoundAudioBus.Get(), 1.0f);
+					}
+				}
+
 				MetasoundGraphEditor->RegisterActiveTimer(0.0f,
 					FWidgetActiveTimerDelegate::CreateLambda([](double InCurrentTime, float InDeltaTime)
 						{
diff --git a/Engine/Plugins/Runtime/Metasound/Source/MetasoundEditor/Private/MetasoundEditor.h b/Engine/Plugins/Runtime/Metasound/Source/MetasoundEditor/Private/MetasoundEditor.h
--- a/Engine/Plugins/Runtime/Metasound/Source/MetasoundEditor/Private/MetasoundEditor.h
+++ b/Engine/Plugins/Runtime/Metasound/Source/MetasoundEditor/Private/MetasoundEditor.h
@@ -14,7 +14,12 @@
 #include "Textures/SlateIcon.h"
 #include "Toolkits/IToolkitHost.h"
 #include "UObject/GCObject.h"
-
+#include "Widgets/SPanel.h"
+#include "Sound/AudioBus.h"
+#include "SAudioMeter.h"
+#include "AudioMeterStyle.h"
+#include "AudioSynesthesia/Classes/Meter.h"
+#include "UObject/StrongObjectPtr.h"
 
 // Forward Declarations
 class FSlateRect;
@@ -24,9 +29,11 @@ class IToolkitHost;
 class SDockableTab;
 class SGraphEditor;
 class SMetasoundPalette;
+class SVerticalBox;
 class UEdGraphNode;
 class UMetasound;
 
+struct FMeterResults;
 struct FPropertyChangedEvent;
 
 namespace Metasound
@@ -46,6 +53,7 @@ namespace Metasound
 
 			/** IMetasoundEditor interface */
 			virtual UObject* GetMetasoundObject() const override;
+			virtual UObject* GetMetasoundAudioBusObject() const override;
 			virtual void SetSelection(const TArray<UObject*>& SelectedObjects) override;
 			virtual bool GetBoundsForSelectedNodes(FSlateRect& Rect, float Padding) override;
 
@@ -78,6 +86,8 @@ namespace Metasound
 				return MetasoundGraphEditor->GetSelectedNodes().Num();
 			}
 
+			void OnMeterOutput(UMeterAnalyzer* InMeterAnalyzer, int32 ChannelIndex, const FMeterResults& MeterResults);
+
 		protected:
 			/** Called when the selection changes in the GraphEditor */
 			void OnSelectedNodesChanged(const TSet<UObject*>& NewSelection);
@@ -122,6 +132,9 @@ namespace Metasound
 			/** Creates all internal widgets for the tabs to point at */
 			void CreateInternalWidgets();
 
+			/** Creates analyzers */
+			void CreateAnalyzers();
+
 			/** Builds the toolbar widget for the Metasound editor */
 			void ExtendToolbar();
 
@@ -179,7 +192,13 @@ namespace Metasound
 			/** Properties tab */
 			TSharedPtr<IDetailsView> MetasoundProperties;
 
-			/** Palette of Sound Node types */
+			/** Metasound Output Meter widget */
+			TSharedPtr<SAudioMeter> MetasoundMeter;
+
+			/** Metasound channel info for the meter. */
+			TArray<FMeterChannelInfo> MetasoundChannelInfo;
+
+			/** Palette of Node types */
 			TSharedPtr<SMetasoundPalette> Palette;
 
 			/** Command list for this editor */
@@ -188,6 +207,14 @@ namespace Metasound
 			/** The Metasound asset being edited */
 			UObject* Metasound = nullptr;
 
+			/** The preview audio bus. Used for analysis. */
+			TStrongObjectPtr<UAudioBus> MetasoundAudioBus;
+
+			/** Metasound analyzer object. */
+			TStrongObjectPtr<UMeterAnalyzer> MetasoundMeterAnalyzer;
+
+			TStrongObjectPtr<UMeterSettings> MetasoundMeterAnalyzerSettings;
+
 			/** Whether or not metasound being edited is valid */
 			bool bPassedValidation = true;
 
diff --git a/Engine/Plugins/Runtime/Metasound/Source/MetasoundEditor/Private/MetasoundEditorTabFactory.h b/Engine/Plugins/Runtime/Metasound/Source/MetasoundEditor/Private/MetasoundEditorTabFactory.h
--- a/Engine/Plugins/Runtime/Metasound/Source/MetasoundEditor/Private/MetasoundEditorTabFactory.h
+++ b/Engine/Plugins/Runtime/Metasound/Source/MetasoundEditor/Private/MetasoundEditorTabFactory.h
@@ -6,6 +6,7 @@
 #include "Framework/Docking/TabManager.h"
 #include "Internationalization/Internationalization.h"
 #include "Widgets/Docking/SDockTab.h"
+#include "SAudioMeter.h"
 #include "Templates/SharedPointer.h"
 
 #define LOCTEXT_NAMESPACE "MetasoundEditor"
@@ -18,11 +19,24 @@ namespace Metasound
 		{
 			namespace Names
 			{
+				const FName Analyzers = "MetasoundEditor_Analyzers";
 				const FName GraphCanvas = "MetasoundEditor_GraphCanvas";
 				const FName Palette = "MetasoundEditor_Palette";
 				const FName Properties = "MetasoundEditor_Properties";
 			}
 
+			TSharedRef<SDockTab> CreateAnalyzersTab(TSharedPtr<SAudioMeter> InAudioMeter, const FSpawnTabArgs& Args)
+			{
+				check(Args.GetTabId() == Names::Analyzers);
+
+				return SNew(SDockTab)
+					.Icon(FEditorStyle::GetBrush("Kismet.Tabs.Palette"))
+					.Label(LOCTEXT("MetasoundAnalyzersTitle", "Analyzers"))
+					[
+						InAudioMeter.ToSharedRef()
+					];
+			}
+
 			TSharedRef<SDockTab> CreateGraphCanvasTab(TSharedPtr<SGraphEditor> GraphEditor, const FSpawnTabArgs& Args)
 			{
 				check(Args.GetTabId() == Names::GraphCanvas);
diff --git a/Engine/Plugins/Runtime/Metasound/Source/MetasoundEditor/Private/SMetasoundGraphNode.cpp b/Engine/Plugins/Runtime/Metasound/Source/MetasoundEditor/Private/SMetasoundGraphNode.cpp
--- a/Engine/Plugins/Runtime/Metasound/Source/MetasoundEditor/Private/SMetasoundGraphNode.cpp
+++ b/Engine/Plugins/Runtime/Metasound/Source/MetasoundEditor/Private/SMetasoundGraphNode.cpp
@@ -19,6 +19,7 @@
 #include "SCommentBubble.h"
 #include "SGraphNode.h"
 #include "SGraphPin.h"
+#include "Widgets/Layout/SSpacer.h"
 #include "SLevelOfDetailBranchNode.h"
 #include "Styling/SlateStyleRegistry.h"
 #include "Styling/SlateColor.h"
diff --git a/Engine/Plugins/Runtime/Metasound/Source/MetasoundEditor/Public/IMetasoundEditor.h b/Engine/Plugins/Runtime/Metasound/Source/MetasoundEditor/Public/IMetasoundEditor.h
--- a/Engine/Plugins/Runtime/Metasound/Source/MetasoundEditor/Public/IMetasoundEditor.h
+++ b/Engine/Plugins/Runtime/Metasound/Source/MetasoundEditor/Public/IMetasoundEditor.h
@@ -12,6 +12,7 @@
 class IMetasoundEditor : public FAssetEditorToolkit
 {
 	virtual UObject* GetMetasoundObject() const = 0;
+	virtual UObject* GetMetasoundAudioBusObject() const = 0;
 	virtual void SetSelection(const TArray<UObject*>& SelectedObjects) = 0;
 	virtual bool GetBoundsForSelectedNodes(FSlateRect& Rect, float Padding) = 0;
 };
