Update Metasound Graph Core Metadata
- Rename FNodeInfo to FNodeClassMetadata
- Add FGuid to Metasound node instances
- Metasound node class names are now Namespace.Name.Variant
- Add displayname for Metasound node classes.

diff --git a/Engine/Plugins/Runtime/Metasound/Source/MetasoundEditor/Private/MetasoundEditorGraphNode.cpp b/Engine/Plugins/Runtime/Metasound/Source/MetasoundEditor/Private/MetasoundEditorGraphNode.cpp
--- a/Engine/Plugins/Runtime/Metasound/Source/MetasoundEditor/Private/MetasoundEditorGraphNode.cpp
+++ b/Engine/Plugins/Runtime/Metasound/Source/MetasoundEditor/Private/MetasoundEditorGraphNode.cpp
@@ -251,7 +251,7 @@ FText UMetasoundEditorGraphNode::GetNodeTitle(ENodeTitleType::Type TitleType) co
 		case EMetasoundFrontendClassType::External:
 		case EMetasoundFrontendClassType::Graph:
 		{
-			return FText::FromString(NodeHandle->GetClassName());
+			return NodeHandle->GetClassDisplayName();
 		}
 		break;
 	}
diff --git a/Engine/Plugins/Runtime/Metasound/Source/MetasoundEditor/Private/MetasoundEditorGraphSchema.cpp b/Engine/Plugins/Runtime/Metasound/Source/MetasoundEditor/Private/MetasoundEditorGraphSchema.cpp
--- a/Engine/Plugins/Runtime/Metasound/Source/MetasoundEditor/Private/MetasoundEditorGraphSchema.cpp
+++ b/Engine/Plugins/Runtime/Metasound/Source/MetasoundEditor/Private/MetasoundEditorGraphSchema.cpp
@@ -659,8 +659,6 @@ void UMetasoundEditorGraphSchema::GetAllMetasoundActions(FGraphActionMenuBuilder
 			? Metadata.Description
 			: FText::Format(LOCTEXT("MetasoundTooltipAuthorFormat", "{0}\nAuthor: {1}"), Metadata.Description, Metadata.Author);
 
-		FString DisplayName = ClassInfo.NodeName;
-
 		const FText* MenuName = &NodeMenuName;
 		FText CategoriesText = FText::GetEmpty();
 		if (!Metadata.CategoryHierarchy.IsEmpty() && !Metadata.CategoryHierarchy[0].CompareTo(ConvertMenuName))
@@ -675,7 +673,7 @@ void UMetasoundEditorGraphSchema::GetAllMetasoundActions(FGraphActionMenuBuilder
 		TSharedPtr<FMetasoundGraphSchemaAction_NewNode> NewNodeAction = MakeShared<FMetasoundGraphSchemaAction_NewNode>
 		(
 			FText::Format(MenuJoinFormat, *MenuName, CategoriesText),
-			FText::FromString(DisplayName),
+			Metadata.DisplayName,
 			Tooltip,
 			0
 		);
diff --git a/Engine/Plugins/Runtime/Metasound/Source/MetasoundEditor/Private/MetasoundFactory.cpp b/Engine/Plugins/Runtime/Metasound/Source/MetasoundEditor/Private/MetasoundFactory.cpp
--- a/Engine/Plugins/Runtime/Metasound/Source/MetasoundEditor/Private/MetasoundFactory.cpp
+++ b/Engine/Plugins/Runtime/Metasound/Source/MetasoundEditor/Private/MetasoundFactory.cpp
@@ -23,7 +23,7 @@ UObject* UMetasoundFactory::FactoryCreateNew(UClass* InClass, UObject* InParent,
 
 	FMetasoundFrontendClassMetadata Metadata;
 
-	Metadata.Name = FMetasoundFrontendClassName{TEXT(""), NewMetasound->GetName(), TEXT("")};
+	Metadata.ClassName = FMetasoundFrontendClassName{TEXT(""), *NewMetasound->GetName(), TEXT("")};
 	Metadata.Version.Major = 1;
 	Metadata.Version.Minor = 0;
 	Metadata.Type = EMetasoundFrontendClassType::Graph;
@@ -52,7 +52,7 @@ UObject* UMetasoundSourceFactory::FactoryCreateNew(UClass* InClass, UObject* InP
 
 	FMetasoundFrontendClassMetadata Metadata;
 
-	Metadata.Name = FMetasoundFrontendClassName{TEXT(""), MetasoundSource->GetName(), TEXT("")};
+	Metadata.ClassName = FMetasoundFrontendClassName{TEXT(""), *MetasoundSource->GetName(), TEXT("")};
 	Metadata.Version.Major = 1;
 	Metadata.Version.Minor = 0;
 	Metadata.Type = EMetasoundFrontendClassType::Graph;
diff --git a/Engine/Plugins/Runtime/Metasound/Source/MetasoundEngine/Private/MetasoundEngineNodesNames.cpp b/Engine/Plugins/Runtime/Metasound/Source/MetasoundEngine/Private/MetasoundEngineNodesNames.cpp
new file mode 100644
--- /dev/null
+++ b/Engine/Plugins/Runtime/Metasound/Source/MetasoundEngine/Private/MetasoundEngineNodesNames.cpp
@@ -0,0 +1,14 @@
+// Copyright Epic Games, Inc. All Rights Reserved.
+
+#include "MetasoundEngineNodesNames.h"
+
+namespace Metasound
+{
+	namespace EngineNodes
+	{
+		const FName Namespace = TEXT("UE");
+
+		const FName MonoVariant = TEXT("Mono");
+		const FName StereoVariant = TEXT("Stereo");
+	}
+}
diff --git a/Engine/Plugins/Runtime/Metasound/Source/MetasoundEngine/Private/MetasoundUObjectRegistry.cpp b/Engine/Plugins/Runtime/Metasound/Source/MetasoundEngine/Private/MetasoundUObjectRegistry.cpp
--- a/Engine/Plugins/Runtime/Metasound/Source/MetasoundEngine/Private/MetasoundUObjectRegistry.cpp
+++ b/Engine/Plugins/Runtime/Metasound/Source/MetasoundEngine/Private/MetasoundUObjectRegistry.cpp
@@ -115,7 +115,7 @@ namespace Metasound
 					PackageToSaveTo = GetTransientPackage();
 				}
 
-				UObject* NewMetasoundObject = InEntry.NewObject(PackageToSaveTo, *InDocument.RootGraph.Metadata.Name.GetFullName());
+				UObject* NewMetasoundObject = InEntry.NewObject(PackageToSaveTo, *InDocument.RootGraph.Metadata.ClassName.GetFullName().ToString());
 				FMetasoundAssetBase* NewAssetBase = InEntry.Cast(NewMetasoundObject);
 				if (ensure(nullptr != NewAssetBase))
 				{
diff --git a/Engine/Plugins/Runtime/Metasound/Source/MetasoundEngine/Private/MetasoundWavePlayerNode.cpp b/Engine/Plugins/Runtime/Metasound/Source/MetasoundEngine/Private/MetasoundWavePlayerNode.cpp
--- a/Engine/Plugins/Runtime/Metasound/Source/MetasoundEngine/Private/MetasoundWavePlayerNode.cpp
+++ b/Engine/Plugins/Runtime/Metasound/Source/MetasoundEngine/Private/MetasoundWavePlayerNode.cpp
@@ -3,6 +3,7 @@
 
 #include "AudioResampler.h"
 #include "MetasoundBuildError.h"
+#include "MetasoundEngineNodesNames.h"
 #include "MetasoundExecutableOperator.h"
 #include "MetasoundNodeRegistrationMacro.h"
 #include "MetasoundPrimitives.h"
@@ -283,14 +284,15 @@ namespace Metasound
 		);
 	}
 
-	const FNodeInfo& FWavePlayerNode::GetNodeInfo()
+	const FNodeClassMetadata& FWavePlayerNode::GetNodeInfo()
 	{
-		auto InitNodeInfo = []() -> FNodeInfo
+		auto InitNodeInfo = []() -> FNodeClassMetadata
 		{
-			FNodeInfo Info;
-			Info.ClassName = FName(TEXT("Wave Player"));
+			FNodeClassMetadata Info;
+			Info.ClassName = {EngineNodes::Namespace, TEXT("WavePlayer"), EngineNodes::StereoVariant};
 			Info.MajorVersion = 1;
 			Info.MinorVersion = 0;
+			Info.DisplayName = LOCTEXT("Metasound_WavePlayerNodeDisplayName", "Wave Player");
 			Info.Description = LOCTEXT("Metasound_WavePlayerNodeDescription", "Plays a supplied Wave");
 			Info.Author = PluginAuthor;
 			Info.PromptIfMissing = PluginNodeMissingPrompt;
@@ -299,20 +301,20 @@ namespace Metasound
 			return Info;
 		};
 
-		static const FNodeInfo Info = InitNodeInfo();
+		static const FNodeClassMetadata Info = InitNodeInfo();
 
 		return Info;
 	}
 
-	FWavePlayerNode::FWavePlayerNode(const FString& InName)
-		:	FNode(InName, GetNodeInfo())
-		,	Factory(MakeOperatorFactoryRef<FWavePlayerNode::FOperatorFactory>())
-		,	Interface(DeclareVertexInterface())
+	FWavePlayerNode::FWavePlayerNode(const FString& InName, const FGuid& InInstanceID)
+	: FNode(InName, InInstanceID, GetNodeInfo())
+	, Factory(MakeOperatorFactoryRef<FWavePlayerNode::FOperatorFactory>())
+	, Interface(DeclareVertexInterface())
 	{
 	}
 
 	FWavePlayerNode::FWavePlayerNode(const FNodeInitData& InInitData)
-		: FWavePlayerNode(InInitData.InstanceName)
+	: FWavePlayerNode(InInitData.InstanceName, InInitData.InstanceID)
 	{
 	}
 
diff --git a/Engine/Plugins/Runtime/Metasound/Source/MetasoundEngine/Private/MetasoundWaveSelectorNode.cpp b/Engine/Plugins/Runtime/Metasound/Source/MetasoundEngine/Private/MetasoundWaveSelectorNode.cpp
--- a/Engine/Plugins/Runtime/Metasound/Source/MetasoundEngine/Private/MetasoundWaveSelectorNode.cpp
+++ b/Engine/Plugins/Runtime/Metasound/Source/MetasoundEngine/Private/MetasoundWaveSelectorNode.cpp
@@ -2,6 +2,7 @@
 #include "MetasoundWaveSelectorNode.h"
 
 #include "MetasoundBuildError.h"
+#include "MetasoundEngineNodesNames.h"
 #include "MetasoundExecutableOperator.h"
 #include "MetasoundNodeRegistrationMacro.h"
 #include "MetasoundPrimitives.h"
@@ -188,14 +189,15 @@ namespace Metasound
 		);
 	}
 
-	const FNodeInfo& FWaveSelectorNode::GetNodeInfo()
+	const FNodeClassMetadata& FWaveSelectorNode::GetNodeInfo()
 	{
-		auto InitNodeInfo = []() -> FNodeInfo
+		auto InitNodeInfo = []() -> FNodeClassMetadata
 		{
-			FNodeInfo Info;
-			Info.ClassName = FName(TEXT("Wave Selector"));
+			FNodeClassMetadata Info;
+			Info.ClassName = {EngineNodes::Namespace, TEXT("WaveSelector"), TEXT("")};
 			Info.MajorVersion = 1;
 			Info.MinorVersion = 0;
+			Info.DisplayName = LOCTEXT("Metasound_WaveSelectorNodeDisplayName", "Wave Selector");
 			Info.Description = LOCTEXT("Metasound_WaveSelectorNodeDescription", "Allows selection of a Wave Asset from a list");
 			Info.Author = PluginAuthor;
 			Info.PromptIfMissing = PluginNodeMissingPrompt;
@@ -204,20 +206,20 @@ namespace Metasound
 			return Info;
 		};
 
-		static const FNodeInfo Info = InitNodeInfo();
+		static const FNodeClassMetadata Info = InitNodeInfo();
 
 		return Info;
 	}
 
-	FWaveSelectorNode::FWaveSelectorNode(const FString& InName)
-		:	FNode(InName, GetNodeInfo())
+	FWaveSelectorNode::FWaveSelectorNode(const FString& InName, const FGuid& InInstanceID)
+		:	FNode(InName, InInstanceID, GetNodeInfo())
 		,	Factory(MakeOperatorFactoryRef<FWaveSelectorNode::FOperatorFactory>())
 		,	Interface(DeclareVertexInterface())
 	{
 	}
 
 	FWaveSelectorNode::FWaveSelectorNode(const FNodeInitData& InInitData)
-		: FWaveSelectorNode(InInitData.InstanceName)
+		: FWaveSelectorNode(InInitData.InstanceName, InInitData.InstanceID)
 	{
 	}
 
diff --git a/Engine/Plugins/Runtime/Metasound/Source/MetasoundEngine/Public/MetasoundEngineNodesNames.h b/Engine/Plugins/Runtime/Metasound/Source/MetasoundEngine/Public/MetasoundEngineNodesNames.h
new file mode 100644
--- /dev/null
+++ b/Engine/Plugins/Runtime/Metasound/Source/MetasoundEngine/Public/MetasoundEngineNodesNames.h
@@ -0,0 +1,16 @@
+// Copyright Epic Games, Inc. All Rights Reserved.
+
+#pragma once
+
+#include "CoreMinimal.h"
+
+namespace Metasound
+{
+	namespace EngineNodes
+	{
+		extern const FName Namespace;
+
+		extern const FName MonoVariant;
+		extern const FName StereoVariant;
+	}
+}
diff --git a/Engine/Plugins/Runtime/Metasound/Source/MetasoundEngine/Public/MetasoundWavePlayerNode.h b/Engine/Plugins/Runtime/Metasound/Source/MetasoundEngine/Public/MetasoundWavePlayerNode.h
--- a/Engine/Plugins/Runtime/Metasound/Source/MetasoundEngine/Public/MetasoundWavePlayerNode.h
+++ b/Engine/Plugins/Runtime/Metasound/Source/MetasoundEngine/Public/MetasoundWavePlayerNode.h
@@ -21,9 +21,9 @@ namespace Metasound
 
 	public:
 		static FVertexInterface DeclareVertexInterface();
-		static const FNodeInfo& GetNodeInfo();
+		static const FNodeClassMetadata& GetNodeInfo();
 
-		FWavePlayerNode(const FString& InName);
+		FWavePlayerNode(const FString& InName, const FGuid& InInstanceID);
 
 		// constructor used by the Metasound Frontend.
 		FWavePlayerNode(const FNodeInitData& InInitData);
diff --git a/Engine/Plugins/Runtime/Metasound/Source/MetasoundEngine/Public/MetasoundWaveSelectorNode.h b/Engine/Plugins/Runtime/Metasound/Source/MetasoundEngine/Public/MetasoundWaveSelectorNode.h
--- a/Engine/Plugins/Runtime/Metasound/Source/MetasoundEngine/Public/MetasoundWaveSelectorNode.h
+++ b/Engine/Plugins/Runtime/Metasound/Source/MetasoundEngine/Public/MetasoundWaveSelectorNode.h
@@ -21,9 +21,9 @@ namespace Metasound
 
 	public:
 		static FVertexInterface DeclareVertexInterface();
-		static const FNodeInfo& GetNodeInfo();
+		static const FNodeClassMetadata& GetNodeInfo();
 
-		FWaveSelectorNode(const FString& InName);
+		FWaveSelectorNode(const FString& InName, const FGuid& InInstanceID);
 
 		// constructor used by the Metasound Frontend.
 		FWaveSelectorNode(const FNodeInitData& InInitData);
diff --git a/Engine/Plugins/Runtime/Metasound/Source/MetasoundFrontend/Private/MetasoundFrontend.cpp b/Engine/Plugins/Runtime/Metasound/Source/MetasoundFrontend/Private/MetasoundFrontend.cpp
--- a/Engine/Plugins/Runtime/Metasound/Source/MetasoundFrontend/Private/MetasoundFrontend.cpp
+++ b/Engine/Plugins/Runtime/Metasound/Source/MetasoundFrontend/Private/MetasoundFrontend.cpp
@@ -38,7 +38,6 @@ namespace Metasound
 			for (auto& NodeClassTuple : Registry)
 			{
 				FNodeClassInfo ClassInfo;
-				ClassInfo.NodeName = NodeClassTuple.Key.NodeName.ToString();
 				ClassInfo.NodeType = EMetasoundFrontendClassType::External;
 				ClassInfo.LookupKey = NodeClassTuple.Key;
 
@@ -55,32 +54,12 @@ namespace Metasound
 			return GenerateClassDescription(InInfo).Metadata;
 		}
 
-		FMetasoundFrontendClass GenerateClassDescription(const FNodeInfo& InNodeMetadata, EMetasoundFrontendClassType ClassType)
+		FMetasoundFrontendClass GenerateClassDescription(const FNodeClassMetadata& InNodeMetadata, EMetasoundFrontendClassType ClassType)
 		{
-			FMetasoundFrontendClassName ClassName
-			{
-				TEXT(""), // Namespace
-				InNodeMetadata.ClassName.ToString(),
-				TEXT("") // Variant
-			};
-
 			FMetasoundFrontendClass ClassDescription;
 
-			ClassDescription.Metadata = FMetasoundFrontendClassMetadata
-			{
-				ClassName,
-				FMetasoundFrontendVersionNumber
-				{
-					InNodeMetadata.MajorVersion,
-					InNodeMetadata.MinorVersion
-				},
-				ClassType,
-				InNodeMetadata.Description,
-				InNodeMetadata.PromptIfMissing,
-				InNodeMetadata.Author,
-				InNodeMetadata.Keywords,
-				InNodeMetadata.CategoryHierarchy,
-			};
+			ClassDescription.Metadata = FMetasoundFrontendClassMetadata(InNodeMetadata);
+			ClassDescription.Metadata.Type = ClassType;
 
 			FMetasoundFrontendClassStyleDisplay DisplayStyle(InNodeMetadata.DisplayStyle);
 			ClassDescription.Style = FMetasoundFrontendClassStyle
@@ -143,7 +122,7 @@ namespace Metasound
 			}
 			else
 			{
-				ensureAlwaysMsgf(false, TEXT("Tried to get Class Description for unknown node %s!"), *InInfo.NodeName);
+				ensureAlwaysMsgf(false, TEXT("Cannot generate description of unregistered node [NodeClassName:%s]"), *InInfo.LookupKey.NodeClassFullName.ToString());
 				return FMetasoundFrontendClass();
 			}
 		}
diff --git a/Engine/Plugins/Runtime/Metasound/Source/MetasoundFrontend/Private/MetasoundFrontendDocument.cpp b/Engine/Plugins/Runtime/Metasound/Source/MetasoundFrontend/Private/MetasoundFrontendDocument.cpp
--- a/Engine/Plugins/Runtime/Metasound/Source/MetasoundFrontend/Private/MetasoundFrontendDocument.cpp
+++ b/Engine/Plugins/Runtime/Metasound/Source/MetasoundFrontend/Private/MetasoundFrontendDocument.cpp
@@ -5,6 +5,11 @@
 #include "MetasoundFrontend.h"
 #include "MetasoundFrontendRegistries.h"
 
+namespace Metasound
+{
+	const FGuid FrontendInvalidID = FGuid();
+}
+
 void FMetasoundFrontendLiteral::Set(bool InValue)
 {
 	Clear();
@@ -83,7 +88,7 @@ FMetasoundFrontendNodeInterface::FMetasoundFrontendNodeInterface(const FMetasoun
 
 FMetasoundFrontendNode::FMetasoundFrontendNode(const FMetasoundFrontendClass& InClass)
 : ClassID(InClass.ID)
-, Name(InClass.Metadata.Name.Name)
+, Name(InClass.Metadata.ClassName.Name.ToString())
 , Interface(InClass.Interface)
 {
 
@@ -104,10 +109,31 @@ bool FMetasoundFrontendClassVertex::IsFunctionalEquivalent(const FMetasoundFront
 	return FMetasoundFrontendVertex::IsFunctionalEquivalent(InLHS, InRHS) && FMetasoundFrontendVertexBehavior::IsFunctionalEquivalent(InLHS.Behavior, InRHS.Behavior);
 }
 
-FString FMetasoundFrontendClassName::GetFullName() const
+FMetasoundFrontendClassName::FMetasoundFrontendClassName(const FName& InNamespace, const FName& InName, const FName& InVariant)
+: Namespace(InNamespace)
+, Name(InName)
+, Variant(InVariant)
+{
+}
+
+FMetasoundFrontendClassName::FMetasoundFrontendClassName(const Metasound::FNodeClassName& InName)
+: FMetasoundFrontendClassName(InName.GetNamespace(), InName.GetName(), InName.GetVariant())
+{
+}
+
+FName FMetasoundFrontendClassName::GetScopedName() const
 {
-	//return FString::Format(TEXT("{0}.{1}.{2}"), {Namespace, Name, Variant});
-	return Name;
+	return Metasound::FNodeClassName::FormatScopedName(Namespace, Name);
+}
+
+FName FMetasoundFrontendClassName::GetFullName() const
+{
+	return Metasound::FNodeClassName::FormatFullName(Namespace, Name, Variant);
+}
+
+FString FMetasoundFrontendClassName::ToString() const
+{
+	return GetFullName().ToString();
 }
 
 bool operator==(const FMetasoundFrontendClassName& InLHS, const FMetasoundFrontendClassName& InRHS)
@@ -115,6 +141,19 @@ bool operator==(const FMetasoundFrontendClassName& InLHS, const FMetasoundFronte
 	return (InLHS.Namespace == InRHS.Namespace) && (InLHS.Name == InRHS.Name) && (InLHS.Variant == InRHS.Variant);
 }
 
+FMetasoundFrontendClassMetadata::FMetasoundFrontendClassMetadata(const Metasound::FNodeClassMetadata& InNodeClassMetadata)
+: ClassName(InNodeClassMetadata.ClassName)
+, Version{InNodeClassMetadata.MajorVersion, InNodeClassMetadata.MinorVersion}
+, Type(EMetasoundFrontendClassType::External)
+, DisplayName(InNodeClassMetadata.DisplayName)
+, Description(InNodeClassMetadata.Description)
+, PromptIfMissing(InNodeClassMetadata.PromptIfMissing)
+, Author(InNodeClassMetadata.Author)
+, Keywords(InNodeClassMetadata.Keywords)
+, CategoryHierarchy(InNodeClassMetadata.CategoryHierarchy)
+{
+}
+
 FMetasoundFrontendClassInput::FMetasoundFrontendClassInput(const FMetasoundFrontendClassVertex& InOther)
 :	FMetasoundFrontendClassVertex(InOther)
 {
diff --git a/Engine/Plugins/Runtime/Metasound/Source/MetasoundFrontend/Private/MetasoundFrontendGraph.cpp b/Engine/Plugins/Runtime/Metasound/Source/MetasoundFrontend/Private/MetasoundFrontendGraph.cpp
--- a/Engine/Plugins/Runtime/Metasound/Source/MetasoundFrontend/Private/MetasoundFrontendGraph.cpp
+++ b/Engine/Plugins/Runtime/Metasound/Source/MetasoundFrontend/Private/MetasoundFrontendGraph.cpp
@@ -13,8 +13,8 @@
 
 namespace Metasound
 {
-	FFrontendGraph::FFrontendGraph(const FString& InInstanceName)
-	:	FGraph(InInstanceName)
+	FFrontendGraph::FFrontendGraph(const FString& InInstanceName, const FGuid& InInstanceID)
+	:	FGraph(InInstanceName, InInstanceID)
 	{
 	}
 
@@ -162,6 +162,7 @@ namespace Metasound
 					FInputNodeConstructorParams InitParams =
 					{
 						InNode.Name,
+						InNode.ID,
 						InputVertex.Name,
 						MoveTemp(Literal)
 					};
@@ -195,6 +196,7 @@ namespace Metasound
 			FOutputNodeConstructorParams InitParams =
 			{
 				InNode.Name,
+				InNode.ID,
 				OutputVertex.Name,
 			};
 
@@ -215,6 +217,9 @@ namespace Metasound
 		InitData.InstanceName.AppendChar('_');
 		InitData.InstanceName.Append(InNode.ID.ToString());
 
+		InitData.InstanceID = InNode.ID;
+
+
 		// Copy over our initialization params.
 		/*
 		for (auto& StaticParamTuple : InNode.StaticParameters)
@@ -233,7 +238,7 @@ namespace Metasound
 
 		Metasound::Frontend::FNodeRegistryKey Key = FMetasoundFrontendRegistryContainer::GetRegistryKey(InClass.Metadata);
 
-		return FMetasoundFrontendRegistryContainer::Get()->ConstructExternalNode(Key.NodeName, Key.NodeHash, InitData);
+		return FMetasoundFrontendRegistryContainer::Get()->ConstructExternalNode(Key.NodeClassFullName, Key.NodeHash, InitData);
 	}
 
 
@@ -604,7 +609,7 @@ namespace Metasound
 			ClassMap.Add(ExtClass.ID, &ExtClass);
 		}
 
-		TUniquePtr<FFrontendGraph> MetasoundGraph = MakeUnique<FFrontendGraph>(InGraph.Metadata.Name.Name);
+		TUniquePtr<FFrontendGraph> MetasoundGraph = MakeUnique<FFrontendGraph>(InGraph.Metadata.ClassName.GetFullName().ToString(), FGuid::NewGuid());
 
 		// TODO: will likely want to bubble up errors here for case where
 		// a datatype or node is not registered. 
diff --git a/Engine/Plugins/Runtime/Metasound/Source/MetasoundFrontend/Private/MetasoundFrontendGraph.h b/Engine/Plugins/Runtime/Metasound/Source/MetasoundFrontend/Private/MetasoundFrontendGraph.h
--- a/Engine/Plugins/Runtime/Metasound/Source/MetasoundFrontend/Private/MetasoundFrontendGraph.h
+++ b/Engine/Plugins/Runtime/Metasound/Source/MetasoundFrontend/Private/MetasoundFrontendGraph.h
@@ -21,8 +21,9 @@ namespace Metasound
 			/** FFrontendGraph constructor.
 			 *
 			 * @parma InInstanceName - Name of this graph.
+			 * @parma InInstanceID - ID of this graph.
 			 */
-			FFrontendGraph(const FString& InInstanceName);
+			FFrontendGraph(const FString& InInstanceName, const FGuid& InInstanceID);
 
 			virtual ~FFrontendGraph() = default;
 
diff --git a/Engine/Plugins/Runtime/Metasound/Source/MetasoundFrontend/Private/MetasoundFrontendInvalidController.h b/Engine/Plugins/Runtime/Metasound/Source/MetasoundFrontend/Private/MetasoundFrontendInvalidController.h
--- a/Engine/Plugins/Runtime/Metasound/Source/MetasoundFrontend/Private/MetasoundFrontendInvalidController.h
+++ b/Engine/Plugins/Runtime/Metasound/Source/MetasoundFrontend/Private/MetasoundFrontendInvalidController.h
@@ -154,8 +154,9 @@ namespace Metasound
 			bool RemoveOutput(FGuid InPointID) override { return false; }
 
 			EMetasoundFrontendClassType GetClassType() const override { return EMetasoundFrontendClassType::Invalid; }
-			const FString& GetClassName() const override { return MetasoundFrontendInvalidControllerPrivate::GetInvalid<FString>(); }
+			const FMetasoundFrontendClassName& GetClassName() const override { return MetasoundFrontendInvalidControllerPrivate::GetInvalid<FMetasoundFrontendClassName>(); }
 			FMetasoundFrontendVersionNumber GetClassVersionNumber() const override { return MetasoundFrontendInvalidControllerPrivate::GetInvalid<FMetasoundFrontendVersionNumber>(); }
+			const FText& GetClassDisplayName() const override { return MetasoundFrontendInvalidControllerPrivate::GetInvalid<FText>(); }
 			const FText& GetClassDescription() const override { return MetasoundFrontendInvalidControllerPrivate::GetInvalid<FText>(); }
 			const FMetasoundFrontendClassStyle& GetClassStyle() const override { static const FMetasoundFrontendClassStyle Invalid; return Invalid; }
 
diff --git a/Engine/Plugins/Runtime/Metasound/Source/MetasoundFrontend/Private/MetasoundFrontendQuerySteps.cpp b/Engine/Plugins/Runtime/Metasound/Source/MetasoundFrontend/Private/MetasoundFrontendQuerySteps.cpp
--- a/Engine/Plugins/Runtime/Metasound/Source/MetasoundFrontend/Private/MetasoundFrontendQuerySteps.cpp
+++ b/Engine/Plugins/Runtime/Metasound/Source/MetasoundFrontend/Private/MetasoundFrontendQuerySteps.cpp
@@ -50,14 +50,14 @@ namespace Metasound
 		);
 	}
 
-	FFilterClassesByClassName::FFilterClassesByClassName(const FString& InClassName)
+	FFilterClassesByClassName::FFilterClassesByClassName(const FMetasoundFrontendClassName& InClassName)
 	: ClassName(InClassName)
 	{
 	}
 
 	bool FFilterClassesByClassName::Filter(const FFrontendQueryEntry& InEntry) const 
 	{
-		return InEntry.Value.Get<FMetasoundFrontendClass>().Metadata.Name.Name == ClassName;
+		return InEntry.Value.Get<FMetasoundFrontendClass>().Metadata.ClassName == ClassName;
 	}
 }
 
diff --git a/Engine/Plugins/Runtime/Metasound/Source/MetasoundFrontend/Private/MetasoundFrontendRegistries.cpp b/Engine/Plugins/Runtime/Metasound/Source/MetasoundFrontend/Private/MetasoundFrontendRegistries.cpp
--- a/Engine/Plugins/Runtime/Metasound/Source/MetasoundFrontend/Private/MetasoundFrontendRegistries.cpp
+++ b/Engine/Plugins/Runtime/Metasound/Source/MetasoundFrontend/Private/MetasoundFrontendRegistries.cpp
@@ -18,10 +18,10 @@ namespace MetasoundFrontendRegistryPrivate
 	{
 		Metasound::Frontend::FNodeRegistryKey Key;
 
-		Key.NodeName = InClassName;
+		Key.NodeClassFullName = InClassName;
 
 		// NodeHash is hash of node name and major version.
-		Key.NodeHash = FCrc::StrCrc32(*Key.NodeName.ToString());
+		Key.NodeHash = FCrc::StrCrc32(*Key.NodeClassFullName.ToString());
 		Key.NodeHash = HashCombine(Key.NodeHash, FCrc::TypeCrc32(InMajorVersion));
 
 		return Key;
@@ -171,10 +171,10 @@ Metasound::FLiteral FMetasoundFrontendRegistryContainer::GenerateLiteralForUObje
 	}
 }
 
-TUniquePtr<Metasound::INode> FMetasoundFrontendRegistryContainer::ConstructExternalNode(const FName& InNodeType, uint32 InNodeHash, const Metasound::FNodeInitData& InInitData)
+TUniquePtr<Metasound::INode> FMetasoundFrontendRegistryContainer::ConstructExternalNode(const FName& InNodeClassFullName, uint32 InNodeHash, const Metasound::FNodeInitData& InInitData)
 {
 	Metasound::Frontend::FNodeRegistryKey RegistryKey;
-	RegistryKey.NodeName = InNodeType;
+	RegistryKey.NodeClassFullName = InNodeClassFullName;
 	RegistryKey.NodeHash = InNodeHash;
 
 	if (!ExternalNodeRegistry.Contains(RegistryKey))
@@ -342,7 +342,7 @@ bool FMetasoundFrontendRegistryContainer::RegisterExternalNode(Metasound::FCreat
 	if (ensure(FMetasoundFrontendRegistryContainer::GetRegistryKey(RegistryElement, Key)))
 	{
 		// check to see if an identical node was already registered, and log
-		ensureAlwaysMsgf(!ExternalNodeRegistry.Contains(Key), TEXT("Node with identical name, inputs and outputs to node %s was already registered. The previously registered node will be overwritten. This could also happen because METASOUND_REGISTER_NODE is in a public header."), *Key.NodeName.ToString());
+		ensureAlwaysMsgf(!ExternalNodeRegistry.Contains(Key), TEXT("Node with identical name, inputs and outputs to node %s was already registered. The previously registered node will be overwritten. This could also happen because METASOUND_REGISTER_NODE is in a public header."), *Key.NodeClassFullName.ToString());
 
 		ExternalNodeRegistry.Add(MoveTemp(Key), MoveTemp(RegistryElement));
 
@@ -366,14 +366,14 @@ bool FMetasoundFrontendRegistryContainer::GetRegistryKey(const Metasound::Fronte
 	return false;
 }
 
-Metasound::Frontend::FNodeRegistryKey FMetasoundFrontendRegistryContainer::GetRegistryKey(const FNodeInfo& InNodeMetadata)
+Metasound::Frontend::FNodeRegistryKey FMetasoundFrontendRegistryContainer::GetRegistryKey(const FNodeClassMetadata& InNodeMetadata)
 {
-	return MetasoundFrontendRegistryPrivate::GetRegistryKey(InNodeMetadata.ClassName, InNodeMetadata.MajorVersion);
+	return MetasoundFrontendRegistryPrivate::GetRegistryKey(InNodeMetadata.ClassName.GetFullName(), InNodeMetadata.MajorVersion);
 }
 
 Metasound::Frontend::FNodeRegistryKey FMetasoundFrontendRegistryContainer::GetRegistryKey(const FMetasoundFrontendClassMetadata& InNodeMetadata)
 {
-	return MetasoundFrontendRegistryPrivate::GetRegistryKey(FName(InNodeMetadata.Name.Name), InNodeMetadata.Version.Major);
+	return MetasoundFrontendRegistryPrivate::GetRegistryKey(InNodeMetadata.ClassName.GetFullName(), InNodeMetadata.Version.Major);
 }
 
 bool FMetasoundFrontendRegistryContainer::GetFrontendClassFromRegistered(const FMetasoundFrontendClassMetadata& InMetadata, FMetasoundFrontendClass& OutClass)
diff --git a/Engine/Plugins/Runtime/Metasound/Source/MetasoundFrontend/Private/MetasoundFrontendStandardController.cpp b/Engine/Plugins/Runtime/Metasound/Source/MetasoundFrontend/Private/MetasoundFrontendStandardController.cpp
--- a/Engine/Plugins/Runtime/Metasound/Source/MetasoundFrontend/Private/MetasoundFrontendStandardController.cpp
+++ b/Engine/Plugins/Runtime/Metasound/Source/MetasoundFrontend/Private/MetasoundFrontendStandardController.cpp
@@ -773,13 +773,13 @@ namespace Metasound
 
 			if (ConverterInputs.Num() < 1)
 			{
-				UE_LOG(LogMetasound, Warning, TEXT("Converter node [Name: %s] does not support preferred input vertex [Vertex: %s]"), *InConverterInfo.NodeKey.NodeName.ToString(), *InConverterInfo.PreferredConverterInputPin);
+				UE_LOG(LogMetasound, Warning, TEXT("Converter node [Name: %s] does not support preferred input vertex [Vertex: %s]"), *InConverterInfo.NodeKey.NodeClassFullName.ToString(), *InConverterInfo.PreferredConverterInputPin);
 				return false;
 			}
 
 			if (ConverterOutputs.Num() < 1)
 			{
-				UE_LOG(LogMetasound, Warning, TEXT("Converter node [Name: %s] does not support preferred output vertex [Vertex: %s]"), *InConverterInfo.NodeKey.NodeName.ToString(), *InConverterInfo.PreferredConverterOutputPin);
+				UE_LOG(LogMetasound, Warning, TEXT("Converter node [Name: %s] does not support preferred output vertex [Vertex: %s]"), *InConverterInfo.NodeKey.NodeClassFullName.ToString(), *InConverterInfo.PreferredConverterOutputPin);
 				return false;
 			}
 
@@ -1054,6 +1054,15 @@ namespace Metasound
 			return Metasound::FrontendInvalidID;
 		}
 
+		const FText& FBaseNodeController::GetClassDisplayName() const
+		{
+			if (const FMetasoundFrontendClass* NodeClass = ClassPtr.Get())
+			{
+				return NodeClass->Metadata.DisplayName;
+			}
+			return FText::GetEmpty();
+		}
+
 		FMetasoundFrontendVersionNumber FBaseNodeController::GetClassVersionNumber() const
 		{
 			if (ClassPtr.IsValid())
@@ -1120,14 +1129,14 @@ namespace Metasound
 			return EMetasoundFrontendClassType::Invalid;
 		}
 
-		const FString& FBaseNodeController::GetClassName() const
+		const FMetasoundFrontendClassName& FBaseNodeController::GetClassName() const
 		{
 			if (ClassPtr.IsValid())
 			{
-				return ClassPtr->Metadata.Name.Name;
+				return ClassPtr->Metadata.ClassName;
 			}
 
-			return FrontendControllerIntrinsics::GetInvalidValueConstRef<FString>();
+			return FrontendControllerIntrinsics::GetInvalidValueConstRef<FMetasoundFrontendClassName>();
 		}
 
 		bool FBaseNodeController::CanAddInput(const FString& InVertexName) const
@@ -2393,7 +2402,6 @@ namespace Metasound
 			// Construct a FNodeClassInfo from this lookup key.
 			FNodeClassInfo ClassInfo;
 			ClassInfo.LookupKey = InNodeClass;
-			ClassInfo.NodeName = InNodeClass.NodeName.ToString();
 			ClassInfo.NodeType = EMetasoundFrontendClassType::External;
 
 			return AddNode(ClassInfo);
@@ -2459,7 +2467,7 @@ namespace Metasound
 				{
 					if (const FMetasoundFrontendClass* ExistingDependency = OwningDocument->FindClass(InInfo).Get())
 					{
-						UE_LOG(LogMetasound, Error, TEXT("Cannot add new subgraph. Metasound class already exists with matching metadata Name: \"%s\", Version %d.%d"), *(ExistingDependency->Metadata.Name.GetFullName()), ExistingDependency->Metadata.Version.Major, ExistingDependency->Metadata.Version.Minor);
+						UE_LOG(LogMetasound, Error, TEXT("Cannot add new subgraph. Metasound class already exists with matching metadata Name: \"%s\", Version %d.%d"), *(ExistingDependency->Metadata.ClassName.GetFullName().ToString()), ExistingDependency->Metadata.Version.Major, ExistingDependency->Metadata.Version.Minor);
 					}
 					//else if (const FMetasoundFrontendClass* DependencyDescription = OwningDocument->FindOrAddClass(InInfo))
 					else if (FConstClassAccessPtr DependencyDescription = OwningDocument->FindOrAddClass(InInfo))
@@ -3171,7 +3179,7 @@ namespace Metasound
 					}
 					else
 					{
-						UE_LOG(LogMetasound, Error, TEXT("Cannot add external dependency. No Metasound class found with matching metadata Name: \"%s\", Version %d.%d. Suggested solution \"%s\" by %s."), *InMetadata.Name.GetFullName(), InMetadata.Version.Major, InMetadata.Version.Minor, *InMetadata.PromptIfMissing.ToString(), *InMetadata.Author.ToString());
+						UE_LOG(LogMetasound, Error, TEXT("Cannot add external dependency. No Metasound class found with matching metadata Name: \"%s\", Version %d.%d. Suggested solution \"%s\" by %s."), *InMetadata.ClassName.GetFullName().ToString(), InMetadata.Version.Major, InMetadata.Version.Minor, *InMetadata.PromptIfMissing.ToString(), *InMetadata.Author.ToString());
 					}
 				} 
 				else if (EMetasoundFrontendClassType::Graph == InMetadata.Type)
@@ -3184,7 +3192,7 @@ namespace Metasound
 				}
 				else
 				{
-					UE_LOG(LogMetasound, Error, TEXT("Unsupported metasound class type for node: \"%s\", Version %d.%d."), *InMetadata.Name.GetFullName(), InMetadata.Version.Major, InMetadata.Version.Minor);
+					UE_LOG(LogMetasound, Error, TEXT("Unsupported metasound class type for node: \"%s\", Version %d.%d."), *InMetadata.ClassName.GetFullName().ToString(), InMetadata.Version.Major, InMetadata.Version.Minor);
 					checkNoEntry();
 				}
 
@@ -3345,7 +3353,7 @@ namespace Metasound
 		{
 			if (InMetadataA.Type == InMetadataB.Type)
 			{
-				if (InMetadataA.Name == InMetadataB.Name)
+				if (InMetadataA.ClassName == InMetadataB.ClassName)
 				{
 					return FRegistry::GetRegistryKey(InMetadataA) == FRegistry::GetRegistryKey(InMetadataB);
 				}
@@ -3357,10 +3365,7 @@ namespace Metasound
 		{
 			if (InNodeClass.NodeType == InMetadata.Type)
 			{
-				if (InNodeClass.NodeName == InMetadata.Name.Name)
-				{
-					return InNodeClass.LookupKey == FRegistry::GetRegistryKey(InMetadata);
-				}
+				return InNodeClass.LookupKey == FRegistry::GetRegistryKey(InMetadata);
 			}
 			return false;
 		}
diff --git a/Engine/Plugins/Runtime/Metasound/Source/MetasoundFrontend/Private/MetasoundFrontendStandardController.h b/Engine/Plugins/Runtime/Metasound/Source/MetasoundFrontend/Private/MetasoundFrontendStandardController.h
--- a/Engine/Plugins/Runtime/Metasound/Source/MetasoundFrontend/Private/MetasoundFrontendStandardController.h
+++ b/Engine/Plugins/Runtime/Metasound/Source/MetasoundFrontend/Private/MetasoundFrontendStandardController.h
@@ -310,6 +310,8 @@ namespace Metasound
 			// Info about this node.
 			FGuid GetID() const override;
 			FGuid GetClassID() const override;
+			const FText& GetClassDisplayName() const override;
+
 			FMetasoundFrontendVersionNumber GetClassVersionNumber() const override;
 			const FText& GetClassDescription() const override;
 			const FMetasoundFrontendClassStyle& GetClassStyle() const override;
@@ -328,7 +330,7 @@ namespace Metasound
 			FInputHandle AddOutput(const FString& InVertexName, const FMetasoundFrontendLiteral* InDefault) override;
 			bool RemoveOutput(FGuid InPointID) override;
 
-			const FString& GetClassName() const override;
+			const FMetasoundFrontendClassName& GetClassName() const override;
 
 			/** Returns all node inputs. */
 			TArray<FInputHandle> GetInputs() override;
diff --git a/Engine/Plugins/Runtime/Metasound/Source/MetasoundFrontend/Public/MetasoundAutoConverterNode.h b/Engine/Plugins/Runtime/Metasound/Source/MetasoundFrontend/Public/MetasoundAutoConverterNode.h
--- a/Engine/Plugins/Runtime/Metasound/Source/MetasoundFrontend/Public/MetasoundAutoConverterNode.h
+++ b/Engine/Plugins/Runtime/Metasound/Source/MetasoundFrontend/Public/MetasoundAutoConverterNode.h
@@ -49,9 +49,9 @@ namespace Metasound
 			);
 		}
 
-		static const FNodeInfo& GetAutoConverterNodeMetadata()
+		static const FNodeClassMetadata& GetAutoConverterNodeMetadata()
 		{
-			auto InitNodeInfo = []() -> FNodeInfo
+			auto InitNodeInfo = []() -> FNodeClassMetadata
 			{
 				const FString& InputDisplayName = GetInputName();
 				const FString& OutputDisplayName = GetOutputName();
@@ -61,10 +61,11 @@ namespace Metasound
 				DisplayStyle.bShowInputNames = false;
 				DisplayStyle.bShowOutputNames = false;
 
-				FNodeInfo Info;
-				Info.ClassName = FName(*FString::Printf(TEXT("%s To %s"), *InputDisplayName, *OutputDisplayName));
+				FNodeClassMetadata Info;
+				Info.ClassName = {TEXT("Convert"), GetMetasoundDataTypeName<ToDataType>(), GetMetasoundDataTypeName<FromDataType>()}; 
 				Info.MajorVersion = 1;
 				Info.MinorVersion = 0;
+				Info.DisplayName = FText::Format(LOCTEXT("Metasound_AutoConverterNodeDisplayNameFormat", "Convert {0} to {1}"), FText::FromName(GetMetasoundDataTypeName<FromDataType>()), FText::FromName(GetMetasoundDataTypeName<ToDataType>()));
 				Info.Description = LOCTEXT("Metasound_ConverterNodeDescription", "Converts between two different data types.");
 				Info.Author = PluginAuthor;
 				Info.DisplayStyle = DisplayStyle;
@@ -77,7 +78,7 @@ namespace Metasound
 				return Info;
 			};
 
-			static const FNodeInfo Info = InitNodeInfo();
+			static const FNodeClassMetadata Info = InitNodeInfo();
 
 			return Info;
 		}
@@ -157,7 +158,7 @@ namespace Metasound
 		public:
 
 			TAutoConverterNode(const FNodeInitData& InInitData)
-				: FNode(InInitData.InstanceName, GetAutoConverterNodeMetadata())
+				: FNode(InInitData.InstanceName, InInitData.InstanceID, GetAutoConverterNodeMetadata())
 				, Interface(DeclareVertexInterface())
 				, Factory(MakeOperatorFactoryRef<FCoverterOperatorFactory>())
 			{
diff --git a/Engine/Plugins/Runtime/Metasound/Source/MetasoundFrontend/Public/MetasoundConverterNodeRegistrationMacro.h b/Engine/Plugins/Runtime/Metasound/Source/MetasoundFrontend/Public/MetasoundConverterNodeRegistrationMacro.h
--- a/Engine/Plugins/Runtime/Metasound/Source/MetasoundFrontend/Public/MetasoundConverterNodeRegistrationMacro.h
+++ b/Engine/Plugins/Runtime/Metasound/Source/MetasoundFrontend/Public/MetasoundConverterNodeRegistrationMacro.h
@@ -12,7 +12,7 @@
 #include "MetasoundFrontendRegistries.h"
 
 template<typename TNodeType, typename TFromType, typename TToType>
-bool RegisterConversionNode(const FString& FromPin, const FString& ToPin, const Metasound::FNodeInfo& InNodeMetadata)
+bool RegisterConversionNode(const FString& FromPin, const FString& ToPin, const Metasound::FNodeClassMetadata& InNodeMetadata)
 {
 	using FConverterNodeRegistryKey = ::Metasound::Frontend::FConverterNodeRegistryKey;
 	using FConverterNodeInfo = ::Metasound::Frontend::FConverterNodeInfo;
diff --git a/Engine/Plugins/Runtime/Metasound/Source/MetasoundFrontend/Public/MetasoundDataTypeRegistrationMacro.h b/Engine/Plugins/Runtime/Metasound/Source/MetasoundFrontend/Public/MetasoundDataTypeRegistrationMacro.h
--- a/Engine/Plugins/Runtime/Metasound/Source/MetasoundFrontend/Public/MetasoundDataTypeRegistrationMacro.h
+++ b/Engine/Plugins/Runtime/Metasound/Source/MetasoundFrontend/Public/MetasoundDataTypeRegistrationMacro.h
@@ -89,7 +89,7 @@ namespace Metasound
 	{
 		using FConverterNode = Metasound::TAutoConverterNode<TFromDataType, TToDataType>;
 
-		const FNodeInfo& Metadata = FConverterNode::GetAutoConverterNodeMetadata();
+		const FNodeClassMetadata& Metadata = FConverterNode::GetAutoConverterNodeMetadata();
 		const Metasound::Frontend::FNodeRegistryKey Key = FMetasoundFrontendRegistryContainer::GetRegistryKey(Metadata);
 
 		if (!std::is_same<TFromDataType, TToDataType>::value && !FMetasoundFrontendRegistryContainer::Get()->IsNodeRegistered(Key))
@@ -141,27 +141,27 @@ namespace Metasound
 		// Lambdas that generate our template-instantiated input and output nodes:
 		FCreateInputNodeFunction InputNodeConstructor = [](FInputNodeConstructorParams&& InParams) -> TUniquePtr<INode>
 		{
-			return TUniquePtr<INode>(new TInputNode<TDataType>(InParams.InNodeName, InParams.InVertexName, MoveTemp(InParams.InitParam)));
+			return TUniquePtr<INode>(new TInputNode<TDataType>(InParams.InNodeName, InParams.InInstanceID, InParams.InVertexName, MoveTemp(InParams.InitParam)));
 		};
 
 		FCreateMetasoundFrontendClassFunction CreateFrontendInputClass = []() -> FMetasoundFrontendClass
 		{
 			// Create class info using prototype node
 			// TODO: register input nodes with static class info.
-  			static TInputNode<TDataType> Prototype(TEXT(""), TEXT(""), FLiteral());
+  			static TInputNode<TDataType> Prototype(TEXT(""), FGuid(), TEXT(""), FLiteral());
 			return Metasound::Frontend::GenerateClassDescription(Prototype.GetMetadata(), EMetasoundFrontendClassType::Input);
 		};
 
 		FCreateOutputNodeFunction OutputNodeConstructor = [](const FOutputNodeConstructorParams& InParams) -> TUniquePtr<INode>
 		{
-			return TUniquePtr<INode>(new TOutputNode<TDataType>(InParams.InNodeName, InParams.InVertexName));
+			return TUniquePtr<INode>(new TOutputNode<TDataType>(InParams.InNodeName, InParams.InInstanceID, InParams.InVertexName));
 		};
 
 		FCreateMetasoundFrontendClassFunction CreateFrontendOutputClass = []() -> FMetasoundFrontendClass
 		{
 			// Create class info using prototype node
 			// TODO: register input nodes with static class info.
-			static TOutputNode<TDataType> Prototype(TEXT(""), TEXT(""));
+			static TOutputNode<TDataType> Prototype(TEXT(""), FGuid(), TEXT(""));
 			return Metasound::Frontend::GenerateClassDescription(Prototype.GetMetadata(), EMetasoundFrontendClassType::Output);
 		};
 		
diff --git a/Engine/Plugins/Runtime/Metasound/Source/MetasoundFrontend/Public/MetasoundFrontend.h b/Engine/Plugins/Runtime/Metasound/Source/MetasoundFrontend/Public/MetasoundFrontend.h
--- a/Engine/Plugins/Runtime/Metasound/Source/MetasoundFrontend/Public/MetasoundFrontend.h
+++ b/Engine/Plugins/Runtime/Metasound/Source/MetasoundFrontend/Public/MetasoundFrontend.h
@@ -22,9 +22,6 @@ namespace Metasound
 		// and also used in FGraphHandle::AddNewNode.
 		struct METASOUNDFRONTEND_API FNodeClassInfo
 		{
-			// The descriptive name of this node class.
-			FString NodeName;
-
 			// The type for this node.
 			EMetasoundFrontendClassType NodeType = EMetasoundFrontendClassType::Invalid;
 
@@ -49,7 +46,7 @@ namespace Metasound
 		 *
 		 * @return Class description for external node.
 		 */
-		METASOUNDFRONTEND_API FMetasoundFrontendClass GenerateClassDescription(const FNodeInfo& InNodeMetadata, EMetasoundFrontendClassType ClassType=EMetasoundFrontendClassType::External);
+		METASOUNDFRONTEND_API FMetasoundFrontendClass GenerateClassDescription(const FNodeClassMetadata& InNodeMetadata, EMetasoundFrontendClassType ClassType=EMetasoundFrontendClassType::External);
 
 		/** Generates a new FMetasoundFrontendClass from node lookup info.
 		 *
diff --git a/Engine/Plugins/Runtime/Metasound/Source/MetasoundFrontend/Public/MetasoundFrontendBaseClasses.h b/Engine/Plugins/Runtime/Metasound/Source/MetasoundFrontend/Public/MetasoundFrontendBaseClasses.h
--- a/Engine/Plugins/Runtime/Metasound/Source/MetasoundFrontend/Public/MetasoundFrontendBaseClasses.h
+++ b/Engine/Plugins/Runtime/Metasound/Source/MetasoundFrontend/Public/MetasoundFrontendBaseClasses.h
@@ -26,7 +26,7 @@ namespace Metasound
 		// Convenience functions to create an INodeB corresponding to a specific externally declared node type.
 		// InNodeType and InNodeHash can be retrieved from the FNodeClassInfo generated from the node registry queries in the metasound frontend (GetAllAvailableNodeClasses, GetAllNodeClassesInNamespace, etc.)
 		// @returns nullptr if the type given wasn't found.
-		METASOUNDFRONTEND_API TUniquePtr<INode> ConstructExternalNode(const FName& InNodeType, uint32 InNodeHash, const FNodeInitData& InInitData);
+		METASOUNDFRONTEND_API TUniquePtr<INode> ConstructExternalNode(const FName& InNodeClassFullName, uint32 InNodeHash, const FNodeInitData& InInitData);
 
 
 		// Utility functions for building a ::Metasound::FDataInitParam corresponding to a literal.
diff --git a/Engine/Plugins/Runtime/Metasound/Source/MetasoundFrontend/Public/MetasoundFrontendController.h b/Engine/Plugins/Runtime/Metasound/Source/MetasoundFrontend/Public/MetasoundFrontendController.h
--- a/Engine/Plugins/Runtime/Metasound/Source/MetasoundFrontend/Public/MetasoundFrontendController.h
+++ b/Engine/Plugins/Runtime/Metasound/Source/MetasoundFrontend/Public/MetasoundFrontendController.h
@@ -362,8 +362,9 @@ namespace Metasound
 			virtual EMetasoundFrontendClassType GetClassType() const = 0;
 
 			/** Returns the name of the node class. */
-			virtual const FString& GetClassName() const = 0;
+			virtual const FMetasoundFrontendClassName& GetClassName() const = 0;
 			virtual FMetasoundFrontendVersionNumber GetClassVersionNumber() const = 0;
+			virtual const FText& GetClassDisplayName() const = 0;
 			virtual const FText& GetClassDescription() const = 0;
 			virtual const FMetasoundFrontendClassStyle& GetClassStyle() const = 0;
 
diff --git a/Engine/Plugins/Runtime/Metasound/Source/MetasoundFrontend/Public/MetasoundFrontendDocument.h b/Engine/Plugins/Runtime/Metasound/Source/MetasoundFrontend/Public/MetasoundFrontendDocument.h
--- a/Engine/Plugins/Runtime/Metasound/Source/MetasoundFrontend/Public/MetasoundFrontendDocument.h
+++ b/Engine/Plugins/Runtime/Metasound/Source/MetasoundFrontend/Public/MetasoundFrontendDocument.h
@@ -12,7 +12,7 @@
 
 namespace Metasound
 {
-	static const FGuid FrontendInvalidID;
+	extern const FGuid METASOUNDFRONTEND_API FrontendInvalidID;
 }
 
 // Forward declare
@@ -580,20 +580,33 @@ struct METASOUNDFRONTEND_API FMetasoundFrontendClassName
 {
 	GENERATED_BODY()
 
+	FMetasoundFrontendClassName() = default;
+
+	FMetasoundFrontendClassName(const FName& InNamespace, const FName& InName, const FName& InVariant);
+
+	FMetasoundFrontendClassName(const Metasound::FNodeClassName& InName);
+	
+
 	// Namespace of class.
-	UPROPERTY()
-	FString Namespace;
+	UPROPERTY(EditAnywhere, Category = General)
+	FName Namespace;
 
 	// Name of class.
-	UPROPERTY()
-	FString Name;
+	UPROPERTY(EditAnywhere, Category = General)
+	FName Name;
 
 	// Variant of class. The Variant is used to describe an equivalent class which performs the same operation but on differing types.
-	UPROPERTY()
-	FString Variant;
+	UPROPERTY(EditAnywhere, Category = General)
+	FName Variant;
 
 	// Returns a full name of the class.
-	FString GetFullName() const;
+	FName GetFullName() const;
+
+	// Returns scoped name representing namespace and name. 
+	FName GetScopedName() const;
+
+	// Return string version of full name.
+	FString ToString() const;
 
 	METASOUNDFRONTEND_API friend bool operator==(const FMetasoundFrontendClassName& InLHS, const FMetasoundFrontendClassName& InRHS);
 };
@@ -604,8 +617,12 @@ struct FMetasoundFrontendClassMetadata
 {
 	GENERATED_BODY()
 
-	UPROPERTY()
-	FMetasoundFrontendClassName Name;
+	FMetasoundFrontendClassMetadata() = default;
+
+	FMetasoundFrontendClassMetadata(const Metasound::FNodeClassMetadata& InNodeClassMetadata);
+
+	UPROPERTY(EditAnywhere, Category = General)
+	FMetasoundFrontendClassName ClassName;
 
 	UPROPERTY(EditAnywhere, Category = General)
 	FMetasoundFrontendVersionNumber Version;
@@ -613,19 +630,22 @@ struct FMetasoundFrontendClassMetadata
 	UPROPERTY(VisibleAnywhere, Category = General)
 	EMetasoundFrontendClassType Type = EMetasoundFrontendClassType::Invalid;
 
-	UPROPERTY(EditAnywhere, Category = General, meta = (DisplayName = "Description"))
+	UPROPERTY(EditAnywhere, Category = General)
+	FText DisplayName;
+
+	UPROPERTY(EditAnywhere, Category = General)
 	FText Description;
 
-	UPROPERTY()
+	UPROPERTY(EditAnywhere, Category = General)
 	FText PromptIfMissing;
 
-	UPROPERTY()
+	UPROPERTY(EditAnywhere, Category = General)
 	FText Author;
 
-	UPROPERTY()
+	UPROPERTY(EditAnywhere, Category = General)
 	TArray<FName> Keywords;
 
-	UPROPERTY()
+	UPROPERTY(EditAnywhere, Category = General)
 	TArray<FText> CategoryHierarchy;
 };
 
diff --git a/Engine/Plugins/Runtime/Metasound/Source/MetasoundFrontend/Public/MetasoundFrontendQuerySteps.h b/Engine/Plugins/Runtime/Metasound/Source/MetasoundFrontend/Public/MetasoundFrontendQuerySteps.h
--- a/Engine/Plugins/Runtime/Metasound/Source/MetasoundFrontend/Public/MetasoundFrontendQuerySteps.h
+++ b/Engine/Plugins/Runtime/Metasound/Source/MetasoundFrontend/Public/MetasoundFrontendQuerySteps.h
@@ -53,12 +53,12 @@ namespace Metasound
 	class METASOUNDFRONTEND_API FFilterClassesByClassName : public IFrontendQueryFilterStep
 	{
 	public: 
-		FFilterClassesByClassName(const FString& InClassName);
+		FFilterClassesByClassName(const FMetasoundFrontendClassName& InClassName);
 
 		bool Filter(const FFrontendQueryEntry& InEntry) const override;
 
 	private:
-		FString ClassName;
+		FMetasoundFrontendClassName ClassName;
 	};
 }
 
diff --git a/Engine/Plugins/Runtime/Metasound/Source/MetasoundFrontend/Public/MetasoundFrontendRegistries.h b/Engine/Plugins/Runtime/Metasound/Source/MetasoundFrontend/Public/MetasoundFrontendRegistries.h
--- a/Engine/Plugins/Runtime/Metasound/Source/MetasoundFrontend/Public/MetasoundFrontendRegistries.h
+++ b/Engine/Plugins/Runtime/Metasound/Source/MetasoundFrontend/Public/MetasoundFrontendRegistries.h
@@ -57,14 +57,14 @@ namespace Metasound
 		struct METASOUNDFRONTEND_API FNodeRegistryKey
 		{
 			// The class name for the node.
-			FName NodeName;
+			FName NodeClassFullName;
 
 			// A hash generated from the input types and output types for this node.
 			uint32 NodeHash = 0;
 
 			FORCEINLINE bool operator==(const FNodeRegistryKey& Other) const
 			{
-				return NodeHash == Other.NodeHash && NodeName == Other.NodeName;
+				return NodeHash == Other.NodeHash && NodeClassFullName == Other.NodeClassFullName;
 			}
 
 			friend uint32 GetTypeHash(const Metasound::Frontend::FNodeRegistryKey& InKey)
@@ -163,7 +163,7 @@ class METASOUNDFRONTEND_API FMetasoundFrontendRegistryContainer
 
 	using FDataTypeRegistryInfo = Metasound::FDataTypeRegistryInfo;
 	using FDataTypeConstructorCallbacks = ::Metasound::FDataTypeConstructorCallbacks;
-	using FNodeInfo = Metasound::FNodeInfo;
+	using FNodeClassMetadata = Metasound::FNodeClassMetadata;
 
 public:
 	static FMetasoundFrontendRegistryContainer* Get();
@@ -224,7 +224,7 @@ public:
 		return ExternalNodeRegistry.Contains(InKey);
 	}
 
-	static FNodeRegistryKey GetRegistryKey(const FNodeInfo& InNodeMetadata);
+	static FNodeRegistryKey GetRegistryKey(const FNodeClassMetadata& InNodeMetadata);
 	static FNodeRegistryKey GetRegistryKey(const FMetasoundFrontendClassMetadata& InNodeMetadata);
 	static bool GetRegistryKey(const FNodeRegistryElement& InElement, FNodeRegistryKey& OutKey);
 	static bool GetFrontendClassFromRegistered(const FMetasoundFrontendClassMetadata& InMetadata, FMetasoundFrontendClass& OutClass);
diff --git a/Engine/Plugins/Runtime/Metasound/Source/MetasoundFrontend/Public/MetasoundInputNode.h b/Engine/Plugins/Runtime/Metasound/Source/MetasoundFrontend/Public/MetasoundInputNode.h
--- a/Engine/Plugins/Runtime/Metasound/Source/MetasoundFrontend/Public/MetasoundInputNode.h
+++ b/Engine/Plugins/Runtime/Metasound/Source/MetasoundFrontend/Public/MetasoundInputNode.h
@@ -117,16 +117,14 @@ namespace Metasound
 				);
 			}
 
-			static FNodeInfo GetNodeInfo(const FString& InVertexName)
+			static FNodeClassMetadata GetNodeInfo(const FString& InVertexName)
 			{
-				static const FString ClassNameString = FString(TEXT("Input_")) + GetMetasoundDataTypeName<DataType>().ToString();
-				static const FName ClassName(*ClassNameString);
+				FNodeClassMetadata Info;
 
-				FNodeInfo Info;
-
-				Info.ClassName = ClassName;
+				Info.ClassName = {TEXT("Input"), GetMetasoundDataTypeName<DataType>(), TEXT("")};
 				Info.MajorVersion = 1;
 				Info.MinorVersion = 0;
+				Info.DisplayName = FText::Format(LOCTEXT("Metasound_InputNodeDisplayNameFormat", "Input {0}"), FText::FromName(GetMetasoundDataTypeName<DataType>()));
 				Info.Description = LOCTEXT("Metasound_InputNodeDescription", "Input into the parent Metasound graph.");
 				Info.Author = PluginAuthor;
 				Info.PromptIfMissing = PluginNodeMissingPrompt;
@@ -139,8 +137,8 @@ namespace Metasound
 			/* Construct a TInputNode using the TInputOperatorFactory<> and forwarding 
 			 * Args to the TInputOperatorFactory constructor.*/
 			template<typename... ArgTypes>
-			TInputNode(const FString& InNodeDescription, const FString& InVertexName, ArgTypes&&... Args)
-			:	FNode(InNodeDescription, GetNodeInfo(InVertexName))
+			TInputNode(const FString& InNodeDescription, const FGuid& InInstanceID, const FString& InVertexName, ArgTypes&&... Args)
+			:	FNode(InNodeDescription, InInstanceID, GetNodeInfo(InVertexName))
 			,	VertexName(InVertexName)
 			,	Interface(DeclareVertexInterface(InVertexName))
 			,	Factory(MakeOperatorFactoryRef<TInputOperatorFactory<DataType>>(Forward<ArgTypes>(Args)...))
@@ -149,8 +147,8 @@ namespace Metasound
 
 			/* Construct a TInputNode using the TInputOperatorLiteralFactory<> and moving
 			 * InParam to the TInputOperatorLiteralFactory constructor.*/
-			explicit TInputNode(const FString& InNodeDescription, const FString& InVertexName, FLiteral&& InParam)
-			:	FNode(InNodeDescription, GetNodeInfo(InVertexName))
+			explicit TInputNode(const FString& InNodeDescription, const FGuid& InInstanceID, const FString& InVertexName, FLiteral&& InParam)
+			:	FNode(InNodeDescription, InInstanceID, GetNodeInfo(InVertexName))
 			,	VertexName(InVertexName)
 			,	Interface(DeclareVertexInterface(InVertexName))
 			, 	Factory(MakeOperatorFactoryRef<TInputOperatorLiteralFactory<DataType>>(MoveTemp(InParam)))
diff --git a/Engine/Plugins/Runtime/Metasound/Source/MetasoundFrontend/Public/MetasoundNodeRegistrationMacro.h b/Engine/Plugins/Runtime/Metasound/Source/MetasoundFrontend/Public/MetasoundNodeRegistrationMacro.h
--- a/Engine/Plugins/Runtime/Metasound/Source/MetasoundFrontend/Public/MetasoundNodeRegistrationMacro.h
+++ b/Engine/Plugins/Runtime/Metasound/Source/MetasoundFrontend/Public/MetasoundNodeRegistrationMacro.h
@@ -27,7 +27,7 @@ namespace Metasound
 
 
 	template <typename FNodeType>
-	bool RegisterNodeWithFrontend(const Metasound::FNodeInfo& InMetadata)
+	bool RegisterNodeWithFrontend(const Metasound::FNodeClassMetadata& InMetadata)
 	{
 		// if we reenter this code (because DECLARE_METASOUND_DATA_REFERENCE_TYPES was called twice with the same type),
 		// we catch it here.
diff --git a/Engine/Plugins/Runtime/Metasound/Source/MetasoundFrontend/Public/MetasoundOutputNode.h b/Engine/Plugins/Runtime/Metasound/Source/MetasoundFrontend/Public/MetasoundOutputNode.h
--- a/Engine/Plugins/Runtime/Metasound/Source/MetasoundFrontend/Public/MetasoundOutputNode.h
+++ b/Engine/Plugins/Runtime/Metasound/Source/MetasoundFrontend/Public/MetasoundOutputNode.h
@@ -80,15 +80,16 @@ namespace Metasound
 			);
 		}
 
-		static FNodeInfo GetNodeInfo(const FString& InVertexName)
+		static FNodeClassMetadata GetNodeInfo(const FString& InVertexName)
 		{
 			static const FString ClassNameString = FString(TEXT("Output_")) + GetMetasoundDataTypeName<DataType>().ToString();
 
-			FNodeInfo Info;
+			FNodeClassMetadata Info;
 
-			Info.ClassName = FName(*ClassNameString);
+			Info.ClassName = {TEXT("Output"), GetMetasoundDataTypeName<DataType>(), TEXT("")};
 			Info.MajorVersion = 1;
 			Info.MinorVersion = 0;
+			Info.DisplayName = FText::Format(LOCTEXT("Metasound_OutputNodeDisplayNameFormat", "Output {0}"), FText::FromName(GetMetasoundDataTypeName<DataType>()));
 			Info.Description = LOCTEXT("Metasound_OutputNodeDescription", "Output from the parent Metasound graph.");
 			Info.Author = PluginAuthor;
 			Info.PromptIfMissing = PluginNodeMissingPrompt;
@@ -100,8 +101,8 @@ namespace Metasound
 
 
 		public:
-			TOutputNode(const FString& InInstanceName, const FString& InVertexName)
-			:	FNode(InInstanceName, GetNodeInfo(InVertexName))
+			TOutputNode(const FString& InInstanceName, const FGuid& InInstanceID, const FString& InVertexName)
+			:	FNode(InInstanceName, InInstanceID, GetNodeInfo(InVertexName))
 			,	VertexInterface(GetVertexInterface(InVertexName))
 			,	Factory(MakeShared<FOutputOperatorFactory, ESPMode::ThreadSafe>(InVertexName))
 			{
diff --git a/Engine/Plugins/Runtime/Metasound/Source/MetasoundFrontend/Public/MetasoundReceiveNode.h b/Engine/Plugins/Runtime/Metasound/Source/MetasoundFrontend/Public/MetasoundReceiveNode.h
--- a/Engine/Plugins/Runtime/Metasound/Source/MetasoundFrontend/Public/MetasoundReceiveNode.h
+++ b/Engine/Plugins/Runtime/Metasound/Source/MetasoundFrontend/Public/MetasoundReceiveNode.h
@@ -38,9 +38,9 @@ namespace Metasound
 			return OutputName;
 		}
 
-		static FString GetClassNameForDataType(const FName& InDataTypeName)
+		static FNodeClassName GetClassNameForDataType(const FName& InDataTypeName)
 		{
-			return FString::Printf(TEXT("Receive.%s"), *InDataTypeName.ToString());
+			return FNodeClassName{TEXT("Receive"), InDataTypeName, TEXT("")};
 		}
 	};
 
@@ -62,14 +62,15 @@ namespace Metasound
 			);
 		}
 
-		static const FNodeInfo& GetNodeInfo()
+		static const FNodeClassMetadata& GetNodeInfo()
 		{
-			auto InitNodeInfo = []() -> FNodeInfo
+			auto InitNodeInfo = []() -> FNodeClassMetadata
 			{
-				FNodeInfo Info;
-				Info.ClassName = FName(FReceiveNodeNames::GetClassNameForDataType(GetMetasoundDataTypeName<TDataType>()));
+				FNodeClassMetadata Info;
+				Info.ClassName = FReceiveNodeNames::GetClassNameForDataType(GetMetasoundDataTypeName<TDataType>());
 				Info.MajorVersion = 1;
 				Info.MinorVersion = 0;
+				Info.DisplayName = FText::Format(LOCTEXT("Metasound_ReceiveNodeDisplayNameFormat", "Receive {0}"), FText::FromName(GetMetasoundDataTypeName<TDataType>()));
 				Info.Description = LOCTEXT("Metasound_ReceiveNodeDescription", "Receives data from a send node with the same name.");
 				Info.Author = PluginAuthor;
 				Info.PromptIfMissing = PluginNodeMissingPrompt;
@@ -80,7 +81,7 @@ namespace Metasound
 				return Info;
 			};
 
-			static const FNodeInfo Info = InitNodeInfo();
+			static const FNodeClassMetadata Info = InitNodeInfo();
 
 			return Info;
 		}
@@ -184,7 +185,7 @@ namespace Metasound
 
 			// TODO: default value of received object.
 			TReceiveNode(const FNodeInitData& InInitData)
-				: FNode(InInitData.InstanceName, GetNodeInfo())
+				: FNode(InInitData.InstanceName, InInitData.InstanceID, GetNodeInfo())
 				, Interface(DeclareVertexInterface())
 				, Factory(MakeOperatorFactoryRef<FReceiverOperatorFactory>())
 			{
diff --git a/Engine/Plugins/Runtime/Metasound/Source/MetasoundFrontend/Public/MetasoundSendNode.h b/Engine/Plugins/Runtime/Metasound/Source/MetasoundFrontend/Public/MetasoundSendNode.h
--- a/Engine/Plugins/Runtime/Metasound/Source/MetasoundFrontend/Public/MetasoundSendNode.h
+++ b/Engine/Plugins/Runtime/Metasound/Source/MetasoundFrontend/Public/MetasoundSendNode.h
@@ -45,16 +45,17 @@ namespace Metasound
 			);
 		}
 
-		static const FNodeInfo& GetNodeInfo()
+		static const FNodeClassMetadata& GetNodeInfo()
 		{
-			auto InitNodeInfo = []() -> FNodeInfo
+			auto InitNodeInfo = []() -> FNodeClassMetadata
 			{
 				const FString& InputName = GetSendInputName();
-				FNodeInfo Info;
+				FNodeClassMetadata Info;
 
-				Info.ClassName = FName(*FString::Printf(TEXT("Send %s"), *InputName));
+				Info.ClassName = {TEXT("Send"), GetMetasoundDataTypeName<TDataType>(), TEXT("")};
 				Info.MajorVersion = 1;
 				Info.MinorVersion = 0;
+				Info.DisplayName = FText::Format(LOCTEXT("Metasound_SendNodeDisplayNameFormat", "Send {0}"), FText::FromName(GetMetasoundDataTypeName<TDataType>()));
 				Info.Description = LOCTEXT("Metasound_SendNodeDescription", "Sends data from a send node with the same name.");
 				Info.Author = PluginAuthor;
 				Info.PromptIfMissing = PluginNodeMissingPrompt;
@@ -65,7 +66,7 @@ namespace Metasound
 				return Info;
 			};
 
-			static const FNodeInfo Info = InitNodeInfo();
+			static const FNodeClassMetadata Info = InitNodeInfo();
 
 			return Info;
 		}
@@ -138,7 +139,7 @@ namespace Metasound
 		public:
 
 			TSendNode(const FNodeInitData& InInitData)
-				: FNode(InInitData.InstanceName, GetNodeInfo())
+				: FNode(InInitData.InstanceName, InInitData.InstanceID, GetNodeInfo())
 				, Interface(DeclareVertexInterface())
 				, Factory(MakeOperatorFactoryRef<FSendOperatorFactory>())
 			{
diff --git a/Engine/Plugins/Runtime/Metasound/Source/MetasoundGraphCore/Private/MetasoundGraph.cpp b/Engine/Plugins/Runtime/Metasound/Source/MetasoundGraphCore/Private/MetasoundGraph.cpp
--- a/Engine/Plugins/Runtime/Metasound/Source/MetasoundGraphCore/Private/MetasoundGraph.cpp
+++ b/Engine/Plugins/Runtime/Metasound/Source/MetasoundGraphCore/Private/MetasoundGraph.cpp
@@ -6,9 +6,10 @@
 
 namespace Metasound
 {
-	FGraph::FGraph(const FString& InInstanceName)
+	FGraph::FGraph(const FString& InInstanceName, const FGuid& InInstanceID)
 	: InstanceName(InInstanceName)
-	, Metadata(FNodeInfo::GetEmpty())
+	, InstanceID(InInstanceID)
+	, Metadata(FNodeClassMetadata::GetEmpty())
 	{
 	}
 
@@ -17,7 +18,12 @@ namespace Metasound
 		return InstanceName;
 	}
 
-	const FNodeInfo& FGraph::GetMetadata() const
+	const FGuid& FGraph::GetInstanceID() const
+	{
+		return InstanceID;
+	}
+
+	const FNodeClassMetadata& FGraph::GetMetadata() const
 	{
 		return Metadata;
 	}
diff --git a/Engine/Plugins/Runtime/Metasound/Source/MetasoundGraphCore/Private/MetasoundNode.cpp b/Engine/Plugins/Runtime/Metasound/Source/MetasoundGraphCore/Private/MetasoundNode.cpp
--- a/Engine/Plugins/Runtime/Metasound/Source/MetasoundGraphCore/Private/MetasoundNode.cpp
+++ b/Engine/Plugins/Runtime/Metasound/Source/MetasoundGraphCore/Private/MetasoundNode.cpp
@@ -6,9 +6,10 @@
 
 namespace Metasound
 {
-	FNode::FNode(const FString& InInstanceName, const FNodeInfo& InInfo)
-	:	InstanceName(InInstanceName)
-	,	Info(InInfo)
+	FNode::FNode(const FString& InInstanceName, const FGuid& InInstanceID, const FNodeClassMetadata& InInfo)
+	: InstanceName(InInstanceName)
+	, InstanceID(InInstanceID)
+	, Info(InInfo)
 	{
 	}
 
@@ -18,8 +19,14 @@ namespace Metasound
 		return InstanceName;
 	}
 
+	/** Return the ID of this specific instance of the node class. */
+	const FGuid& FNode::GetInstanceID() const
+	{
+		return InstanceID;
+	}
+
 	/** Return the type name of this node. */
-	const FNodeInfo& FNode::GetMetadata() const
+	const FNodeClassMetadata& FNode::GetMetadata() const
 	{
 		return Info;
 	}
diff --git a/Engine/Plugins/Runtime/Metasound/Source/MetasoundGraphCore/Private/MetasoundNodeInterface.cpp b/Engine/Plugins/Runtime/Metasound/Source/MetasoundGraphCore/Private/MetasoundNodeInterface.cpp
--- a/Engine/Plugins/Runtime/Metasound/Source/MetasoundGraphCore/Private/MetasoundNodeInterface.cpp
+++ b/Engine/Plugins/Runtime/Metasound/Source/MetasoundGraphCore/Private/MetasoundNodeInterface.cpp
@@ -4,6 +4,68 @@
 
 namespace Metasound
 {
+	FNodeClassName::FNodeClassName()
+	{
+	}
+
+	FNodeClassName::FNodeClassName(const FName& InNamespace, const FName& InName, const FName& InVariant)
+	: Namespace(InNamespace)
+	, Name(InName)
+	, Variant(InVariant)
+	{
+		ScopedName = FormatScopedName(Namespace, Name);
+		FullName = FormatFullName(Namespace, Name, Variant);
+	}
+
+	/** Namespace of node class. */
+	const FName& FNodeClassName::GetNamespace() const
+	{
+		return Namespace;
+	}
+
+	/** Name of node class. */
+	const FName& FNodeClassName::GetName() const
+	{
+		return Name;
+	}
+
+	/** Variant of node class. */
+	const FName& FNodeClassName::GetVariant() const
+	{
+		return Variant;
+	}
+
+	/** Namespace and name of the node class. */
+	const FName& FNodeClassName::GetScopedName() const
+	{
+		return ScopedName;
+	}
+
+	/** Namespace, name and variant of the node class. */
+	const FName& FNodeClassName::GetFullName() const
+	{
+		return FullName;
+	}
+
+	FName FNodeClassName::FormatFullName(const FName& InNamespace, const FName& InName, const FName& InVariant)
+	{
+		FName ScopedName = FormatScopedName(InNamespace, InName);
+
+		if (InVariant != NAME_None)
+		{
+			return FName(FString::Printf(TEXT("%s.%s"), *ScopedName.ToString(), *InVariant.ToString()));
+		}
+		else
+		{
+			return ScopedName;
+		}
+	}
+
+	FName FNodeClassName::FormatScopedName(const FName& InNamespace, const FName& InName)
+	{
+		return FName(FString::Printf(TEXT("%s.%s"), *InNamespace.ToString(), *InName.ToString()));
+	}
+
 	bool operator==(const FOutputDataSource& InLeft, const FOutputDataSource& InRight)
 	{
 		return (InLeft.Node == InRight.Node) && (InLeft.Vertex == InRight.Vertex);
diff --git a/Engine/Plugins/Runtime/Metasound/Source/MetasoundGraphCore/Public/MetasoundFacade.h b/Engine/Plugins/Runtime/Metasound/Source/MetasoundGraphCore/Public/MetasoundFacade.h
--- a/Engine/Plugins/Runtime/Metasound/Source/MetasoundGraphCore/Public/MetasoundFacade.h
+++ b/Engine/Plugins/Runtime/Metasound/Source/MetasoundGraphCore/Public/MetasoundFacade.h
@@ -40,9 +40,9 @@ namespace Metasound
 			template<typename T, T> 
 			struct Helper;
 
-			// Check for "static const FNodeInfo& U::CreateOperator(const FCreateOperatorParams& Inparams, FBuildErrorArray& OutErrors)"
+			// Check for "static const FNodeClassMetadata& U::CreateOperator(const FCreateOperatorParams& Inparams, FBuildErrorArray& OutErrors)"
 			template<typename T>
-			static uint8 Check(Helper<const FNodeInfo&(*)(), &T::GetNodeInfo>*);
+			static uint8 Check(Helper<const FNodeClassMetadata&(*)(), &T::GetNodeInfo>*);
 
 			template<typename T> static uint16 Check(...);
 
@@ -65,10 +65,10 @@ namespace Metasound
 		// Require static TUniquePtr<IOperator> OperatorType::CreateOperator(const FCreateOperatorParams&, TArray<TUniquePtr<IOperatorBuildError>>&) exists.
 		static_assert(TIsFactoryMethodDeclared<OperatorType>::Value, "To use the FNodeFacade constructor, the OperatorType must have the static function \"static TUniquePtr<IOperator> OperatorType::CreateOperator(const FCreateOperatorParams&, TArray<TUniquePtr<IOperatorBuildError>>&)\"");
 
-		// Require static const FNodeInfo& OperatorType::GetNodeInfo() exists.
-		static_assert(TIsNodeInfoDeclared<OperatorType>::Value, "To use the FNodeFacade constructor, the OperatorType must have the static function \"static const FNodeInfo& OperatorType::GetNodeInfo()\"");
+		// Require static const FNodeClassMetadata& OperatorType::GetNodeInfo() exists.
+		static_assert(TIsNodeInfoDeclared<OperatorType>::Value, "To use the FNodeFacade constructor, the OperatorType must have the static function \"static const FNodeClassMetadata& OperatorType::GetNodeInfo()\"");
 
-		static const FNodeInfo& GetNodeInfo() 
+		static const FNodeClassMetadata& GetNodeInfo() 
 		{ 
 			return OperatorType::GetNodeInfo(); 
 		}
@@ -113,12 +113,13 @@ namespace Metasound
 			 * to get the node info, operator factory method and vertex interface. 
 			 *
 			 * @param InInstanceName - Instance name for the node.
+			 * @param InInstanceID - Instance ID for the node.
 			 * @param OperatorClass - Template class wrapper for the underlying 
 			 * 						  IOperator which is created by this node.
 			 */
 			template<typename OperatorType>
-			FNodeFacade(const FString& InInstanceName, TFacadeOperatorClass<OperatorType> OperatorClass)
-			:	FNode(InInstanceName, OperatorType::GetNodeInfo())
+			FNodeFacade(const FString& InInstanceName, const FGuid& InInstanceID, TFacadeOperatorClass<OperatorType> OperatorClass)
+			:	FNode(InInstanceName, InInstanceID, OperatorType::GetNodeInfo())
 			, 	Factory(MakeShared<FFactory, ESPMode::ThreadSafe>(&OperatorType::CreateOperator))
 			{
 				VertexInterface = GetMetadata().DefaultInterface;
diff --git a/Engine/Plugins/Runtime/Metasound/Source/MetasoundGraphCore/Public/MetasoundGraph.h b/Engine/Plugins/Runtime/Metasound/Source/MetasoundGraphCore/Public/MetasoundGraph.h
--- a/Engine/Plugins/Runtime/Metasound/Source/MetasoundGraphCore/Public/MetasoundGraph.h
+++ b/Engine/Plugins/Runtime/Metasound/Source/MetasoundGraphCore/Public/MetasoundGraph.h
@@ -17,14 +17,17 @@ namespace Metasound
 	class METASOUNDGRAPHCORE_API FGraph : public IGraph
 	{
 		public:
-			FGraph(const FString& InInstanceName);
+			FGraph(const FString& InInstanceName, const FGuid& InInstanceID);
 			virtual ~FGraph() = default;
 
 			/** Return the name of this specific instance of the node class. */
 			const FString& GetInstanceName() const override;
 
+			/** Return the ID of this specific instance of the node class. */
+			const FGuid& GetInstanceID() const override;
+
 			/** Return metadata about this graph. */
-			const FNodeInfo& GetMetadata() const override;
+			const FNodeClassMetadata& GetMetadata() const override;
 
 			/** Retrieve all the edges associated with a graph. */
 			const TArray<FDataEdge>& GetDataEdges() const override;
@@ -106,8 +109,9 @@ namespace Metasound
 
 		private:
 			FString InstanceName;
-			FNodeInfo Metadata;
+			FGuid InstanceID;
 
+			FNodeClassMetadata Metadata;
 
 			TArray<FDataEdge> Edges;
 
diff --git a/Engine/Plugins/Runtime/Metasound/Source/MetasoundGraphCore/Public/MetasoundLiteral.h b/Engine/Plugins/Runtime/Metasound/Source/MetasoundGraphCore/Public/MetasoundLiteral.h
--- a/Engine/Plugins/Runtime/Metasound/Source/MetasoundGraphCore/Public/MetasoundLiteral.h
+++ b/Engine/Plugins/Runtime/Metasound/Source/MetasoundGraphCore/Public/MetasoundLiteral.h
@@ -6,6 +6,7 @@
 
 #include "IAudioProxyInitializer.h"
 #include "MetasoundOperatorSettings.h" 
+#include "Misc/Guid.h"
 #include "Misc/TVariant.h"
 #include <type_traits>
 
@@ -218,6 +219,7 @@ namespace Metasound
 	{
 		// the instance name and name of the specific connection that we should use.
 		const FString& InNodeName;
+		const FGuid& InInstanceID;
 		const FString& InVertexName;
 
 		FLiteral InitParam;
@@ -227,6 +229,7 @@ namespace Metasound
 	{
 		// the instance name and name of the specific connection that we should use.
 		const FString& InNodeName;
+		const FGuid& InInstanceID;
 		const FString& InVertexName;
 	};
 }
diff --git a/Engine/Plugins/Runtime/Metasound/Source/MetasoundGraphCore/Public/MetasoundNode.h b/Engine/Plugins/Runtime/Metasound/Source/MetasoundGraphCore/Public/MetasoundNode.h
--- a/Engine/Plugins/Runtime/Metasound/Source/MetasoundGraphCore/Public/MetasoundNode.h
+++ b/Engine/Plugins/Runtime/Metasound/Source/MetasoundGraphCore/Public/MetasoundNode.h
@@ -10,19 +10,23 @@ namespace Metasound
 	class METASOUNDGRAPHCORE_API FNode : public INode
 	{
 		public:
-			FNode(const FString& InInstanceName, const FNodeInfo& InInfo);
+			FNode(const FString& InInstanceName, const FGuid& InInstanceID, const FNodeClassMetadata& InInfo);
 
 			virtual ~FNode() = default;
 
 			/** Return the name of this specific instance of the node class. */
 			const FString& GetInstanceName() const override;
 
+			/** Return the ID of this specific instance of the node class. */
+			const FGuid& GetInstanceID() const override;
+
 			/** Return metadata associated with this node. */
-			const FNodeInfo& GetMetadata() const override;
+			const FNodeClassMetadata& GetMetadata() const override;
 
 		private:
 
 			FString InstanceName;
-			FNodeInfo Info;
+			FGuid InstanceID;
+			FNodeClassMetadata Info;
 	};
 }
diff --git a/Engine/Plugins/Runtime/Metasound/Source/MetasoundGraphCore/Public/MetasoundNodeInterface.h b/Engine/Plugins/Runtime/Metasound/Source/MetasoundGraphCore/Public/MetasoundNodeInterface.h
--- a/Engine/Plugins/Runtime/Metasound/Source/MetasoundGraphCore/Public/MetasoundNodeInterface.h
+++ b/Engine/Plugins/Runtime/Metasound/Source/MetasoundGraphCore/Public/MetasoundNodeInterface.h
@@ -9,6 +9,7 @@
 #include "MetasoundOperatorInterface.h"
 #include "MetasoundVertex.h"
 #include "MetasoundLiteral.h"
+#include "Misc/Guid.h"
 
 namespace Metasound
 {
@@ -22,6 +23,7 @@ namespace Metasound
 	struct FNodeInitData
 	{
 		FString InstanceName;
+		FGuid InstanceID;
 		TMap<FName, FLiteral> ParamMap;
 
 		template<typename ParamType>
@@ -51,21 +53,56 @@ namespace Metasound
 		bool bShowOutputNames = true;
 	};
 
-	// TODO: update FNodeInfo with this as class name.
-	struct FNodeClassName
+	/** Name of a node class.
+	 *
+	 * FNodeClassName is used for lookup and declaring interoperability.
+	 *
+	 * Namespaces are provided as a convenience to simply name collisions.
+	 *
+	 * Nodes with equal Namespace and Name, but different Variants are considered
+	 * to be interoperable. They can be used to define nodes that perform the same
+	 * function, but have differing vertex types.
+	 */
+	class METASOUNDGRAPHCORE_API FNodeClassName
 	{
-		FString Namespace;
-		FString Name;
-		FString Variant;
+	public:
+		FNodeClassName();
+
+		FNodeClassName(const FName& InNamespace, const FName& InName, const FName& InVariant);
+
+		/** Namespace of node class. */
+		const FName& GetNamespace() const;
+
+		/** Name of node class. */
+		const FName& GetName() const;
+
+		/** Variant of node class. */
+		const FName& GetVariant() const;
+
+		/** Namespace and name of the node class. */
+		const FName& GetScopedName() const;
+
+		/** Namespace, name and variant of the node class. */
+		const FName& GetFullName() const;
+
+		static FName FormatFullName(const FName& InNamespace, const FName& InName, const FName& InVariant);
+
+		static FName FormatScopedName(const FName& InNamespace, const FName& InName);
+
+	private:
+
+		FName Namespace;
+		FName Name;
+		FName Variant;
+		FName ScopedName;
+		FName FullName;
 	};
 
 	/** Provides metadata for a given node. */
-	struct FNodeInfo
+	struct FNodeClassMetadata
 	{
-		// TODO: rename this class to FNodeMetadata
-
 		/** Name of class. Used for registration and lookup. */
-		FName ClassName;
+		FNodeClassName ClassName;
 
 		/** Major version of node. Used for registration and lookup. */
 		int32 MajorVersion = -1;
@@ -73,6 +110,9 @@ namespace Metasound
 		/** Minor version of node. */
 		int32 MinorVersion = -1;
 
+		/** Display name of node class. */
+		FText DisplayName;
+
 		/** Human readable description of node. */
 		FText Description;
 
@@ -94,10 +134,10 @@ namespace Metasound
 		/** Display style for node when visualized. */
 		FNodeDisplayStyle DisplayStyle;
 
-		/** Returns an empty FNodeInfo object. */
-		static const FNodeInfo& GetEmpty()
+		/** Returns an empty FNodeClassMetadata object. */
+		static const FNodeClassMetadata& GetEmpty()
 		{
-			static const FNodeInfo EmptyInfo;
+			static const FNodeClassMetadata EmptyInfo;
 			return EmptyInfo;
 		}
 	};
@@ -114,8 +154,11 @@ namespace Metasound
 			/** Return the name of this specific instance of the node class. */
 			virtual const FString& GetInstanceName() const = 0;
 
+			/** Return the ID of this node instance. */
+			virtual const FGuid& GetInstanceID() const = 0;
+
 			/** Return the type name of this node. */
-			virtual const FNodeInfo& GetMetadata() const = 0;
+			virtual const FNodeClassMetadata& GetMetadata() const = 0;
 
 			/** Return the current vertex interface. */
 			virtual const FVertexInterface& GetVertexInterface() const = 0;
diff --git a/Engine/Plugins/Runtime/Metasound/Source/MetasoundStandardNodes/Private/MetasoundADSRNode.cpp b/Engine/Plugins/Runtime/Metasound/Source/MetasoundStandardNodes/Private/MetasoundADSRNode.cpp
--- a/Engine/Plugins/Runtime/Metasound/Source/MetasoundStandardNodes/Private/MetasoundADSRNode.cpp
+++ b/Engine/Plugins/Runtime/Metasound/Source/MetasoundStandardNodes/Private/MetasoundADSRNode.cpp
@@ -8,6 +8,7 @@
 #include "MetasoundExecutableOperator.h"
 #include "MetasoundNodeRegistrationMacro.h"
 #include "MetasoundPrimitives.h"
+#include "MetasoundStandardNodesNames.h"
 #include "MetasoundTime.h"
 #include "MetasoundTrigger.h"
 
@@ -28,7 +29,7 @@ namespace Metasound
 				FFloatTimeReadRef Release;
 			};
 
-			static const FNodeInfo& GetNodeInfo();
+			static const FNodeClassMetadata& GetNodeInfo();
 			static FVertexInterface DeclareVertexInterface();
 			static TUniquePtr<IOperator> CreateOperator(const FCreateOperatorParams& InParams, FBuildErrorArray& OutErrors);
 
@@ -252,14 +253,15 @@ namespace Metasound
 		return Interface;
 	}
 
-	const FNodeInfo& FADSROperator::GetNodeInfo()
+	const FNodeClassMetadata& FADSROperator::GetNodeInfo()
 	{
-		auto InitNodeInfo = []() -> FNodeInfo
+		auto InitNodeInfo = []() -> FNodeClassMetadata
 		{
-			FNodeInfo Info;
-			Info.ClassName = FName(TEXT("ADSR"));
+			FNodeClassMetadata Info;
+			Info.ClassName = {Metasound::StandardNodes::Namespace, TEXT("ADSR"), Metasound::StandardNodes::AudioVariant};
 			Info.MajorVersion = 1;
 			Info.MinorVersion = 0;
+			Info.DisplayName = LOCTEXT("Metasound_ADSRNodeDisplayName", "ADSR");
 			Info.Description = LOCTEXT("Metasound_ADSRNodeDescription", "Emits an ADSR (Attack, decay, sustain, & release) envelope when triggered.");
 			Info.Author = PluginAuthor;
 			Info.PromptIfMissing = PluginNodeMissingPrompt;
@@ -268,7 +270,7 @@ namespace Metasound
 			return Info;
 		};
 
-		static const FNodeInfo Info = InitNodeInfo();
+		static const FNodeClassMetadata Info = InitNodeInfo();
 
 		return Info;
 	}
@@ -294,8 +296,8 @@ namespace Metasound
 		return MakeUnique<FADSROperator>(InParams.OperatorSettings, TriggerAttack, TriggerRelease, ADSRReferences);
 	}
 
-	FADSRNode::FADSRNode(const FString& InName, float InDefaultAttackMs, float InDefaultDecayMs, float InDefaultSustainLevel, float InDefaultReleaseMs)
-	:	FNodeFacade(InName, TFacadeOperatorClass<FADSROperator>())
+	FADSRNode::FADSRNode(const FString& InName, const FGuid& InInstanceID, float InDefaultAttackMs, float InDefaultDecayMs, float InDefaultSustainLevel, float InDefaultReleaseMs)
+	:	FNodeFacade(InName, InInstanceID, TFacadeOperatorClass<FADSROperator>())
 	,	DefaultSustainLevel(InDefaultSustainLevel)
 	,	DefaultAttackMs(InDefaultAttackMs)
 	,	DefaultDecayMs(InDefaultDecayMs)
@@ -304,7 +306,7 @@ namespace Metasound
 	}
 
 	FADSRNode::FADSRNode(const FNodeInitData& InitData)
-		: FADSRNode(InitData.InstanceName, 10.0f, 20.0f, 0.7f, 20.0f)
+		: FADSRNode(InitData.InstanceName, InitData.InstanceID, 10.0f, 20.0f, 0.7f, 20.0f)
 	{
 	}
 
diff --git a/Engine/Plugins/Runtime/Metasound/Source/MetasoundStandardNodes/Private/MetasoundAudioFormatNodes.cpp b/Engine/Plugins/Runtime/Metasound/Source/MetasoundStandardNodes/Private/MetasoundAudioFormatNodes.cpp
--- a/Engine/Plugins/Runtime/Metasound/Source/MetasoundStandardNodes/Private/MetasoundAudioFormatNodes.cpp
+++ b/Engine/Plugins/Runtime/Metasound/Source/MetasoundStandardNodes/Private/MetasoundAudioFormatNodes.cpp
@@ -5,6 +5,7 @@
 #include "MetasoundAudioFormats.h"
 #include "MetasoundFacade.h"
 #include "MetasoundOperatorInterface.h"
+#include "MetasoundStandardNodesNames.h"
 
 #define LOCTEXT_NAMESPACE "MetasoundStandardNodes"
 
@@ -63,18 +64,19 @@ namespace Metasound
 			return MakeUnique<FMonoAudioFormatOperator>(Mono);
 		}
 
-		static const FNodeInfo& GetNodeInfo()
+		static const FNodeClassMetadata& GetNodeInfo()
 		{
-			auto InitNodeInfo = []() -> FNodeInfo
+			auto InitNodeInfo = []() -> FNodeClassMetadata
 			{
 				FNodeDisplayStyle DisplayStyle;
 				DisplayStyle.bShowName = false;
 
-				FNodeInfo Info;
-				Info.ClassName = "Buffer to Mono";
+				FNodeClassMetadata Info;
+				Info.ClassName = {Metasound::StandardNodes::Namespace, TEXT("FormatAudio"), TEXT("Mono")};
 				Info.DisplayStyle = DisplayStyle;
 				Info.MajorVersion = 1;
 				Info.MinorVersion = 0;
+				Info.DisplayName = LOCTEXT("MonoAudioFormat_NodeDisplayName", "Mono");
 				Info.Description = LOCTEXT("MonoAudioFormat_NodeDescription", "Converts Audio:Buffer To Mono Format");
 				Info.CategoryHierarchy = { LOCTEXT("Metasound_ConvertNodeCategory", "Conversions") };
 				Info.Author = PluginAuthor;
@@ -92,18 +94,18 @@ namespace Metasound
 				return Info;
 			};
 
-			static const FNodeInfo Info = InitNodeInfo();
+			static const FNodeClassMetadata Info = InitNodeInfo();
 			return Info;
 		}
 	};
 
-	FMonoAudioFormatNode::FMonoAudioFormatNode(const FString& InInstanceName)
-	: FNodeFacade(InInstanceName, TFacadeOperatorClass<FMonoAudioFormatOperator>())
+	FMonoAudioFormatNode::FMonoAudioFormatNode(const FString& InInstanceName, const FGuid& InInstanceID)
+	: FNodeFacade(InInstanceName, InInstanceID, TFacadeOperatorClass<FMonoAudioFormatOperator>())
 	{
 	}
 
 	FMonoAudioFormatNode::FMonoAudioFormatNode(const FNodeInitData& InInitData)
-	: FMonoAudioFormatNode(InInitData.InstanceName)
+	: FMonoAudioFormatNode(InInitData.InstanceName, InInitData.InstanceID)
 	{
 	}
 
@@ -169,18 +171,19 @@ namespace Metasound
 			return MakeUnique<FStereoAudioFormatOperator>(Stereo);
 		}
 
-		static const FNodeInfo& GetNodeInfo()
+		static const FNodeClassMetadata& GetNodeInfo()
 		{
-			auto InitNodeInfo = []() -> FNodeInfo
+			auto InitNodeInfo = []() -> FNodeClassMetadata
 			{
 				FNodeDisplayStyle DisplayStyle;
 				DisplayStyle.bShowName = false;
 
-				FNodeInfo Info;
-				Info.ClassName = "Buffer to Stereo";
+				FNodeClassMetadata Info;
+				Info.ClassName = {Metasound::StandardNodes::Namespace, TEXT("FormatAudio"), TEXT("Stereo")};
 				Info.DisplayStyle = DisplayStyle;
 				Info.MajorVersion = 1;
 				Info.MinorVersion = 0;
+				Info.DisplayName = LOCTEXT("StereoAudioFormat_NodeDisplayName", "Stereo");
 				Info.Description = LOCTEXT("StereoAudioFormat_NodeDescription", "Converts Audio:Buffer to Stereo Format");
 				Info.CategoryHierarchy = { LOCTEXT("Metasound_ConvertNodeCategory", "Conversions") };
 				Info.Author = PluginAuthor;
@@ -199,18 +202,18 @@ namespace Metasound
 				return Info;
 			};
 
-			static const FNodeInfo Info = InitNodeInfo();
+			static const FNodeClassMetadata Info = InitNodeInfo();
 			return Info;
 		}
 	};
 
-	FStereoAudioFormatNode::FStereoAudioFormatNode(const FString& InInstanceName)
-	: FNodeFacade(InInstanceName, TFacadeOperatorClass<FStereoAudioFormatOperator>())
+	FStereoAudioFormatNode::FStereoAudioFormatNode(const FString& InInstanceName, const FGuid& InInstanceID)
+	: FNodeFacade(InInstanceName, InInstanceID, TFacadeOperatorClass<FStereoAudioFormatOperator>())
 	{
 	}
 
 	FStereoAudioFormatNode::FStereoAudioFormatNode(const FNodeInitData& InInitData)
-	: FStereoAudioFormatNode(InInitData.InstanceName)
+	: FStereoAudioFormatNode(InInitData.InstanceName, InInitData.InstanceID)
 	{
 	}
 
diff --git a/Engine/Plugins/Runtime/Metasound/Source/MetasoundStandardNodes/Private/MetasoundBufferAddFloatNode.cpp b/Engine/Plugins/Runtime/Metasound/Source/MetasoundStandardNodes/Private/MetasoundBufferAddFloatNode.cpp
--- a/Engine/Plugins/Runtime/Metasound/Source/MetasoundStandardNodes/Private/MetasoundBufferAddFloatNode.cpp
+++ b/Engine/Plugins/Runtime/Metasound/Source/MetasoundStandardNodes/Private/MetasoundBufferAddFloatNode.cpp
@@ -9,6 +9,7 @@
 #include "MetasoundNodeRegistrationMacro.h"
 #include "MetasoundOperatorSettings.h"
 #include "MetasoundPrimitives.h"
+#include "MetasoundStandardNodesNames.h"
 
 #define LOCTEXT_NAMESPACE "MetasoundBufferAddFloatNode"
 
@@ -17,7 +18,7 @@ namespace Metasound
 	class FBufferAddFloatOperator : public TExecutableOperator<FBufferAddFloatOperator>
 	{
 		public:
-			static const FNodeInfo& GetNodeInfo();
+			static const FNodeClassMetadata& GetNodeInfo();
 
 			static FVertexInterface DeclareVertexInterface();
 
@@ -66,13 +67,13 @@ namespace Metasound
 			FAudioBufferWriteRef OutputBuffer;
 	};
 
-	FBufferAddFloatNode::FBufferAddFloatNode(const FString& InInstanceName)
-	:	FNodeFacade(InInstanceName, TFacadeOperatorClass<FBufferAddFloatOperator>())
+	FBufferAddFloatNode::FBufferAddFloatNode(const FString& InInstanceName, const FGuid& InInstanceID)
+	:	FNodeFacade(InInstanceName, InInstanceID, TFacadeOperatorClass<FBufferAddFloatOperator>())
 	{
 	}
 
 	FBufferAddFloatNode::FBufferAddFloatNode(const FNodeInitData& InInitData)
-	:	FBufferAddFloatNode(InInitData.InstanceName)
+	:	FBufferAddFloatNode(InInitData.InstanceName, InInitData.InstanceID)
 	{
 	}
 
@@ -102,9 +103,9 @@ namespace Metasound
 		return MakeUnique<FBufferAddFloatOperator>(InParams.OperatorSettings, InputBuffer, InputGain);
 	}
 
-	const FNodeInfo& FBufferAddFloatOperator::GetNodeInfo()
+	const FNodeClassMetadata& FBufferAddFloatOperator::GetNodeInfo()
 	{
-		auto InitNodeInfo = []() -> FNodeInfo
+		auto InitNodeInfo = []() -> FNodeClassMetadata
 		{
 			FNodeDisplayStyle DisplayStyle;
 			DisplayStyle.ImageName = "MetasoundEditor.Graph.Node.Math.Add";
@@ -112,11 +113,12 @@ namespace Metasound
 			DisplayStyle.bShowInputNames = false;
 			DisplayStyle.bShowOutputNames = false;
 
-			FNodeInfo Info;
-			Info.ClassName = "Add Float to AudioBuffer";
+			FNodeClassMetadata Info;
+			Info.ClassName = {Metasound::StandardNodes::Namespace, TEXT("Add"), TEXT("FloatToAudioBuffer")};
 			Info.DisplayStyle = DisplayStyle;
 			Info.MajorVersion = 1;
 			Info.MinorVersion = 0;
+			Info.DisplayName = LOCTEXT("BufferAddFloat_NodeDisplayName", "Add");
 			Info.Description = LOCTEXT("BufferAddFloat_NodeDescription", "Add Float to AudioBuffer.");
 			Info.CategoryHierarchy = { LOCTEXT("Metasound_MathCategory", "Math") };
 			Info.Author = PluginAuthor;
@@ -126,7 +128,7 @@ namespace Metasound
 			return Info;
 		};
 
-		static const FNodeInfo Info = InitNodeInfo();
+		static const FNodeClassMetadata Info = InitNodeInfo();
 
 		return Info;
 	};
diff --git a/Engine/Plugins/Runtime/Metasound/Source/MetasoundStandardNodes/Private/MetasoundGainMultiplyNode.cpp b/Engine/Plugins/Runtime/Metasound/Source/MetasoundStandardNodes/Private/MetasoundGainMultiplyNode.cpp
--- a/Engine/Plugins/Runtime/Metasound/Source/MetasoundStandardNodes/Private/MetasoundGainMultiplyNode.cpp
+++ b/Engine/Plugins/Runtime/Metasound/Source/MetasoundStandardNodes/Private/MetasoundGainMultiplyNode.cpp
@@ -9,6 +9,7 @@
 #include "MetasoundNodeRegistrationMacro.h"
 #include "MetasoundOperatorSettings.h"
 #include "MetasoundPrimitives.h"
+#include "MetasoundStandardNodesNames.h"
 
 #define LOCTEXT_NAMESPACE "MetasoundGainMultiplyNode"
 
@@ -17,7 +18,7 @@ namespace Metasound
 	class FGainMultiplyOperator : public TExecutableOperator<FGainMultiplyOperator>
 	{
 		public:
-			static const FNodeInfo& GetNodeInfo();
+			static const FNodeClassMetadata& GetNodeInfo();
 
 			static FVertexInterface DeclareVertexInterface();
 
@@ -84,13 +85,13 @@ namespace Metasound
 			int32 BlockSize;
 	};
 
-	FGainMultiplyNode::FGainMultiplyNode(const FString& InInstanceName)
-	:	FNodeFacade(InInstanceName, TFacadeOperatorClass<FGainMultiplyOperator>())
+	FGainMultiplyNode::FGainMultiplyNode(const FString& InInstanceName, const FGuid& InInstanceID)
+	:	FNodeFacade(InInstanceName, InInstanceID, TFacadeOperatorClass<FGainMultiplyOperator>())
 	{
 	}
 
 	FGainMultiplyNode::FGainMultiplyNode(const FNodeInitData& InInitData)
-	:	FGainMultiplyNode(InInitData.InstanceName)
+	:	FGainMultiplyNode(InInitData.InstanceName, InInitData.InstanceID)
 	{
 	}
 
@@ -120,9 +121,9 @@ namespace Metasound
 		return MakeUnique<FGainMultiplyOperator>(InParams.OperatorSettings, InputBuffer, InputGain);
 	}
 
-	const FNodeInfo& FGainMultiplyOperator::GetNodeInfo()
+	const FNodeClassMetadata& FGainMultiplyOperator::GetNodeInfo()
 	{
-		auto InitNodeInfo = []() -> FNodeInfo
+		auto InitNodeInfo = []() -> FNodeClassMetadata
 		{
 			FNodeDisplayStyle DisplayStyle;
 			DisplayStyle.ImageName = "MetasoundEditor.Graph.Node.Math.Multiply";
@@ -130,11 +131,12 @@ namespace Metasound
 			DisplayStyle.bShowInputNames = false;
 			DisplayStyle.bShowOutputNames = false;
 
-			FNodeInfo Info;
-			Info.ClassName = "Multiply AudioBuffer by Gain";
+			FNodeClassMetadata Info;
+			Info.ClassName = FNodeClassName(StandardNodes::Namespace, TEXT("Multiply"), TEXT("AudioBufferByGain"));
 			Info.DisplayStyle = DisplayStyle;
 			Info.MajorVersion = 1;
 			Info.MinorVersion = 0;
+			Info.DisplayName = LOCTEXT("GainMultiply_NodeDisplayName", "Multiply Audio by Float");
 			Info.Description = LOCTEXT("GainMultiply_NodeDescription", "Multiply AudioBuffer by float.");
 			Info.CategoryHierarchy = { LOCTEXT("Metasound_MathCategory", "Math") };
 			Info.Author = PluginAuthor;
@@ -144,7 +146,7 @@ namespace Metasound
 			return Info;
 		};
 
-		static const FNodeInfo Info = InitNodeInfo();
+		static const FNodeClassMetadata Info = InitNodeInfo();
 
 		return Info;
 	};
diff --git a/Engine/Plugins/Runtime/Metasound/Source/MetasoundStandardNodes/Private/MetasoundMathNodes.cpp b/Engine/Plugins/Runtime/Metasound/Source/MetasoundStandardNodes/Private/MetasoundMathNodes.cpp
--- a/Engine/Plugins/Runtime/Metasound/Source/MetasoundStandardNodes/Private/MetasoundMathNodes.cpp
+++ b/Engine/Plugins/Runtime/Metasound/Source/MetasoundStandardNodes/Private/MetasoundMathNodes.cpp
@@ -1,6 +1,7 @@
 // Copyright Epic Games, Inc. All Rights Reserved.
 #include "MetasoundMathNodes.h"
 
+#define LOCTEXT_NAMESPACE "MetasoundMathOpNode"
 
 namespace Metasound
 {
@@ -37,3 +38,5 @@ namespace Metasound
 	METASOUND_REGISTER_NODE(FRandRangeFrequencyNode)
 	METASOUND_REGISTER_NODE(FRandRangeFloatTimeNode)
 }
+
+#undef LOCTEXT_NAMESPACE // MetasoundMathOpNode
diff --git a/Engine/Plugins/Runtime/Metasound/Source/MetasoundStandardNodes/Private/MetasoundMixerNode.cpp b/Engine/Plugins/Runtime/Metasound/Source/MetasoundStandardNodes/Private/MetasoundMixerNode.cpp
--- a/Engine/Plugins/Runtime/Metasound/Source/MetasoundStandardNodes/Private/MetasoundMixerNode.cpp
+++ b/Engine/Plugins/Runtime/Metasound/Source/MetasoundStandardNodes/Private/MetasoundMixerNode.cpp
@@ -10,6 +10,7 @@
 #include "MetasoundNodeRegistrationMacro.h"
 #include "MetasoundOperatorSettings.h"
 #include "MetasoundPrimitives.h"
+#include "MetasoundStandardNodesNames.h"
 #include "MetasoundTrigger.h"
 #include "MetasoundVertex.h"
 
@@ -22,7 +23,7 @@ namespace Metasound
 	{
 	public:
 		static TUniquePtr<IOperator> CreateOperator(const FCreateOperatorParams& InParams, FBuildErrorArray& OutErrors);
-		static const FNodeInfo& GetNodeInfo();
+		static const FNodeClassMetadata& GetNodeInfo();
 		static FVertexInterface DeclareVertexInterface();	
 
 		FMixerOperator(const FOperatorSettings& InSettings, TArray<FAudioBufferReadRef>&& InBuffers, TArray<FGainReadRef>&& InGains);
@@ -116,14 +117,15 @@ namespace Metasound
 		return Interface;
 	}
 
-	const FNodeInfo& FMixerOperator::GetNodeInfo()
+	const FNodeClassMetadata& FMixerOperator::GetNodeInfo()
 	{
-		auto InitNodeInfo = []() -> FNodeInfo
+		auto InitNodeInfo = []() -> FNodeClassMetadata
 		{
-			FNodeInfo Info;
-			Info.ClassName = FName(TEXT("Mixer"));
+			FNodeClassMetadata Info;
+			Info.ClassName = {Metasound::StandardNodes::Namespace, TEXT("Mixer"), Metasound::StandardNodes::AudioVariant};
 			Info.MajorVersion = 1;
 			Info.MinorVersion = 0;
+			Info.DisplayName = LOCTEXT("Metasound_MixerNodeDisplayName", "Mixer");
 			Info.Description = LOCTEXT("Metasound_MixerNodeDescription", "Mixes 1 or more input signals together");
 			Info.Author = PluginAuthor;
 			Info.PromptIfMissing = PluginNodeMissingPrompt;
@@ -132,19 +134,19 @@ namespace Metasound
 			return Info;
 		};
 
-		static const FNodeInfo Info = InitNodeInfo();
+		static const FNodeClassMetadata Info = InitNodeInfo();
 
 		return Info;
 	}
 
-	FMixerNode::FMixerNode(const FString& InName, float InDefaultMixGainLinear)
-	:	FNodeFacade(InName, TFacadeOperatorClass<FMixerOperator>())
-		, DefaultMixGainLinear(InDefaultMixGainLinear)
+	FMixerNode::FMixerNode(const FString& InName, const FGuid& InInstanceID, float InDefaultMixGainLinear)
+	: FNodeFacade(InName, InInstanceID, TFacadeOperatorClass<FMixerOperator>())
+	, DefaultMixGainLinear(InDefaultMixGainLinear)
 	{
 	}
 
 	FMixerNode::FMixerNode(const FNodeInitData& InInitData)
-		: FMixerNode(InInitData.InstanceName, 1.0f)
+	: FMixerNode(InInitData.InstanceName, InInitData.InstanceID, 1.0f)
 	{
 	}
 
diff --git a/Engine/Plugins/Runtime/Metasound/Source/MetasoundStandardNodes/Private/MetasoundOscNode.cpp b/Engine/Plugins/Runtime/Metasound/Source/MetasoundStandardNodes/Private/MetasoundOscNode.cpp
--- a/Engine/Plugins/Runtime/Metasound/Source/MetasoundStandardNodes/Private/MetasoundOscNode.cpp
+++ b/Engine/Plugins/Runtime/Metasound/Source/MetasoundStandardNodes/Private/MetasoundOscNode.cpp
@@ -8,6 +8,7 @@
 #include "MetasoundFrequency.h"
 #include "MetasoundNodeRegistrationMacro.h"
 #include "MetasoundPrimitives.h"
+#include "MetasoundStandardNodesNames.h"
 #include "MetasoundTrigger.h"
 #include "MetasoundVertex.h"
 
@@ -21,7 +22,7 @@ namespace Metasound
 	public:
 
 		static TUniquePtr<IOperator> CreateOperator(const FCreateOperatorParams& InParams, FBuildErrorArray& OutErrors);
-		static const FNodeInfo& GetNodeInfo();
+		static const FNodeClassMetadata& GetNodeInfo();
 		static FVertexInterface DeclareVertexInterface();
 
 		FOscOperator(const FOperatorSettings& InSettings, const FFrequencyReadRef& InFrequency, const FBoolReadRef& InActivate);
@@ -134,14 +135,15 @@ namespace Metasound
 	}
 
 
-	const FNodeInfo& FOscOperator::GetNodeInfo()
+	const FNodeClassMetadata& FOscOperator::GetNodeInfo()
 	{
-		auto InitNodeInfo = []() -> FNodeInfo
+		auto InitNodeInfo = []() -> FNodeClassMetadata
 		{
-			FNodeInfo Info;
-			Info.ClassName = FName(TEXT("Osc"));
+			FNodeClassMetadata Info;
+			Info.ClassName = {Metasound::StandardNodes::Namespace, TEXT("Osc"), Metasound::StandardNodes::AudioVariant};
 			Info.MajorVersion = 1;
 			Info.MinorVersion = 0;
+			Info.DisplayName = LOCTEXT("Metasound_OscNodeDisplayName", "Oscillator");
 			Info.Description = LOCTEXT("Metasound_OscNodeDescription", "Emits an audio signal of a sinusoid.");
 			Info.Author = PluginAuthor;
 			Info.PromptIfMissing = PluginNodeMissingPrompt;
@@ -150,20 +152,20 @@ namespace Metasound
 			return Info;
 		};
 
-		static const FNodeInfo Info = InitNodeInfo();
+		static const FNodeClassMetadata Info = InitNodeInfo();
 
 		return Info;
 	}
 
-	FOscNode::FOscNode(const FString& InName, float InDefaultFrequency, bool bInDefaultEnablement)
-	:	FNodeFacade(InName, TFacadeOperatorClass<FOscOperator>())
+	FOscNode::FOscNode(const FString& InInstanceName, const FGuid& InInstanceID, float InDefaultFrequency, bool bInDefaultEnablement)
+	:	FNodeFacade(InInstanceName, InInstanceID, TFacadeOperatorClass<FOscOperator>())
 	,	DefaultFrequency(InDefaultFrequency)
 	,	bDefaultEnablement(bInDefaultEnablement)
 	{
 	}
 
 	FOscNode::FOscNode(const FNodeInitData& InInitData)
-		: FOscNode(InInitData.InstanceName, 440.0f, true)
+		: FOscNode(InInitData.InstanceName, InInitData.InstanceID, 440.0f, true)
 	{
 	}
 
diff --git a/Engine/Plugins/Runtime/Metasound/Source/MetasoundStandardNodes/Private/MetasoundPeriodicTriggerNode.cpp b/Engine/Plugins/Runtime/Metasound/Source/MetasoundStandardNodes/Private/MetasoundPeriodicTriggerNode.cpp
--- a/Engine/Plugins/Runtime/Metasound/Source/MetasoundStandardNodes/Private/MetasoundPeriodicTriggerNode.cpp
+++ b/Engine/Plugins/Runtime/Metasound/Source/MetasoundStandardNodes/Private/MetasoundPeriodicTriggerNode.cpp
@@ -13,6 +13,7 @@
 #include "MetasoundNodeRegistrationMacro.h"
 #include "MetasoundOperatorInterface.h"
 #include "MetasoundPrimitives.h"
+#include "MetasoundStandardNodesNames.h"
 
 #define LOCTEXT_NAMESPACE "MetasoundStandardNodes"
 
@@ -23,7 +24,7 @@ namespace Metasound
 	class FPeriodicTriggerOperator : public TExecutableOperator<FPeriodicTriggerOperator>
 	{
 		public:
-			static const FNodeInfo& GetNodeInfo();
+			static const FNodeClassMetadata& GetNodeInfo();
 			static FVertexInterface DeclareVertexInterface();
 			static TUniquePtr<IOperator> CreateOperator(const FCreateOperatorParams& InParams, FBuildErrorArray& OutErrors);
 
@@ -152,14 +153,15 @@ namespace Metasound
 		return Interface;
 	}
 
-	const FNodeInfo& FPeriodicTriggerOperator::GetNodeInfo()
+	const FNodeClassMetadata& FPeriodicTriggerOperator::GetNodeInfo()
 	{
-		auto InitNodeInfo = []() -> FNodeInfo
+		auto InitNodeInfo = []() -> FNodeClassMetadata
 		{
-			FNodeInfo Info;
-			Info.ClassName = "Periodic Trigger";
+			FNodeClassMetadata Info;
+			Info.ClassName = {Metasound::StandardNodes::Namespace, TEXT("PeriodicTrigger"), TEXT("")};
 			Info.MajorVersion = 1;
 			Info.MinorVersion = 0;
+			Info.DisplayName = LOCTEXT("Metasound_PeriodicTriggerNodeDisplayName", "Periodic Trigger");
 			Info.Description = LOCTEXT("Metasound_PeriodicTriggerNodeDescription", "Emits a trigger periodically based on the period duration given.");
 			Info.Author = PluginAuthor;
 			Info.PromptIfMissing = PluginNodeMissingPrompt;
@@ -168,20 +170,20 @@ namespace Metasound
 			return Info;
 		};
 
-		static const FNodeInfo Info = InitNodeInfo();
+		static const FNodeClassMetadata Info = InitNodeInfo();
 
 		return Info;
 	}
 
 
-	FPeriodicTriggerNode::FPeriodicTriggerNode(const FString& InName, float InDefaultPeriodInSeconds)
-	:	FNodeFacade(InName, TFacadeOperatorClass<FPeriodicTriggerOperator>())
+	FPeriodicTriggerNode::FPeriodicTriggerNode(const FString& InInstanceName, const FGuid& InInstanceID, float InDefaultPeriodInSeconds)
+	:	FNodeFacade(InInstanceName, InInstanceID, TFacadeOperatorClass<FPeriodicTriggerOperator>())
 	,	DefaultPeriod(InDefaultPeriodInSeconds)
 	{
 	}
 
 	FPeriodicTriggerNode::FPeriodicTriggerNode(const FNodeInitData& InInitData)
-		: FPeriodicTriggerNode(InInitData.InstanceName, 1.0f)
+		: FPeriodicTriggerNode(InInitData.InstanceName, InInitData.InstanceID, 1.0f)
 	{}
 
 	float FPeriodicTriggerNode::GetDefaultPeriodInSeconds() const
diff --git a/Engine/Plugins/Runtime/Metasound/Source/MetasoundStandardNodes/Private/MetasoundStandardNodesNames.cpp b/Engine/Plugins/Runtime/Metasound/Source/MetasoundStandardNodes/Private/MetasoundStandardNodesNames.cpp
new file mode 100644
--- /dev/null
+++ b/Engine/Plugins/Runtime/Metasound/Source/MetasoundStandardNodes/Private/MetasoundStandardNodesNames.cpp
@@ -0,0 +1,13 @@
+// Copyright Epic Games, Inc. All Rights Reserved.
+
+#include "MetasoundStandardNodesNames.h"
+
+
+namespace Metasound
+{
+	namespace StandardNodes
+	{
+		const FName Namespace = FName(TEXT("UE"));
+		const FName AudioVariant = FName(TEXT("Audio"));
+	}
+}
diff --git a/Engine/Plugins/Runtime/Metasound/Source/MetasoundStandardNodes/Private/MetasoundSwitchGateNode.cpp b/Engine/Plugins/Runtime/Metasound/Source/MetasoundStandardNodes/Private/MetasoundSwitchGateNode.cpp
--- a/Engine/Plugins/Runtime/Metasound/Source/MetasoundStandardNodes/Private/MetasoundSwitchGateNode.cpp
+++ b/Engine/Plugins/Runtime/Metasound/Source/MetasoundStandardNodes/Private/MetasoundSwitchGateNode.cpp
@@ -8,6 +8,7 @@
 #include "MetasoundFrequency.h"
 #include "MetasoundNodeRegistrationMacro.h"
 #include "MetasoundPrimitives.h"
+#include "MetasoundStandardNodesNames.h"
 #include "MetasoundVertex.h"
 
 #define LOCTEXT_NAMESPACE "MetasoundStandardNodes"
@@ -40,7 +41,7 @@ namespace Metasound
 	{
 		public:
 			static TUniquePtr<IOperator> CreateOperator(const FCreateOperatorParams& InParams, FBuildErrorArray& OutErrors);
-			static const FNodeInfo& GetNodeInfo();
+			static const FNodeClassMetadata& GetNodeInfo();
 			static FVertexInterface DeclareVertexInterface();
 
 			FSwitchOperator(const FOperatorSettings& InSettings, FInt32ReadRef InSelector, FBoolReadRef InShouldCrossfade, FAudioBufferReadRef InBuffers[PREVIEW_INPUT_COUNT])
@@ -191,14 +192,15 @@ namespace Metasound
 		return CachedInterface;
 	}
 
-	const FNodeInfo& FSwitchOperator::GetNodeInfo()
+	const FNodeClassMetadata& FSwitchOperator::GetNodeInfo()
 	{
-		auto InitNodeInfo = []() -> FNodeInfo
+		auto InitNodeInfo = []() -> FNodeClassMetadata
 		{
-			FNodeInfo Info;
-			Info.ClassName = FName(TEXT("Switch"));
+			FNodeClassMetadata Info;
+			Info.ClassName = {Metasound::StandardNodes::Namespace, TEXT("Switch"), Metasound::StandardNodes::AudioVariant};
 			Info.MajorVersion = 1;
 			Info.MinorVersion = 0;
+			Info.DisplayName = LOCTEXT("Metasound_SwitchNodeDisplayName", "Switch");
 			Info.Description = LOCTEXT("Metasound_SwitchNodeDescription", "Selects one audio buffer from any number of inputs.");
 			Info.Author = PluginAuthor;
 			Info.PromptIfMissing = PluginNodeMissingPrompt;
@@ -207,7 +209,7 @@ namespace Metasound
 			return Info;
 		};
 
-		static const FNodeInfo Info = InitNodeInfo();
+		static const FNodeClassMetadata Info = InitNodeInfo();
 		return Info;
 	}
 
@@ -232,13 +234,13 @@ namespace Metasound
 		return MakeUnique<FSwitchOperator>(InParams.OperatorSettings, Selector, Crossfade, PreviewBuffers);
 	}
 
-	FSwitchNode::FSwitchNode(const FString& InName)
-		: FNodeFacade(InName, TFacadeOperatorClass<FSwitchOperator>())
+	FSwitchNode::FSwitchNode(const FString& InName, const FGuid& InInstanceID)
+		: FNodeFacade(InName, InInstanceID, TFacadeOperatorClass<FSwitchOperator>())
 	{
 	}
 
 	FSwitchNode::FSwitchNode(const FNodeInitData& InInitData)
-		: FSwitchNode(InInitData.InstanceName)
+		: FSwitchNode(InInitData.InstanceName, InInitData.InstanceID)
 	{
 	}
 
@@ -248,7 +250,7 @@ namespace Metasound
 	{
 	public:
 		static TUniquePtr<IOperator> CreateOperator(const FCreateOperatorParams& InParams, FBuildErrorArray& OutErrors);
-		static const FNodeInfo& GetNodeInfo();
+		static const FNodeClassMetadata& GetNodeInfo();
 		static FVertexInterface DeclareVertexInterface();
 
 		FGateOperator(const FOperatorSettings& InSettings, FInt32ReadRef InSelector, FBoolReadRef InShouldCrossfade, FAudioBufferReadRef InBuffer)
@@ -409,14 +411,15 @@ namespace Metasound
 		return CachedInterface;
 	}
 
-	const FNodeInfo& FGateOperator::GetNodeInfo()
+	const FNodeClassMetadata& FGateOperator::GetNodeInfo()
 	{
-		auto InitNodeInfo = []() -> FNodeInfo
+		auto InitNodeInfo = []() -> FNodeClassMetadata
 		{
-			FNodeInfo Info;
-			Info.ClassName = FName(TEXT("Gate"));
+			FNodeClassMetadata Info;
+			Info.ClassName = {Metasound::StandardNodes::Namespace, TEXT("Gate"), Metasound::StandardNodes::AudioVariant};
 			Info.MajorVersion = 1;
 			Info.MinorVersion = 0;
+			Info.DisplayName = LOCTEXT("Metasound_GateNodeDisplayName", "Gate");
 			Info.Description = LOCTEXT("Metasound_GateNodeDescription", "Routes an audio buffer to one of a set of output buffers based on an input value.");
 			Info.Author = PluginAuthor;
 			Info.PromptIfMissing = PluginNodeMissingPrompt;
@@ -425,7 +428,7 @@ namespace Metasound
 			return Info;
 		};
 
-		static const FNodeInfo Info = InitNodeInfo();
+		static const FNodeClassMetadata Info = InitNodeInfo();
 		return Info;
 	}
 
@@ -444,13 +447,13 @@ namespace Metasound
 		return MakeUnique<FGateOperator>(InParams.OperatorSettings, Selector, Crossfade, Input);
 	}
 
-	FGateNode::FGateNode(const FString& InName)
-		: FNodeFacade(InName, TFacadeOperatorClass<FGateOperator>())
+	FGateNode::FGateNode(const FString& InName, const FGuid& InInstanceID)
+		: FNodeFacade(InName, InInstanceID, TFacadeOperatorClass<FGateOperator>())
 	{
 	}
 
 	FGateNode::FGateNode(const FNodeInitData& InInitData)
-		: FGateNode(InInitData.InstanceName)
+		: FGateNode(InInitData.InstanceName, InInitData.InstanceID)
 	{
 	}
 
diff --git a/Engine/Plugins/Runtime/Metasound/Source/MetasoundStandardNodes/Private/MetasoundTriggerAccumulatorNode.cpp b/Engine/Plugins/Runtime/Metasound/Source/MetasoundStandardNodes/Private/MetasoundTriggerAccumulatorNode.cpp
--- a/Engine/Plugins/Runtime/Metasound/Source/MetasoundStandardNodes/Private/MetasoundTriggerAccumulatorNode.cpp
+++ b/Engine/Plugins/Runtime/Metasound/Source/MetasoundStandardNodes/Private/MetasoundTriggerAccumulatorNode.cpp
@@ -6,6 +6,7 @@
 #include "MetasoundNodeRegistrationMacro.h"
 #include "MetasoundPrimitives.h"
 #include "MetasoundTrigger.h"
+#include "MetasoundStandardNodesNames.h"
 #include "MetasoundVertex.h"
 
 
@@ -18,7 +19,7 @@ namespace Metasound
 	class FTriggerAccumulatorOperator : public TExecutableOperator<FTriggerAccumulatorOperator>
 	{
 		public:
-			static const FNodeInfo& GetNodeInfo();
+			static const FNodeClassMetadata& GetNodeInfo();
 			static FVertexInterface DeclareVertexInterface();
 			static TUniquePtr<IOperator> CreateOperator(const FCreateOperatorParams& InParams, FBuildErrorArray& OutErrors);
 
@@ -183,14 +184,15 @@ namespace Metasound
 		return Interface;
 	}
 
-	const FNodeInfo& FTriggerAccumulatorOperator::GetNodeInfo()
+	const FNodeClassMetadata& FTriggerAccumulatorOperator::GetNodeInfo()
 	{
-		auto InitNodeInfo = []() -> FNodeInfo
+		auto InitNodeInfo = []() -> FNodeClassMetadata
 		{
-			FNodeInfo Info;
-			Info.ClassName = FName(TEXT("Trigger Accumulator"));
+			FNodeClassMetadata Info;
+			Info.ClassName = {Metasound::StandardNodes::Namespace, TEXT("TriggerAccumulator"), TEXT("")};
 			Info.MajorVersion = 1;
 			Info.MinorVersion = 0;
+			Info.DisplayName = LOCTEXT("Metasound_TriggerAccumulatorNodeDisplayName", "Trigger Accumulator");
 			Info.Description = LOCTEXT("Metasound_TriggerAccumulatorNodeDescription", "Accumulates Triggers until a Count is hit and then it Triggers");
 			Info.Author = PluginAuthor;
 			Info.PromptIfMissing = PluginNodeMissingPrompt;
@@ -199,7 +201,7 @@ namespace Metasound
 			return Info;
 		};
 
-		static const FNodeInfo Info = InitNodeInfo();
+		static const FNodeClassMetadata Info = InitNodeInfo();
 
 		return Info;
 	}
@@ -218,14 +220,14 @@ namespace Metasound
 	}
 
 	// Node.
-	FTriggerAccumulatorNode::FTriggerAccumulatorNode(const FString& InName, int32 InDefaultTriggerAtCount)
-		: FNodeFacade(InName, TFacadeOperatorClass<FTriggerAccumulatorOperator>())
+	FTriggerAccumulatorNode::FTriggerAccumulatorNode(const FString& InName, const FGuid& InInstanceID, int32 InDefaultTriggerAtCount)
+		: FNodeFacade(InName, InInstanceID, TFacadeOperatorClass<FTriggerAccumulatorOperator>())
 		, DefaultTriggerAtCount(InDefaultTriggerAtCount)
 	{}
 
 	FTriggerAccumulatorNode::FTriggerAccumulatorNode(const FNodeInitData& InitData)
-		: FTriggerAccumulatorNode(InitData.InstanceName, 1)
+		: FTriggerAccumulatorNode(InitData.InstanceName, InitData.InstanceID, 1)
 	{}
 }
 
-#undef LOCTEXT_NAMESPACE
\ No newline at end of file
+#undef LOCTEXT_NAMESPACE
diff --git a/Engine/Plugins/Runtime/Metasound/Source/MetasoundStandardNodes/Private/MetasoundTriggerDelayNode.cpp b/Engine/Plugins/Runtime/Metasound/Source/MetasoundStandardNodes/Private/MetasoundTriggerDelayNode.cpp
--- a/Engine/Plugins/Runtime/Metasound/Source/MetasoundStandardNodes/Private/MetasoundTriggerDelayNode.cpp
+++ b/Engine/Plugins/Runtime/Metasound/Source/MetasoundStandardNodes/Private/MetasoundTriggerDelayNode.cpp
@@ -12,6 +12,7 @@
 #include "MetasoundNodeRegistrationMacro.h"
 #include "MetasoundOperatorInterface.h"
 #include "MetasoundPrimitives.h"
+#include "MetasoundStandardNodesNames.h"
 
 #define LOCTEXT_NAMESPACE "MetasoundStandardNodes"
 
@@ -22,7 +23,7 @@ namespace Metasound
 	class FTriggerDelayOperator : public TExecutableOperator<FTriggerDelayOperator>
 	{
 		public:
-			static const FNodeInfo& GetNodeInfo();
+			static const FNodeClassMetadata& GetNodeInfo();
 			static FVertexInterface DeclareVertexInterface();
 			static TUniquePtr<IOperator> CreateOperator(const FCreateOperatorParams& InParams, FBuildErrorArray& OutErrors);
 
@@ -143,14 +144,15 @@ namespace Metasound
 		return Interface;
 	}
 
-	const FNodeInfo& FTriggerDelayOperator::GetNodeInfo()
+	const FNodeClassMetadata& FTriggerDelayOperator::GetNodeInfo()
 	{
-		auto InitNodeInfo = []() -> FNodeInfo
+		auto InitNodeInfo = []() -> FNodeClassMetadata
 		{
-			FNodeInfo Info;
-			Info.ClassName = "Trigger Delay";
+			FNodeClassMetadata Info;
+			Info.ClassName = {Metasound::StandardNodes::Namespace, TEXT("Trigger Delay"), TEXT("")};
 			Info.MajorVersion = 1;
 			Info.MinorVersion = 0;
+			Info.DisplayName = LOCTEXT("Metasound_TriggerDelayNodeDisplayName", "Trigger Delay");
 			Info.Description = LOCTEXT("Metasound_TriggerDelayNodeDescription", "Executes output trigger after the given delay time from the most recent execution of the input trigger .");
 			Info.Author = PluginAuthor;
 			Info.PromptIfMissing = PluginNodeMissingPrompt;
@@ -159,7 +161,7 @@ namespace Metasound
 			return Info;
 		};
 
-		static const FNodeInfo Info = InitNodeInfo();
+		static const FNodeClassMetadata Info = InitNodeInfo();
 		return Info;
 	}
 
@@ -168,14 +170,14 @@ namespace Metasound
 		SamplesUntilTrigger = -1;
 	}
 
-	FTriggerDelayNode::FTriggerDelayNode(const FString& InName, float InDefaultDelayInSeconds)
-	:	FNodeFacade(InName, TFacadeOperatorClass<FTriggerDelayOperator>())
+	FTriggerDelayNode::FTriggerDelayNode(const FString& InName, const FGuid& InInstanceID, float InDefaultDelayInSeconds)
+	:	FNodeFacade(InName, InInstanceID, TFacadeOperatorClass<FTriggerDelayOperator>())
 	,	DefaultDelay(InDefaultDelayInSeconds)
 	{
 	}
 
 	FTriggerDelayNode::FTriggerDelayNode(const FNodeInitData& InInitData)
-		: FTriggerDelayNode(InInitData.InstanceName, 0.0f)
+		: FTriggerDelayNode(InInitData.InstanceName, InInitData.InstanceID, 0.0f)
 	{
 	}
 
diff --git a/Engine/Plugins/Runtime/Metasound/Source/MetasoundStandardNodes/Private/MetasoundTriggerPipeNode.cpp b/Engine/Plugins/Runtime/Metasound/Source/MetasoundStandardNodes/Private/MetasoundTriggerPipeNode.cpp
--- a/Engine/Plugins/Runtime/Metasound/Source/MetasoundStandardNodes/Private/MetasoundTriggerPipeNode.cpp
+++ b/Engine/Plugins/Runtime/Metasound/Source/MetasoundStandardNodes/Private/MetasoundTriggerPipeNode.cpp
@@ -11,6 +11,7 @@
 #include "MetasoundNodeRegistrationMacro.h"
 #include "MetasoundOperatorInterface.h"
 #include "MetasoundPrimitives.h"
+#include "MetasoundStandardNodesNames.h"
 #include "MetasoundTime.h"
 #include "MetasoundTrigger.h"
 
@@ -24,7 +25,7 @@ namespace Metasound
 	class FTriggerPipeOperator : public TExecutableOperator<FTriggerPipeOperator>
 	{
 		public:
-			static const FNodeInfo& GetNodeInfo();
+			static const FNodeClassMetadata& GetNodeInfo();
 			static FVertexInterface DeclareVertexInterface();
 			static TUniquePtr<IOperator> CreateOperator(const FCreateOperatorParams& InParams, FBuildErrorArray& OutErrors);
 
@@ -141,14 +142,15 @@ namespace Metasound
 		return Interface;
 	}
 
-	const FNodeInfo& FTriggerPipeOperator::GetNodeInfo()
+	const FNodeClassMetadata& FTriggerPipeOperator::GetNodeInfo()
 	{
-		auto InitNodeInfo = []() -> FNodeInfo
+		auto InitNodeInfo = []() -> FNodeClassMetadata
 		{
-			FNodeInfo Info;
-			Info.ClassName = "Trigger Pipe";
+			FNodeClassMetadata Info;
+			Info.ClassName = FNodeClassName(StandardNodes::Namespace, TEXT("Pipe"), TEXT("Trigger"));
 			Info.MajorVersion = 1;
 			Info.MinorVersion = 0;
+			Info.DisplayName = LOCTEXT("PipeTriggerNode_NodeDisplayName", "Trigger Pipe");
 			Info.Description = LOCTEXT("Metasound_DelayNodeDescription", "Delays execution of the input trigger by the given delay for all input trigger executions.");
 			Info.Author = PluginAuthor;
 			Info.PromptIfMissing = PluginNodeMissingPrompt;
@@ -157,18 +159,18 @@ namespace Metasound
 			return Info;
 		};
 
-		static const FNodeInfo Info = InitNodeInfo();
+		static const FNodeClassMetadata Info = InitNodeInfo();
 		return Info;
 	}
 
-	FTriggerPipeNode::FTriggerPipeNode(const FString& InName, float InDefaultDelayInSeconds)
-	:	FNodeFacade(InName, TFacadeOperatorClass<FTriggerPipeOperator>())
+	FTriggerPipeNode::FTriggerPipeNode(const FString& InName, const FGuid& InInstanceID, float InDefaultDelayInSeconds)
+	:	FNodeFacade(InName, InInstanceID, TFacadeOperatorClass<FTriggerPipeOperator>())
 	,	DefaultDelay(InDefaultDelayInSeconds)
 	{
 	}
 
 	FTriggerPipeNode::FTriggerPipeNode(const FNodeInitData& InInitData)
-		: FTriggerPipeNode(InInitData.InstanceName, 0.0f)
+		: FTriggerPipeNode(InInitData.InstanceName, InInitData.InstanceID, 0.0f)
 	{
 	}
 
diff --git a/Engine/Plugins/Runtime/Metasound/Source/MetasoundStandardNodes/Public/MetasoundADSRNode.h b/Engine/Plugins/Runtime/Metasound/Source/MetasoundStandardNodes/Public/MetasoundADSRNode.h
--- a/Engine/Plugins/Runtime/Metasound/Source/MetasoundStandardNodes/Public/MetasoundADSRNode.h
+++ b/Engine/Plugins/Runtime/Metasound/Source/MetasoundStandardNodes/Public/MetasoundADSRNode.h
@@ -17,13 +17,14 @@ namespace Metasound
 
 			/** FADSR node constructor.
 			 *
-			 * @param InName - Name of this node.
+			 * @param InInstanceName - Name of this node.
+			 * @param InInstanceID - ID of this node.
 			 * @param InDefaultAttackMs - Default attack duration in milliseconds.
 			 * @param InDefaultDecayMs - Default decay duration in milliseconds.
 			 * @param InDefaultSustainLevel - Default sustain level [0.0f, 1.0f].
 			 * @param InDefaultReleaseMs - Default release duration in milliseconds.
 			 */
-			FADSRNode(const FString& InName, float InDefaultAttackMs, float InDefaultDecayMs, float InDefaultSustainLevel, float InDefaultReleaseMs);
+			FADSRNode(const FString& InInstanceName, const FGuid& InInstanceID, float InDefaultAttackMs, float InDefaultDecayMs, float InDefaultSustainLevel, float InDefaultReleaseMs);
 
 			/**
 			 * Constructor used by the Metasound Frontend.
diff --git a/Engine/Plugins/Runtime/Metasound/Source/MetasoundStandardNodes/Public/MetasoundAudioFormatNodes.h b/Engine/Plugins/Runtime/Metasound/Source/MetasoundStandardNodes/Public/MetasoundAudioFormatNodes.h
--- a/Engine/Plugins/Runtime/Metasound/Source/MetasoundStandardNodes/Public/MetasoundAudioFormatNodes.h
+++ b/Engine/Plugins/Runtime/Metasound/Source/MetasoundStandardNodes/Public/MetasoundAudioFormatNodes.h
@@ -11,14 +11,14 @@ namespace Metasound
 	class METASOUNDSTANDARDNODES_API FMonoAudioFormatNode : public FNodeFacade
 	{
 		public:
-			FMonoAudioFormatNode(const FString& InInstanceName);
+			FMonoAudioFormatNode(const FString& InInstanceName, const FGuid& InInstanceID);
 			FMonoAudioFormatNode(const FNodeInitData& InInitData);
 	};
 
 	class METASOUNDSTANDARDNODES_API FStereoAudioFormatNode : public FNodeFacade
 	{
 		public:
-			FStereoAudioFormatNode(const FString& InInstanceName);
+			FStereoAudioFormatNode(const FString& InInstanceName, const FGuid& InInstnaceID);
 			FStereoAudioFormatNode(const FNodeInitData& InInitData);
 	};
 }
diff --git a/Engine/Plugins/Runtime/Metasound/Source/MetasoundStandardNodes/Public/MetasoundBufferAddFloatNode.h b/Engine/Plugins/Runtime/Metasound/Source/MetasoundStandardNodes/Public/MetasoundBufferAddFloatNode.h
--- a/Engine/Plugins/Runtime/Metasound/Source/MetasoundStandardNodes/Public/MetasoundBufferAddFloatNode.h
+++ b/Engine/Plugins/Runtime/Metasound/Source/MetasoundStandardNodes/Public/MetasoundBufferAddFloatNode.h
@@ -12,7 +12,7 @@ namespace Metasound
 	class METASOUNDSTANDARDNODES_API FBufferAddFloatNode : public FNodeFacade
 	{
 		public:
-			FBufferAddFloatNode(const FString& InInstanceName);
+			FBufferAddFloatNode(const FString& InInstanceName, const FGuid& InInstanceID);
 			FBufferAddFloatNode(const FNodeInitData& InInitData);
 	};
 } // namespace Metasound
diff --git a/Engine/Plugins/Runtime/Metasound/Source/MetasoundStandardNodes/Public/MetasoundGainMultiplyNode.h b/Engine/Plugins/Runtime/Metasound/Source/MetasoundStandardNodes/Public/MetasoundGainMultiplyNode.h
--- a/Engine/Plugins/Runtime/Metasound/Source/MetasoundStandardNodes/Public/MetasoundGainMultiplyNode.h
+++ b/Engine/Plugins/Runtime/Metasound/Source/MetasoundStandardNodes/Public/MetasoundGainMultiplyNode.h
@@ -12,7 +12,7 @@ namespace Metasound
 	class METASOUNDSTANDARDNODES_API FGainMultiplyNode : public FNodeFacade
 	{
 		public:
-			FGainMultiplyNode(const FString& InInstanceName);
+			FGainMultiplyNode(const FString& InInstanceName, const FGuid& InInstanceID);
 			FGainMultiplyNode(const FNodeInitData& InInitData);
 	};
 } // namespace Metasound
diff --git a/Engine/Plugins/Runtime/Metasound/Source/MetasoundStandardNodes/Public/MetasoundMathNodes.h b/Engine/Plugins/Runtime/Metasound/Source/MetasoundStandardNodes/Public/MetasoundMathNodes.h
--- a/Engine/Plugins/Runtime/Metasound/Source/MetasoundStandardNodes/Public/MetasoundMathNodes.h
+++ b/Engine/Plugins/Runtime/Metasound/Source/MetasoundStandardNodes/Public/MetasoundMathNodes.h
@@ -15,18 +15,20 @@
 #include "MetasoundNodeInterface.h"
 #include "MetasoundOperatorInterface.h"
 #include "MetasoundPrimitives.h"
+#include "MetasoundStandardNodesNames.h"
 #include "MetasoundTime.h"
 #include "MetasoundTrigger.h"
 
 
 #define LOCTEXT_NAMESPACE "MetasoundMathOpNode"
 
-#define DEFINE_METASOUND_MATHOP(OpName, DataTypeName, DataClass, Description) \
+#define DEFINE_METASOUND_MATHOP(OpName, DataTypeName, DataClass, DisplayName, Description) \
 	class F##OpName##DataTypeName##Node \
 		: public ::Metasound::TMathOpNode<F##OpName##DataTypeName##Node, ::Metasound::TMathOp##OpName<DataClass>, DataClass> \
 	{ \
 	public: \
-		static FName GetClassName() { return #OpName " " #DataTypeName; } \
+		static FNodeClassName GetClassName() { return {::Metasound::StandardNodes::Namespace, TEXT(#OpName), TEXT(#DataTypeName)}; } \
+		static FText GetDisplayName() { return DisplayName; } \
 		static FText GetDescription() { return Description; } \
 		static FName GetImageName() { return "MetasoundEditor.Graph.Node.Math." #OpName; } \
 		F##OpName##DataTypeName##Node(const FNodeInitData& InInitData) : TMathOpNode(InInitData) { } \
@@ -34,6 +36,44 @@
 
 namespace Metasound
 {
+	namespace MathOpDisplayNames
+	{
+		template<typename DataType>
+		const FText GetAddDisplayName()
+		{
+			static const FText DisplayName = FText::Format(LOCTEXT("AddNodeDisplayNamePattern", "Add {0}"), FText::FromString(GetMetasoundDataTypeString<DataType>()));
+			return DisplayName;
+		}
+
+		template<typename DataType>
+		const FText GetSubtractDisplayName()
+		{
+			static const FText DisplayName = FText::Format(LOCTEXT("SubtractNodeDisplayNamePattern", "Subtract {0}"), FText::FromString(GetMetasoundDataTypeString<DataType>()));
+			return DisplayName;
+		}
+
+		template<typename DataType>
+		const FText GetMultiplyDisplayName()
+		{
+			static const FText DisplayName = FText::Format(LOCTEXT("MultiplyNodeDisplayNamePattern", "Multiply {0}"), FText::FromString(GetMetasoundDataTypeString<DataType>()));
+			return DisplayName;
+		}
+
+		template<typename DataType>
+		const FText GetDivideDisplayName()
+		{
+			static const FText DisplayName = FText::Format(LOCTEXT("DivideNodeDisplayNamePattern", "Divide {0}"), FText::FromString(GetMetasoundDataTypeString<DataType>()));
+			return DisplayName;
+		}
+
+		template<typename DataType>
+		const FText GetRandRangeDisplayName()
+		{
+			static const FText DisplayName = FText::Format(LOCTEXT("RandRangeNodeDisplayNamePattern", "RandRange {0}"), FText::FromString(GetMetasoundDataTypeString<DataType>()));
+			return DisplayName;
+		}
+	}
+
 	template <typename TDerivedClass, typename TMathOpClass, typename TDataClass>
 	class TMathOpNode : public FNodeFacade
 	{
@@ -63,9 +103,9 @@ namespace Metasound
 				return Interface;
 			}
 
-			static const FNodeInfo& GetNodeInfo()
+			static const FNodeClassMetadata& GetNodeInfo()
 			{
-				auto InitNodeInfo = []() -> FNodeInfo
+				auto InitNodeInfo = []() -> FNodeClassMetadata
 				{
 					FNodeDisplayStyle DisplayStyle;
 					DisplayStyle.ImageName = TDerivedClass::GetImageName();
@@ -73,10 +113,11 @@ namespace Metasound
 					DisplayStyle.bShowInputNames = false;
 					DisplayStyle.bShowOutputNames = false;
 
-					FNodeInfo Info;
+					FNodeClassMetadata Info;
 					Info.ClassName = TDerivedClass::GetClassName();
 					Info.MajorVersion = 1;
 					Info.MinorVersion = 0;
+					Info.DisplayName = TDerivedClass::GetDisplayName();
 					Info.Description = TDerivedClass::GetDescription();
 					Info.Author = PluginAuthor;
 					Info.PromptIfMissing = PluginNodeMissingPrompt;
@@ -87,7 +128,7 @@ namespace Metasound
 					return Info;
 				};
 
-				static const FNodeInfo Info = InitNodeInfo();
+				static const FNodeClassMetadata Info = InitNodeInfo();
 				return Info;
 			}
 
@@ -134,7 +175,7 @@ namespace Metasound
 
 	public:
 		TMathOpNode(const FNodeInitData& InInitData)
-			: FNodeFacade(InInitData.InstanceName, TFacadeOperatorClass<TMathOperator>())
+			: FNodeFacade(InInitData.InstanceName, InInitData.InstanceID, TFacadeOperatorClass<TMathOperator>())
 		{
 		}
 	};
@@ -555,39 +596,39 @@ namespace Metasound
 		}
 	};
 
-// Definitions
-	DEFINE_METASOUND_MATHOP(Add, Float, float, LOCTEXT("Metasound_MathAddFloatNodeDescription", "Adds floats."))
-	DEFINE_METASOUND_MATHOP(Add, Int32, int32, LOCTEXT("Metasound_MathAddInt32NodeDescription", "Adds int32s."))
-	DEFINE_METASOUND_MATHOP(Add, Int64, int64, LOCTEXT("Metasound_MathAddInt64NodeDescription", "Adds int64s."))
-	DEFINE_METASOUND_MATHOP(Add, Frequency, FFrequency, LOCTEXT("Metasound_MathAddFrequencyNodeDescription", "Adds frequency values."))
-	DEFINE_METASOUND_MATHOP(Add, FloatTime, FFloatTime, LOCTEXT("Metasound_MathAddFloatTimeNodeDescription", "Adds time values."))
-	DEFINE_METASOUND_MATHOP(Add, AudioBuffer, FAudioBuffer, LOCTEXT("Metasound_MathAddBufferNodeDescription", "Adds buffers together by sample."))
-
-	DEFINE_METASOUND_MATHOP(Subtract, Float, float, LOCTEXT("Metasound_MathSubractFloatNodeDescription", "Subtracts floats."))
-	DEFINE_METASOUND_MATHOP(Subtract, Int32, int32, LOCTEXT("Metasound_MathSubractInt32NodeDescription", "Subtracts int32s."))
-	DEFINE_METASOUND_MATHOP(Subtract, Int64, int64, LOCTEXT("Metasound_MathSubractInt64NodeDescription", "Subtracts int64s."))
-	DEFINE_METASOUND_MATHOP(Subtract, Frequency, FFrequency, LOCTEXT("Metasound_MathSubractFrequencyNodeDescription", "Subtracts frequency values."))
-	DEFINE_METASOUND_MATHOP(Subtract, FloatTime, FFloatTime, LOCTEXT("Metasound_MathSubractFloatTimeNodeDescription", "Subtracts time values."))
-	DEFINE_METASOUND_MATHOP(Subtract, AudioBuffer, FAudioBuffer, LOCTEXT("Metasound_MathSubtractBufferNodeDescription", "Subtracts buffers by sample."))
-
-	DEFINE_METASOUND_MATHOP(Multiply, Float, float, LOCTEXT("Metasound_MathMultiplyFloatNodeDescription", "Multiplies floats."))
-	DEFINE_METASOUND_MATHOP(Multiply, Int32, int32, LOCTEXT("Metasound_MathMultiplyInt32NodeDescription", "Multiplies int32s."))
-	DEFINE_METASOUND_MATHOP(Multiply, Int64, int64, LOCTEXT("Metasound_MathMultiplyInt64NodeDescription", "Multiplies int64s."))
-	DEFINE_METASOUND_MATHOP(Multiply, Frequency, FFrequency, LOCTEXT("Metasound_MathMultiplyFrequencyNodeDescription", "Multiplies frequency values."))
-	DEFINE_METASOUND_MATHOP(Multiply, FloatTime, FFloatTime, LOCTEXT("Metasound_MathMultiplyFloatTimeNodeDescription", "Multiplies time values."))
-	DEFINE_METASOUND_MATHOP(Multiply, AudioBuffer, FAudioBuffer, LOCTEXT("Metasound_MathMultiplyBufferNodeDescription", "Multiplies buffers together by sample."))
-
-	DEFINE_METASOUND_MATHOP(Divide, Float, float, LOCTEXT("Metasound_MathDivideFloatNodeDescription", "Divide float by another float."))
-	DEFINE_METASOUND_MATHOP(Divide, Int32, int32, LOCTEXT("Metasound_MathDivideInt32NodeDescription", "Divide int32 by another int32."))
-	DEFINE_METASOUND_MATHOP(Divide, Int64, int64, LOCTEXT("Metasound_MathDivideInt64NodeDescription", "Divide int64 by another int64."))
-	DEFINE_METASOUND_MATHOP(Divide, Frequency, FFrequency, LOCTEXT("Metasound_MathDivideFrequencyNodeDescription", "Divide frequency by another frequency."))
-	DEFINE_METASOUND_MATHOP(Divide, FloatTime, FFloatTime, LOCTEXT("Metasound_MathDivideFloatTimeNodeDescription", "Divide time by another time."))
-
-	DEFINE_METASOUND_MATHOP(RandRange, Float, float, LOCTEXT("Metasound_MathRandRangeFloatNodeDescription", "Computes random float value in the range [Min, Max]."))
-	DEFINE_METASOUND_MATHOP(RandRange, Int32, int32, LOCTEXT("Metasound_MathRandRangeInt32NodeDescription", "Computes random int32 in the range [Min, Max]."))
-	DEFINE_METASOUND_MATHOP(RandRange, Int64, int64, LOCTEXT("Metasound_MathRandRangeInt64NodeDescription", "Computes random int64 in the range [Min, Max]."))
-	DEFINE_METASOUND_MATHOP(RandRange, Frequency, FFrequency, LOCTEXT("Metasound_MathRandRangeFrequencyNodeDescription", "Computes random frequency value in the range [Min, Max]."))
-	DEFINE_METASOUND_MATHOP(RandRange, FloatTime, FFloatTime, LOCTEXT("Metasound_MathRandRangeFloatTimeNodeDescription", "Computes random time value in the range [Min, Max]."))
+	// Definitions
+	DEFINE_METASOUND_MATHOP(Add, Float, float, MathOpDisplayNames::GetAddDisplayName<float>(), LOCTEXT("Metasound_MathAddFloatNodeDescription", "Adds floats."))
+	DEFINE_METASOUND_MATHOP(Add, Int32, int32, MathOpDisplayNames::GetAddDisplayName<int32>(), LOCTEXT("Metasound_MathAddInt32NodeDescription", "Adds int32s."))
+	DEFINE_METASOUND_MATHOP(Add, Int64, int64, MathOpDisplayNames::GetAddDisplayName<int64>(), LOCTEXT("Metasound_MathAddInt64NodeDescription", "Adds int64s."))
+	DEFINE_METASOUND_MATHOP(Add, Frequency, FFrequency, MathOpDisplayNames::GetAddDisplayName<FFrequency>(), LOCTEXT("Metasound_MathAddFrequencyNodeDescription", "Adds frequency values."))
+	DEFINE_METASOUND_MATHOP(Add, FloatTime, FFloatTime, MathOpDisplayNames::GetAddDisplayName<FFloatTime>(), LOCTEXT("Metasound_MathAddFloatTimeNodeDescription", "Adds time values."))
+	DEFINE_METASOUND_MATHOP(Add, AudioBuffer, FAudioBuffer, MathOpDisplayNames::GetAddDisplayName<FAudioBuffer>(), LOCTEXT("Metasound_MathAddBufferNodeDescription", "Adds buffers together by sample."))
+
+	DEFINE_METASOUND_MATHOP(Subtract, Float, float, MathOpDisplayNames::GetSubtractDisplayName<float>(), LOCTEXT("Metasound_MathSubractFloatNodeDescription", "Subtracts floats."))
+	DEFINE_METASOUND_MATHOP(Subtract, Int32, int32, MathOpDisplayNames::GetSubtractDisplayName<int32>(), LOCTEXT("Metasound_MathSubractInt32NodeDescription", "Subtracts int32s."))
+	DEFINE_METASOUND_MATHOP(Subtract, Int64, int64, MathOpDisplayNames::GetSubtractDisplayName<int64>(), LOCTEXT("Metasound_MathSubractInt64NodeDescription", "Subtracts int64s."))
+	DEFINE_METASOUND_MATHOP(Subtract, Frequency, FFrequency, MathOpDisplayNames::GetSubtractDisplayName<FFrequency>(), LOCTEXT("Metasound_MathSubractFrequencyNodeDescription", "Subtracts frequency values."))
+	DEFINE_METASOUND_MATHOP(Subtract, FloatTime, FFloatTime, MathOpDisplayNames::GetSubtractDisplayName<FFloatTime>(), LOCTEXT("Metasound_MathSubractFloatTimeNodeDescription", "Subtracts time values."))
+	DEFINE_METASOUND_MATHOP(Subtract, AudioBuffer, FAudioBuffer, MathOpDisplayNames::GetSubtractDisplayName<FAudioBuffer>(), LOCTEXT("Metasound_MathSubtractBufferNodeDescription", "Subtracts buffers by sample."))
+
+	DEFINE_METASOUND_MATHOP(Multiply, Float, float, MathOpDisplayNames::GetMultiplyDisplayName<float>(), LOCTEXT("Metasound_MathMultiplyFloatNodeDescription", "Multiplies floats."))
+	DEFINE_METASOUND_MATHOP(Multiply, Int32, int32, MathOpDisplayNames::GetMultiplyDisplayName<int32>(), LOCTEXT("Metasound_MathMultiplyInt32NodeDescription", "Multiplies int32s."))
+	DEFINE_METASOUND_MATHOP(Multiply, Int64, int64, MathOpDisplayNames::GetMultiplyDisplayName<int64>(), LOCTEXT("Metasound_MathMultiplyInt64NodeDescription", "Multiplies int64s."))
+	DEFINE_METASOUND_MATHOP(Multiply, Frequency, FFrequency, MathOpDisplayNames::GetMultiplyDisplayName<FFrequency>(), LOCTEXT("Metasound_MathMultiplyFrequencyNodeDescription", "Multiplies frequency values."))
+	DEFINE_METASOUND_MATHOP(Multiply, FloatTime, FFloatTime, MathOpDisplayNames::GetMultiplyDisplayName<FFloatTime>(), LOCTEXT("Metasound_MathMultiplyFloatTimeNodeDescription", "Multiplies time values."))
+	DEFINE_METASOUND_MATHOP(Multiply, AudioBuffer, FAudioBuffer, MathOpDisplayNames::GetMultiplyDisplayName<FAudioBuffer>(), LOCTEXT("Metasound_MathMultiplyBufferNodeDescription", "Multiplies buffers together by sample."))
+
+	DEFINE_METASOUND_MATHOP(Divide, Float, float, MathOpDisplayNames::GetDivideDisplayName<float>(), LOCTEXT("Metasound_MathDivideFloatNodeDescription", "Divide float by another float."))
+	DEFINE_METASOUND_MATHOP(Divide, Int32, int32, MathOpDisplayNames::GetDivideDisplayName<int32>(), LOCTEXT("Metasound_MathDivideInt32NodeDescription", "Divide int32 by another int32."))
+	DEFINE_METASOUND_MATHOP(Divide, Int64, int64, MathOpDisplayNames::GetDivideDisplayName<int64>(), LOCTEXT("Metasound_MathDivideInt64NodeDescription", "Divide int64 by another int64."))
+	DEFINE_METASOUND_MATHOP(Divide, Frequency, FFrequency, MathOpDisplayNames::GetDivideDisplayName<FFrequency>(), LOCTEXT("Metasound_MathDivideFrequencyNodeDescription", "Divide frequency by another frequency."))
+	DEFINE_METASOUND_MATHOP(Divide, FloatTime, FFloatTime, MathOpDisplayNames::GetDivideDisplayName<FFloatTime>(), LOCTEXT("Metasound_MathDivideFloatTimeNodeDescription", "Divide time by another time."))
+
+	DEFINE_METASOUND_MATHOP(RandRange, Float, float, MathOpDisplayNames::GetRandRangeDisplayName<float>(), LOCTEXT("Metasound_MathRandRangeFloatNodeDescription", "Computes random float value in the range [Min, Max]."))
+	DEFINE_METASOUND_MATHOP(RandRange, Int32, int32, MathOpDisplayNames::GetRandRangeDisplayName<int32>(), LOCTEXT("Metasound_MathRandRangeInt32NodeDescription", "Computes random int32 in the range [Min, Max]."))
+	DEFINE_METASOUND_MATHOP(RandRange, Int64, int64, MathOpDisplayNames::GetRandRangeDisplayName<int64>(), LOCTEXT("Metasound_MathRandRangeInt64NodeDescription", "Computes random int64 in the range [Min, Max]."))
+	DEFINE_METASOUND_MATHOP(RandRange, Frequency, FFrequency, MathOpDisplayNames::GetRandRangeDisplayName<FFrequency>(), LOCTEXT("Metasound_MathRandRangeFrequencyNodeDescription", "Computes random frequency value in the range [Min, Max]."))
+	DEFINE_METASOUND_MATHOP(RandRange, FloatTime, FFloatTime, MathOpDisplayNames::GetRandRangeDisplayName<FFloatTime>(), LOCTEXT("Metasound_MathRandRangeFloatTimeNodeDescription", "Computes random time value in the range [Min, Max]."))
 }
 
 #undef LOCTEXT_NAMESPACE // MetasoundMathOpNode
diff --git a/Engine/Plugins/Runtime/Metasound/Source/MetasoundStandardNodes/Public/MetasoundMixerNode.h b/Engine/Plugins/Runtime/Metasound/Source/MetasoundStandardNodes/Public/MetasoundMixerNode.h
--- a/Engine/Plugins/Runtime/Metasound/Source/MetasoundStandardNodes/Public/MetasoundMixerNode.h
+++ b/Engine/Plugins/Runtime/Metasound/Source/MetasoundStandardNodes/Public/MetasoundMixerNode.h
@@ -17,7 +17,7 @@ namespace Metasound
 	class METASOUNDSTANDARDNODES_API FMixerNode : public FNodeFacade
 	{
 	public:
-		FMixerNode(const FString& InName, float InDefaultMixGainCents);
+		FMixerNode(const FString& InName, const FGuid& InInstanceID, float InDefaultMixGainCents);
 		FMixerNode(const FNodeInitData& InInitData);
 
 		float GetDefaultMixGainLinear() const;
diff --git a/Engine/Plugins/Runtime/Metasound/Source/MetasoundStandardNodes/Public/MetasoundOscNode.h b/Engine/Plugins/Runtime/Metasound/Source/MetasoundStandardNodes/Public/MetasoundOscNode.h
--- a/Engine/Plugins/Runtime/Metasound/Source/MetasoundStandardNodes/Public/MetasoundOscNode.h
+++ b/Engine/Plugins/Runtime/Metasound/Source/MetasoundStandardNodes/Public/MetasoundOscNode.h
@@ -19,7 +19,7 @@ namespace Metasound
 	{
 		public:
 
-		FOscNode(const FString& InName, float InDefaultFrequency, bool bInDefaultEnablement);
+		FOscNode(const FString& InInstanceName, const FGuid& InInstanceID, float InDefaultFrequency, bool bInDefaultEnablement);
 
 		FOscNode(const FNodeInitData& InInitData);
 
diff --git a/Engine/Plugins/Runtime/Metasound/Source/MetasoundStandardNodes/Public/MetasoundPeriodicTriggerNode.h b/Engine/Plugins/Runtime/Metasound/Source/MetasoundStandardNodes/Public/MetasoundPeriodicTriggerNode.h
--- a/Engine/Plugins/Runtime/Metasound/Source/MetasoundStandardNodes/Public/MetasoundPeriodicTriggerNode.h
+++ b/Engine/Plugins/Runtime/Metasound/Source/MetasoundStandardNodes/Public/MetasoundPeriodicTriggerNode.h
@@ -19,7 +19,7 @@ namespace Metasound
 	class METASOUNDSTANDARDNODES_API FPeriodicTriggerNode : public FNodeFacade
 	{
 		public:
-			FPeriodicTriggerNode(const FString& InName, float InDefaultPeriodInSeconds);
+			FPeriodicTriggerNode(const FString& InInstanceName, const FGuid& InInstanceID, float InDefaultPeriodInSeconds);
 			FPeriodicTriggerNode(const FNodeInitData& InInitData);
 
 			float GetDefaultPeriodInSeconds() const;
diff --git a/Engine/Plugins/Runtime/Metasound/Source/MetasoundStandardNodes/Public/MetasoundStandardNodesNames.h b/Engine/Plugins/Runtime/Metasound/Source/MetasoundStandardNodes/Public/MetasoundStandardNodesNames.h
new file mode 100644
--- /dev/null
+++ b/Engine/Plugins/Runtime/Metasound/Source/MetasoundStandardNodes/Public/MetasoundStandardNodesNames.h
@@ -0,0 +1,14 @@
+// Copyright Epic Games, Inc. All Rights Reserved.
+
+#pragma once
+
+#include "CoreMinimal.h"
+
+namespace Metasound
+{
+	namespace StandardNodes
+	{
+		extern const FName Namespace;
+		extern const FName AudioVariant;
+	}
+}
diff --git a/Engine/Plugins/Runtime/Metasound/Source/MetasoundStandardNodes/Public/MetasoundSwitchGateNode.h b/Engine/Plugins/Runtime/Metasound/Source/MetasoundStandardNodes/Public/MetasoundSwitchGateNode.h
--- a/Engine/Plugins/Runtime/Metasound/Source/MetasoundStandardNodes/Public/MetasoundSwitchGateNode.h
+++ b/Engine/Plugins/Runtime/Metasound/Source/MetasoundStandardNodes/Public/MetasoundSwitchGateNode.h
@@ -28,7 +28,7 @@ namespace Metasound
 	class METASOUNDSTANDARDNODES_API FSwitchNode : public FNodeFacade
 	{
 	public:
-		FSwitchNode(const FString& InName);
+		FSwitchNode(const FString& InName, const FGuid& InInstanceID);
 		FSwitchNode(const FNodeInitData& InInitData);
 	};
 
@@ -40,7 +40,7 @@ namespace Metasound
 	class METASOUNDSTANDARDNODES_API FGateNode : public FNodeFacade
 	{
 	public:
-		FGateNode(const FString& InName);
+		FGateNode(const FString& InName, const FGuid& InInstanceID);
 		FGateNode(const FNodeInitData& InInitData);
 	};
 }
diff --git a/Engine/Plugins/Runtime/Metasound/Source/MetasoundStandardNodes/Public/MetasoundTriggerAccumulatorNode.h b/Engine/Plugins/Runtime/Metasound/Source/MetasoundStandardNodes/Public/MetasoundTriggerAccumulatorNode.h
--- a/Engine/Plugins/Runtime/Metasound/Source/MetasoundStandardNodes/Public/MetasoundTriggerAccumulatorNode.h
+++ b/Engine/Plugins/Runtime/Metasound/Source/MetasoundStandardNodes/Public/MetasoundTriggerAccumulatorNode.h
@@ -20,10 +20,7 @@ namespace Metasound
 		 * @param InName - Name of this node.
 		 * @param InTriggerAtCount - The count at which after accumulated triggers we trigger.
 		 */
-		FTriggerAccumulatorNode(
-			const FString& InName,
-			int32 InDefaultTriggerAtCount
-		);
+		FTriggerAccumulatorNode(const FString& InName, const FGuid& InInstanceID, int32 InDefaultTriggerAtCount);
 
 		/**
 		 * Constructor used by the Metasound Frontend.
diff --git a/Engine/Plugins/Runtime/Metasound/Source/MetasoundStandardNodes/Public/MetasoundTriggerDelayNode.h b/Engine/Plugins/Runtime/Metasound/Source/MetasoundStandardNodes/Public/MetasoundTriggerDelayNode.h
--- a/Engine/Plugins/Runtime/Metasound/Source/MetasoundStandardNodes/Public/MetasoundTriggerDelayNode.h
+++ b/Engine/Plugins/Runtime/Metasound/Source/MetasoundStandardNodes/Public/MetasoundTriggerDelayNode.h
@@ -19,7 +19,7 @@ namespace Metasound
 	class METASOUNDSTANDARDNODES_API FTriggerDelayNode : public FNodeFacade
 	{
 		public:
-			FTriggerDelayNode(const FString& InName, float InDefaultDelayInSeconds);
+			FTriggerDelayNode(const FString& InName, const FGuid& InInstanceID, float InDefaultDelayInSeconds);
 			FTriggerDelayNode(const FNodeInitData& InInitData);
 
 			float GetDefaultDelayInSeconds() const;
diff --git a/Engine/Plugins/Runtime/Metasound/Source/MetasoundStandardNodes/Public/MetasoundTriggerPipeNode.h b/Engine/Plugins/Runtime/Metasound/Source/MetasoundStandardNodes/Public/MetasoundTriggerPipeNode.h
--- a/Engine/Plugins/Runtime/Metasound/Source/MetasoundStandardNodes/Public/MetasoundTriggerPipeNode.h
+++ b/Engine/Plugins/Runtime/Metasound/Source/MetasoundStandardNodes/Public/MetasoundTriggerPipeNode.h
@@ -16,7 +16,7 @@ namespace Metasound
 	class METASOUNDSTANDARDNODES_API FTriggerPipeNode : public FNodeFacade
 	{
 		public:
-			FTriggerPipeNode(const FString& InName, float InDefaultDelayInSeconds);
+			FTriggerPipeNode(const FString& InName, const FGuid& InInstanceID, float InDefaultDelayInSeconds);
 			FTriggerPipeNode(const FNodeInitData& InInitData);
 
 			float GetDefaultDelayInSeconds() const;
