Fix for shutdown crash on metasound editor due to analysis delegate not unregister when editor deletes

diff --git a/Engine/Plugins/Runtime/Metasound/Source/MetasoundEditor/Private/MetasoundEditor.cpp b/Engine/Plugins/Runtime/Metasound/Source/MetasoundEditor/Private/MetasoundEditor.cpp
--- a/Engine/Plugins/Runtime/Metasound/Source/MetasoundEditor/Private/MetasoundEditor.cpp
+++ b/Engine/Plugins/Runtime/Metasound/Source/MetasoundEditor/Private/MetasoundEditor.cpp
@@ -93,6 +93,11 @@ namespace Metasound
 				Stop();
 			}
 
+			if (MetasoundMeterAnalyzer.IsValid() && ResultsDelegateHandle.IsValid())
+			{
+				MetasoundMeterAnalyzer->OnLatestPerChannelMeterResultsNative.Remove(ResultsDelegateHandle);
+			}
+
 			check(GEditor);
 			GEditor->UnregisterForUndo(this);
 		}
@@ -295,7 +300,7 @@ namespace Metasound
 
 			MetasoundMeterAnalyzer = TStrongObjectPtr<UMeterAnalyzer>(NewObject<UMeterAnalyzer>());
 
-			MetasoundMeterAnalyzer->OnLatestPerChannelMeterResultsNative.AddRaw(this, &FEditor::OnMeterOutput);
+			ResultsDelegateHandle = MetasoundMeterAnalyzer->OnLatestPerChannelMeterResultsNative.AddRaw(this, &FEditor::OnMeterOutput);
 
 			MetasoundMeterAnalyzerSettings = TStrongObjectPtr<UMeterSettings>(NewObject<UMeterSettings>());
 			MetasoundMeterAnalyzerSettings->PeakHoldTime = 4000.0f;
diff --git a/Engine/Plugins/Runtime/Metasound/Source/MetasoundEditor/Private/MetasoundEditor.h b/Engine/Plugins/Runtime/Metasound/Source/MetasoundEditor/Private/MetasoundEditor.h
--- a/Engine/Plugins/Runtime/Metasound/Source/MetasoundEditor/Private/MetasoundEditor.h
+++ b/Engine/Plugins/Runtime/Metasound/Source/MetasoundEditor/Private/MetasoundEditor.h
@@ -215,6 +215,9 @@ namespace Metasound
 
 			TStrongObjectPtr<UMeterSettings> MetasoundMeterAnalyzerSettings;
 
+			/** Handle for results delegate for metasound meter analyzer. */
+			FDelegateHandle ResultsDelegateHandle;
+
 			/** Whether or not metasound being edited is valid */
 			bool bPassedValidation = true;
 
