Fix for output format on metasound postload

diff --git a/Engine/Plugins/Runtime/Metasound/Source/MetasoundEngine/Private/Metasound.cpp b/Engine/Plugins/Runtime/Metasound/Source/MetasoundEngine/Private/Metasound.cpp
--- a/Engine/Plugins/Runtime/Metasound/Source/MetasoundEngine/Private/Metasound.cpp
+++ b/Engine/Plugins/Runtime/Metasound/Source/MetasoundEngine/Private/Metasound.cpp
@@ -72,17 +72,3 @@ const FMetasoundFrontendArchetype& UMetasound::GetPreferredMetasoundArchetype(co
 	return InDocument.Archetype;
 }
 
-// This can be used to update the metadata (name, author, etc) for this metasound.
-// @param InMetadata may be updated with any corrections we do to the input metadata.
-void UMetasound::SetMetadata(FMetasoundFrontendClassMetadata& InMetadata)
-{
-	FMetasoundAssetBase::SetMetadata(InMetadata);
-	MarkPackageDirty();
-}
-
-void UMetasound::PostLoad()
-{
-	ConformDocumentToMetasoundArchetype();
-	Super::PostLoad();
-}
-
diff --git a/Engine/Plugins/Runtime/Metasound/Source/MetasoundEngine/Private/MetasoundSource.cpp b/Engine/Plugins/Runtime/Metasound/Source/MetasoundEngine/Private/MetasoundSource.cpp
--- a/Engine/Plugins/Runtime/Metasound/Source/MetasoundEngine/Private/MetasoundSource.cpp
+++ b/Engine/Plugins/Runtime/Metasound/Source/MetasoundEngine/Private/MetasoundSource.cpp
@@ -40,11 +40,6 @@ UMetasoundSource::UMetasoundSource(const FObjectInitializer& ObjectInitializer)
 
 }
 
-void UMetasoundSource::SetMetadata(FMetasoundFrontendClassMetadata& InMetadata)
-{
-	RootMetasoundDocument.RootGraph.Metadata = InMetadata;
-}
-
 #if WITH_EDITOR
 
 void UMetasoundSource::PostEditChangeProperty(FPropertyChangedEvent& InEvent)
@@ -209,13 +204,6 @@ TUniquePtr<IAudioInstanceTransmitter> UMetasoundSource::CreateInstanceTransmitte
 	return MakeUnique<Metasound::FMetasoundInstanceTransmitter>(InitParams);
 }
 
-void UMetasoundSource::PostLoad()
-{
-	Super::PostLoad();
-
-	ConformDocumentToMetasoundArchetype();
-}
-
 bool UMetasoundSource::GetReceiveNodeMetadataForDataType(const FName& InTypeName, FMetasoundFrontendClassMetadata& OutMetadata) const
 {
 	using namespace Metasound;
diff --git a/Engine/Plugins/Runtime/Metasound/Source/MetasoundEngine/Public/Metasound.h b/Engine/Plugins/Runtime/Metasound/Source/MetasoundEngine/Public/Metasound.h
--- a/Engine/Plugins/Runtime/Metasound/Source/MetasoundEngine/Public/Metasound.h
+++ b/Engine/Plugins/Runtime/Metasound/Source/MetasoundEngine/Public/Metasound.h
@@ -68,11 +68,14 @@ public:
 	virtual void SetGraph(UEdGraph* InGraph) override;
 #endif // WITH_EDITORONLY_DATA
 
-	UObject* GetOwningAsset() const override
+	UObject* GetOwningAsset() override
 	{
-		// Hack to allow handles to manipulate while providing
-		// ability to inspect via a handle from the const version of GetRootGraphHandle()
-		return const_cast<UObject*>(CastChecked<const UObject>(this));
+		return this;
+	}
+
+	const UObject* GetOwningAsset() const override
+	{
+		return this;
 	}
 
 	const TArray<FMetasoundFrontendArchetype>& GetPreferredMetasoundArchetypes() const override;
@@ -81,10 +84,4 @@ public:
 
 	const FMetasoundFrontendArchetype& GetPreferredMetasoundArchetype(const FMetasoundFrontendDocument& InDocument) const override;
 
-	// Updates the Metasound's metadata (name, author, etc).
-	// @param InMetadata Metadata containing corrections to the class metadata.
-	void SetMetadata(FMetasoundFrontendClassMetadata& InMetadata) override;
-
-	void PostLoad() override;
-
 };
diff --git a/Engine/Plugins/Runtime/Metasound/Source/MetasoundEngine/Public/MetasoundSource.h b/Engine/Plugins/Runtime/Metasound/Source/MetasoundEngine/Public/MetasoundSource.h
--- a/Engine/Plugins/Runtime/Metasound/Source/MetasoundEngine/Public/MetasoundSource.h
+++ b/Engine/Plugins/Runtime/Metasound/Source/MetasoundEngine/Public/MetasoundSource.h
@@ -98,23 +98,21 @@ public:
 
 #endif // WITH_EDITOR
 
-	UObject* GetOwningAsset() const override
+	UObject* GetOwningAsset() override
 	{
-		// Hack to allow handles to manipulate while providing
-		// ability to inspect via a handle from the const version of GetRootGraphHandle()
-		return const_cast<UObject*>(CastChecked<const UObject>(this));
+		return this;
 	}
 
-	// Updates the Metasound's metadata (name, author, etc).
-	// @param InMetadata Metadata containing corrections to the class metadata.
-	void SetMetadata(FMetasoundFrontendClassMetadata& InMetadata) override;
+	const UObject* GetOwningAsset() const override
+	{
+		return this;
+	}
 
 	bool IsPlayable() const override;
 	bool SupportsSubtitles() const override;
 	float GetDuration() override;
 	ISoundGeneratorPtr CreateSoundGenerator(const FSoundGeneratorInitParams& InParams) override;
 	TUniquePtr<IAudioInstanceTransmitter> CreateInstanceTransmitter(const FAudioInstanceTransmitterInitParams& InParams) const override;
-	void PostLoad() override;
 
 	// Get the most up to date archetype for metasound sources.
 	const TArray<FMetasoundFrontendArchetype>& GetPreferredMetasoundArchetypes() const override;
diff --git a/Engine/Plugins/Runtime/Metasound/Source/MetasoundFrontend/Private/MetasoundAssetBase.cpp b/Engine/Plugins/Runtime/Metasound/Source/MetasoundFrontend/Private/MetasoundAssetBase.cpp
--- a/Engine/Plugins/Runtime/Metasound/Source/MetasoundFrontend/Private/MetasoundAssetBase.cpp
+++ b/Engine/Plugins/Runtime/Metasound/Source/MetasoundFrontend/Private/MetasoundAssetBase.cpp
@@ -33,6 +33,8 @@ void FMetasoundAssetBase::SetMetadata(FMetasoundFrontendClassMetadata& InMetadat
 		UE_LOG(LogMetasound, Display, TEXT("Forcing class type to EMetasoundFrontendClassType::Graph on root graph metadata"));
 		Doc.RootGraph.Metadata.Type = EMetasoundFrontendClassType::Graph;
 	}
+
+	MarkMetasoundDocumentDirty();
 }
 
 const FMetasoundFrontendArchetype& FMetasoundAssetBase::GetMetasoundArchetype() const
@@ -106,6 +108,17 @@ void FMetasoundAssetBase::ConformDocumentToMetasoundArchetype()
 
 	Metasound::Frontend::FMatchRootGraphToArchetype Transform;
 	Transform.Transform(GetDocumentHandle());
+
+	MarkMetasoundDocumentDirty();
+}
+
+bool FMetasoundAssetBase::MarkMetasoundDocumentDirty() const
+{
+	if (const UObject* OwningAsset = GetOwningAsset())
+	{
+		return ensure(OwningAsset->MarkPackageDirty());
+	}
+	return false;
 }
 
 FMetasoundFrontendClassMetadata FMetasoundAssetBase::GetMetadata()
@@ -138,7 +151,14 @@ bool FMetasoundAssetBase::ImportFromJSON(const FString& InJSON)
 	Metasound::Frontend::TAccessPtr<FMetasoundFrontendDocument> Document = GetDocument();
 	if (ensure(Document.IsValid()))
 	{
-		return Metasound::Frontend::ImportJSONToMetasound(InJSON, *Document);
+		bool bSuccess = Metasound::Frontend::ImportJSONToMetasound(InJSON, *Document);
+
+		if (bSuccess)
+		{
+			ensure(MarkMetasoundDocumentDirty());
+		}
+
+		return bSuccess;
 	}
 	return false;
 }
@@ -148,7 +168,14 @@ bool FMetasoundAssetBase::ImportFromJSONAsset(const FString& InAbsolutePath)
 	Metasound::Frontend::TAccessPtr<FMetasoundFrontendDocument> Document = GetDocument();
 	if (ensure(Document.IsValid()))
 	{
-		return Metasound::Frontend::ImportJSONAssetToMetasound(InAbsolutePath, *Document);
+		bool bSuccess = Metasound::Frontend::ImportJSONAssetToMetasound(InAbsolutePath, *Document);
+
+		if (bSuccess)
+		{
+			ensure(MarkMetasoundDocumentDirty());
+		}
+
+		return bSuccess;
 	}
 	return false;
 }
diff --git a/Engine/Plugins/Runtime/Metasound/Source/MetasoundFrontend/Public/MetasoundAssetBase.h b/Engine/Plugins/Runtime/Metasound/Source/MetasoundFrontend/Public/MetasoundAssetBase.h
--- a/Engine/Plugins/Runtime/Metasound/Source/MetasoundFrontend/Public/MetasoundAssetBase.h
+++ b/Engine/Plugins/Runtime/Metasound/Source/MetasoundFrontend/Public/MetasoundAssetBase.h
@@ -120,6 +120,9 @@ public:
 	// This must be called on UObject::PostLoad, as well as in this asset's UFactory, to fix up the root document based on the most recent version of the archetype.
 	void ConformDocumentToMetasoundArchetype();
 
+	// Calls the outermost package and marks it dirty. 
+	bool MarkMetasoundDocumentDirty() const;
+
 protected:
 
 	// Returns an access pointer to the document.
@@ -129,7 +132,10 @@ protected:
 	virtual Metasound::Frontend::TAccessPtr<const FMetasoundFrontendDocument> GetDocument() const = 0;
 
 	// Returns the owning asset responsible for transactions applied to metasound
-	virtual UObject* GetOwningAsset() const = 0;
+	virtual UObject* GetOwningAsset() = 0;
+
+	// Returns the owning asset responsible for transactions applied to metasound
+	virtual const UObject* GetOwningAsset() const = 0;
 
 private:
 
