Fixed issue with WavePlayerNode not retriggering if the sound had finished.

diff --git a/Engine/Plugins/Runtime/Metasound/Source/MetasoundEngine/Private/MetasoundWavePlayerNode.cpp b/Engine/Plugins/Runtime/Metasound/Source/MetasoundEngine/Private/MetasoundWavePlayerNode.cpp
--- a/Engine/Plugins/Runtime/Metasound/Source/MetasoundEngine/Private/MetasoundWavePlayerNode.cpp
+++ b/Engine/Plugins/Runtime/Metasound/Source/MetasoundEngine/Private/MetasoundWavePlayerNode.cpp
@@ -3,12 +3,12 @@
 
 #include "AudioResampler.h"
 #include "MetasoundBuildError.h"
-#include "MetasoundEngineNodesNames.h"
 #include "MetasoundExecutableOperator.h"
 #include "MetasoundNodeRegistrationMacro.h"
 #include "MetasoundPrimitives.h"
 #include "MetasoundWave.h"
 #include "MetasoundTrigger.h"
+#include "MetasoundEngineNodesNames.h"
 
 
 #define LOCTEXT_NAMESPACE "MetasoundWaveNode"
@@ -24,7 +24,7 @@
 // 			return i;
 // 		}
 // 	}
-// 	return -1; // all good!
+// 	return SMOOTH; // all good!
 // }
 
 namespace Metasound
@@ -69,7 +69,6 @@ namespace Metasound
 
 			if (Wave->SoundWaveProxy.IsValid())
 			{
-				ResetDecoder();
 				CurrentSoundWaveName = Wave->SoundWaveProxy->GetFName();
 			}
 
@@ -116,11 +115,6 @@ namespace Metasound
 			FMemory::Memzero(AudioBufferL->GetData(), OutputBlockSizeInFrames * sizeof(float));
 			FMemory::Memzero(AudioBufferR->GetData(), OutputBlockSizeInFrames * sizeof(float));
 
-			if (!Decoder.CanGenerateAudio())
-			{
-				return;
-			}
-
 			TrigIn->ExecuteBlock(
 				// OnPreTrigger
 				[&](int32 StartFrame, int32 EndFrame)
@@ -133,14 +127,12 @@ namespace Metasound
 				// OnTrigger
 				[&](int32 StartFrame, int32 EndFrame)
 				{
+					ResetDecoder();
+
 					if (!bIsPlaying)
 					{
 						bIsPlaying = Decoder.CanGenerateAudio();
 					}
-					else
-					{
-						ResetDecoder();
-					}
 
 					ExecuteInternal(StartFrame, EndFrame);
 				}
@@ -289,10 +281,10 @@ namespace Metasound
 		auto InitNodeInfo = []() -> FNodeClassMetadata
 		{
 			FNodeClassMetadata Info;
-			Info.ClassName = {EngineNodes::Namespace, TEXT("WavePlayer"), EngineNodes::StereoVariant};
+			Info.ClassName = { Metasound::EngineNodes::Namespace, TEXT("Wave Player"), Metasound::EngineNodes::StereoVariant };
 			Info.MajorVersion = 1;
 			Info.MinorVersion = 0;
-			Info.DisplayName = LOCTEXT("Metasound_WavePlayerNodeDisplayName", "Wave Player");
+			Info.DisplayName = LOCTEXT("Metasound_WavePlayerNodeDisplayName", "Wave Player Node");
 			Info.Description = LOCTEXT("Metasound_WavePlayerNodeDescription", "Plays a supplied Wave");
 			Info.Author = PluginAuthor;
 			Info.PromptIfMissing = PluginNodeMissingPrompt;
@@ -307,14 +299,14 @@ namespace Metasound
 	}
 
 	FWavePlayerNode::FWavePlayerNode(const FString& InName, const FGuid& InInstanceID)
-	: FNode(InName, InInstanceID, GetNodeInfo())
-	, Factory(MakeOperatorFactoryRef<FWavePlayerNode::FOperatorFactory>())
-	, Interface(DeclareVertexInterface())
+		:	FNode(InName, InInstanceID, GetNodeInfo())
+		,	Factory(MakeOperatorFactoryRef<FWavePlayerNode::FOperatorFactory>())
+		,	Interface(DeclareVertexInterface())
 	{
 	}
 
 	FWavePlayerNode::FWavePlayerNode(const FNodeInitData& InInitData)
-	: FWavePlayerNode(InInitData.InstanceName, InInitData.InstanceID)
+		: FWavePlayerNode(InInitData.InstanceName, InInitData.InstanceID)
 	{
 	}
 
