Metasound Filters

diff --git a/Engine/Plugins/Runtime/Metasound/Source/MetasoundStandardNodes/Private/MetasoundBasicFilters.cpp b/Engine/Plugins/Runtime/Metasound/Source/MetasoundStandardNodes/Private/MetasoundBasicFilters.cpp
new file mode 100644
--- /dev/null
+++ b/Engine/Plugins/Runtime/Metasound/Source/MetasoundStandardNodes/Private/MetasoundBasicFilters.cpp
@@ -0,0 +1,858 @@
+// Copyright Epic Games, Inc. All Rights Reserved.
+
+#include "MetasoundBasicFilters.h"
+#include "MetasoundAudioBuffer.h"
+#include "MetasoundPrimitives.h"
+#include "MetasoundExecutableOperator.h"
+#include "MetasoundFacade.h"
+#include "MetasoundNodeRegistrationMacro.h"
+#include "MetasoundStandardNodesNames.h"
+#include "MetasoundOperatorSettings.h"
+
+#include "DSP/Filter.h"
+#include "DSP/InterpolatedOnePole.h"
+
+#define LOCTEXT_NAMESPACE "MetasoundBasicFilterNodes"
+
+namespace Metasound
+{
+#pragma region Ladder Filter
+
+	class FLadderFilterOperator : public TExecutableOperator<FLadderFilterOperator>
+	{
+	public:
+		static const FNodeClassMetadata& GetNodeInfo();
+
+		static FVertexInterface DeclareVertexInterface();
+
+		static TUniquePtr<IOperator> CreateOperator(const FCreateOperatorParams& InParams, FBuildErrorArray& OutErrors);
+
+		virtual FDataReferenceCollection GetInputs() const override
+		{
+			FDataReferenceCollection InputDataReferences;
+			InputDataReferences.AddDataReadReference(TEXT("In"), FAudioBufferReadRef(AudioInput));
+			InputDataReferences.AddDataReadReference(TEXT("Freq."), FFloatReadRef(Frequency));
+			InputDataReferences.AddDataReadReference(TEXT("Res."), FFloatReadRef(Resonance));
+			InputDataReferences.AddDataReadReference(TEXT("Pass-Band Gain Compensation"), FFloatReadRef(PassBandGainCompensation));
+
+			return InputDataReferences;
+		}
+
+		virtual FDataReferenceCollection GetOutputs() const override
+		{
+			// expose read access to our output buffer for other processors in the graph
+			FDataReferenceCollection OutputDataReferences;
+			OutputDataReferences.AddDataReadReference(TEXT("Out"), FAudioBufferReadRef(AudioOutput)); 
+
+			return OutputDataReferences;
+		}
+
+		void Execute();
+
+
+	private: // members
+		// input pins
+		FAudioBufferReadRef AudioInput;
+		FFloatReadRef Frequency;
+		FFloatReadRef Resonance;
+		FFloatReadRef PassBandGainCompensation;
+
+		// cached data
+		float LastFrequency;
+		float LastResonance;
+		float LastPassBandGainCompensation;
+
+		// output pins
+		FAudioBufferWriteRef AudioOutput;
+
+		// data
+		const int32 BlockSize;
+		float SampleRate;
+
+		// dsp
+		Audio::FLadderFilter LadderFilter;
+
+
+	public:
+		// ctor
+		FLadderFilterOperator(
+			const FOperatorSettings& InSettings,
+			const FAudioBufferReadRef& InAudioInput,
+			const FFloatReadRef& InFrequency,
+			const FFloatReadRef& InResonance,
+			const FFloatReadRef& InPassBandGainCompensation
+		)
+			: AudioInput(InAudioInput)
+			, Frequency(InFrequency)
+			, Resonance(InResonance)
+			, PassBandGainCompensation(InPassBandGainCompensation)
+			, AudioOutput(FAudioBufferWriteRef::CreateNew(InSettings))
+			, BlockSize(InSettings.GetNumFramesPerBlock())
+			, SampleRate(InSettings.GetSampleRate())
+		{
+			// verify our buffer sizes:
+			check(AudioOutput->Num() == BlockSize);
+
+			LadderFilter.Init(SampleRate, 1);
+		}
+	};
+
+	FLadderFilterNode::FLadderFilterNode(const FString& InInstanceName, const FGuid& InInstanceID)
+		: FNodeFacade(InInstanceName, InInstanceID, TFacadeOperatorClass<FLadderFilterOperator>())
+	{
+	}
+
+	FLadderFilterNode::FLadderFilterNode(const FNodeInitData& InInitData)
+		: FLadderFilterNode(InInitData.InstanceName, InInitData.InstanceID)
+	{
+	}
+
+	const FNodeClassMetadata& FLadderFilterOperator::GetNodeInfo()
+	{
+		auto InitNodeInfo = []() -> FNodeClassMetadata
+		{
+			FNodeClassMetadata Info;
+			Info.ClassName = { Metasound::StandardNodes::Namespace, TEXT("Ladder Filter"), Metasound::StandardNodes::AudioVariant };
+			Info.MajorVersion = 1;
+			Info.MinorVersion = 0;
+			Info.DisplayName = LOCTEXT("Metasound_LadderFilterNodeDisplayName", "Ladder Filter");
+			Info.Description = LOCTEXT("Ladder_Filter_NodeDescription", "Ladder filter"),
+			Info.Author = PluginAuthor;
+			Info.PromptIfMissing = PluginNodeMissingPrompt;
+			Info.DefaultInterface = DeclareVertexInterface();
+
+			return Info;
+		};
+
+		static const FNodeClassMetadata Info = InitNodeInfo();
+
+		return Info;
+	}
+
+	FVertexInterface FLadderFilterOperator::DeclareVertexInterface()
+	{
+		static const FVertexInterface Interface(
+			FInputVertexInterface(
+				TInputDataVertexModel<FAudioBuffer>(TEXT("In"), LOCTEXT("AudioInputTooltip", "Audio Input")),
+				TInputDataVertexModel<float>(TEXT("Freq."), LOCTEXT("FrequencyTooltip", "Controls cutoff frequency")),
+				TInputDataVertexModel<float>(TEXT("Res."), LOCTEXT("ResonanceTooltip", "Controls filter resonance")),
+				TInputDataVertexModel<float>(TEXT("Pass-Band Gain Comp."), LOCTEXT("PassBandGainCompensationTooltip", "Controls pass-band gain compenation"))
+			),
+			FOutputVertexInterface(
+				TOutputDataVertexModel<FAudioBuffer>(TEXT("Out"), LOCTEXT("FilterOutputToolTip", "Audio Out"))
+			)
+		);
+
+		return Interface;
+	}
+
+	TUniquePtr<IOperator> FLadderFilterOperator::CreateOperator(const FCreateOperatorParams& InParams, FBuildErrorArray& OutErrors)
+	{
+		const FLadderFilterNode& LadderFilterNode = static_cast<const FLadderFilterNode&>(InParams.Node);
+
+		const FDataReferenceCollection& InputDataRefs = InParams.InputDataReferences;
+
+		// inputs
+		FAudioBufferReadRef AudioIn = InputDataRefs.GetDataReadReferenceOrConstruct<FAudioBuffer>(TEXT("In"), InParams.OperatorSettings);
+		FFloatReadRef FrequencyIn = InputDataRefs.GetDataReadReferenceOrConstruct<float>(TEXT("Freq."));
+		FFloatReadRef ResonanceIn = InputDataRefs.GetDataReadReferenceOrConstruct<float>(TEXT("Res."));
+		FFloatReadRef PassBandGainCompensationIn = InputDataRefs.GetDataReadReferenceOrConstruct<float>(TEXT("Pass-Band Gain Compensation"));
+
+		return MakeUnique<FLadderFilterOperator>(
+			InParams.OperatorSettings
+			, AudioIn
+			, FrequencyIn
+			, ResonanceIn
+			, PassBandGainCompensationIn
+			);
+	}
+
+	void FLadderFilterOperator::Execute()
+	{
+		bool bNeedsUpdate = false;
+		bNeedsUpdate |= (!FMath::IsNearlyEqual(LastFrequency, *Frequency));
+		bNeedsUpdate |= (!FMath::IsNearlyEqual(LastResonance, *Resonance));
+		bNeedsUpdate |= (!FMath::IsNearlyEqual(LastPassBandGainCompensation, *PassBandGainCompensation));
+
+		if (bNeedsUpdate)
+		{
+			LadderFilter.SetQ(*Resonance);
+			LadderFilter.SetFrequency(*Frequency);
+			LadderFilter.SetPassBandGainCompensation(*PassBandGainCompensation);
+
+			LadderFilter.Update();
+		}
+
+		LadderFilter.ProcessAudio(AudioInput->GetData(), AudioInput->Num(), AudioOutput->GetData());
+
+		LastFrequency = *Frequency;
+		LastResonance = *Resonance;
+		LastPassBandGainCompensation = *PassBandGainCompensation;
+	}
+
+	METASOUND_REGISTER_NODE(FLadderFilterNode);
+
+#pragma endregion
+
+
+#pragma region State Variable Filter
+
+	class FStateVariableFilterOperator : public TExecutableOperator<FStateVariableFilterOperator>
+	{
+	public:
+		static const FNodeClassMetadata& GetNodeInfo();
+
+		static FVertexInterface DeclareVertexInterface();
+
+		static TUniquePtr<IOperator> CreateOperator(const FCreateOperatorParams& InParams, FBuildErrorArray& OutErrors);
+
+		virtual FDataReferenceCollection GetInputs() const override
+		{
+			FDataReferenceCollection InputDataReferences;
+			InputDataReferences.AddDataReadReference(TEXT("In"), FAudioBufferReadRef(AudioInput));
+			InputDataReferences.AddDataReadReference(TEXT("Freq."), FFloatReadRef(Frequency));
+			InputDataReferences.AddDataReadReference(TEXT("Res."), FFloatReadRef(Resonance));
+			InputDataReferences.AddDataReadReference(TEXT("Pass-Band Gain Compensation"), FFloatReadRef(BandStopControl));
+
+			return InputDataReferences;
+		}
+
+		virtual FDataReferenceCollection GetOutputs() const override
+		{
+			// expose read access to our output buffer for other processors in the graph
+			FDataReferenceCollection OutputDataReferences;
+			OutputDataReferences.AddDataReadReference(TEXT("LP"), FAudioBufferReadRef(LowPassOutput));
+			OutputDataReferences.AddDataReadReference(TEXT("HP"), FAudioBufferReadRef(HighPassOutput));
+			OutputDataReferences.AddDataReadReference(TEXT("BP"), FAudioBufferReadRef(BandPassOutput));
+			OutputDataReferences.AddDataReadReference(TEXT("BS"), FAudioBufferReadRef(BandStopOutput));
+
+			return OutputDataReferences;
+		}
+
+		void Execute();
+
+
+	private: // members
+		// input pins
+		FAudioBufferReadRef AudioInput;
+		FFloatReadRef Frequency;
+		FFloatReadRef Resonance;
+		FFloatReadRef BandStopControl;
+
+		// cached data
+		float LastFrequency;
+		float LastResonance;
+		float LastBandStopControl;
+
+		// output pins
+		FAudioBufferWriteRef LowPassOutput;
+		FAudioBufferWriteRef HighPassOutput;
+		FAudioBufferWriteRef BandPassOutput;
+		FAudioBufferWriteRef BandStopOutput;
+
+		// data
+		const int32 BlockSize;
+		float SampleRate;
+
+		// dsp
+		Audio::FStateVariableFilter StateVariableFilter;
+
+
+	public:
+		// ctor
+		FStateVariableFilterOperator(
+			const FOperatorSettings& InSettings,
+			const FAudioBufferReadRef& InAudioInput,
+			const FFloatReadRef& InFrequency,
+			const FFloatReadRef& InResonance,
+			const FFloatReadRef& InBandStopControl
+		)
+			: AudioInput(InAudioInput)
+			, Frequency(InFrequency)
+			, Resonance(InResonance)
+			, BandStopControl(InBandStopControl)
+			, LowPassOutput(FAudioBufferWriteRef::CreateNew(InSettings))
+			, HighPassOutput(FAudioBufferWriteRef::CreateNew(InSettings))
+			, BandPassOutput(FAudioBufferWriteRef::CreateNew(InSettings))
+			, BandStopOutput(FAudioBufferWriteRef::CreateNew(InSettings))
+			, BlockSize(InSettings.GetNumFramesPerBlock())
+			, SampleRate(InSettings.GetSampleRate())
+		{
+			// verify our buffer sizes:
+			check(LowPassOutput->Num() == BlockSize);
+			check(HighPassOutput->Num() == BlockSize);
+			check(BandPassOutput->Num() == BlockSize);
+			check(BandPassOutput->Num() == BlockSize);
+
+			StateVariableFilter.Init(SampleRate, 1);
+		}
+	};
+
+	FStateVariableFilterNode::FStateVariableFilterNode(const FString& InInstanceName, const FGuid& InInstanceID)
+		: FNodeFacade(InInstanceName, InInstanceID, TFacadeOperatorClass<FStateVariableFilterOperator>())
+	{
+	}
+
+	FStateVariableFilterNode::FStateVariableFilterNode(const FNodeInitData& InInitData)
+		: FStateVariableFilterNode(InInitData.InstanceName, InInitData.InstanceID)
+	{
+	}
+
+	const FNodeClassMetadata& FStateVariableFilterOperator::GetNodeInfo()
+	{
+		auto InitNodeInfo = []() -> FNodeClassMetadata
+		{
+			FNodeClassMetadata Info;
+			Info.ClassName = { Metasound::StandardNodes::Namespace, TEXT("State Variable Filter"), Metasound::StandardNodes::AudioVariant };
+			Info.MajorVersion = 1;
+			Info.MinorVersion = 0;
+			Info.DisplayName = LOCTEXT("Metasound_StateVariableFilterNodeDisplayName", "State Variable Filter");
+			Info.Description = LOCTEXT("State_Variable_Filter_NodeDescription", "State Variable filter"),
+			Info.Author = PluginAuthor;
+			Info.PromptIfMissing = PluginNodeMissingPrompt;
+			Info.DefaultInterface = DeclareVertexInterface();
+
+			return Info;
+		};
+
+		static const FNodeClassMetadata Info = InitNodeInfo();
+
+		return Info;
+	}
+
+	FVertexInterface FStateVariableFilterOperator::DeclareVertexInterface()
+	{
+		static const FVertexInterface Interface(
+			FInputVertexInterface(
+				TInputDataVertexModel<FAudioBuffer>(TEXT("In"), LOCTEXT("AudioInputTooltip", "Audio Input")),
+				TInputDataVertexModel<float>(TEXT("Freq."), LOCTEXT("FrequencyTooltip", "Controls cutoff frequency")),
+				TInputDataVertexModel<float>(TEXT("Res."), LOCTEXT("ResonanceTooltip", "Controls filter resonance")),
+				TInputDataVertexModel<float>(TEXT("Band-Stop Control"), LOCTEXT("PassBandGainCompensationTooltip", "Band Stop Control"))
+			),
+			FOutputVertexInterface(
+				TOutputDataVertexModel<FAudioBuffer>(TEXT("LP"), LOCTEXT("FilterOutputToolTip", "Low Pass Output")),
+				TOutputDataVertexModel<FAudioBuffer>(TEXT("HP"), LOCTEXT("FilterOutputToolTip", "High Pass Output")),
+				TOutputDataVertexModel<FAudioBuffer>(TEXT("BP"), LOCTEXT("FilterOutputToolTip", "Band Pass Output")),
+				TOutputDataVertexModel<FAudioBuffer>(TEXT("BS"), LOCTEXT("FilterOutputToolTip", "Band Stop Output"))
+			)
+		);
+
+		return Interface;
+	}
+
+	TUniquePtr<IOperator> FStateVariableFilterOperator::CreateOperator(const FCreateOperatorParams& InParams, FBuildErrorArray& OutErrors)
+	{
+		const FStateVariableFilterNode& StateVariableFilterNode = static_cast<const FStateVariableFilterNode&>(InParams.Node);
+
+		const FDataReferenceCollection& InputDataRefs = InParams.InputDataReferences;
+
+		// inputs
+		FAudioBufferReadRef AudioIn = InputDataRefs.GetDataReadReferenceOrConstruct<FAudioBuffer>(TEXT("In"), InParams.OperatorSettings);
+		FFloatReadRef FrequencyIn = InputDataRefs.GetDataReadReferenceOrConstruct<float>(TEXT("Freq."));
+		FFloatReadRef ResonanceIn = InputDataRefs.GetDataReadReferenceOrConstruct<float>(TEXT("Res."));
+		FFloatReadRef PassBandGainCompensationIn = InputDataRefs.GetDataReadReferenceOrConstruct<float>(TEXT("Band-Stop Control"));
+
+		return MakeUnique<FStateVariableFilterOperator>(
+			InParams.OperatorSettings
+			, AudioIn
+			, FrequencyIn
+			, ResonanceIn
+			, PassBandGainCompensationIn
+			);
+	}
+
+	void FStateVariableFilterOperator::Execute()
+	{
+		bool bNeedsUpdate = false;
+		bNeedsUpdate |= (!FMath::IsNearlyEqual(LastFrequency, *Frequency));
+		bNeedsUpdate |= (!FMath::IsNearlyEqual(LastResonance, *Resonance));
+		bNeedsUpdate |= (!FMath::IsNearlyEqual(LastBandStopControl, *BandStopControl));
+
+		if (bNeedsUpdate)
+		{
+			StateVariableFilter.SetQ(*Resonance);
+			StateVariableFilter.SetFrequency(*Frequency);
+			StateVariableFilter.SetPassBandGainCompensation(*BandStopControl);
+
+			StateVariableFilter.Update();
+		}
+
+		StateVariableFilter.ProcessAudio(
+			  AudioInput->GetData()
+			, AudioInput->Num()
+			, LowPassOutput->GetData()
+			, HighPassOutput->GetData()
+			, BandPassOutput->GetData()
+			, BandStopOutput->GetData()
+		);
+
+		LastFrequency = *Frequency;
+		LastResonance = *Resonance;
+		LastBandStopControl = *BandStopControl;
+	}
+
+	METASOUND_REGISTER_NODE(FStateVariableFilterNode);
+
+#pragma endregion
+
+
+#pragma region OnePoleLowPass Filter
+
+	class FOnePoleLowPassFilterOperator : public TExecutableOperator<FOnePoleLowPassFilterOperator>
+	{
+	public:
+		static const FNodeClassMetadata& GetNodeInfo();
+
+		static FVertexInterface DeclareVertexInterface();
+
+		static TUniquePtr<IOperator> CreateOperator(const FCreateOperatorParams& InParams, FBuildErrorArray& OutErrors);
+
+		virtual FDataReferenceCollection GetInputs() const override
+		{
+			FDataReferenceCollection InputDataReferences;
+			InputDataReferences.AddDataReadReference(TEXT("In"), FAudioBufferReadRef(AudioInput));
+			InputDataReferences.AddDataReadReference(TEXT("Freq."), FFloatReadRef(Frequency));
+
+			return InputDataReferences;
+		}
+
+		virtual FDataReferenceCollection GetOutputs() const override
+		{
+			// expose read access to our output buffer for other processors in the graph
+			FDataReferenceCollection OutputDataReferences;
+			OutputDataReferences.AddDataReadReference(TEXT("Out"), FAudioBufferReadRef(AudioOutput));
+
+			return OutputDataReferences;
+		}
+
+		void Execute();
+
+
+	private: // members
+		// input pins
+		FAudioBufferReadRef AudioInput;
+		FFloatReadRef Frequency;
+
+		// output pins
+		FAudioBufferWriteRef AudioOutput;
+
+		// data
+		const int32 BlockSize;
+		float SampleRate;
+
+		// dsp
+		Audio::FInterpolatedLPF OnePoleLowPassFilter;
+
+
+	public:
+		// ctor
+		FOnePoleLowPassFilterOperator(
+			const FOperatorSettings& InSettings,
+			const FAudioBufferReadRef& InAudioInput,
+			const FFloatReadRef& InFrequency
+		)
+			: AudioInput(InAudioInput)
+			, Frequency(InFrequency)
+			, AudioOutput(FAudioBufferWriteRef::CreateNew(InSettings))
+			, BlockSize(InSettings.GetNumFramesPerBlock())
+			, SampleRate(InSettings.GetSampleRate())
+		{
+			// verify our buffer sizes:
+			check(AudioOutput->Num() == BlockSize);
+
+			OnePoleLowPassFilter.Init(SampleRate, 1); // mono
+		}
+	};
+
+	FOnePoleLowPassFilterNode::FOnePoleLowPassFilterNode(const FString& InInstanceName, const FGuid& InInstanceID)
+		: FNodeFacade(InInstanceName, InInstanceID, TFacadeOperatorClass<FOnePoleLowPassFilterOperator>())
+	{
+	}
+
+	FOnePoleLowPassFilterNode::FOnePoleLowPassFilterNode(const FNodeInitData& InInitData)
+		: FOnePoleLowPassFilterNode(InInitData.InstanceName, InInitData.InstanceID)
+	{
+	}
+
+	const FNodeClassMetadata& FOnePoleLowPassFilterOperator::GetNodeInfo()
+	{
+		auto InitNodeInfo = []() -> FNodeClassMetadata
+		{
+			FNodeClassMetadata Info;
+			Info.ClassName = { Metasound::StandardNodes::Namespace, TEXT("One-Pole Low Pass Filter"), Metasound::StandardNodes::AudioVariant };
+			Info.MajorVersion = 1;
+			Info.MinorVersion = 0;
+			Info.DisplayName = LOCTEXT("Metasound_OnePoleLpfNodeDisplayName", "One-Pole Low Pass Filter");
+			Info.Description = LOCTEXT("One_Pole_Low_Pass_Filter_NodeDescription", "One-Pole Low Pass Filter"),
+			Info.Author = PluginAuthor;
+			Info.PromptIfMissing = PluginNodeMissingPrompt;
+			Info.DefaultInterface = DeclareVertexInterface();
+
+			return Info;
+		};
+
+		static const FNodeClassMetadata Info = InitNodeInfo();
+
+		return Info;
+	}
+
+	FVertexInterface FOnePoleLowPassFilterOperator::DeclareVertexInterface()
+	{
+		static const FVertexInterface Interface(
+			FInputVertexInterface(
+				TInputDataVertexModel<FAudioBuffer>(TEXT("In"), LOCTEXT("AudioInputTooltip", "Audio Input")),
+				TInputDataVertexModel<float>(TEXT("Freq."), LOCTEXT("FrequencyTooltip", "Controls filter cutoff"))
+			),
+			FOutputVertexInterface(
+				TOutputDataVertexModel<FAudioBuffer>(TEXT("Out"), LOCTEXT("FilterOutputToolTip", "Audio Output"))
+			)
+		);
+
+		return Interface;
+	}
+
+	TUniquePtr<IOperator> FOnePoleLowPassFilterOperator::CreateOperator(const FCreateOperatorParams& InParams, FBuildErrorArray& OutErrors)
+	{
+		const FOnePoleLowPassFilterNode& OnePoleLowPassFilterNode = static_cast<const FOnePoleLowPassFilterNode&>(InParams.Node);
+
+		const FDataReferenceCollection& InputDataRefs = InParams.InputDataReferences;
+
+		// inputs
+		FAudioBufferReadRef AudioIn = InputDataRefs.GetDataReadReferenceOrConstruct<FAudioBuffer>(TEXT("In"), InParams.OperatorSettings);
+		FFloatReadRef FrequencyIn = InputDataRefs.GetDataReadReferenceOrConstruct<float>(TEXT("Freq."));
+
+		return MakeUnique<FOnePoleLowPassFilterOperator>(
+			InParams.OperatorSettings
+			, AudioIn
+			, FrequencyIn
+			);
+	}
+
+	void FOnePoleLowPassFilterOperator::Execute()
+	{
+		OnePoleLowPassFilter.StartFrequencyInterpolation(*Frequency);
+		OnePoleLowPassFilter.ProcessAudioBuffer(AudioInput->GetData(), AudioOutput->GetData(), AudioInput->Num());
+		OnePoleLowPassFilter.StopFrequencyInterpolation();
+	}
+
+	METASOUND_REGISTER_NODE(FOnePoleLowPassFilterNode);
+
+#pragma endregion
+
+
+#pragma region OnePoleHighPass Filter
+
+	class FOnePoleHighPassFilterOperator : public TExecutableOperator<FOnePoleHighPassFilterOperator>
+	{
+	public:
+		static const FNodeClassMetadata& GetNodeInfo();
+
+		static FVertexInterface DeclareVertexInterface();
+
+		static TUniquePtr<IOperator> CreateOperator(const FCreateOperatorParams& InParams, FBuildErrorArray& OutErrors);
+
+		virtual FDataReferenceCollection GetInputs() const override
+		{
+			FDataReferenceCollection InputDataReferences;
+			InputDataReferences.AddDataReadReference(TEXT("In"), FAudioBufferReadRef(AudioInput));
+			InputDataReferences.AddDataReadReference(TEXT("Freq."), FFloatReadRef(Frequency));
+
+			return InputDataReferences;
+		}
+
+		virtual FDataReferenceCollection GetOutputs() const override
+		{
+			// expose read access to our output buffer for other processors in the graph
+			FDataReferenceCollection OutputDataReferences;
+			OutputDataReferences.AddDataReadReference(TEXT("Out"), FAudioBufferReadRef(AudioOutput));
+
+			return OutputDataReferences;
+		}
+
+		void Execute();
+
+
+	private: // members
+		// input pins
+		FAudioBufferReadRef AudioInput;
+		FFloatReadRef Frequency;
+
+		// output pins
+		FAudioBufferWriteRef AudioOutput;
+
+		// data
+		const int32 BlockSize;
+		float SampleRate;
+
+		// dsp
+		Audio::FInterpolatedHPF OnePoleHighPassFilter;
+
+
+	public:
+		// ctor
+		FOnePoleHighPassFilterOperator(
+			const FOperatorSettings& InSettings,
+			const FAudioBufferReadRef& InAudioInput,
+			const FFloatReadRef& InFrequency
+		)
+			: AudioInput(InAudioInput)
+			, Frequency(InFrequency)
+			, AudioOutput(FAudioBufferWriteRef::CreateNew(InSettings))
+			, BlockSize(InSettings.GetNumFramesPerBlock())
+			, SampleRate(InSettings.GetSampleRate())
+		{
+			// verify our buffer sizes:
+			check(AudioOutput->Num() == BlockSize);
+
+			OnePoleHighPassFilter.Init(SampleRate, 1); // mono
+		}
+	};
+
+	FOnePoleHighPassFilterNode::FOnePoleHighPassFilterNode(const FString& InInstanceName, const FGuid& InInstanceID)
+		: FNodeFacade(InInstanceName, InInstanceID, TFacadeOperatorClass<FOnePoleHighPassFilterOperator>())
+	{
+	}
+
+	FOnePoleHighPassFilterNode::FOnePoleHighPassFilterNode(const FNodeInitData& InInitData)
+		: FOnePoleHighPassFilterNode(InInitData.InstanceName, InInitData.InstanceID)
+	{
+	}
+
+	const FNodeClassMetadata& FOnePoleHighPassFilterOperator::GetNodeInfo()
+	{
+		auto InitNodeInfo = []() -> FNodeClassMetadata
+		{
+			FNodeClassMetadata Info;
+			Info.ClassName = { Metasound::StandardNodes::Namespace, TEXT("One-Pole High Pass Filter"), Metasound::StandardNodes::AudioVariant };
+			Info.MajorVersion = 1;
+			Info.MinorVersion = 0;
+			Info.DisplayName = LOCTEXT("Metasound_OnePoleHpfNodeDisplayName", "One-Pole High Pass Filter");
+			Info.Description = LOCTEXT("One_Pole_High_Pass_Filter_NodeDescription", "One-Pole High Pass Filter"),
+			Info.Author = PluginAuthor;
+			Info.PromptIfMissing = PluginNodeMissingPrompt;
+			Info.DefaultInterface = DeclareVertexInterface();
+
+			return Info;
+		};
+
+		static const FNodeClassMetadata Info = InitNodeInfo();
+
+		return Info;
+	}
+
+	FVertexInterface FOnePoleHighPassFilterOperator::DeclareVertexInterface()
+	{
+		static const FVertexInterface Interface(
+			FInputVertexInterface(
+				TInputDataVertexModel<FAudioBuffer>(TEXT("In"), LOCTEXT("AudioInputTooltip", "Audio In")),
+				TInputDataVertexModel<float>(TEXT("Freq."), LOCTEXT("FrequencyTooltip", "Freq In"))
+				),
+			FOutputVertexInterface(
+				TOutputDataVertexModel<FAudioBuffer>(TEXT("Out"), LOCTEXT("FilterOutputToolTip", "Audio Out"))
+			)
+		);
+
+		return Interface;
+	}
+
+	TUniquePtr<IOperator> FOnePoleHighPassFilterOperator::CreateOperator(const FCreateOperatorParams& InParams, FBuildErrorArray& OutErrors)
+	{
+		const FOnePoleHighPassFilterNode& OnePoleHighPassFilterNode = static_cast<const FOnePoleHighPassFilterNode&>(InParams.Node);
+
+		const FDataReferenceCollection& InputDataRefs = InParams.InputDataReferences;
+
+		// inputs
+		FAudioBufferReadRef AudioIn = InputDataRefs.GetDataReadReferenceOrConstruct<FAudioBuffer>(TEXT("In"), InParams.OperatorSettings);
+		FFloatReadRef FrequencyIn = InputDataRefs.GetDataReadReferenceOrConstruct<float>(TEXT("Freq."));
+
+		return MakeUnique<FOnePoleHighPassFilterOperator>(
+			InParams.OperatorSettings
+			, AudioIn
+			, FrequencyIn
+			);
+	}
+
+	void FOnePoleHighPassFilterOperator::Execute()
+	{
+		OnePoleHighPassFilter.StartFrequencyInterpolation(*Frequency);
+		OnePoleHighPassFilter.ProcessAudioBuffer(AudioInput->GetData(), AudioOutput->GetData(), AudioInput->Num());
+		OnePoleHighPassFilter.StopFrequencyInterpolation();
+	}
+
+	METASOUND_REGISTER_NODE(FOnePoleHighPassFilterNode);
+
+#pragma endregion
+
+
+#pragma region Biquad Filter
+
+	class FBiquadFilterOperator : public TExecutableOperator<FBiquadFilterOperator>
+	{
+	public:
+		static const FNodeClassMetadata& GetNodeInfo();
+
+		static FVertexInterface DeclareVertexInterface();
+
+		static TUniquePtr<IOperator> CreateOperator(const FCreateOperatorParams& InParams, FBuildErrorArray& OutErrors);
+
+		virtual FDataReferenceCollection GetInputs() const override
+		{
+			FDataReferenceCollection InputDataReferences;
+			InputDataReferences.AddDataReadReference(TEXT("In"), FAudioBufferReadRef(AudioInput));
+			InputDataReferences.AddDataReadReference(TEXT("Freq."), FFloatReadRef(Frequency));
+			InputDataReferences.AddDataReadReference(TEXT("Bandwidth"), FFloatReadRef(Bandwidth));
+
+			return InputDataReferences;
+		}
+
+		virtual FDataReferenceCollection GetOutputs() const override
+		{
+			// expose read access to our output buffer for other processors in the graph
+			FDataReferenceCollection OutputDataReferences;
+			OutputDataReferences.AddDataReadReference(TEXT("Out"), FAudioBufferReadRef(AudioOutput));
+
+			return OutputDataReferences;
+		}
+
+		void Execute();
+
+
+	private: // members
+		// input pins
+		FAudioBufferReadRef AudioInput;
+		FFloatReadRef Frequency;
+		FFloatReadRef Bandwidth;
+
+		// cached data
+		float LastFrequency;
+		float LastBandwidth;
+
+		// output pins
+		FAudioBufferWriteRef AudioOutput;
+
+		// data
+		const int32 BlockSize;
+		float SampleRate;
+
+		// dsp
+		Audio::EBiquadFilter::Type LastFilterType;
+		Audio::FBiquadFilter BiquadFilter;
+
+
+	public:
+		// ctor
+		FBiquadFilterOperator(
+			const FOperatorSettings& InSettings,
+			const FAudioBufferReadRef& InAudioInput,
+			const FFloatReadRef& InFrequency,
+			const FFloatReadRef& InBandwidth
+		)
+			: AudioInput(InAudioInput)
+			, Frequency(InFrequency)
+			, Bandwidth(InBandwidth)
+			, AudioOutput(FAudioBufferWriteRef::CreateNew(InSettings))
+			, BlockSize(InSettings.GetNumFramesPerBlock())
+			, SampleRate(InSettings.GetSampleRate())
+		{
+			// verify our buffer sizes:
+			check(AudioOutput->Num() == BlockSize);
+
+			BiquadFilter.Init(SampleRate, 1, Audio::EBiquadFilter::Type::ButterworthLowPass);
+		}
+	};
+
+	FBiquadFilterNode::FBiquadFilterNode(const FString& InInstanceName, const FGuid& InInstanceID)
+		: FNodeFacade(InInstanceName, InInstanceID, TFacadeOperatorClass<FBiquadFilterOperator>())
+	{
+	}
+
+	FBiquadFilterNode::FBiquadFilterNode(const FNodeInitData& InInitData)
+		: FBiquadFilterNode(InInitData.InstanceName, InInitData.InstanceID)
+	{
+	}
+
+	const FNodeClassMetadata& FBiquadFilterOperator::GetNodeInfo()
+	{
+		auto InitNodeInfo = []() -> FNodeClassMetadata
+		{
+			FNodeClassMetadata Info;
+			Info.ClassName = { Metasound::StandardNodes::Namespace, TEXT("Biquad Filter"), Metasound::StandardNodes::AudioVariant };
+			Info.MajorVersion = 1;
+			Info.MinorVersion = 0;
+			Info.DisplayName = LOCTEXT("Metasound_BiquadFilterNodeDisplayName", "Biquad Filter");
+			Info.Description = LOCTEXT("Biquad_Filter_NodeDescription", "Biquad filter"),
+			Info.Author = PluginAuthor;
+			Info.PromptIfMissing = PluginNodeMissingPrompt;
+			Info.DefaultInterface = DeclareVertexInterface();
+
+			return Info;
+		};
+
+		static const FNodeClassMetadata Info = InitNodeInfo();
+
+		return Info;
+	}
+
+	FVertexInterface FBiquadFilterOperator::DeclareVertexInterface()
+	{
+		static const FVertexInterface Interface(
+			FInputVertexInterface(
+				TInputDataVertexModel<FAudioBuffer>(TEXT("In"), LOCTEXT("AudioInputTooltip", "Audio Input")),
+				TInputDataVertexModel<float>(TEXT("Freq."), LOCTEXT("FrequencyTooltip", "Controls Filter Cutoff")),
+				TInputDataVertexModel<float>(TEXT("Bandwidth"), LOCTEXT("BandwidthTooltip", "Bandwidth Control"))
+			),
+			FOutputVertexInterface(
+				TOutputDataVertexModel<FAudioBuffer>(TEXT("Out"), LOCTEXT("FilterOutputToolTip", "Audio Output"))
+			)
+		);
+
+		return Interface;
+	}
+
+	TUniquePtr<IOperator> FBiquadFilterOperator::CreateOperator(const FCreateOperatorParams& InParams, FBuildErrorArray& OutErrors)
+	{
+		const FBiquadFilterNode& BiquadFilterNode = static_cast<const FBiquadFilterNode&>(InParams.Node);
+
+		const FDataReferenceCollection& InputDataRefs = InParams.InputDataReferences;
+
+		// inputs
+		FAudioBufferReadRef AudioIn = InputDataRefs.GetDataReadReferenceOrConstruct<FAudioBuffer>(TEXT("In"), InParams.OperatorSettings);
+		FFloatReadRef FrequencyIn = InputDataRefs.GetDataReadReferenceOrConstruct<float>(TEXT("Freq."));
+		FFloatReadRef BandwidthIn = InputDataRefs.GetDataReadReferenceOrConstruct<float>(TEXT("Bandwidth"));
+
+		return MakeUnique<FBiquadFilterOperator>(
+			InParams.OperatorSettings
+			, AudioIn
+			, FrequencyIn
+			, BandwidthIn
+			);
+	}
+
+	void FBiquadFilterOperator::Execute()
+	{
+		if (!FMath::IsNearlyEqual(LastFrequency, *Frequency))
+		{
+			BiquadFilter.SetFrequency(*Frequency);
+		}
+
+		if (!FMath::IsNearlyEqual(LastBandwidth, *Bandwidth))
+		{
+			BiquadFilter.SetBandwidth(*Bandwidth);
+		}
+
+		BiquadFilter.ProcessAudio(AudioInput->GetData(), AudioInput->Num(), AudioOutput->GetData());
+
+		LastFrequency = *Frequency;
+		LastBandwidth = *Bandwidth;
+	}
+
+	METASOUND_REGISTER_NODE(FBiquadFilterNode);
+
+#pragma endregion
+
+
+} // namespace Metasound
+
+
+
+#undef LOCTEXT_NAMESPACE //MetasoundBasicFilterNodes
\ No newline at end of file
diff --git a/Engine/Plugins/Runtime/Metasound/Source/MetasoundStandardNodes/Public/MetasoundBasicFilters.h b/Engine/Plugins/Runtime/Metasound/Source/MetasoundStandardNodes/Public/MetasoundBasicFilters.h
new file mode 100644
--- /dev/null
+++ b/Engine/Plugins/Runtime/Metasound/Source/MetasoundStandardNodes/Public/MetasoundBasicFilters.h
@@ -0,0 +1,62 @@
+// Copyright Epic Games, Inc. All Rights Reserved.
+#pragma once
+
+#include "CoreMinimal.h"
+#include "Internationalization/Text.h"
+
+#include "MetasoundFacade.h"
+
+
+namespace Metasound
+{
+	class METASOUNDSTANDARDNODES_API FLadderFilterNode : public FNodeFacade
+	{
+	public:
+		// public node api needs to define two conversion constructors:
+		// 1.) from FString
+		FLadderFilterNode(const FString& InInstanceName, const FGuid& InInstanceID);
+		// 2.) From an NodeInitData struct
+		FLadderFilterNode(const FNodeInitData& InInitData);
+	};
+
+	class METASOUNDSTANDARDNODES_API FStateVariableFilterNode : public FNodeFacade
+	{
+	public:
+		// public node api needs to define two conversion constructors:
+		// 1.) from FString
+		FStateVariableFilterNode(const FString& InInstanceName, const FGuid& InInstanceID);
+		// 2.) From an NodeInitData struct
+		FStateVariableFilterNode(const FNodeInitData& InInitData);
+	};
+
+	class METASOUNDSTANDARDNODES_API FOnePoleLowPassFilterNode : public FNodeFacade
+	{
+	public:
+		// public node api needs to define two conversion constructors:
+		// 1.) from FString
+		FOnePoleLowPassFilterNode(const FString& InInstanceName, const FGuid& InInstanceID);
+		// 2.) From an NodeInitData struct
+		FOnePoleLowPassFilterNode(const FNodeInitData& InInitData);
+	};
+
+	class METASOUNDSTANDARDNODES_API FOnePoleHighPassFilterNode : public FNodeFacade
+	{
+	public:
+		// public node api needs to define two conversion constructors:
+		// 1.) from FString
+		FOnePoleHighPassFilterNode(const FString& InInstanceName, const FGuid& InInstanceID);
+		// 2.) From an NodeInitData struct
+		FOnePoleHighPassFilterNode(const FNodeInitData& InInitData);
+	};
+
+	class METASOUNDSTANDARDNODES_API FBiquadFilterNode : public FNodeFacade
+	{
+	public:
+		// public node api needs to define two conversion constructors:
+		// 1.) from FString
+		FBiquadFilterNode(const FString& InInstanceName, const FGuid& InInstanceID);
+		// 2.) From an NodeInitData struct
+		FBiquadFilterNode(const FNodeInitData& InInitData);
+	};
+
+} // namespace Metasound
diff --git a/Engine/Source/Runtime/SignalProcessing/Private/Filter.cpp b/Engine/Source/Runtime/SignalProcessing/Private/Filter.cpp
--- a/Engine/Source/Runtime/SignalProcessing/Private/Filter.cpp
+++ b/Engine/Source/Runtime/SignalProcessing/Private/Filter.cpp
@@ -614,6 +614,32 @@ namespace Audio
 		}
 	}
 
+
+	void FStateVariableFilter::ProcessAudio(const float* InSamples, const int32 InNumSamples
+		, float* LpfOutput, float* HpfOutput, float* BpfOutput, float* BsfOutput)
+	{
+		for (int32 SampleIndex = 0; SampleIndex < InNumSamples; SampleIndex += NumChannels)
+		{
+			for (int32 Channel = 0; Channel < NumChannels; ++Channel)
+			{
+				const float HPF = InputScale * (InSamples[SampleIndex + Channel] - Feedback * FilterState[Channel].Z1_1 - FilterState[Channel].Z1_2);
+				float BPF = Audio::FastTanh(A0 * HPF + FilterState[Channel].Z1_1);
+
+				const float LPF = A0 * BPF + FilterState[Channel].Z1_2;
+				const float Dampening = 0.5f / Q;
+				const float BSF = BandStopParam * HPF + (1.0f - BandStopParam) * LPF;
+
+				FilterState[Channel].Z1_1 = A0 * HPF + BPF;
+				FilterState[Channel].Z1_2 = A0 * BPF + LPF;
+
+				LpfOutput[SampleIndex + Channel] = LPF;
+				HpfOutput[SampleIndex + Channel] = HPF;
+				BpfOutput[SampleIndex + Channel] = BPF;
+				BsfOutput[SampleIndex + Channel] = BSF;
+			}
+		}
+	}
+
 	FLadderFilter::FLadderFilter()
 		: K(0.0f)
 		, Gamma(0.0f)
diff --git a/Engine/Source/Runtime/SignalProcessing/Private/InterpolatedOnePole.cpp b/Engine/Source/Runtime/SignalProcessing/Private/InterpolatedOnePole.cpp
--- a/Engine/Source/Runtime/SignalProcessing/Private/InterpolatedOnePole.cpp
+++ b/Engine/Source/Runtime/SignalProcessing/Private/InterpolatedOnePole.cpp
@@ -49,7 +49,7 @@ namespace Audio
 		}
 	}
 
-	void FInterpolatedLPF::ProcessAudioFrame(float* RESTRICT InputFrame, float* RESTRICT OutputFrame)
+	void FInterpolatedLPF::ProcessAudioFrame(const float* RESTRICT InputFrame, float* RESTRICT OutputFrame)
 	{
 		B1Curr += B1Delta; // step forward coefficient
 
@@ -75,7 +75,7 @@ namespace Audio
 		}
 	}
 
-	void FInterpolatedLPF::ProcessAudioBuffer(float *RESTRICT InputBuffer, float *RESTRICT OutputBuffer, const int32 NumSamples)
+	void FInterpolatedLPF::ProcessAudioBuffer(const float *RESTRICT InputBuffer, float *RESTRICT OutputBuffer, const int32 NumSamples)
 	{
 		for (int SampleIndex = 0; SampleIndex < NumSamples; ++SampleIndex)
 		{
@@ -141,7 +141,7 @@ namespace Audio
 
 		if (!FMath::IsNearlyEqual(InTargetFrequency, CutoffFrequency))
 		{
-			CutoffFrequency = FMath::Min(InTargetFrequency, NyquistLimit);
+			CutoffFrequency = FMath::Clamp(InTargetFrequency, 0.f, NyquistLimit);
 
 			// G computation is a reduced form of the following set of equations:
 			// OmegaDigital = 2.0f * PI * CutoffFrequency;
@@ -159,7 +159,7 @@ namespace Audio
 		}
 	}
 
-	void FInterpolatedHPF::ProcessAudioFrame(float* RESTRICT InputFrame, float* RESTRICT OutputFrame)
+	void FInterpolatedHPF::ProcessAudioFrame(const float* RESTRICT InputFrame, float* RESTRICT OutputFrame)
 	{
 		A0Curr += A0Delta; // step forward coefficient
 
@@ -174,7 +174,7 @@ namespace Audio
 		}
 	}
 
-	void FInterpolatedHPF::ProcessAudioBuffer(float *RESTRICT InputBuffer, float *RESTRICT OutputBuffer, const int32 NumSamples)
+	void FInterpolatedHPF::ProcessAudioBuffer(const float *RESTRICT InputBuffer, float *RESTRICT OutputBuffer, const int32 NumSamples)
 	{
 		for (int SampleIndex = 0; SampleIndex < NumSamples; ++SampleIndex)
 		{
diff --git a/Engine/Source/Runtime/SignalProcessing/Public/DSP/Filter.h b/Engine/Source/Runtime/SignalProcessing/Public/DSP/Filter.h
--- a/Engine/Source/Runtime/SignalProcessing/Public/DSP/Filter.h
+++ b/Engine/Source/Runtime/SignalProcessing/Public/DSP/Filter.h
@@ -262,6 +262,9 @@ namespace Audio
 		virtual void Update() override;
 		virtual void ProcessAudio(const float* InSamples, const int32 InNumSamples, float* OutSamples) override;
 
+		void ProcessAudio(const float* InSamples, const int32 InNumSamples
+			, float* LpfOutput, float* HpfOutput, float* BpfOutput, float* BsfOutput);
+
 	protected:
 		float InputScale;
 		float A0;
@@ -293,8 +296,10 @@ namespace Audio
 		virtual void Update() override;
 		virtual void SetQ(const float InQ) override;
 		virtual void SetPassBandGainCompensation(const float InPassBandGainCompensation) override;
+
 		virtual void ProcessAudio(const float* InSamples, const int32 InNumSamples, float* OutSamples) override;
 
+
 	protected:
 
 		// Ladder filter is made of 4 LPF filters
diff --git a/Engine/Source/Runtime/SignalProcessing/Public/DSP/InterpolatedOnePole.h b/Engine/Source/Runtime/SignalProcessing/Public/DSP/InterpolatedOnePole.h
--- a/Engine/Source/Runtime/SignalProcessing/Public/DSP/InterpolatedOnePole.h
+++ b/Engine/Source/Runtime/SignalProcessing/Public/DSP/InterpolatedOnePole.h
@@ -32,10 +32,10 @@ namespace Audio
 		void StartFrequencyInterpolation(const float InTargetFrequency, const int32 InterpLength = 1);
 
 		// interpolates coefficient and processes a sample
-		void ProcessAudioFrame(float* RESTRICT InputFrame, float* RESTRICT OutputFrame);
+		void ProcessAudioFrame(const float* RESTRICT InputFrame, float* RESTRICT OutputFrame);
 
 		// interpolates coefficient and processes a buffer
-		void ProcessAudioBuffer(float* RESTRICT InputBuffer, float* RESTRICT OutputBuffer, const int32 NumSamples);
+		void ProcessAudioBuffer(const float* RESTRICT InputBuffer, float* RESTRICT OutputBuffer, const int32 NumSamples);
 
 		/*
 			StopFrequencyInterpolation() needs to be called manually when the interpolation should be done.
@@ -97,10 +97,10 @@ namespace Audio
 		void StartFrequencyInterpolation(const float InTargetFrequency, const int32 InterpLength = 1);
 
 		// interpolates coefficient and processes a sample
-		void ProcessAudioFrame(float* RESTRICT InputFrame, float* RESTRICT OutputFrame);
+		void ProcessAudioFrame(const float* RESTRICT InputFrame, float* RESTRICT OutputFrame);
 
 		// interpolates coefficient and processes a buffer
-		void ProcessAudioBuffer(float* RESTRICT InputBuffer, float* RESTRICT OutputBuffer, const int32 NumSamples);
+		void ProcessAudioBuffer(const float* RESTRICT InputBuffer, float* RESTRICT OutputBuffer, const int32 NumSamples);
 
 		/*
 			StopFrequencyInterpolation() needs to be called manually when the interpolation should be done.
