Add ability to define default literals on input vertex in Metasound core
- Update default values for standard nodes to sensible values

diff --git a/Engine/Plugins/Runtime/Metasound/Source/MetasoundEditor/Private/MetasoundEditorGraphBuilder.cpp b/Engine/Plugins/Runtime/Metasound/Source/MetasoundEditor/Private/MetasoundEditorGraphBuilder.cpp
--- a/Engine/Plugins/Runtime/Metasound/Source/MetasoundEditor/Private/MetasoundEditorGraphBuilder.cpp
+++ b/Engine/Plugins/Runtime/Metasound/Source/MetasoundEditor/Private/MetasoundEditorGraphBuilder.cpp
@@ -1,6 +1,7 @@
 // Copyright Epic Games, Inc. All Rights Reserved.
 #include "MetasoundEditorGraphBuilder.h"
 
+#include "AssetRegistry/AssetRegistryModule.h"
 #include "EdGraph/EdGraph.h"
 #include "EdGraph/EdGraphNode.h"
 #include "EdGraph/EdGraphPin.h"
@@ -53,70 +54,6 @@ namespace Metasound
 
 				InMetasound.MarkPackageDirty();
 			}
-
-			void SetInputLiteral(Frontend::FNodeHandle InInputNodeHandle, FGuid PointID, const FName InTypeName, const FLiteral& InDefaultValue)
-			{
-				Frontend::FGraphHandle GraphHandle = InInputNodeHandle->GetOwningGraph();
-
-				const FString Name = InInputNodeHandle->GetNodeName();
-				switch (InDefaultValue.GetType())
-				{
-					case ELiteralType::Boolean:
-					{
-						GraphHandle->SetDefaultInputToLiteral(Name, PointID, InDefaultValue.Value.Get<bool>());
-					}
-					break;
-
-					case ELiteralType::Float:
-					{
-						GraphHandle->SetDefaultInputToLiteral(Name, PointID, InDefaultValue.Value.Get<float>());
-					}
-					break;
-
-					case ELiteralType::Integer:
-					{
-						GraphHandle->SetDefaultInputToLiteral(Name, PointID, InDefaultValue.Value.Get<int32>());
-					}
-					break;
-
-					case ELiteralType::String:
-					{
-						GraphHandle->SetDefaultInputToLiteral(Name, PointID, InDefaultValue.Value.Get<FString>());
-					}
-					break;
-
-					case ELiteralType::UObjectProxy:
-					{
-						UClass* ClassToUse = FMetasoundFrontendRegistryContainer::Get()->GetLiteralUClassForDataType(InTypeName);
-						if (ClassToUse)
-						{
-							GraphHandle->SetDefaultInputToLiteral(Name, PointID, ClassToUse->ClassDefaultObject);
-						}
-					}
-					break;
-
-					case ELiteralType::UObjectProxyArray:
-					{
-						// TODO: Support default UObject array value on node
-						UClass* ClassToUse = FMetasoundFrontendRegistryContainer::Get()->GetLiteralUClassForDataType(InTypeName);
-						if (ClassToUse)
-						{
-							TArray<UObject*> ObjectArray;
-							ObjectArray.Add(ClassToUse->ClassDefaultObject);
-							GraphHandle->SetDefaultInputToLiteral(Name, PointID, ObjectArray);
-						}
-					}
-					break;
-
-					case ELiteralType::Invalid:
-					case ELiteralType::None:
-					default:
-					{
-						static_assert(static_cast<int32>(ELiteralType::Invalid) == 7, "Possible missing ELiteralType case coverage");
-					}
-					break;
-				}
-			}
 		}
 
 		UEdGraphNode* FGraphBuilder::AddNode(UObject& InMetasound, Frontend::FNodeHandle& InNodeHandle, bool bInSelectNewNode)
@@ -229,16 +166,13 @@ namespace Metasound
 			return AddNode(InMetasound, NodeHandle, bInSelectNewNode);
 		}
 
-		void FGraphBuilder::AddOrUpdateLiteralInput(UObject& InMetasound, Frontend::FNodeHandle InNodeHandle, const UEdGraphPin& InInputPin)
+		void FGraphBuilder::AddOrUpdateLiteralInput(UObject& InMetasound, Frontend::FNodeHandle InNodeHandle, UEdGraphPin& InInputPin)
 		{
 			using namespace Metasound::Frontend;
 
-			const FScopedTransaction Transaction(LOCTEXT("SetMetasoundGraphNode", "Set Metasound Literal Input"));
+			const FScopedTransaction Transaction(LOCTEXT("SetMetasoundGraphNode", "Add/Update Metasound Literal Input"));
 
 			const FString& InInputName = InInputPin.GetName();
-			const FString& InStringValue = InInputPin.DefaultValue;
-			const UObject* InObjectValue = InInputPin.DefaultObject;
-
 			TArray<FInputHandle> InputHandles = InNodeHandle->GetInputsWithVertexName(InInputName);
 			if (!ensureAlways(InputHandles.Num() == 1))
 			{
@@ -248,44 +182,72 @@ namespace Metasound
 			FInputHandle InputHandle = InputHandles[0];
 			IMetasoundEditorModule& EditorModule = FModuleManager::GetModuleChecked<IMetasoundEditorModule>("MetasoundEditor");
 
-			FLiteral Literal;
+			FMetasoundFrontendLiteral PinDataTypeDefaultLiteral;
+
+			const FString& InStringValue = InInputPin.DefaultValue;
 			const FName TypeName = InputHandle->GetDataType();
 			const FEditorDataType DataType = EditorModule.FindDataType(TypeName);
 			switch (DataType.RegistryInfo.PreferredLiteralType)
 			{
 				case ELiteralType::Boolean:
 				{
-					Literal.Value.Set<bool>(FCString::ToBool(*InStringValue));
+					PinDataTypeDefaultLiteral.Set(FCString::ToBool(*InStringValue));
 				}
 				break;
 
 				case ELiteralType::Float:
 				{
-					Literal.Value.Set<float>(FCString::Atof(*InStringValue));
+					PinDataTypeDefaultLiteral.Set(FCString::Atof(*InStringValue));
 				}
 				break;
 
 				case ELiteralType::Integer:
 				{
-					Literal.Value.Set<int32>(FCString::Atoi(*InStringValue));
+					PinDataTypeDefaultLiteral.Set(FCString::Atoi(*InStringValue));
 				}
 				break;
 
 				case ELiteralType::String:
 				{
-					Literal.Value.Set<FString>(InStringValue);
+					PinDataTypeDefaultLiteral.Set(InStringValue);
 				}
 				break;
 
-				// TODO: Support UObjects/UObject arrays for default literals
 				case ELiteralType::UObjectProxy:
-				case ELiteralType::UObjectProxyArray:
 				{
-					return;
+					bool bObjectFound = false;
+					if (!InInputPin.DefaultValue.IsEmpty())
+					{
+						FMetasoundFrontendRegistryContainer* FrontendRegistry = FMetasoundFrontendRegistryContainer::Get();
+						check(FrontendRegistry);
+
+						if (UClass* Class = FrontendRegistry->GetLiteralUClassForDataType(TypeName))
+						{
+							FAssetRegistryModule& AssetRegistryModule = FModuleManager::LoadModuleChecked<FAssetRegistryModule>("AssetRegistry");
+							FARFilter Filter;
+							Filter.bRecursiveClasses = false;
+							Filter.ObjectPaths.Add(FName(*InInputPin.DefaultValue));
+							Filter.ClassNames.Add(Class->GetFName());
+							TArray<FAssetData> AssetData;
+							AssetRegistryModule.Get().GetAssets(Filter, AssetData);
+
+							if (AssetData.Num() > 0)
+							{
+								PinDataTypeDefaultLiteral.Set(AssetData[0].GetAsset());
+								bObjectFound = true;
+							}
+						}
+					}
+					if (!bObjectFound)
+					{
+						return;
+					}
 				}
 				break;
 
-				// If no literal supported, no need for private input
+				// TODO: Remove UObject proxy arrays.
+				// If no literal supported, no need for literal input
+				case ELiteralType::UObjectProxyArray:
 				case ELiteralType::None:
 				{
 					return;
@@ -301,17 +263,33 @@ namespace Metasound
 				break;
 			}
 
-			bool bNewValueSet = false;
 			FOutputHandle OutputHandle = InputHandle->GetCurrentlyConnectedOutput();
 			if (!OutputHandle->IsValid())
 			{
 				FMetasoundFrontendNodeStyle Style;
 				Style.Display.Visibility = EMetasoundFrontendNodeStyleDisplayVisibility::Hidden;
 
-				Frontend::FNodeHandle NewInputNode = AddInputNodeHandle(InMetasound, GenerateUniqueInputName(InMetasound, "LiteralInput"), TypeName, Style, nullptr /* ToolTip */, &Literal);
+				const FString NewInputName = GenerateUniqueInputName(InMetasound, "LiteralInput");
+				Frontend::FNodeHandle NewInputNode = AddInputNodeHandle(InMetasound, NewInputName, TypeName, Style, nullptr /* ToolTip */, &PinDataTypeDefaultLiteral);
 				NewInputNode->SetNodeStyle(Style);
 
-				bNewValueSet = true;
+				Frontend::FGraphHandle GraphHandle = InNodeHandle->GetOwningGraph();
+				const FString Name = InNodeHandle->GetNodeName();
+
+				const TArray<FGuid> PointIDs = GraphHandle->GetDefaultIDsForInputVertex(NewInputNode->GetNodeName());
+				// TODO: Support multiple literal point ids for arrays
+				if (ensureAlways(PointIDs.Num() == 1))
+				{
+					InNodeHandle->GetClassDescription();
+					const TArray<FMetasoundFrontendVertexLiteral>& Defaults = InputHandle->GetDefaults();
+					if (!Defaults.IsEmpty() && Defaults[0].Value.Type != EMetasoundFrontendLiteralType::Invalid && Defaults[0].Value.Type != EMetasoundFrontendLiteralType::None)
+					{
+						const FString NewInputNodeName = NewInputNode->GetNodeName();
+						GraphHandle->SetDefaultInputToFrontendLiteral(NewInputNodeName, PointIDs[0], TypeName, Defaults[0].Value);
+						InInputPin.DefaultValue = Defaults[0].Value.ToString();
+					}
+				}
+
 				TArray<Metasound::Frontend::FOutputHandle> OutputHandles = NewInputNode->GetOutputs();
 				if (ensureAlways(OutputHandles.Num() == 1))
 				{
@@ -319,31 +297,32 @@ namespace Metasound
 				}
 
 				ensureAlways(InputHandle->Connect(*OutputHandle));
-			}
 
-			if (OutputHandle->IsValid())
+				InMetasound.PostEditChange();
+				InMetasound.MarkPackageDirty();
+			}
+			else
 			{
 				FNodeHandle InputNode = OutputHandle->GetOwningNode();
 				const FMetasoundFrontendNodeStyle& Style = InputNode->GetNodeStyle();
 				if (InputNode->IsValid() && Style.Display.Visibility == EMetasoundFrontendNodeStyleDisplayVisibility::Hidden)
 				{
-					if (!bNewValueSet)
+					Frontend::FGraphHandle GraphHandle = InputNode->GetOwningGraph();
+					const TArray<FGuid> PointIDs = GraphHandle->GetDefaultIDsForInputVertex(InputNode->GetNodeName());
+
+					// TODO: Support multiple literal point ids for arrays
+					if (ensureAlways(PointIDs.Num() == 1))
 					{
-						// TODO: Support multiple literal point ids for arrays
-						Frontend::FGraphHandle GraphHandle = InputNode->GetOwningGraph();
-						TArray<FGuid> PointIDs = GraphHandle->GetDefaultIDsForInputVertex(InputNode->GetNodeName());
-						if (ensureAlways(PointIDs.Num() == 1))
-						{
-							GraphBuilderPrivate::SetInputLiteral(InputNode, PointIDs[0], TypeName, Literal);
-						}
+						GraphHandle->SetDefaultInputToFrontendLiteral(InputNode->GetNodeName(), PointIDs[0], TypeName, PinDataTypeDefaultLiteral);
 					}
+
 					InMetasound.PostEditChange();
 					InMetasound.MarkPackageDirty();
 				}
 			}
 		}
 
-		Frontend::FNodeHandle FGraphBuilder::AddInputNodeHandle(UObject& InMetasound, const FString& InName, const FName InTypeName, const FMetasoundFrontendNodeStyle& InNodeStyle, const FText* InToolTip, const FLiteral* InDefaultValue)
+		Frontend::FNodeHandle FGraphBuilder::AddInputNodeHandle(UObject& InMetasound, const FString& InName, const FName InTypeName, const FMetasoundFrontendNodeStyle& InNodeStyle, const FText* InToolTip, const FMetasoundFrontendLiteral* InDefaultValue)
 		{
 			FMetasoundAssetBase* MetasoundAsset = IMetasoundUObjectRegistry::Get().GetObjectAsAssetBase(&InMetasound);
 			check(MetasoundAsset);
@@ -368,28 +347,22 @@ namespace Metasound
 			Description.Defaults.Add(DefaultVertex);
 
 			Frontend::FNodeHandle NodeHandle = GraphHandle->AddInputVertex(Description);
-			if (!ensureAlways(NodeHandle->IsValid()))
-			{
-				return NodeHandle;
-			}
-
-			GraphHandle->SetInputDisplayName(InName, FText::FromString(InName));
-
-			Metasound::FLiteral LiteralParam = Frontend::GetDefaultParamForDataType(InTypeName);
-			if (!LiteralParam.IsValid())
+			if (ensureAlways(NodeHandle->IsValid()))
 			{
-				return NodeHandle;
-			}
+				GraphHandle->SetInputDisplayName(InName, FText::FromString(InName));
 
-			if (InDefaultValue)
-			{
-				if (!ensureAlways(InDefaultValue->GetType() == LiteralParam.GetType()))
+				if (InDefaultValue)
 				{
-					return NodeHandle;
+					if (InDefaultValue->Type != EMetasoundFrontendLiteralType::None && InDefaultValue->Type != EMetasoundFrontendLiteralType::Invalid)
+					{
+						GraphHandle->SetDefaultInputToFrontendLiteral(NodeHandle->GetNodeName(), PointID, InTypeName, *InDefaultValue);
+					}
 				}
-			};
-
-			GraphBuilderPrivate::SetInputLiteral(NodeHandle, PointID, InTypeName, InDefaultValue ? *InDefaultValue : LiteralParam);
+				else
+				{
+					GraphHandle->SetDefaultInputToTypeDefaultLiteral(NodeHandle->GetNodeName(), PointID, InTypeName);
+				}
+			}
 
 			return NodeHandle;
 		}
diff --git a/Engine/Plugins/Runtime/Metasound/Source/MetasoundEditor/Private/MetasoundEditorGraphBuilder.h b/Engine/Plugins/Runtime/Metasound/Source/MetasoundEditor/Private/MetasoundEditorGraphBuilder.h
--- a/Engine/Plugins/Runtime/Metasound/Source/MetasoundEditor/Private/MetasoundEditorGraphBuilder.h
+++ b/Engine/Plugins/Runtime/Metasound/Source/MetasoundEditor/Private/MetasoundEditorGraphBuilder.h
@@ -64,10 +64,13 @@ namespace Metasound
 
 			static UEdGraphNode* AddInput(UObject& InMetasound, const FString& InName, const FName InTypeName, const FMetasoundFrontendNodeStyle& InNodeStyle, const FText& InToolTip, bool bInSelectNewNode = true);
 
-			static void AddOrUpdateLiteralInput(UObject& InMetasound, Frontend::FNodeHandle InNodeHandle, const UEdGraphPin& InInputPin);
+			// Adds or updates a hidden input that is set to the default literal value.  If the literal input is added, value defaults to
+			// 1. what is set on the input vertex of the node's inputer vertex interface, and if not defined there,
+			// 2. the pin's DataType literal default.
+			static void AddOrUpdateLiteralInput(UObject& InMetasound, Frontend::FNodeHandle InNodeHandle, UEdGraphPin& InInputPin);
 
 			// Adds an input node handle with the given class info
-			static Frontend::FNodeHandle AddInputNodeHandle(UObject& InMetasound, const FString& InName, const FName InTypeName, const FMetasoundFrontendNodeStyle& InNodeStyle, const FText* InToolTip = nullptr, const FLiteral* InDefaultValue = nullptr);
+			static Frontend::FNodeHandle AddInputNodeHandle(UObject& InMetasound, const FString& InName, const FName InTypeName, const FMetasoundFrontendNodeStyle& InNodeStyle, const FText* InToolTip = nullptr, const FMetasoundFrontendLiteral* InDefaultValue = nullptr);
 
 			static UEdGraphNode* AddOutput(UObject& InMetasound, const FString& InName, const FName InTypeName, const FMetasoundFrontendNodeStyle& InNodeStyle, const FText& InToolTip, bool bInSelectNewNode = true);
 
diff --git a/Engine/Plugins/Runtime/Metasound/Source/MetasoundFrontend/Private/MetasoundFrontend.cpp b/Engine/Plugins/Runtime/Metasound/Source/MetasoundFrontend/Private/MetasoundFrontend.cpp
--- a/Engine/Plugins/Runtime/Metasound/Source/MetasoundFrontend/Private/MetasoundFrontend.cpp
+++ b/Engine/Plugins/Runtime/Metasound/Source/MetasoundFrontend/Private/MetasoundFrontend.cpp
@@ -80,6 +80,15 @@ namespace Metasound
 				ClassInput.Metadata.DisplayName = FText::FromString(InputTuple.Value.GetVertexName());
 				ClassInput.Metadata.Description = InputTuple.Value.GetDescription();
 
+				const FLiteral& DefaultLiteral = InputTuple.Value.GetDefaultValue();
+				if (DefaultLiteral.GetType() != ELiteralType::None)
+				{
+					FMetasoundFrontendVertexLiteral VertexLiteral;
+					VertexLiteral.Value.SetFromLiteral(DefaultLiteral);
+					VertexLiteral.PointID = FGuid::NewGuid();
+					ClassInput.Defaults.Add(VertexLiteral);
+				}
+
 				ClassDescription.Interface.Inputs.Add(ClassInput);
 			}
 
diff --git a/Engine/Plugins/Runtime/Metasound/Source/MetasoundFrontend/Private/MetasoundFrontendDocument.cpp b/Engine/Plugins/Runtime/Metasound/Source/MetasoundFrontend/Private/MetasoundFrontendDocument.cpp
--- a/Engine/Plugins/Runtime/Metasound/Source/MetasoundFrontend/Private/MetasoundFrontendDocument.cpp
+++ b/Engine/Plugins/Runtime/Metasound/Source/MetasoundFrontend/Private/MetasoundFrontendDocument.cpp
@@ -52,6 +52,88 @@ void FMetasoundFrontendLiteral::Set(const TArray<UObject*>& InValue)
 	Type = EMetasoundFrontendLiteralType::UObjectArray;
 }
 
+void FMetasoundFrontendLiteral::SetFromLiteral(const Metasound::FLiteral& InLiteral)
+{
+	using namespace Metasound;
+
+	Clear();
+	switch (InLiteral.GetType())
+	{
+		case ELiteralType::Boolean:
+		{
+			Set(InLiteral.Value.Get<bool>());
+		}
+		break;
+
+		case ELiteralType::Float:
+		{
+			Set(InLiteral.Value.Get<float>());
+		}
+		break;
+
+		case ELiteralType::Integer:
+		{
+			Set(InLiteral.Value.Get<int32>());
+		}
+		break;
+
+		case ELiteralType::String:
+		{
+			Set(InLiteral.Value.Get<FString>());
+		}
+		break;
+
+		case ELiteralType::UObjectProxy:
+		case ELiteralType::UObjectProxyArray:
+		case ELiteralType::None:
+		case ELiteralType::Invalid:
+		default:
+		{
+			static_assert(static_cast<int32>(ELiteralType::Invalid) == 7, "Possible missing literal type switch coverage");
+		}
+	}
+}
+
+FString FMetasoundFrontendLiteral::ToString() const
+{
+	switch (Type)
+	{
+		case EMetasoundFrontendLiteralType::Bool:
+		return FString::Printf(TEXT("%s"), AsBool ? TEXT("true") : TEXT("false"));
+
+		case EMetasoundFrontendLiteralType::Float:
+		return FString::Printf(TEXT("%f"), AsFloat);
+
+		case EMetasoundFrontendLiteralType::Integer:
+		return FString::Printf(TEXT("%d"), AsInteger);
+
+		case EMetasoundFrontendLiteralType::String:
+		return AsString;
+
+		case EMetasoundFrontendLiteralType::UObject:
+		{
+			if (AsUObject)
+			{
+				return AsUObject->GetFullName();
+			}
+			else
+			{
+				return FString();
+			}
+		}
+
+		case EMetasoundFrontendLiteralType::UObjectArray:
+		case EMetasoundFrontendLiteralType::None:
+		case EMetasoundFrontendLiteralType::Invalid:
+		default:
+		{
+			static_assert(static_cast<int32>(EMetasoundFrontendLiteralType::Invalid) == 7, "Possible missing literal type switch coverage");
+		}
+
+		return FString();
+	}
+}
+
 void FMetasoundFrontendLiteral::Clear()
 {
 	Type = EMetasoundFrontendLiteralType::None;
diff --git a/Engine/Plugins/Runtime/Metasound/Source/MetasoundFrontend/Private/MetasoundFrontendInvalidController.h b/Engine/Plugins/Runtime/Metasound/Source/MetasoundFrontend/Private/MetasoundFrontendInvalidController.h
--- a/Engine/Plugins/Runtime/Metasound/Source/MetasoundFrontend/Private/MetasoundFrontendInvalidController.h
+++ b/Engine/Plugins/Runtime/Metasound/Source/MetasoundFrontend/Private/MetasoundFrontendInvalidController.h
@@ -87,11 +87,15 @@ namespace Metasound
 			bool IsValid() const override { return false; }
 			FGuid GetID() const override { return Metasound::FrontendInvalidID; }
 			bool IsConnected() const override { return false; }
-			const FName& GetDataType() const override { return MetasoundFrontendInvalidControllerPrivate::GetInvalid<FName>(); }  
+			const FName& GetDataType() const override { return MetasoundFrontendInvalidControllerPrivate::GetInvalid<FName>(); }
 			const FString& GetName() const override { return MetasoundFrontendInvalidControllerPrivate::GetInvalid<FString>(); }
 			const FText& GetDisplayName() const override { return FText::GetEmpty(); }
 			const FText& GetTooltip() const override { return FText::GetEmpty(); }
-
+			const TArray<FMetasoundFrontendVertexLiteral>& GetDefaults() const override
+			{
+				static const TArray<FMetasoundFrontendVertexLiteral> InvalidDefaults;
+				return InvalidDefaults;
+			}
 			FGuid GetOwningNodeID() const override { return Metasound::FrontendInvalidID; }
 			TSharedRef<INodeController> GetOwningNode() override;
 			TSharedRef<const INodeController> GetOwningNode() const override;
@@ -242,13 +246,8 @@ namespace Metasound
 
 			// These can be used to set the default value for a given input on this graph.
 			// @returns false if the input name couldn't be found, or if the literal type was incompatible with the Data Type of this input.
-			bool SetDefaultInputToLiteral(const FString& InInputName, FGuid InPointID, bool bInValue) override { return false; }
-			bool SetDefaultInputToLiteral(const FString& InInputName, FGuid InPointID, int32 InValue) override { return false; }
-			bool SetDefaultInputToLiteral(const FString& InInputName, FGuid InPointID, float InValue) override { return false; }
-			bool SetDefaultInputToLiteral(const FString& InInputName, FGuid InPointID, const FString& InValue) override { return false; }
-			bool SetDefaultInputToLiteral(const FString& InInputName, FGuid InPointID, UObject* InValue) override { return false; }
-			bool SetDefaultInputToLiteral(const FString& InInputName, FGuid InPointID, const TArray<UObject*>& InValue) override { return false; }
-
+			bool SetDefaultInputToFrontendLiteral(const FString& InInputName, FGuid InPointID, FName InDataTypeName, const FMetasoundFrontendLiteral& InLiteral) override { return false; }
+			bool SetDefaultInputToTypeDefaultLiteral(const FString& InInputName, FGuid InPointID, FName InDataTypeName) override { return false; }
 
 			// Set the display name for the input with the given name
 			void SetInputDisplayName(FString InName, const FText& InDisplayName) override {}
diff --git a/Engine/Plugins/Runtime/Metasound/Source/MetasoundFrontend/Private/MetasoundFrontendRegistries.cpp b/Engine/Plugins/Runtime/Metasound/Source/MetasoundFrontend/Private/MetasoundFrontendRegistries.cpp
--- a/Engine/Plugins/Runtime/Metasound/Source/MetasoundFrontend/Private/MetasoundFrontendRegistries.cpp
+++ b/Engine/Plugins/Runtime/Metasound/Source/MetasoundFrontend/Private/MetasoundFrontendRegistries.cpp
@@ -249,7 +249,7 @@ UClass* FMetasoundFrontendRegistryContainer::GetLiteralUClassForDataType(FName I
 {
 	if (!DataTypeRegistry.Contains(InDataType))
 	{
-		ensureAlwaysMsgf(false, TEXT("couldn't find DataType %s in the registry."), *InDataType.ToString());
+		ensureAlwaysMsgf(false, TEXT("DataType %s not registered."), *InDataType.ToString());
 		return nullptr;
 	}
 	else
@@ -262,7 +262,7 @@ bool FMetasoundFrontendRegistryContainer::DoesDataTypeSupportLiteralType(FName I
 {
 	if (!DataTypeRegistry.Contains(InDataType))
 	{
-		ensureAlwaysMsgf(false, TEXT("couldn't find DataType %s in the registry."), *InDataType.ToString());
+		ensureAlwaysMsgf(false, TEXT("DataType %s not registered."), *InDataType.ToString());
 		return false;
 	}
 
diff --git a/Engine/Plugins/Runtime/Metasound/Source/MetasoundFrontend/Private/MetasoundFrontendStandardController.cpp b/Engine/Plugins/Runtime/Metasound/Source/MetasoundFrontend/Private/MetasoundFrontendStandardController.cpp
--- a/Engine/Plugins/Runtime/Metasound/Source/MetasoundFrontend/Private/MetasoundFrontendStandardController.cpp
+++ b/Engine/Plugins/Runtime/Metasound/Source/MetasoundFrontend/Private/MetasoundFrontendStandardController.cpp
@@ -10,12 +10,14 @@
 #include "MetasoundFrontendController.h"
 #include "MetasoundFrontendDocument.h"
 #include "MetasoundFrontendGraph.h"
+#include "MetasoundFrontendRegistries.h"
 #include "MetasoundJsonBackend.h"
 #include "MetasoundOperatorBuilder.h"
 #include "StructDeserializer.h"
 #include "StructSerializer.h"
 #include "Serialization/MemoryReader.h"
 
+
 namespace Metasound
 {
 	namespace Frontend
@@ -917,6 +919,17 @@ namespace Metasound
 			return FrontendControllerIntrinsics::GetInvalidValueConstRef<FText>();
 		}
 
+		const TArray<FMetasoundFrontendVertexLiteral>& FNodeInputController::GetDefaults() const
+		{
+			if (OwningGraphClassInputPtr.IsValid())
+			{
+				return OwningGraphClassInputPtr->Defaults;
+			}
+
+			static const TArray<FMetasoundFrontendVertexLiteral> InvalidDefaults;
+			return InvalidDefaults;
+		}
+
 		const FText& FNodeInputController::GetTooltip() const
 		{
 			if (OwningGraphClassInputPtr.IsValid())
@@ -2297,37 +2310,104 @@ namespace Metasound
 			return nullptr;
 		}
 
-
-		// These can be used to set the default value for a given input on this graph.
-		// @returns false if the input name couldn't be found, or if the literal type was incompatible with the Data Type of this input.
-		bool FGraphController::SetDefaultInputToLiteral(const FString& InInputName, FGuid InPointID, bool bInValue)
+		bool FGraphController::SetDefaultInputToFrontendLiteral(const FString& InInputName, FGuid InPointID, FName InDataTypeName, const FMetasoundFrontendLiteral& InLiteral)
 		{
-			return SetDefaultInputToLiteralInternal(InInputName, InPointID, bInValue);
-		}
+			switch (InLiteral.Type)
+			{
+				case EMetasoundFrontendLiteralType::Bool:
+				{
+					return SetDefaultInputToLiteralInternal(InInputName, InPointID, InLiteral.AsBool);
+				}
+				break;
 
-		bool FGraphController::SetDefaultInputToLiteral(const FString& InInputName, FGuid InPointID, int32 InValue)
-		{
-			return SetDefaultInputToLiteralInternal(InInputName, InPointID, InValue);
-		}
+				case EMetasoundFrontendLiteralType::Float:
+				{
+					return SetDefaultInputToLiteralInternal(InInputName, InPointID, InLiteral.AsFloat);
+				}
+				break;
 
-		bool FGraphController::SetDefaultInputToLiteral(const FString& InInputName, FGuid InPointID, float InValue)
-		{
-			return SetDefaultInputToLiteralInternal(InInputName, InPointID, InValue);
-		}
+				case EMetasoundFrontendLiteralType::Integer:
+				{
+					return SetDefaultInputToLiteralInternal(InInputName, InPointID, InLiteral.AsInteger);
+				}
+				break;
 
-		bool FGraphController::SetDefaultInputToLiteral(const FString& InInputName, FGuid InPointID, const FString& InValue)
-		{
-			return SetDefaultInputToLiteralInternal(InInputName, InPointID, InValue);
-		}
+				case EMetasoundFrontendLiteralType::String:
+				{
+					return SetDefaultInputToLiteralInternal(InInputName, InPointID, InLiteral.AsString);
+				}
+				break;
 
-		bool FGraphController::SetDefaultInputToLiteral(const FString& InInputName, FGuid InPointID, UObject* InValue)
-		{
-			return SetDefaultInputToLiteralInternal(InInputName, InPointID, InValue);
+				case EMetasoundFrontendLiteralType::UObject:
+				{
+					return SetDefaultInputToLiteralInternal(InInputName, InPointID, InLiteral.AsUObject);
+				}
+				break;
+
+				// TODO: Remove as array support will not be specialized to object type
+				case EMetasoundFrontendLiteralType::UObjectArray:
+				{
+					return false;
+				}
+				break;
+
+				case EMetasoundFrontendLiteralType::Invalid:
+				case EMetasoundFrontendLiteralType::None:
+				default:
+				{
+					static_assert(static_cast<int32>(ELiteralType::Invalid) == 7, "Possible missing ELiteralType case coverage");
+					return false;
+				}
+				break;
+			}
 		}
 
-		bool FGraphController::SetDefaultInputToLiteral(const FString& InInputName, FGuid InPointID, const TArray<UObject*>& InValue)
+		bool FGraphController::SetDefaultInputToTypeDefaultLiteral(const FString& InInputName, FGuid InPointID, FName InDataTypeName)
 		{
-			return SetDefaultInputToLiteralInternal(InInputName, InPointID, InValue);
+			Metasound::FLiteral LiteralParam = Frontend::GetDefaultParamForDataType(InDataTypeName);
+			switch (LiteralParam.GetType())
+			{
+				case ELiteralType::Boolean:
+				{
+					return SetDefaultInputToLiteralInternal(InInputName, InPointID, LiteralParam.Value.Get<bool>());
+				}
+				break;
+
+				case ELiteralType::Float:
+				{
+					return SetDefaultInputToLiteralInternal(InInputName, InPointID, LiteralParam.Value.Get<float>());
+				}
+				break;
+
+				case ELiteralType::Integer:
+				{
+					return SetDefaultInputToLiteralInternal(InInputName, InPointID, LiteralParam.Value.Get<int32>());
+				}
+				break;
+
+				case ELiteralType::String:
+				{
+					return SetDefaultInputToLiteralInternal(InInputName, InPointID, LiteralParam.Value.Get<FString>());
+				}
+				break;
+
+				// TODO: Remove.
+				case ELiteralType::UObjectProxyArray:
+				{
+					return false;
+				}
+				break;
+
+				case ELiteralType::UObjectProxy:	// UObject literals do not currently support default data type references
+				case ELiteralType::Invalid:
+				case ELiteralType::None:
+				default:
+				{
+					static_assert(static_cast<int32>(ELiteralType::Invalid) == 7, "Possible missing ELiteralType case coverage");
+					return false;
+				}
+				break;
+			}
 		}
 
 		// Set the display name for the input with the given name
@@ -3329,7 +3409,6 @@ namespace Metasound
 			return false;
 		}
 
-
 		FString FDocumentController::ExportToJSON() const
 		{
 			TArray<uint8> WriterBuffer;
diff --git a/Engine/Plugins/Runtime/Metasound/Source/MetasoundFrontend/Private/MetasoundFrontendStandardController.h b/Engine/Plugins/Runtime/Metasound/Source/MetasoundFrontend/Private/MetasoundFrontendStandardController.h
--- a/Engine/Plugins/Runtime/Metasound/Source/MetasoundFrontend/Private/MetasoundFrontendStandardController.h
+++ b/Engine/Plugins/Runtime/Metasound/Source/MetasoundFrontend/Private/MetasoundFrontendStandardController.h
@@ -9,6 +9,8 @@
 #include "MetasoundFrontendDocument.h"
 #include "MetasoundFrontendInvalidController.h"
 #include "MetasoundLog.h"
+#include "UObject/Object.h"
+
 
 namespace Metasound
 {
@@ -178,6 +180,12 @@ namespace Metasound
 			const FName& GetDataType() const override;
 			const FString& GetName() const override;
 
+			const TArray<FMetasoundFrontendVertexLiteral>& GetDefaults() const override
+			{
+				static const TArray<FMetasoundFrontendVertexLiteral> InvalidDefaults;
+				return InvalidDefaults;
+			}
+
 			// Owning node info
 			FGuid GetOwningNodeID() const override;
 			FNodeHandle GetOwningNode() override;
@@ -232,6 +240,8 @@ namespace Metasound
 
 			bool IsValid() const override;
 
+			const TArray<FMetasoundFrontendVertexLiteral>& GetDefaults() const;
+
 			// Input metadata
 			const FText& GetDisplayName() const override;
 			const FText& GetTooltip() const override;
@@ -618,6 +628,33 @@ namespace Metasound
 				return false;
 			}
 
+			template<>
+			bool SetDefaultInputToLiteralInternal<UObject*>(const FString& InInputName, FGuid InPointID, UObject* InValue)
+			{
+				if (FMetasoundFrontendClassInput* Desc = FindInputDescriptionWithName(InInputName))
+				{
+					auto IsLiteralWithSamePointID = [&](const FMetasoundFrontendVertexLiteral& InVertexLiteral)
+					{
+						return InVertexLiteral.PointID == InPointID;
+					};
+
+					FMetasoundFrontendVertexLiteral* VertexLiteral = Desc->Defaults.FindByPredicate(IsLiteralWithSamePointID);
+					if (nullptr == VertexLiteral)
+					{
+						VertexLiteral = &Desc->Defaults.AddDefaulted_GetRef();
+						VertexLiteral->PointID = InPointID;
+					}
+
+					if (ensure(FMetasoundFrontendRegistryContainer::Get()->DoesDataTypeSupportLiteralType(Desc->TypeName, Metasound::ELiteralType::UObjectProxy)))
+					{
+						VertexLiteral->Value.Set(InValue);
+						return true;
+					}
+				}
+
+				return false;
+			}
+
 			// Private token only allows members or friends to call constructor.
 			enum EPrivateToken { Token };
 
@@ -699,12 +736,8 @@ namespace Metasound
 
 			// These can be used to set the default value for a given input on this graph.
 			// @returns false if the input name couldn't be found, or if the literal type was incompatible with the Data Type of this input.
-			bool SetDefaultInputToLiteral(const FString& InInputName, FGuid InPointID, bool bInValue) override;
-			bool SetDefaultInputToLiteral(const FString& InInputName, FGuid InPointID, int32 InValue) override;
-			bool SetDefaultInputToLiteral(const FString& InInputName, FGuid InPointID, float InValue) override;
-			bool SetDefaultInputToLiteral(const FString& InInputName, FGuid InPointID, const FString& InValue) override;
-			bool SetDefaultInputToLiteral(const FString& InInputName, FGuid InPointID, UObject* InValue) override;
-			bool SetDefaultInputToLiteral(const FString& InInputName, FGuid InPointID, const TArray<UObject*>& InValue) override;
+			bool SetDefaultInputToFrontendLiteral(const FString& InInputName, FGuid InPointID, FName InDataTypeName, const FMetasoundFrontendLiteral& InLiteral) override;
+			bool SetDefaultInputToTypeDefaultLiteral(const FString& InInputName, FGuid InPointID, FName InDataTypeName) override;
 
 			// Set the display name for the input with the given name
 			void SetInputDisplayName(FString InName, const FText& InDisplayName) override;
diff --git a/Engine/Plugins/Runtime/Metasound/Source/MetasoundFrontend/Public/MetasoundFrontendController.h b/Engine/Plugins/Runtime/Metasound/Source/MetasoundFrontend/Public/MetasoundFrontendController.h
--- a/Engine/Plugins/Runtime/Metasound/Source/MetasoundFrontend/Public/MetasoundFrontendController.h
+++ b/Engine/Plugins/Runtime/Metasound/Source/MetasoundFrontend/Public/MetasoundFrontendController.h
@@ -8,6 +8,7 @@
 #include "MetasoundFrontend.h"
 #include "MetasoundFrontendBaseClasses.h"
 #include "MetasoundFrontendDocument.h"
+#include "MetasoundFrontendRegistries.h"
 #include "MetasoundGraph.h"
 
 /* Metasound Controllers and Handles provide a object oriented interface for  manipulating Metasound Documents. 
@@ -242,6 +243,9 @@ namespace Metasound
 			/** Returns the data type name associated with this input. */
 			virtual const FText& GetDisplayName() const = 0;
 
+			/** Returns the default value of the given input. */
+			virtual const TArray<FMetasoundFrontendVertexLiteral>& GetDefaults() const = 0;
+
 			/** Returns the data type name associated with this input. */
 			virtual const FText& GetTooltip() const = 0;
 
@@ -532,60 +536,17 @@ namespace Metasound
 			virtual TArray<FGuid> GetDefaultIDsForInputVertex(const FString& InInputName) const = 0;
 			virtual TArray<FGuid> GetDefaultIDsForOutputVertex(const FString& InOutputName) const = 0;
 
-			/** Set the default value for the graph input.
-			 *
-			 * @param InInputName - Name of the graph input.
-			 * @param InValue - Value to give to the graph input. 
-			 *
-			 * @return True on success. False if the input does not exist or if the literal type was incompatible with the input.
-			 */
-			virtual bool SetDefaultInputToLiteral(const FString& InInputName, FGuid InPointID, bool InValue) = 0;
-
-			/** Set the default value for the graph input.
-			 *
-			 * @param InInputName - Name of the graph input.
-			 * @param InValue - Value to give to the graph input. 
-			 *
-			 * @return True on success. False if the input does not exist or if the literal type was incompatible with the input.
-			 */
-			virtual bool SetDefaultInputToLiteral(const FString& InInputName, FGuid InPointID, int32 InValue) = 0;
-
-			/** Set the default value for the graph input.
-			 *
-			 * @param InInputName - Name of the graph input.
-			 * @param InValue - Value to give to the graph input. 
-			 *
-			 * @return True on success. False if the input does not exist or if the literal type was incompatible with the input.
-			 */
-			virtual bool SetDefaultInputToLiteral(const FString& InInputName, FGuid InPointID, float InValue) = 0;
-
-			/** Set the default value for the graph input.
-			 *
-			 * @param InInputName - Name of the graph input.
-			 * @param InValue - Value to give to the graph input. 
-			 *
-			 * @return True on success. False if the input does not exist or if the literal type was incompatible with the input.
-			 */
-			virtual bool SetDefaultInputToLiteral(const FString& InInputName, FGuid InPointID, const FString& InValue) = 0;
-
-			/** Set the default value for the graph input.
-			 *
-			 * @param InInputName - Name of the graph input.
-			 * @param InValue - Value to give to the graph input. 
-			 *
-			 * @return True on success. False if the input does not exist or if the literal type was incompatible with the input.
-			 */
-			virtual bool SetDefaultInputToLiteral(const FString& InInputName, FGuid InPointID, UObject* InValue) = 0;
+			virtual bool SetDefaultInputToFrontendLiteral(const FString& InInputName, FGuid InPointID, FName InDataTypeName, const FMetasoundFrontendLiteral& InLiteral) = 0;
 
 			/** Set the default value for the graph input.
 			 *
 			 * @param InInputName - Name of the graph input.
-			 * @param InValue - Value to give to the graph input. 
+			 * @param InPointID - Point to set to DataType default.
+			 * @param InDataTypeName - Name of datatype to set to default.
 			 *
 			 * @return True on success. False if the input does not exist or if the literal type was incompatible with the input.
 			 */
-			virtual bool SetDefaultInputToLiteral(const FString& InInputName, FGuid InPointID, const TArray<UObject*>& InValue) = 0;
-
+			virtual bool SetDefaultInputToTypeDefaultLiteral(const FString& InInputName, FGuid InPointID, FName InDataTypeName) = 0;
 
 			/** Set the display name for the input with the given name. */
 			virtual void SetInputDisplayName(FString InName, const FText& InDisplayName) = 0;
diff --git a/Engine/Plugins/Runtime/Metasound/Source/MetasoundFrontend/Public/MetasoundFrontendDocument.h b/Engine/Plugins/Runtime/Metasound/Source/MetasoundFrontend/Public/MetasoundFrontendDocument.h
--- a/Engine/Plugins/Runtime/Metasound/Source/MetasoundFrontend/Public/MetasoundFrontendDocument.h
+++ b/Engine/Plugins/Runtime/Metasound/Source/MetasoundFrontend/Public/MetasoundFrontendDocument.h
@@ -10,14 +10,17 @@
 
 #include "MetasoundFrontendDocument.generated.h"
 
+
+// Forward Declarations
+struct FMetasoundFrontendClass;
+struct FMetasoundFrontendClassInterface;
 namespace Metasound
 {
+	struct FLiteral;
+
 	extern const FGuid METASOUNDFRONTEND_API FrontendInvalidID;
-}
+} // namespace Metasound
 
-// Forward declare
-struct FMetasoundFrontendClass;
-struct FMetasoundFrontendClassInterface;
 
 UENUM()
 enum class EMetasoundFrontendClassType : uint8
@@ -115,10 +118,15 @@ struct METASOUNDFRONTEND_API FMetasoundFrontendLiteral
 	void Set(const FString& InValue);
 	void Set(UObject* InValue);
 	void Set(const TArray<UObject*>& InValue);
+
+	void SetFromLiteral(const Metasound::FLiteral& InLiteral);
+
+	FString ToString() const;
+
 	void Clear();
 };
 
-// A FMetasoundFrontendVertex provides an named connection point of a node.
+// An FMetasoundFrontendVertex provides a named connection point of a node.
 USTRUCT() 
 struct METASOUNDFRONTEND_API FMetasoundFrontendVertex
 {
@@ -126,7 +134,7 @@ struct METASOUNDFRONTEND_API FMetasoundFrontendVertex
 
 	virtual ~FMetasoundFrontendVertex() = default;
 
-	// Name of the vertex. Unique amongst other vertexes on the same interface.
+	// Name of the vertex. Unique amongst other vertices on the same interface.
 	UPROPERTY(VisibleAnywhere, Category = CustomView)
 	FString Name;
 
@@ -138,7 +146,7 @@ struct METASOUNDFRONTEND_API FMetasoundFrontendVertex
 	UPROPERTY()	
 	TArray<FGuid> PointIDs;
 
-	// Returns true if vertexes have equal name, type and number of IDs. 
+	// Returns true if vertices have equal name, type and number of IDs.
 	static bool IsFunctionalEquivalent(const FMetasoundFrontendVertex& InLHS, const FMetasoundFrontendVertex& InRHS);
 };
 
diff --git a/Engine/Plugins/Runtime/Metasound/Source/MetasoundGraphCore/Private/MetasoundVertex.cpp b/Engine/Plugins/Runtime/Metasound/Source/MetasoundGraphCore/Private/MetasoundVertex.cpp
--- a/Engine/Plugins/Runtime/Metasound/Source/MetasoundGraphCore/Private/MetasoundVertex.cpp
+++ b/Engine/Plugins/Runtime/Metasound/Source/MetasoundGraphCore/Private/MetasoundVertex.cpp
@@ -53,6 +53,11 @@ namespace Metasound
 		return VertexModel->DataTypeName;
 	}
 
+	const FLiteral& FInputDataVertex::GetDefaultValue() const
+	{
+		return VertexModel->DefaultValue;
+	}
+
 	const FText& FInputDataVertex::GetDescription() const
 	{
 		return VertexModel->Description;
diff --git a/Engine/Plugins/Runtime/Metasound/Source/MetasoundGraphCore/Public/MetasoundLiteral.h b/Engine/Plugins/Runtime/Metasound/Source/MetasoundGraphCore/Public/MetasoundLiteral.h
--- a/Engine/Plugins/Runtime/Metasound/Source/MetasoundGraphCore/Public/MetasoundLiteral.h
+++ b/Engine/Plugins/Runtime/Metasound/Source/MetasoundGraphCore/Public/MetasoundLiteral.h
@@ -10,6 +10,8 @@
 #include "Misc/TVariant.h"
 #include <type_traits>
 
+#define METASOUND_DEBUG_LITERALS 1
+
 namespace Metasound
 {
 	enum class ELiteralType : uint8
@@ -142,10 +144,56 @@ namespace Metasound
 		FLiteral(FLiteral&& Other) = default;
 		FLiteral& operator=(FLiteral&& Other) = default;
 
-		typedef TVariant<FInvalidLiteralState, FEmptyLiteralState, bool, int, float, FString, Audio::IProxyDataPtr, TArray<Audio::IProxyDataPtr>> FVariantType;
+		using FVariantType = TVariant<FInvalidLiteralState, FEmptyLiteralState, bool, int, float, FString, Audio::IProxyDataPtr, TArray<Audio::IProxyDataPtr>>;
 
 		FVariantType Value;
 
+#if METASOUND_DEBUG_LITERALS
+		private:
+		FString DebugString;
+
+		void InitDebugString()
+		{
+			switch (Value.GetIndex())
+			{
+				case 0:
+					DebugString = TEXT("INVALID");
+				break;
+
+				case 1:
+					DebugString = TEXT("EMPTY");
+				break;
+
+				case 2:
+					DebugString = FString::Printf(TEXT("Bool: %s"), Value.Get<bool>() ? TEXT("true") : TEXT("false"));
+				break;
+
+				case 3:
+					DebugString = FString::Printf(TEXT("Int32: %d"), Value.Get<int32>());
+				break;
+
+				case 4:
+					DebugString = FString::Printf(TEXT("Float: %f"), Value.Get<float>());
+				break;
+
+				case 5:
+					DebugString = FString::Printf(TEXT("String: %s"), *Value.Get<FString>());
+				break;
+
+				case 6:
+					DebugString = FString::Printf(TEXT("Object: %f"), *Value.Get<Audio::IProxyDataPtr>()->GetProxyTypeName().ToString());
+				break;
+
+				case 7:
+				default:
+					checkNoEntry();
+			}
+		}
+
+		public:
+#endif // METAOUND_DEBUG_LITERALS
+
+
 		// builds an invalid FLiteral.
 		static FLiteral CreateInvalid()
 		{
@@ -191,6 +239,42 @@ namespace Metasound
 			}
 		}
 
+		FLiteral Clone() const
+		{
+			switch (GetType())
+			{
+			case ELiteralType::Invalid:
+				return CreateInvalid();
+
+			case ELiteralType::None:
+				return FLiteral(FEmptyLiteralState());
+
+			case ELiteralType::Boolean:
+				return FLiteral(Value.Get<bool>());
+
+			case ELiteralType::Float:
+				return FLiteral(Value.Get<float>());
+
+			case ELiteralType::Integer:
+				return FLiteral(Value.Get<int32>());
+
+			case ELiteralType::String:
+				return FLiteral(Value.Get<FString>());
+
+			case ELiteralType::UObjectProxy:
+			{
+				TUniquePtr<Audio::IProxyDataPtr> THing;
+				return FLiteral(Value.Get<Audio::IProxyDataPtr>()->Clone());
+			}
+
+			case ELiteralType::UObjectProxyArray:
+				return CreateInvalid();
+
+			default:
+				checkNoEntry();
+				return CreateInvalid();
+			}
+		}
 
 		/** Construct a literal param with a single argument. */
 		template<
@@ -200,6 +284,9 @@ namespace Metasound
 		FLiteral(ArgTypes&&... Args)
 		{
 			Value.Set<typename std::decay<ArgTypes...>::type>(Forward<ArgTypes>(Args)...);
+#if METASOUND_DEBUG_LITERALS
+			InitDebugString();
+#endif
 		}
 
 		/** Construct a literal param with no arguments. */
@@ -210,8 +297,10 @@ namespace Metasound
 		FLiteral(ArgTypes&&... Args)
 		{
 			Value.Set<FEmptyLiteralState>(FEmptyLiteralState());
+#if METASOUND_DEBUG_LITERALS
+			InitDebugString();
+#endif
 		}
-
 	};
 
 	// This struct is passed to FInputNodeConstructorCallback and FOutputNodeConstructorCallback.
diff --git a/Engine/Plugins/Runtime/Metasound/Source/MetasoundGraphCore/Public/MetasoundVertex.h b/Engine/Plugins/Runtime/Metasound/Source/MetasoundGraphCore/Public/MetasoundVertex.h
--- a/Engine/Plugins/Runtime/Metasound/Source/MetasoundGraphCore/Public/MetasoundVertex.h
+++ b/Engine/Plugins/Runtime/Metasound/Source/MetasoundGraphCore/Public/MetasoundVertex.h
@@ -6,6 +6,7 @@
 #include "Containers/SortedMap.h"
 #include "MetasoundDataReference.h"
 #include "MetasoundEnvironment.h"
+#include "MetasoundLiteral.h"
 
 #include <type_traits>
 
@@ -40,7 +41,7 @@ namespace Metasound
 	 */
 	struct FDataVertexModel
 	{
-		/** FDataVertexModel Construtor
+		/** FDataVertexModel Constructor
 		 *
 		 * @InVertexName - Name of vertex.
 		 * @InDataTypeName - Name of data type.
@@ -84,14 +85,24 @@ namespace Metasound
 	 *
 	 * Vertex model for inputs.
 	 */
-	struct FInputDataVertexModel: FDataVertexModel
+	struct FInputDataVertexModel : public FDataVertexModel
 	{
-		// Inherit constructor
-		using FDataVertexModel::FDataVertexModel;
+		FInputDataVertexModel(const FString& InVertexName, const FName& InDataTypeName, const FText& InDescription)
+			: FDataVertexModel(InVertexName, InDataTypeName, InDescription)
+		{
+		}
+
+		FInputDataVertexModel(const FString& InVertexName, const FName& InDataTypeName, const FText& InDescription, const FLiteral& InDefaultValue)
+			: FDataVertexModel(InVertexName, InDataTypeName, InDescription)
+			, DefaultValue(InDefaultValue.Clone())
+		{
+		}
 
 		/** Create a clone of this FInputDataVertexModel */
 		virtual TUniquePtr<FInputDataVertexModel> Clone() const = 0;
 
+		FLiteral DefaultValue;
+
 		friend bool METASOUNDGRAPHCORE_API operator==(const FInputDataVertex& InLHS, const FInputDataVertex& InRHS);
 		friend bool METASOUNDGRAPHCORE_API operator!=(const FInputDataVertex& InLHS, const FInputDataVertex& InRHS);
 
@@ -108,7 +119,7 @@ namespace Metasound
 	 *
 	 * Vertex model for outputs.
 	 */
-	struct FOutputDataVertexModel : FDataVertexModel
+	struct FOutputDataVertexModel : public FDataVertexModel
 	{
 		using FDataVertexModel::FDataVertexModel;
 
@@ -130,9 +141,9 @@ namespace Metasound
 
 	/** TBaseVertexModel provides basic functionality of vertex models. */
 	template<typename DataType, typename VertexModelType>
-	struct TBaseVertexModel : VertexModelType
+	struct TBaseVertexModel : public VertexModelType
 	{
-		/** TBaseVertexModel Construtor
+		/** TBaseVertexModel Constructor
 		 *
 		 * @InVertexName - Name of vertex.
 		 * @InDescription - Human readable vertex description.
@@ -142,6 +153,17 @@ namespace Metasound
 		{
 		}
 
+		/** TBaseVertexModel Constructor
+		 *
+		 * @InVertexName - Name of vertex.
+		 * @InDescription - Human readable vertex description.
+		 * @InDefaultValue - Default value of vertex.
+		 */
+		TBaseVertexModel(const FString& InVertexName, const FText& InDescription, const FLiteral& InDefaultValue)
+		:	VertexModelType(InVertexName, GetMetasoundDataTypeName<DataType>(), InDescription, InDefaultValue)
+		{
+		}
+
 		/** Test if a IDataReference matches the DataType of this vertex model.
 		 *
 		 * @param InReference - Data reference object.
@@ -173,10 +195,31 @@ namespace Metasound
 
 	/** TOuputDataVertexModel creates a simple, unchanging, input vertex. */
 	template<typename DataType>
-	struct TInputDataVertexModel : TBaseVertexModel<DataType, FInputDataVertexModel>
+	struct TInputDataVertexModel : public TBaseVertexModel<DataType, FInputDataVertexModel>
 	{
 		using TBaseVertexModel<DataType, FInputDataVertexModel>::TBaseVertexModel;
 
+		/** TInputDataVertexModel Constructor
+		 *
+		 * @InVertexName - Name of vertex.
+		 * @InDescription - Human readable vertex description.
+		 */
+		TInputDataVertexModel(const FString& InVertexName, const FText& InDescription)
+		:	TBaseVertexModel<DataType, FInputDataVertexModel>(InVertexName, InDescription)
+		{
+		}
+
+		/** TInputDataVertexModel Constructor
+		 *
+		 * @InVertexName - Name of vertex.
+		 * @InDescription - Human readable vertex description.
+		 * @InDefaultValue - Default Value of vertex
+		 */
+		TInputDataVertexModel(const FString& InVertexName, const FText& InDescription, const FLiteral& InDefaultValue)
+		:	TBaseVertexModel<DataType, FInputDataVertexModel>(InVertexName, InDescription, InDefaultValue)
+		{
+		}
+
 		/** Return the vertex type name (not to be confused with the data type name).
 		 *
 		 * @return Name of this vertex type.
@@ -192,7 +235,7 @@ namespace Metasound
 		/** Create a clone of this VertexModelType */
 		virtual TUniquePtr<FInputDataVertexModel> Clone() const override
 		{
-			return MakeUnique<TInputDataVertexModel<DataType>>(this->VertexName, this->Description);
+			return MakeUnique<TInputDataVertexModel<DataType>>(this->VertexName, this->Description, this->DefaultValue);
 		}
 	};
 
@@ -240,6 +283,9 @@ namespace Metasound
 			/** Description of the vertex. */
 			const FText& GetDescription() const;
 
+			/** Default value of the vertex. */
+			const FLiteral& GetDefaultValue() const;
+
 			/** Determine if vertex refers to same data type. 
 			 *
 			 * @param InReference - Data reference object.
@@ -689,7 +735,7 @@ namespace Metasound
 			/** Return an output vertex. */
 			const FOutputDataVertex& GetOutputVertex(const FVertexKey& InKey) const;
 
-			/** Returns true if an outptu vertex with the given key exists. */
+			/** Returns true if an output vertex with the given key exists. */
 			bool ContainsOutputVertex(const FVertexKey& InKey) const;
 
 			/** Return the output interface. */
@@ -701,7 +747,7 @@ namespace Metasound
 			/** Return an output vertex. */
 			const FEnvironmentVertex& GetEnvironmentVertex(const FVertexKey& InKey) const;
 
-			/** Returns true if an outptu vertex with the given key exists. */
+			/** Returns true if an output vertex with the given key exists. */
 			bool ContainsEnvironmentVertex(const FVertexKey& InKey) const;
 
 			/** Test for equality between two interfaces. */
diff --git a/Engine/Plugins/Runtime/Metasound/Source/MetasoundStandardNodes/Private/MetasoundADSRNode.cpp b/Engine/Plugins/Runtime/Metasound/Source/MetasoundStandardNodes/Private/MetasoundADSRNode.cpp
--- a/Engine/Plugins/Runtime/Metasound/Source/MetasoundStandardNodes/Private/MetasoundADSRNode.cpp
+++ b/Engine/Plugins/Runtime/Metasound/Source/MetasoundStandardNodes/Private/MetasoundADSRNode.cpp
@@ -16,8 +16,6 @@
 
 namespace Metasound
 {
-	METASOUND_REGISTER_NODE(FADSRNode)
-
 	class FADSROperator : public TExecutableOperator<FADSROperator>
 	{
 		public:
@@ -30,7 +28,7 @@ namespace Metasound
 			};
 
 			static const FNodeClassMetadata& GetNodeInfo();
-			static FVertexInterface DeclareVertexInterface();
+			static const FVertexInterface& GetVertexInterface();
 			static TUniquePtr<IOperator> CreateOperator(const FCreateOperatorParams& InParams, FBuildErrorArray& OutErrors);
 
 			FADSROperator(const FOperatorSettings& InSettings, const FTriggerReadRef& InTriggerAttack, const FTriggerReadRef& InTriggerRelease, const FADSRReferences& InADSRData);
@@ -44,7 +42,6 @@ namespace Metasound
 			void GenerateEnvelope(int32 InStartFrame, int32 InEndFrame);
 			void Start();
 
-			// TODO: write envelope gen for metasound more suited to this processing structure. 
 			Audio::FEnvelope Envelope;
 
 			FTriggerReadRef TriggerAttack;
@@ -231,16 +228,16 @@ namespace Metasound
 		}
 	}
 
-	FVertexInterface FADSROperator::DeclareVertexInterface()
+	const FVertexInterface& FADSROperator::GetVertexInterface()
 	{
 		static const FVertexInterface Interface(
 			FInputVertexInterface(
-				TInputDataVertexModel<FFloatTime>(TEXT("Attack"), LOCTEXT("AttackDurationTooltip", "Attack duration.")),
-				TInputDataVertexModel<FFloatTime>(TEXT("Decay"), LOCTEXT("DecayDurationTooltip", "Decay duration.")),
-				TInputDataVertexModel<FFloatTime>(TEXT("Release"), LOCTEXT("ReleaseDurationTooltip", "Release duration.")),
-				TInputDataVertexModel<float>(TEXT("Sustain Level"), LOCTEXT("SustainLevelTooltip", "Sustain level [0.0f, 1.0f].")),
-				TInputDataVertexModel<FTrigger>(TEXT("Trigger Attack"), LOCTEXT("TriggerTooltip", "Trigger the envelope's attack.")),
-				TInputDataVertexModel<FTrigger>(TEXT("Trigger Release"), LOCTEXT("TriggerTooltip", "Trigger the envelope's release."))
+				TInputDataVertexModel<FFloatTime>(TEXT("Attack"), LOCTEXT("AttackDurationTooltip", "Attack duration."), 0.01f),
+				TInputDataVertexModel<FFloatTime>(TEXT("Decay"), LOCTEXT("DecayDurationTooltip", "Decay duration."), 0.02f),
+				TInputDataVertexModel<FFloatTime>(TEXT("Release"), LOCTEXT("ReleaseDurationTooltip", "Release duration."), 1.0f),
+				TInputDataVertexModel<float>(TEXT("Sustain Level"), LOCTEXT("SustainLevelTooltip", "Sustain level [0.0f, 1.0f]."), 0.7f),
+				TInputDataVertexModel<FTrigger>(TEXT("Trigger Attack"), LOCTEXT("TriggerAttackTooltip", "Trigger the envelope's attack.")),
+				TInputDataVertexModel<FTrigger>(TEXT("Trigger Release"), LOCTEXT("TriggerReleaseTooltip", "Trigger the envelope's release."))
 			),
 			FOutputVertexInterface(
 				TOutputDataVertexModel<FAudioBuffer>(TEXT("Envelope"), LOCTEXT("EnvelopeTooltip", "The output envelope")),
@@ -265,7 +262,7 @@ namespace Metasound
 			Info.Description = LOCTEXT("Metasound_ADSRNodeDescription", "Emits an ADSR (Attack, decay, sustain, & release) envelope when triggered.");
 			Info.Author = PluginAuthor;
 			Info.PromptIfMissing = PluginNodeMissingPrompt;
-			Info.DefaultInterface = DeclareVertexInterface();
+			Info.DefaultInterface = GetVertexInterface();
 
 			return Info;
 		};
@@ -284,55 +281,37 @@ namespace Metasound
 		FTriggerReadRef TriggerAttack = InputCollection.GetDataReadReferenceOrConstruct<FTrigger>(TEXT("Trigger Attack"), InParams.OperatorSettings);
 		FTriggerReadRef TriggerRelease = InputCollection.GetDataReadReferenceOrConstruct<FTrigger>(TEXT("Trigger Release"), InParams.OperatorSettings);
 
-		// TODO: If none of these are connected, could pre-generate ADSR envelope and return a different operator. 
-		FADSRReferences ADSRReferences = 
+		auto GetOrConstructFloat = [&](const FInputVertexInterface& InputVertices, const FString& InputName)
 		{
-			InputCollection.GetDataReadReferenceOrConstruct<FFloatTime>(TEXT("Attack"), ADSRNode.GetDefaultAttackMs(), ETimeResolution::Milliseconds),
-			InputCollection.GetDataReadReferenceOrConstruct<FFloatTime>(TEXT("Decay"), ADSRNode.GetDefaultDecayMs(), ETimeResolution::Milliseconds),
-			InputCollection.GetDataReadReferenceOrConstruct<float>(TEXT("Sustain Level"), ADSRNode.GetDefaultSustainLevel()),
-			InputCollection.GetDataReadReferenceOrConstruct<FFloatTime>(TEXT("Release"), ADSRNode.GetDefaultReleaseMs(), ETimeResolution::Milliseconds),
+			float DefaultValue = InputVertices[InputName].GetDefaultValue().Value.Get<float>();
+			return InputCollection.GetDataReadReferenceOrConstruct<float>(InputName, DefaultValue);
 		};
-	
-		return MakeUnique<FADSROperator>(InParams.OperatorSettings, TriggerAttack, TriggerRelease, ADSRReferences);
-	}
-
-	FADSRNode::FADSRNode(const FString& InName, const FGuid& InInstanceID, float InDefaultAttackMs, float InDefaultDecayMs, float InDefaultSustainLevel, float InDefaultReleaseMs)
-	:	FNodeFacade(InName, InInstanceID, TFacadeOperatorClass<FADSROperator>())
-	,	DefaultSustainLevel(InDefaultSustainLevel)
-	,	DefaultAttackMs(InDefaultAttackMs)
-	,	DefaultDecayMs(InDefaultDecayMs)
-	,	DefaultReleaseMs(InDefaultReleaseMs)
-	{
-	}
 
-	FADSRNode::FADSRNode(const FNodeInitData& InitData)
-		: FADSRNode(InitData.InstanceName, InitData.InstanceID, 10.0f, 20.0f, 0.7f, 20.0f)
-	{
-	}
-
-	FADSRNode::~FADSRNode()
-	{
-	}
+		auto GetOrConstructTime = [&](const FInputVertexInterface& InputVertices, const FString& InputName, ETimeResolution Resolution = ETimeResolution::Seconds)
+		{
+			float DefaultValue = InputVertices[InputName].GetDefaultValue().Value.Get<float>();
+			return InputCollection.GetDataReadReferenceOrConstruct<FFloatTime>(InputName, DefaultValue, Resolution);
+		};
 
-	float FADSRNode::GetDefaultSustainLevel() const
-	{
-		return DefaultSustainLevel;
-	}
+		// TODO: If none of these are connected, could pre-generate ADSR envelope and return a different operator.
+		const FInputVertexInterface& InputInterface = GetVertexInterface().GetInputInterface();
+		FADSRReferences ADSRReferences =
+		{
+			GetOrConstructTime(InputInterface, TEXT("Attack")),
+			GetOrConstructTime(InputInterface, TEXT("Decay")),
+			GetOrConstructFloat(InputInterface, TEXT("Sustain Level")),
+			GetOrConstructTime(InputInterface, TEXT("Release"))
+		};
 
-	float FADSRNode::GetDefaultAttackMs() const
-	{
-		return DefaultAttackMs;
+		return MakeUnique<FADSROperator>(InParams.OperatorSettings, TriggerAttack, TriggerRelease, ADSRReferences);
 	}
 
-	float FADSRNode::GetDefaultDecayMs() const
+	FADSRNode::FADSRNode(const FNodeInitData& InitData)
+		: FNodeFacade(InitData.InstanceName, InitData.InstanceID, TFacadeOperatorClass<FADSROperator>())
 	{
-		return DefaultDecayMs;
 	}
 
-	float FADSRNode::GetDefaultReleaseMs() const
-	{
-		return DefaultReleaseMs;
-	}
+	METASOUND_REGISTER_NODE(FADSRNode)
 }
 
 #undef LOCTEXT_NAMESPACE
diff --git a/Engine/Plugins/Runtime/Metasound/Source/MetasoundStandardNodes/Private/MetasoundBufferAddFloatNode.cpp b/Engine/Plugins/Runtime/Metasound/Source/MetasoundStandardNodes/Private/MetasoundBufferAddFloatNode.cpp
--- a/Engine/Plugins/Runtime/Metasound/Source/MetasoundStandardNodes/Private/MetasoundBufferAddFloatNode.cpp
+++ b/Engine/Plugins/Runtime/Metasound/Source/MetasoundStandardNodes/Private/MetasoundBufferAddFloatNode.cpp
@@ -82,7 +82,7 @@ namespace Metasound
 		static const FVertexInterface Interface(
 			FInputVertexInterface(
 				TInputDataVertexModel<FAudioBuffer>(TEXT("Buffer"), LOCTEXT("InputBufferTooltip", "Input buffer.")),
-				TInputDataVertexModel<float>(TEXT("Addend"), LOCTEXT("InputAddendTooltip", "Input addend."))
+				TInputDataVertexModel<float>(TEXT("Addend"), LOCTEXT("InputAddendTooltip", "Input addend."), 0.0f)
 			),
 			FOutputVertexInterface(
 				TOutputDataVertexModel<FAudioBuffer>(TEXT("Out"), LOCTEXT("OutpuBufferTooltip", "Add float value to buffer."))
diff --git a/Engine/Plugins/Runtime/Metasound/Source/MetasoundStandardNodes/Private/MetasoundGainMultiplyNode.cpp b/Engine/Plugins/Runtime/Metasound/Source/MetasoundStandardNodes/Private/MetasoundGainMultiplyNode.cpp
--- a/Engine/Plugins/Runtime/Metasound/Source/MetasoundStandardNodes/Private/MetasoundGainMultiplyNode.cpp
+++ b/Engine/Plugins/Runtime/Metasound/Source/MetasoundStandardNodes/Private/MetasoundGainMultiplyNode.cpp
@@ -100,7 +100,7 @@ namespace Metasound
 		static const FVertexInterface Interface(
 			FInputVertexInterface(
 				TInputDataVertexModel<FAudioBuffer>(TEXT("Buffer"), LOCTEXT("InputBufferTooltip", "Input buffer.")),
-				TInputDataVertexModel<float>(TEXT("Gain"), LOCTEXT("InputGainTooltip", "Input gain."))
+				TInputDataVertexModel<float>(TEXT("Gain"), LOCTEXT("InputGainTooltip", "Input gain."), 1.0f)
 			),
 			FOutputVertexInterface(
 				TOutputDataVertexModel<FAudioBuffer>(TEXT("Out"), LOCTEXT("OutpuBufferTooltip", "Output buffer with gain multiplier applied."))
diff --git a/Engine/Plugins/Runtime/Metasound/Source/MetasoundStandardNodes/Private/MetasoundMixerNode.cpp b/Engine/Plugins/Runtime/Metasound/Source/MetasoundStandardNodes/Private/MetasoundMixerNode.cpp
--- a/Engine/Plugins/Runtime/Metasound/Source/MetasoundStandardNodes/Private/MetasoundMixerNode.cpp
+++ b/Engine/Plugins/Runtime/Metasound/Source/MetasoundStandardNodes/Private/MetasoundMixerNode.cpp
@@ -95,19 +95,19 @@ namespace Metasound
 			FInputVertexInterface(
 				// 1
 				TInputDataVertexModel<FAudioBuffer>(GetInputAudioPinName(0), LOCTEXT("MixerInputDescription1", "Audio Input 1 of the Mixer")),		// Can't have dynamic pins yet, so just hard code 4.
-				TInputDataVertexModel<FGain>(GetInputGainPinName(0), LOCTEXT("MixerGainDescription1", "")),
+				TInputDataVertexModel<FGain>(GetInputGainPinName(0), LOCTEXT("MixerGainDescription1", ""), 1.0f),
 
 				// 2
-				TInputDataVertexModel<FAudioBuffer>(GetInputAudioPinName(1), LOCTEXT("MixerInputDescription2", "Audio Input 2 of the Mixer")),	
-				TInputDataVertexModel<FGain>(GetInputGainPinName(1), LOCTEXT("MixerGainDescription2", "")),
+				TInputDataVertexModel<FAudioBuffer>(GetInputAudioPinName(1), LOCTEXT("MixerInputDescription2", "Audio Input 2 of the Mixer")),
+				TInputDataVertexModel<FGain>(GetInputGainPinName(1), LOCTEXT("MixerGainDescription2", ""), 1.0f),
 
 				// 3
-				TInputDataVertexModel<FAudioBuffer>(GetInputAudioPinName(2), LOCTEXT("MixerInputDescription3", "Audio Input 3 of the Mixer")),	
-				TInputDataVertexModel<FGain>(GetInputGainPinName(2), LOCTEXT("MixerGainDescription3", "")),
+				TInputDataVertexModel<FAudioBuffer>(GetInputAudioPinName(2), LOCTEXT("MixerInputDescription3", "Audio Input 3 of the Mixer")),
+				TInputDataVertexModel<FGain>(GetInputGainPinName(2), LOCTEXT("MixerGainDescription3", ""), 1.0f),
 
 				// 4
 				TInputDataVertexModel<FAudioBuffer>(GetInputAudioPinName(3), LOCTEXT("MixerInputDescription3", "Audio Input 4 of the Mixer")),	
-				TInputDataVertexModel<FGain>(GetInputGainPinName(3), LOCTEXT("MixerGainDescription4", ""))
+				TInputDataVertexModel<FGain>(GetInputGainPinName(3), LOCTEXT("MixerGainDescription4", ""), 1.0f)
 			),
 			FOutputVertexInterface(
 				TOutputDataVertexModel<FAudioBuffer>(AudioOutPinName, LOCTEXT("AudioOutTooltip", "Audio Ouput from the mixer"))
diff --git a/Engine/Plugins/Runtime/Metasound/Source/MetasoundStandardNodes/Private/MetasoundOscNode.cpp b/Engine/Plugins/Runtime/Metasound/Source/MetasoundStandardNodes/Private/MetasoundOscNode.cpp
--- a/Engine/Plugins/Runtime/Metasound/Source/MetasoundStandardNodes/Private/MetasoundOscNode.cpp
+++ b/Engine/Plugins/Runtime/Metasound/Source/MetasoundStandardNodes/Private/MetasoundOscNode.cpp
@@ -64,7 +64,7 @@ namespace Metasound
 
 		static constexpr const TCHAR* EnabledPinName = TEXT("Enabled");
 		static constexpr const TCHAR* FrequencyPinName = TEXT("Frequency");
-		static constexpr const TCHAR* AudioOutPinName = TEXT("Audio");		
+		static constexpr const TCHAR* AudioOutPinName = TEXT("Audio");
 	};
 
 	FOscOperator::FOscOperator(const FOperatorSettings& InSettings, const FFrequencyReadRef& InFrequency, const FBoolReadRef& InEnabled)
@@ -123,8 +123,8 @@ namespace Metasound
 	{
 		static const FVertexInterface Interface(
 			FInputVertexInterface(
-				TInputDataVertexModel<FFrequency>(FrequencyPinName, LOCTEXT("OscFrequencyDescription", "The frequency of oscillator.")),
-				TInputDataVertexModel<bool>(EnabledPinName, LOCTEXT("OscActivateDescription", "Enable the oscilator."))
+				TInputDataVertexModel<FFrequency>(FrequencyPinName, LOCTEXT("OscFrequencyDescription", "The frequency of oscillator."), 440.f),
+				TInputDataVertexModel<bool>(EnabledPinName, LOCTEXT("OscActivateDescription", "Enable the oscilator."), true)
 			),
 			FOutputVertexInterface(
 				TOutputDataVertexModel<FAudioBuffer>(TEXT("Audio"), LOCTEXT("AudioTooltip", "The output audio"))
diff --git a/Engine/Plugins/Runtime/Metasound/Source/MetasoundStandardNodes/Private/MetasoundPeriodicTriggerNode.cpp b/Engine/Plugins/Runtime/Metasound/Source/MetasoundStandardNodes/Private/MetasoundPeriodicTriggerNode.cpp
--- a/Engine/Plugins/Runtime/Metasound/Source/MetasoundStandardNodes/Private/MetasoundPeriodicTriggerNode.cpp
+++ b/Engine/Plugins/Runtime/Metasound/Source/MetasoundStandardNodes/Private/MetasoundPeriodicTriggerNode.cpp
@@ -23,6 +23,8 @@ namespace Metasound
 
 	class FPeriodicTriggerOperator : public TExecutableOperator<FPeriodicTriggerOperator>
 	{
+		static constexpr float DefaultPeriod = 1.0f;
+
 		public:
 			static const FNodeClassMetadata& GetNodeInfo();
 			static FVertexInterface DeclareVertexInterface();
@@ -124,13 +126,7 @@ namespace Metasound
 	{
 		const FPeriodicTriggerNode& PeriodicTriggerNode = static_cast<const FPeriodicTriggerNode&>(InParams.Node);
 
-		FFloatTimeReadRef Period = FFloatTimeReadRef::CreateNew(PeriodicTriggerNode.GetDefaultPeriodInSeconds(), ETimeResolution::Seconds);
-
-		if (InParams.InputDataReferences.ContainsDataReadReference<FFloatTime>(TEXT("Period")))
-		{
-			Period = InParams.InputDataReferences.GetDataReadReference<FFloatTime>(TEXT("Period"));
-		}
-
+		FFloatTimeReadRef Period = InParams.InputDataReferences.GetDataReadReferenceOrConstruct<FFloatTime>(TEXT("Period"), DefaultPeriod, ETimeResolution::Seconds);
 		FTriggerReadRef TriggerEnable = InParams.InputDataReferences.GetDataReadReferenceOrConstruct<FTrigger>(TEXT("Activate"), InParams.OperatorSettings);
 		FTriggerReadRef TriggerDisable = InParams.InputDataReferences.GetDataReadReferenceOrConstruct<FTrigger>(TEXT("Deactivate"), InParams.OperatorSettings);
 
@@ -141,7 +137,7 @@ namespace Metasound
 	{
 		static const FVertexInterface Interface(
 			FInputVertexInterface(
-				TInputDataVertexModel<FFloatTime>(TEXT("Period"), LOCTEXT("PeriodTooltip", "The period to trigger in seconds.")),
+				TInputDataVertexModel<FFloatTime>(TEXT("Period"), LOCTEXT("PeriodTooltip", "The period to trigger in seconds."), DefaultPeriod),
 				TInputDataVertexModel<FTrigger>(TEXT("Activate"), LOCTEXT("TriggerEnableTooltip", "Enables executing periodic output triggers.")),
 				TInputDataVertexModel<FTrigger>(TEXT("Deactivate"), LOCTEXT("TriggerDisableTooltip", "Disables executing periodic output triggers."))
 			),
@@ -175,20 +171,9 @@ namespace Metasound
 		return Info;
 	}
 
-
-	FPeriodicTriggerNode::FPeriodicTriggerNode(const FString& InInstanceName, const FGuid& InInstanceID, float InDefaultPeriodInSeconds)
-	:	FNodeFacade(InInstanceName, InInstanceID, TFacadeOperatorClass<FPeriodicTriggerOperator>())
-	,	DefaultPeriod(InDefaultPeriodInSeconds)
-	{
-	}
-
 	FPeriodicTriggerNode::FPeriodicTriggerNode(const FNodeInitData& InInitData)
-		: FPeriodicTriggerNode(InInitData.InstanceName, InInitData.InstanceID, 1.0f)
-	{}
-
-	float FPeriodicTriggerNode::GetDefaultPeriodInSeconds() const
+	:	FNodeFacade(InInitData.InstanceName, InInitData.InstanceID, TFacadeOperatorClass<FPeriodicTriggerOperator>())
 	{
-		return DefaultPeriod;
 	}
 }
 
diff --git a/Engine/Plugins/Runtime/Metasound/Source/MetasoundStandardNodes/Private/MetasoundStandardNodesNames.cpp b/Engine/Plugins/Runtime/Metasound/Source/MetasoundStandardNodes/Private/MetasoundStandardNodesNames.cpp
--- a/Engine/Plugins/Runtime/Metasound/Source/MetasoundStandardNodes/Private/MetasoundStandardNodesNames.cpp
+++ b/Engine/Plugins/Runtime/Metasound/Source/MetasoundStandardNodes/Private/MetasoundStandardNodesNames.cpp
@@ -7,7 +7,7 @@ namespace Metasound
 {
 	namespace StandardNodes
 	{
-		const FName Namespace = FName(TEXT("UE"));
-		const FName AudioVariant = FName(TEXT("Audio"));
+		const FName Namespace = "UE";
+		const FName AudioVariant = "Audio";
 	}
 }
diff --git a/Engine/Plugins/Runtime/Metasound/Source/MetasoundStandardNodes/Private/MetasoundSwitchGateNode.cpp b/Engine/Plugins/Runtime/Metasound/Source/MetasoundStandardNodes/Private/MetasoundSwitchGateNode.cpp
--- a/Engine/Plugins/Runtime/Metasound/Source/MetasoundStandardNodes/Private/MetasoundSwitchGateNode.cpp
+++ b/Engine/Plugins/Runtime/Metasound/Source/MetasoundStandardNodes/Private/MetasoundSwitchGateNode.cpp
@@ -178,11 +178,11 @@ namespace Metasound
 			CachedInterface = FVertexInterface();
 			
 			CachedInterface.GetInputInterface().Add(TInputDataVertexModel<int32>(SWITCH_SELECTOR_NAME, LOCTEXT("SwitchSelectorDescription", "The index of the audio buffer to use (0 based).")));
-			CachedInterface.GetInputInterface().Add(TInputDataVertexModel<bool>(SWITCH_CROSSFADE_NAME, LOCTEXT("SwitchCrossfadeDescription", "Whether changing the buffer while running causes a short crossfade to prevent pops. This should almost always be true.")));
+			CachedInterface.GetInputInterface().Add(TInputDataVertexModel<bool>(SWITCH_CROSSFADE_NAME, LOCTEXT("SwitchCrossfadeDescription", "Whether changing the buffer while running causes a short crossfade to prevent pops. This should almost always be true."), true));
 			for (int32 i=0; i < PREVIEW_INPUT_COUNT; i++)
-			{				
+			{
 				CachedInterface.GetInputInterface().Add(TInputDataVertexModel<FAudioBuffer>(*FString::Printf(SWITCH_INPUT_NAME_FMT, i), LOCTEXT("SwitchInputDescription", "An input buffer to choose from.")));
-			}			
+			}
 
 			CachedInterface.GetOutputInterface().Add(TOutputDataVertexModel<FAudioBuffer>(SWITCH_OUTPUT_NAME, LOCTEXT("SwitchOutputDescription", "The output audio buffer.")));
 			
@@ -197,7 +197,7 @@ namespace Metasound
 		auto InitNodeInfo = []() -> FNodeClassMetadata
 		{
 			FNodeClassMetadata Info;
-			Info.ClassName = {Metasound::StandardNodes::Namespace, TEXT("Switch"), Metasound::StandardNodes::AudioVariant};
+			Info.ClassName = { Metasound::StandardNodes::Namespace, TEXT("Switch"), Metasound::StandardNodes::AudioVariant };
 			Info.MajorVersion = 1;
 			Info.MinorVersion = 0;
 			Info.DisplayName = LOCTEXT("Metasound_SwitchNodeDisplayName", "Switch");
diff --git a/Engine/Plugins/Runtime/Metasound/Source/MetasoundStandardNodes/Private/MetasoundTriggerAccumulatorNode.cpp b/Engine/Plugins/Runtime/Metasound/Source/MetasoundStandardNodes/Private/MetasoundTriggerAccumulatorNode.cpp
--- a/Engine/Plugins/Runtime/Metasound/Source/MetasoundStandardNodes/Private/MetasoundTriggerAccumulatorNode.cpp
+++ b/Engine/Plugins/Runtime/Metasound/Source/MetasoundStandardNodes/Private/MetasoundTriggerAccumulatorNode.cpp
@@ -171,9 +171,9 @@ namespace Metasound
 	{
 		static const FVertexInterface Interface(
 			FInputVertexInterface(
-				TInputDataVertexModel<FTrigger>(AccumulatePinName, LOCTEXT("TriggerToAccumulateTooltop", "Trigger to accumulate")),
+				TInputDataVertexModel<FTrigger>(AccumulatePinName, LOCTEXT("TriggerToAccumulateTooltip", "Trigger to accumulate")),
 				TInputDataVertexModel<FTrigger>(ResetPinName, LOCTEXT("ResetAccumulatorTooltip", "Reset the accumulator")),
-				TInputDataVertexModel<int32>(TriggerCountPinName, LOCTEXT("CountToTriggerAtTooltip", "Count to trigger at"))
+				TInputDataVertexModel<int32>(TriggerCountPinName, LOCTEXT("CountToTriggerAtTooltip", "Count to trigger at"), 1)
 			),
 			FOutputVertexInterface(
 				TOutputDataVertexModel<FTrigger>(CountReachedPinName, LOCTEXT("CountReachedTooltop", "Triggered when the accumulated count is reached")),
diff --git a/Engine/Plugins/Runtime/Metasound/Source/MetasoundStandardNodes/Private/MetasoundTriggerDelayNode.cpp b/Engine/Plugins/Runtime/Metasound/Source/MetasoundStandardNodes/Private/MetasoundTriggerDelayNode.cpp
--- a/Engine/Plugins/Runtime/Metasound/Source/MetasoundStandardNodes/Private/MetasoundTriggerDelayNode.cpp
+++ b/Engine/Plugins/Runtime/Metasound/Source/MetasoundStandardNodes/Private/MetasoundTriggerDelayNode.cpp
@@ -22,6 +22,8 @@ namespace Metasound
 
 	class FTriggerDelayOperator : public TExecutableOperator<FTriggerDelayOperator>
 	{
+		static constexpr float DefaultDelay = 1.0f; // In Seconds
+
 		public:
 			static const FNodeClassMetadata& GetNodeInfo();
 			static FVertexInterface DeclareVertexInterface();
@@ -33,10 +35,10 @@ namespace Metasound
 			virtual FDataReferenceCollection GetOutputs() const override;
 
 			void Execute();
-			void Reset();
 
 		private:
-			int32 SamplesUntilTrigger;
+			int32 SamplesUntilTrigger = -1;
+			int32 LastFrameReset = -1;
 
 			FTriggerReadRef TriggerIn;
 			FTriggerReadRef TriggerReset;
@@ -50,6 +52,7 @@ namespace Metasound
 
 	FTriggerDelayOperator::FTriggerDelayOperator(const FOperatorSettings& InSettings, const FTriggerReadRef& InTriggerReset, const FTriggerReadRef& InTriggerIn, const FFloatTimeReadRef& InDelay)
 	:	SamplesUntilTrigger(-1)
+	,	LastFrameReset(-1)
 	,	TriggerIn(InTriggerIn)
 	,	TriggerReset(InTriggerReset)
 	,	TriggerOut(FTriggerWriteRef::CreateNew(InSettings))
@@ -81,47 +84,55 @@ namespace Metasound
 		// Advance internal counter to get rid of old triggers.
 		TriggerOut->AdvanceBlock();
 
-		if (SamplesUntilTrigger > 0)
-		{
-			int32 SamplesRemaining = SamplesUntilTrigger - FramesPerBlock;
-			if (SamplesRemaining > 0.0f)
+		LastFrameReset = -1;
+		TriggerReset->ExecuteBlock(
+			[&](int32 StartFrame, int32 EndFrame)
 			{
-				SamplesUntilTrigger -= FramesPerBlock;
-			}
-			else
+			},
+			[this](int32 StartFrame, int32 EndFrame)
 			{
-				TriggerOut->TriggerFrame(SamplesRemaining + (int32)FramesPerBlock);
-				Reset();
+				LastFrameReset = StartFrame;
 			}
-		}
+		);
 
 		TriggerIn->ExecuteBlock(
-			[](int32 StartFrame, int32 EndFrame)
+			[&](int32 StartFrame, int32 EndFrame)
 			{
 			},
 			[this](int32 StartFrame, int32 EndFrame)
 			{
-				SamplesUntilTrigger = Delay->GetNumSamples(SampleRate);
+				if (StartFrame > LastFrameReset)
+				{
+					SamplesUntilTrigger = Delay->GetNumSamples(SampleRate);
+					LastFrameReset = -1;
+				}
+				else
+				{
+					SamplesUntilTrigger = -1;
+				}
 			}
 		);
 
-		TriggerReset->ExecuteBlock(
-			[](int32 StartFrame, int32 EndFrame)
+		if (SamplesUntilTrigger >= 0)
+		{
+			int32 SamplesRemaining = SamplesUntilTrigger - FramesPerBlock;
+			if (SamplesRemaining > 0.0f)
 			{
-			},
-			[this](int32 StartFrame, int32 EndFrame)
+				SamplesUntilTrigger -= FramesPerBlock;
+			}
+			else
 			{
-				Reset();
+				TriggerOut->TriggerFrame(SamplesRemaining + (int32)FramesPerBlock);
+				SamplesUntilTrigger = -1;
 			}
-		);
+		}
 	}
 
-
 	TUniquePtr<IOperator> FTriggerDelayOperator::CreateOperator(const FCreateOperatorParams& InParams, FBuildErrorArray& OutErrors)
 	{
 		const FTriggerDelayNode& DelayNode = static_cast<const FTriggerDelayNode&>(InParams.Node);
 
-		FFloatTimeReadRef Delay = InParams.InputDataReferences.GetDataReadReferenceOrConstruct<FFloatTime>(TEXT("Delay"), DelayNode.GetDefaultDelayInSeconds());
+		FFloatTimeReadRef Delay = InParams.InputDataReferences.GetDataReadReferenceOrConstruct<FFloatTime>(TEXT("Delay"), DefaultDelay);
 		FTriggerReadRef TriggerIn = InParams.InputDataReferences.GetDataReadReferenceOrConstruct<FTrigger>(TEXT("In"), InParams.OperatorSettings);
 		FTriggerReadRef TriggerReset = InParams.InputDataReferences.GetDataReadReferenceOrConstruct<FTrigger>(TEXT("Reset"), InParams.OperatorSettings);
 
@@ -132,7 +143,7 @@ namespace Metasound
 	{
 		static const FVertexInterface Interface(
 			FInputVertexInterface(
-				TInputDataVertexModel<FFloatTime>(TEXT("Delay"), LOCTEXT("TriggerDelayTooltip", "Time to delay and execute deferred trigger in seconds.")),
+				TInputDataVertexModel<FFloatTime>(TEXT("Delay"), LOCTEXT("TriggerDelayTooltip", "Time to delay and execute deferred trigger in seconds."), 1.0f),
 				TInputDataVertexModel<FTrigger>(TEXT("In"), LOCTEXT("TriggerDelayInTooltip", "Triggers delay.")),
 				TInputDataVertexModel<FTrigger>(TEXT("Reset"), LOCTEXT("TriggerDelayResetInTooltip", "Resets the trigger delay, clearing the execution task if pending."))
 			),
@@ -165,25 +176,9 @@ namespace Metasound
 		return Info;
 	}
 
-	void FTriggerDelayOperator::Reset()
-	{
-		SamplesUntilTrigger = -1;
-	}
-
-	FTriggerDelayNode::FTriggerDelayNode(const FString& InName, const FGuid& InInstanceID, float InDefaultDelayInSeconds)
-	:	FNodeFacade(InName, InInstanceID, TFacadeOperatorClass<FTriggerDelayOperator>())
-	,	DefaultDelay(InDefaultDelayInSeconds)
-	{
-	}
-
 	FTriggerDelayNode::FTriggerDelayNode(const FNodeInitData& InInitData)
-		: FTriggerDelayNode(InInitData.InstanceName, InInitData.InstanceID, 0.0f)
-	{
-	}
-
-	float FTriggerDelayNode::GetDefaultDelayInSeconds() const
+	:	FNodeFacade(InInitData.InstanceName, InInitData.InstanceID, TFacadeOperatorClass<FTriggerDelayOperator>())
 	{
-		return DefaultDelay;
 	}
 }
 
diff --git a/Engine/Plugins/Runtime/Metasound/Source/MetasoundStandardNodes/Private/MetasoundTriggerPipeNode.cpp b/Engine/Plugins/Runtime/Metasound/Source/MetasoundStandardNodes/Private/MetasoundTriggerPipeNode.cpp
--- a/Engine/Plugins/Runtime/Metasound/Source/MetasoundStandardNodes/Private/MetasoundTriggerPipeNode.cpp
+++ b/Engine/Plugins/Runtime/Metasound/Source/MetasoundStandardNodes/Private/MetasoundTriggerPipeNode.cpp
@@ -24,28 +24,28 @@ namespace Metasound
 
 	class FTriggerPipeOperator : public TExecutableOperator<FTriggerPipeOperator>
 	{
-		public:
-			static const FNodeClassMetadata& GetNodeInfo();
-			static FVertexInterface DeclareVertexInterface();
-			static TUniquePtr<IOperator> CreateOperator(const FCreateOperatorParams& InParams, FBuildErrorArray& OutErrors);
+	public:
+		static const FNodeClassMetadata& GetNodeInfo();
+		static const FVertexInterface& GetVertexInterface();
+		static TUniquePtr<IOperator> CreateOperator(const FCreateOperatorParams& InParams, FBuildErrorArray& OutErrors);
 
-			FTriggerPipeOperator(const FOperatorSettings& InSettings, const FTriggerReadRef& InTriggerReset, const FTriggerReadRef& InTriggerIn, const FFloatTimeReadRef& InDelay);
+		FTriggerPipeOperator(const FOperatorSettings& InSettings, const FTriggerReadRef& InTriggerReset, const FTriggerReadRef& InTriggerIn, const FFloatTimeReadRef& InDelay);
 
-			virtual FDataReferenceCollection GetInputs() const override;
-			virtual FDataReferenceCollection GetOutputs() const override;
-			void Execute();
+		virtual FDataReferenceCollection GetInputs() const override;
+		virtual FDataReferenceCollection GetOutputs() const override;
+		void Execute();
 
-		private:
-			TArray<int32> SamplesUntilTrigger;
+	private:
+		TArray<int32> SamplesUntilTrigger;
 
-			FTriggerReadRef TriggerIn;
-			FTriggerReadRef TriggerReset;
-			FTriggerWriteRef TriggerOut;
+		FTriggerReadRef TriggerIn;
+		FTriggerReadRef TriggerReset;
+		FTriggerWriteRef TriggerOut;
 
-			FFloatTimeReadRef Delay;
+		FFloatTimeReadRef Delay;
 
-			float FramesPerBlock;
-			float SampleRate;
+		float FramesPerBlock;
+		float SampleRate;
 	};
 
 	FTriggerPipeOperator::FTriggerPipeOperator(const FOperatorSettings& InSettings, const FTriggerReadRef& InTriggerReset, const FTriggerReadRef& InTriggerIn, const FFloatTimeReadRef& InDelay)
@@ -80,20 +80,6 @@ namespace Metasound
 		// Advance internal counter to get rid of old triggers.
 		TriggerOut->AdvanceBlock();
 
-		for (int32 i = SamplesUntilTrigger.Num() - 1; i >= 0; --i)
-		{
-			int32 SamplesRemaining = SamplesUntilTrigger[i] - FramesPerBlock;
-			if (SamplesRemaining > 0.0f)
-			{
-				SamplesUntilTrigger[i] -= FramesPerBlock;
-			}
-			else
-			{
-				TriggerOut->TriggerFrame(SamplesRemaining + (int32)FramesPerBlock);
-				SamplesUntilTrigger.RemoveAtSwap(i);
-			}
-		}
-
 		TriggerIn->ExecuteBlock(
 			[](int32 StartFrame, int32 EndFrame)
 			{
@@ -110,27 +96,58 @@ namespace Metasound
 			},
 			[this](int32 StartFrame, int32 EndFrame)
 			{
-				SamplesUntilTrigger.Empty();
+				// Iterate backward and only remove delayed triggers that occur
+				// after the reset trigger's start frame.
+				for (int32 i = SamplesUntilTrigger.Num() - 1; i >= 0; --i)
+				{
+					const int32 SamplesRemaining = SamplesUntilTrigger[i] - FramesPerBlock;
+					if (SamplesRemaining >= StartFrame)
+					{
+						SamplesUntilTrigger.RemoveAtSwap(i);
+					}
+				}
 			}
 		);
+
+		for (int32 i = SamplesUntilTrigger.Num() - 1; i >= 0; --i)
+		{
+			const int32 SamplesRemaining = SamplesUntilTrigger[i] - FramesPerBlock;
+			if (SamplesRemaining >= 0)
+			{
+				SamplesUntilTrigger[i] -= FramesPerBlock;
+			}
+			else
+			{
+				TriggerOut->TriggerFrame(SamplesRemaining + (int32)FramesPerBlock);
+				SamplesUntilTrigger.RemoveAtSwap(i);
+			}
+		}
 	}
 
 	TUniquePtr<IOperator> FTriggerPipeOperator::CreateOperator(const FCreateOperatorParams& InParams, FBuildErrorArray& OutErrors)
 	{
 		const FTriggerPipeNode& DelayNode = static_cast<const FTriggerPipeNode&>(InParams.Node);
+		const FDataReferenceCollection& InputCollection = InParams.InputDataReferences;
+
+		auto GetOrConstructTime = [&](const FInputVertexInterface& InputVertices, const FString& InputName, ETimeResolution Resolution = ETimeResolution::Seconds)
+		{
+			float DefaultValue = InputVertices[InputName].GetDefaultValue().Value.Get<float>();
+			return InputCollection.GetDataReadReferenceOrConstruct<FFloatTime>(InputName, DefaultValue, Resolution);
+		};
 
-		FFloatTimeReadRef Delay = InParams.InputDataReferences.GetDataReadReferenceOrConstruct<FFloatTime>(TEXT("Delay"), DelayNode.GetDefaultDelayInSeconds());
+		const FInputVertexInterface& InputInterface = GetVertexInterface().GetInputInterface();
+		FFloatTimeReadRef Delay = GetOrConstructTime(InputInterface, TEXT("Delay"));
 		FTriggerReadRef TriggerIn = InParams.InputDataReferences.GetDataReadReferenceOrConstruct<FTrigger>(TEXT("In"), InParams.OperatorSettings);
 		FTriggerReadRef TriggerReset = InParams.InputDataReferences.GetDataReadReferenceOrConstruct<FTrigger>(TEXT("Reset"), InParams.OperatorSettings);
 
 		return MakeUnique<FTriggerPipeOperator>(InParams.OperatorSettings, TriggerReset, TriggerIn, Delay);
 	}
 
-	FVertexInterface FTriggerPipeOperator::DeclareVertexInterface()
+	const FVertexInterface& FTriggerPipeOperator::GetVertexInterface()
 	{
 		static const FVertexInterface Interface(
 			FInputVertexInterface(
-				TInputDataVertexModel<FFloatTime>(TEXT("Delay"), LOCTEXT("DelayTooltip", "Time to delay and execute deferred input trigger execution(s) in seconds.")),
+				TInputDataVertexModel<FFloatTime>(TEXT("Delay"), LOCTEXT("DelayTooltip", "Time to delay and execute deferred input trigger execution(s) in seconds."), 1.0f),
 				TInputDataVertexModel<FTrigger>(TEXT("In"), LOCTEXT("TriggerInTooltip", "Trigger to execute at a future time by the given delay amount.")),
 				TInputDataVertexModel<FTrigger>(TEXT("Reset"), LOCTEXT("TriggerDelayResetInTooltip", "Resets the trigger delay, clearing any pending execution tasks."))
 			),
@@ -147,14 +164,14 @@ namespace Metasound
 		auto InitNodeInfo = []() -> FNodeClassMetadata
 		{
 			FNodeClassMetadata Info;
-			Info.ClassName = FNodeClassName(StandardNodes::Namespace, TEXT("Pipe"), TEXT("Trigger"));
+			Info.ClassName = FNodeClassName(StandardNodes::Namespace, "Pipe", "Trigger");
 			Info.MajorVersion = 1;
 			Info.MinorVersion = 0;
 			Info.DisplayName = LOCTEXT("PipeTriggerNode_NodeDisplayName", "Trigger Pipe");
 			Info.Description = LOCTEXT("Metasound_DelayNodeDescription", "Delays execution of the input trigger by the given delay for all input trigger executions.");
 			Info.Author = PluginAuthor;
 			Info.PromptIfMissing = PluginNodeMissingPrompt;
-			Info.DefaultInterface = DeclareVertexInterface();
+			Info.DefaultInterface = GetVertexInterface();
 
 			return Info;
 		};
@@ -163,20 +180,9 @@ namespace Metasound
 		return Info;
 	}
 
-	FTriggerPipeNode::FTriggerPipeNode(const FString& InName, const FGuid& InInstanceID, float InDefaultDelayInSeconds)
-	:	FNodeFacade(InName, InInstanceID, TFacadeOperatorClass<FTriggerPipeOperator>())
-	,	DefaultDelay(InDefaultDelayInSeconds)
-	{
-	}
-
 	FTriggerPipeNode::FTriggerPipeNode(const FNodeInitData& InInitData)
-		: FTriggerPipeNode(InInitData.InstanceName, InInitData.InstanceID, 0.0f)
-	{
-	}
-
-	float FTriggerPipeNode::GetDefaultDelayInSeconds() const
+		: FNodeFacade(InInitData.InstanceName, InInitData.InstanceID, TFacadeOperatorClass<FTriggerPipeOperator>())
 	{
-		return DefaultDelay;
 	}
 }
 
diff --git a/Engine/Plugins/Runtime/Metasound/Source/MetasoundStandardNodes/Public/MetasoundADSRNode.h b/Engine/Plugins/Runtime/Metasound/Source/MetasoundStandardNodes/Public/MetasoundADSRNode.h
--- a/Engine/Plugins/Runtime/Metasound/Source/MetasoundStandardNodes/Public/MetasoundADSRNode.h
+++ b/Engine/Plugins/Runtime/Metasound/Source/MetasoundStandardNodes/Public/MetasoundADSRNode.h
@@ -14,42 +14,9 @@ namespace Metasound
 	class METASOUNDSTANDARDNODES_API FADSRNode : public FNodeFacade
 	{
 		public:
-
-			/** FADSR node constructor.
-			 *
-			 * @param InInstanceName - Name of this node.
-			 * @param InInstanceID - ID of this node.
-			 * @param InDefaultAttackMs - Default attack duration in milliseconds.
-			 * @param InDefaultDecayMs - Default decay duration in milliseconds.
-			 * @param InDefaultSustainLevel - Default sustain level [0.0f, 1.0f].
-			 * @param InDefaultReleaseMs - Default release duration in milliseconds.
-			 */
-			FADSRNode(const FString& InInstanceName, const FGuid& InInstanceID, float InDefaultAttackMs, float InDefaultDecayMs, float InDefaultSustainLevel, float InDefaultReleaseMs);
-
 			/**
 			 * Constructor used by the Metasound Frontend.
 			 */
 			FADSRNode(const FNodeInitData& InitData);
-
-			virtual ~FADSRNode();
-
-			/** Return default sustain gain [0.0f, 1.0f] */
-			float GetDefaultSustainLevel() const;
-
-			/** Return default attack time in milliseconds */
-			float GetDefaultAttackMs() const;
-
-			/** Return default decay time in milliseconds */
-			float GetDefaultDecayMs() const;
-
-			/** Return default release time in milliseconds */
-			float GetDefaultReleaseMs() const;
-
-		private:
-			float DefaultSustainLevel;
-
-			float DefaultAttackMs;
-			float DefaultDecayMs;
-			float DefaultReleaseMs;
 	};
 } // namespace Metasound
diff --git a/Engine/Plugins/Runtime/Metasound/Source/MetasoundStandardNodes/Public/MetasoundMathNodes.h b/Engine/Plugins/Runtime/Metasound/Source/MetasoundStandardNodes/Public/MetasoundMathNodes.h
--- a/Engine/Plugins/Runtime/Metasound/Source/MetasoundStandardNodes/Public/MetasoundMathNodes.h
+++ b/Engine/Plugins/Runtime/Metasound/Source/MetasoundStandardNodes/Public/MetasoundMathNodes.h
@@ -142,8 +142,8 @@ namespace Metasound
 				const TMathOpNode& MathOpNode = static_cast<const TMathOpNode&>(InParams.Node);
 
 				// TODO: Support dynamic number of inputs
-				TDataClassReadRef Op1 = InParams.InputDataReferences.GetDataReadReferenceOrConstruct<TDataClass>(TEXT("Operand1"), TMathOpClass::GetDefault(InParams.OperatorSettings));
-				TDataClassReadRef Op2 = InParams.InputDataReferences.GetDataReadReferenceOrConstruct<TDataClass>(TEXT("Operand2"), TMathOpClass::GetDefault(InParams.OperatorSettings));
+				TDataClassReadRef Op1 = InParams.InputDataReferences.GetDataReadReferenceOrConstruct<TDataClass>(TEXT("Operand1"), TMathOpClass::GetDefault(&InParams.OperatorSettings));
+				TDataClassReadRef Op2 = InParams.InputDataReferences.GetDataReadReferenceOrConstruct<TDataClass>(TEXT("Operand2"), TMathOpClass::GetDefault(&InParams.OperatorSettings));
 
 				TArray<TDataClassReadRef> OperandRefs = { Op1, Op2 };
 
@@ -152,7 +152,7 @@ namespace Metasound
 
 			TMathOperator(const FOperatorSettings& InSettings, const TArray<TDataClassReadRef>& InOperands)
 			:	OperandRefs(InOperands)
-			,	ValueRef(TDataClassWriteRef::CreateNew(TMathOpClass::GetDefault(InSettings)))
+			,	ValueRef(TDataClassWriteRef::CreateNew(TMathOpClass::GetDefault(&InSettings)))
 			{
 			}
 
@@ -188,7 +188,7 @@ namespace Metasound
 		using TDataClassReadRef = TDataReadReference<TDataClass>;
 		using TDataClassWriteRef = TDataWriteReference<TDataClass>;
 
-		static TDataClass GetDefault(const FOperatorSettings& InSettings)
+		static TDataClass GetDefault(const FOperatorSettings* InSettings)
 		{
 			return 0;
 		}
@@ -215,7 +215,7 @@ namespace Metasound
 		using TDataClassReadRef = TDataReadReference<TDataClass>;
 		using TDataClassWriteRef = TDataWriteReference<TDataClass>;
 
-		static TDataClass GetDefault(const FOperatorSettings& InSettings)
+		static TDataClass GetDefault(const FOperatorSettings* InSettings)
 		{
 			return 0;
 		}
@@ -242,7 +242,7 @@ namespace Metasound
 		using TDataClassReadRef = TDataReadReference<TDataClass>;
 		using TDataClassWriteRef = TDataWriteReference<TDataClass>;
 
-		static TDataClass GetDefault(const FOperatorSettings& InSettings)
+		static TDataClass GetDefault(const FOperatorSettings* InSettings)
 		{
 			return 1;
 		}
@@ -269,7 +269,7 @@ namespace Metasound
 		using TDataClassReadRef = TDataReadReference<TDataClass>;
 		using TDataClassWriteRef = TDataWriteReference<TDataClass>;
 
-		static TDataClass GetDefault(const FOperatorSettings& InSettings)
+		static TDataClass GetDefault(const FOperatorSettings* InSettings)
 		{
 			return 1;
 		}
@@ -301,9 +301,9 @@ namespace Metasound
 	class TMathOpAdd<FAudioBuffer>
 	{
 	public:
-		static FAudioBuffer GetDefault(const FOperatorSettings& InSettings)
+		static FAudioBuffer GetDefault(const FOperatorSettings* InSettings)
 		{
-			return FAudioBuffer(InSettings.GetNumFramesPerBlock());
+			return FAudioBuffer(InSettings ? InSettings->GetNumFramesPerBlock() : 0);
 		}
 
 		static void Calculate(const TArray<FAudioBufferReadRef>& InOperands, FAudioBufferWriteRef& OutResult)
@@ -358,9 +358,9 @@ namespace Metasound
 	class TMathOpSubtract<FAudioBuffer>
 	{
 	public:
-		static FAudioBuffer GetDefault(const FOperatorSettings& InSettings)
+		static FAudioBuffer GetDefault(const FOperatorSettings* InSettings)
 		{
-			return FAudioBuffer(InSettings.GetNumFramesPerBlock());
+			return FAudioBuffer(InSettings ? InSettings->GetNumFramesPerBlock() : 0);
 		}
 
 		static void Calculate(TArray<FAudioBufferReadRef>& InOperands, FAudioBufferWriteRef& OutResult)
@@ -413,9 +413,9 @@ namespace Metasound
 	class TMathOpMultiply<FAudioBuffer>
 	{
 	public:
-		static FAudioBuffer GetDefault(const FOperatorSettings& InSettings)
+		static FAudioBuffer GetDefault(const FOperatorSettings* InSettings)
 		{
-			return FAudioBuffer(InSettings.GetNumFramesPerBlock());
+			return FAudioBuffer(InSettings ? InSettings->GetNumFramesPerBlock() : 0);
 		}
 
 		static void Calculate(const TArray<FAudioBufferReadRef>& InOperands, FAudioBufferWriteRef& OutResult)
@@ -468,7 +468,7 @@ namespace Metasound
 	class TMathOpMultiply<FFloatTime>
 	{
 	public:
-		static FFloatTime GetDefault(const FOperatorSettings& InSettings)
+		static FFloatTime GetDefault(const FOperatorSettings* InSettings)
 		{
 			return FFloatTime(1.0f);
 		}
@@ -492,7 +492,7 @@ namespace Metasound
 	class TMathOpDivide<FFloatTime>
 	{
 	public:
-		static FFloatTime GetDefault(const FOperatorSettings& InSettings)
+		static FFloatTime GetDefault(const FOperatorSettings* InSettings)
 		{
 			return FFloatTime(1.0f);
 		}
@@ -526,7 +526,7 @@ namespace Metasound
 		using TDataClassReadRef = TDataReadReference<TDataClass>;
 		using TDataClassWriteRef = TDataWriteReference<TDataClass>;
 
-		static TDataClass GetDefault(const FOperatorSettings& InSettings)
+		static TDataClass GetDefault(const FOperatorSettings* InSettings)
 		{
 			return 0;
 		}
@@ -550,7 +550,7 @@ namespace Metasound
 	class TMathOpRandRange<FFloatTime>
 	{
 	public:
-		static FFloatTime GetDefault(const FOperatorSettings& InSettings)
+		static FFloatTime GetDefault(const FOperatorSettings* InSettings)
 		{
 			return FFloatTime(0.0f);
 		}
@@ -575,7 +575,7 @@ namespace Metasound
 	class TMathOpRandRange<FFrequency>
 	{
 	public:
-		static FFrequency GetDefault(const FOperatorSettings& InSettings)
+		static FFrequency GetDefault(const FOperatorSettings* InSettings)
 		{
 			return FFrequency(0.0f);
 		}
diff --git a/Engine/Plugins/Runtime/Metasound/Source/MetasoundStandardNodes/Public/MetasoundPeriodicTriggerNode.h b/Engine/Plugins/Runtime/Metasound/Source/MetasoundStandardNodes/Public/MetasoundPeriodicTriggerNode.h
--- a/Engine/Plugins/Runtime/Metasound/Source/MetasoundStandardNodes/Public/MetasoundPeriodicTriggerNode.h
+++ b/Engine/Plugins/Runtime/Metasound/Source/MetasoundStandardNodes/Public/MetasoundPeriodicTriggerNode.h
@@ -19,10 +19,9 @@ namespace Metasound
 	class METASOUNDSTANDARDNODES_API FPeriodicTriggerNode : public FNodeFacade
 	{
 		public:
-			FPeriodicTriggerNode(const FString& InInstanceName, const FGuid& InInstanceID, float InDefaultPeriodInSeconds);
 			FPeriodicTriggerNode(const FNodeInitData& InInitData);
 
-			float GetDefaultPeriodInSeconds() const;
+			virtual ~FPeriodicTriggerNode() = default;
 
 		private:
 			float DefaultPeriod = 1.0f;
diff --git a/Engine/Plugins/Runtime/Metasound/Source/MetasoundStandardNodes/Public/MetasoundTriggerDelayNode.h b/Engine/Plugins/Runtime/Metasound/Source/MetasoundStandardNodes/Public/MetasoundTriggerDelayNode.h
--- a/Engine/Plugins/Runtime/Metasound/Source/MetasoundStandardNodes/Public/MetasoundTriggerDelayNode.h
+++ b/Engine/Plugins/Runtime/Metasound/Source/MetasoundStandardNodes/Public/MetasoundTriggerDelayNode.h
@@ -19,12 +19,8 @@ namespace Metasound
 	class METASOUNDSTANDARDNODES_API FTriggerDelayNode : public FNodeFacade
 	{
 		public:
-			FTriggerDelayNode(const FString& InName, const FGuid& InInstanceID, float InDefaultDelayInSeconds);
 			FTriggerDelayNode(const FNodeInitData& InInitData);
 
-			float GetDefaultDelayInSeconds() const;
-
-		private:
-			float DefaultDelay = 1.0f;
+			virtual ~FTriggerDelayNode() = default;
 	};
 }
diff --git a/Engine/Plugins/Runtime/Metasound/Source/MetasoundStandardNodes/Public/MetasoundTriggerPipeNode.h b/Engine/Plugins/Runtime/Metasound/Source/MetasoundStandardNodes/Public/MetasoundTriggerPipeNode.h
--- a/Engine/Plugins/Runtime/Metasound/Source/MetasoundStandardNodes/Public/MetasoundTriggerPipeNode.h
+++ b/Engine/Plugins/Runtime/Metasound/Source/MetasoundStandardNodes/Public/MetasoundTriggerPipeNode.h
@@ -18,10 +18,5 @@ namespace Metasound
 		public:
 			FTriggerPipeNode(const FString& InName, const FGuid& InInstanceID, float InDefaultDelayInSeconds);
 			FTriggerPipeNode(const FNodeInitData& InInitData);
-
-			float GetDefaultDelayInSeconds() const;
-
-		private:
-			float DefaultDelay;
 	};
 }
diff --git a/Engine/Source/Runtime/AudioExtensions/Public/IAudioProxyInitializer.h b/Engine/Source/Runtime/AudioExtensions/Public/IAudioProxyInitializer.h
--- a/Engine/Source/Runtime/AudioExtensions/Public/IAudioProxyInitializer.h
+++ b/Engine/Source/Runtime/AudioExtensions/Public/IAudioProxyInitializer.h
@@ -18,7 +18,7 @@
 #define  IMPL_AUDIOPROXY_CLASS(FClassName) \
 	static FName GetAudioProxyTypeName() \
 	{ \
-		static FName MyClassName = TEXT( #FClassName ); \
+		static FName MyClassName = #FClassName; \
 		return MyClassName; \
 	} \
 	static constexpr bool bWasAudioProxyClassImplemented = true; \
@@ -26,11 +26,15 @@
 	friend class ::Audio::TProxyData<FClassName>;
 
 
-
 namespace Audio
 {
+	// Forward Declarations
+	class IProxyData;
+
 	constexpr bool bShouldCheckAudioProxyTypes = CHECK_AUDIOPROXY_TYPES;
 
+	using IProxyDataPtr = TUniquePtr<IProxyData>;
+
 	/*
 	 * Base class that allows us to typecheck proxy data before downcasting it in debug builds.
 	*/
@@ -55,6 +59,11 @@ namespace Audio
 			}
 		}
 
+		FName GetProxyTypeName() const
+		{
+			return ProxyTypeName;
+		}
+
 		template<typename ProxyType>
 		ProxyType& GetAs()
 		{
@@ -90,6 +99,8 @@ namespace Audio
 		IProxyData(FName InProxyTypeName)
 			: ProxyTypeName(InProxyTypeName)
 		{}
+
+		virtual IProxyDataPtr Clone() const = 0;
 	};
 
 	/**
@@ -115,10 +126,6 @@ namespace Audio
 	{
 		FName NameOfFeatureRequestingProxy;
 	};
-
-	typedef TUniquePtr<IProxyData> IProxyDataPtr;
-
-	
 }
 
 /*
diff --git a/Engine/Source/Runtime/Engine/Classes/Sound/SoundWave.h b/Engine/Source/Runtime/Engine/Classes/Sound/SoundWave.h
--- a/Engine/Source/Runtime/Engine/Classes/Sound/SoundWave.h
+++ b/Engine/Source/Runtime/Engine/Classes/Sound/SoundWave.h
@@ -291,6 +291,11 @@ public:
 	explicit FSoundWaveProxy(USoundWave* InWave);
 	FSoundWaveProxy(const FSoundWaveProxy& Other);
 
+	Audio::IProxyDataPtr Clone() const override
+	{
+		return MakeUnique<FSoundWaveProxy>(*this);
+	}
+
 	// UObject Interface cached data getters
 	const FName& GetFName() const { return NameCached; }
 	const FString& GetFullName() const { return FullNameCached; }
diff --git a/Engine/Source/Runtime/Engine/Private/SoundWave.cpp b/Engine/Source/Runtime/Engine/Private/SoundWave.cpp
--- a/Engine/Source/Runtime/Engine/Private/SoundWave.cpp
+++ b/Engine/Source/Runtime/Engine/Private/SoundWave.cpp
@@ -3031,7 +3031,6 @@ FSoundWaveProxy::FSoundWaveProxy(const FSoundWaveProxy& Other)
 	, CurrentChunkRevision(Other.CurrentChunkRevision)
 #endif // #if WITH_EDITOR
 {
-//	UE_LOG(LogAudio, Display, TEXT("MAXDEBUG: FSoundWaveProxy(const FSoundWaveProxy& Other)"));
 }
 
 void FSoundWaveProxy::ReleaseCompressedAudio()
