Fix mac build, & add generalized way to extract default when initializing input data ref from vertex definition

diff --git a/Engine/Plugins/Runtime/Metasound/Source/MetasoundGraphCore/Public/MetasoundDataReferenceCollection.h b/Engine/Plugins/Runtime/Metasound/Source/MetasoundGraphCore/Public/MetasoundDataReferenceCollection.h
--- a/Engine/Plugins/Runtime/Metasound/Source/MetasoundGraphCore/Public/MetasoundDataReferenceCollection.h
+++ b/Engine/Plugins/Runtime/Metasound/Source/MetasoundGraphCore/Public/MetasoundDataReferenceCollection.h
@@ -4,6 +4,7 @@
 
 #include "CoreMinimal.h"
 #include "MetasoundDataReference.h"
+#include "MetasoundVertex.h"
 
 namespace Metasound
 {
@@ -215,6 +216,31 @@ namespace Metasound
 				}
 			}
 
+			/** Returns a readable data ref from the collection or construct one
+			 * if one is not there with the default provided from the given input vertex collection.
+			 *
+			 * @param InName - Name of readable data.
+			 * @param InputVertices - Collection of input vertices to retrieve default from if reference is not available
+			 * @param ConstructorArgs - Arguments to pass to constructor of TDataReadReference<DataType>
+			 *
+			 * @return A readable data reference.
+			 */
+			template<typename DataType, typename DefaultDataType, typename... ConstructorArgTypes>
+			TDataReadReference<DataType> GetDataReadReferenceOrConstructWithVertexDefault(const FInputVertexInterface& InputVertices, const FString& InName, ConstructorArgTypes&&... ConstructorArgs) const
+			{
+				using FDataRefType = TDataReadReference<DataType>;
+
+				if (ContainsDataReadReference<DataType>(InName))
+				{
+					return GetDataReadReference<DataType>(InName);
+				}
+				else
+				{
+					const DefaultDataType DefaultValue = InputVertices[InName].GetDefaultValue().Value.Get<DefaultDataType>();
+					return FDataRefType::CreateNew(DefaultValue, Forward<ConstructorArgTypes>(ConstructorArgs)...);
+				}
+			}
+
 			/** Returns a writable data ref from the collection.
 			 *
 			 * @param InName - Name of readable data.
diff --git a/Engine/Plugins/Runtime/Metasound/Source/MetasoundStandardNodes/Private/MetasoundADSRNode.cpp b/Engine/Plugins/Runtime/Metasound/Source/MetasoundStandardNodes/Private/MetasoundADSRNode.cpp
--- a/Engine/Plugins/Runtime/Metasound/Source/MetasoundStandardNodes/Private/MetasoundADSRNode.cpp
+++ b/Engine/Plugins/Runtime/Metasound/Source/MetasoundStandardNodes/Private/MetasoundADSRNode.cpp
@@ -297,10 +297,10 @@ namespace Metasound
 		const FInputVertexInterface& InputInterface = GetVertexInterface().GetInputInterface();
 		FADSRReferences ADSRReferences =
 		{
-			GetOrConstructTime(InputInterface, TEXT("Attack")),
-			GetOrConstructTime(InputInterface, TEXT("Decay")),
-			GetOrConstructFloat(InputInterface, TEXT("Sustain Level")),
-			GetOrConstructTime(InputInterface, TEXT("Release"))
+			InputCollection.GetDataReadReferenceOrConstructWithVertexDefault<FFloatTime, float>(InputInterface, TEXT("Attack")),
+			InputCollection.GetDataReadReferenceOrConstructWithVertexDefault<FFloatTime, float>(InputInterface, TEXT("Decay")),
+			InputCollection.GetDataReadReferenceOrConstructWithVertexDefault<float, float>(InputInterface, TEXT("Sustain Level")),
+			InputCollection.GetDataReadReferenceOrConstructWithVertexDefault<FFloatTime, float>(InputInterface, TEXT("Release")),
 		};
 
 		return MakeUnique<FADSROperator>(InParams.OperatorSettings, TriggerAttack, TriggerRelease, ADSRReferences);
diff --git a/Engine/Plugins/Runtime/Metasound/Source/MetasoundStandardNodes/Private/MetasoundTriggerDelayNode.cpp b/Engine/Plugins/Runtime/Metasound/Source/MetasoundStandardNodes/Private/MetasoundTriggerDelayNode.cpp
--- a/Engine/Plugins/Runtime/Metasound/Source/MetasoundStandardNodes/Private/MetasoundTriggerDelayNode.cpp
+++ b/Engine/Plugins/Runtime/Metasound/Source/MetasoundStandardNodes/Private/MetasoundTriggerDelayNode.cpp
@@ -22,11 +22,9 @@ namespace Metasound
 
 	class FTriggerDelayOperator : public TExecutableOperator<FTriggerDelayOperator>
 	{
-		static constexpr float DefaultDelay = 1.0f; // In Seconds
-
 		public:
 			static const FNodeClassMetadata& GetNodeInfo();
-			static FVertexInterface DeclareVertexInterface();
+			static const FVertexInterface& GetVertexInterface();
 			static TUniquePtr<IOperator> CreateOperator(const FCreateOperatorParams& InParams, FBuildErrorArray& OutErrors);
 
 			FTriggerDelayOperator(const FOperatorSettings& InSettings, const FTriggerReadRef& InTriggerReset, const FTriggerReadRef& InTriggerIn, const FFloatTimeReadRef& InDelay);
@@ -132,14 +130,16 @@ namespace Metasound
 	{
 		const FTriggerDelayNode& DelayNode = static_cast<const FTriggerDelayNode&>(InParams.Node);
 
-		FFloatTimeReadRef Delay = InParams.InputDataReferences.GetDataReadReferenceOrConstruct<FFloatTime>(TEXT("Delay"), DefaultDelay);
+		const FInputVertexInterface& InputInterface = GetVertexInterface().GetInputInterface();
+		FFloatTimeReadRef Delay = InParams.InputDataReferences.GetDataReadReferenceOrConstructWithVertexDefault<FFloatTime, float>(InputInterface, TEXT("Delay"));
+
 		FTriggerReadRef TriggerIn = InParams.InputDataReferences.GetDataReadReferenceOrConstruct<FTrigger>(TEXT("In"), InParams.OperatorSettings);
 		FTriggerReadRef TriggerReset = InParams.InputDataReferences.GetDataReadReferenceOrConstruct<FTrigger>(TEXT("Reset"), InParams.OperatorSettings);
 
 		return MakeUnique<FTriggerDelayOperator>(InParams.OperatorSettings, TriggerReset, TriggerIn, Delay);
 	}
 
-	FVertexInterface FTriggerDelayOperator::DeclareVertexInterface()
+	const FVertexInterface& FTriggerDelayOperator::GetVertexInterface()
 	{
 		static const FVertexInterface Interface(
 			FInputVertexInterface(
@@ -167,7 +167,7 @@ namespace Metasound
 			Info.Description = LOCTEXT("Metasound_TriggerDelayNodeDescription", "Executes output trigger after the given delay time from the most recent execution of the input trigger .");
 			Info.Author = PluginAuthor;
 			Info.PromptIfMissing = PluginNodeMissingPrompt;
-			Info.DefaultInterface = DeclareVertexInterface();
+			Info.DefaultInterface = GetVertexInterface();
 
 			return Info;
 		};
