Fix up a couple more spots to use vertex input literal defaults

diff --git a/Engine/Plugins/Runtime/Metasound/Source/MetasoundGraphCore/Public/MetasoundDataReferenceCollection.h b/Engine/Plugins/Runtime/Metasound/Source/MetasoundGraphCore/Public/MetasoundDataReferenceCollection.h
--- a/Engine/Plugins/Runtime/Metasound/Source/MetasoundGraphCore/Public/MetasoundDataReferenceCollection.h
+++ b/Engine/Plugins/Runtime/Metasound/Source/MetasoundGraphCore/Public/MetasoundDataReferenceCollection.h
@@ -225,7 +225,7 @@ namespace Metasound
 			 *
 			 * @return A readable data reference.
 			 */
-			template<typename DataType, typename DefaultDataType, typename... ConstructorArgTypes>
+			template<typename DataType, typename DefaultDataType = DataType, typename... ConstructorArgTypes>
 			TDataReadReference<DataType> GetDataReadReferenceOrConstructWithVertexDefault(const FInputVertexInterface& InputVertices, const FString& InName, ConstructorArgTypes&&... ConstructorArgs) const
 			{
 				using FDataRefType = TDataReadReference<DataType>;
diff --git a/Engine/Plugins/Runtime/Metasound/Source/MetasoundStandardNodes/Private/MetasoundADSRNode.cpp b/Engine/Plugins/Runtime/Metasound/Source/MetasoundStandardNodes/Private/MetasoundADSRNode.cpp
--- a/Engine/Plugins/Runtime/Metasound/Source/MetasoundStandardNodes/Private/MetasoundADSRNode.cpp
+++ b/Engine/Plugins/Runtime/Metasound/Source/MetasoundStandardNodes/Private/MetasoundADSRNode.cpp
@@ -276,30 +276,18 @@ namespace Metasound
 	{
 		const FADSRNode& ADSRNode = static_cast<const FADSRNode&>(InParams.Node);
 		const FDataReferenceCollection& InputCollection = InParams.InputDataReferences;
+		const FInputVertexInterface& InputInterface = GetVertexInterface().GetInputInterface();
 
 		// TODO: Could no-op this if the trigger is not connected.
 		FTriggerReadRef TriggerAttack = InputCollection.GetDataReadReferenceOrConstruct<FTrigger>(TEXT("Trigger Attack"), InParams.OperatorSettings);
 		FTriggerReadRef TriggerRelease = InputCollection.GetDataReadReferenceOrConstruct<FTrigger>(TEXT("Trigger Release"), InParams.OperatorSettings);
 
-		auto GetOrConstructFloat = [&](const FInputVertexInterface& InputVertices, const FString& InputName)
-		{
-			float DefaultValue = InputVertices[InputName].GetDefaultValue().Value.Get<float>();
-			return InputCollection.GetDataReadReferenceOrConstruct<float>(InputName, DefaultValue);
-		};
-
-		auto GetOrConstructTime = [&](const FInputVertexInterface& InputVertices, const FString& InputName, ETimeResolution Resolution = ETimeResolution::Seconds)
-		{
-			float DefaultValue = InputVertices[InputName].GetDefaultValue().Value.Get<float>();
-			return InputCollection.GetDataReadReferenceOrConstruct<FFloatTime>(InputName, DefaultValue, Resolution);
-		};
-
 		// TODO: If none of these are connected, could pre-generate ADSR envelope and return a different operator.
-		const FInputVertexInterface& InputInterface = GetVertexInterface().GetInputInterface();
 		FADSRReferences ADSRReferences =
 		{
 			InputCollection.GetDataReadReferenceOrConstructWithVertexDefault<FFloatTime, float>(InputInterface, TEXT("Attack")),
 			InputCollection.GetDataReadReferenceOrConstructWithVertexDefault<FFloatTime, float>(InputInterface, TEXT("Decay")),
-			InputCollection.GetDataReadReferenceOrConstructWithVertexDefault<float, float>(InputInterface, TEXT("Sustain Level")),
+			InputCollection.GetDataReadReferenceOrConstructWithVertexDefault<float>(InputInterface, TEXT("Sustain Level")),
 			InputCollection.GetDataReadReferenceOrConstructWithVertexDefault<FFloatTime, float>(InputInterface, TEXT("Release")),
 		};
 
diff --git a/Engine/Plugins/Runtime/Metasound/Source/MetasoundStandardNodes/Private/MetasoundBufferAddFloatNode.cpp b/Engine/Plugins/Runtime/Metasound/Source/MetasoundStandardNodes/Private/MetasoundBufferAddFloatNode.cpp
--- a/Engine/Plugins/Runtime/Metasound/Source/MetasoundStandardNodes/Private/MetasoundBufferAddFloatNode.cpp
+++ b/Engine/Plugins/Runtime/Metasound/Source/MetasoundStandardNodes/Private/MetasoundBufferAddFloatNode.cpp
@@ -20,7 +20,7 @@ namespace Metasound
 		public:
 			static const FNodeClassMetadata& GetNodeInfo();
 
-			static FVertexInterface DeclareVertexInterface();
+			static const FVertexInterface& GetVertexInterface();
 
 			static TUniquePtr<IOperator> CreateOperator(const FCreateOperatorParams& InParams, FBuildErrorArray& OutErrors);
 
@@ -77,7 +77,7 @@ namespace Metasound
 	{
 	}
 
-	FVertexInterface FBufferAddFloatOperator::DeclareVertexInterface()
+	const FVertexInterface& FBufferAddFloatOperator::GetVertexInterface()
 	{
 		static const FVertexInterface Interface(
 			FInputVertexInterface(
@@ -96,9 +96,10 @@ namespace Metasound
 	{
 		const FBufferAddFloatNode& BufferAddFloatNode = static_cast<const FBufferAddFloatNode&>(InParams.Node);
 		const FDataReferenceCollection& InputCol = InParams.InputDataReferences;
+		const FInputVertexInterface& InputInterface = GetVertexInterface().GetInputInterface();
 
 		FAudioBufferReadRef InputBuffer = InputCol.GetDataReadReferenceOrConstruct<FAudioBuffer>(TEXT("Buffer"), FAudioBuffer(InParams.OperatorSettings.GetNumFramesPerBlock()));
-		FFloatReadRef InputGain = InputCol.GetDataReadReferenceOrConstruct<float>(TEXT("Addend"), 0.0f);
+		FFloatReadRef InputGain = InputCol.GetDataReadReferenceOrConstructWithVertexDefault<float>(InputInterface, TEXT("Addend"));
 
 		return MakeUnique<FBufferAddFloatOperator>(InParams.OperatorSettings, InputBuffer, InputGain);
 	}
@@ -123,7 +124,7 @@ namespace Metasound
 			Info.CategoryHierarchy = { LOCTEXT("Metasound_MathCategory", "Math") };
 			Info.Author = PluginAuthor;
 			Info.PromptIfMissing = PluginNodeMissingPrompt;
-			Info.DefaultInterface = DeclareVertexInterface();
+			Info.DefaultInterface = GetVertexInterface();
 
 			return Info;
 		};
diff --git a/Engine/Plugins/Runtime/Metasound/Source/MetasoundStandardNodes/Private/MetasoundGainMultiplyNode.cpp b/Engine/Plugins/Runtime/Metasound/Source/MetasoundStandardNodes/Private/MetasoundGainMultiplyNode.cpp
--- a/Engine/Plugins/Runtime/Metasound/Source/MetasoundStandardNodes/Private/MetasoundGainMultiplyNode.cpp
+++ b/Engine/Plugins/Runtime/Metasound/Source/MetasoundStandardNodes/Private/MetasoundGainMultiplyNode.cpp
@@ -20,7 +20,7 @@ namespace Metasound
 		public:
 			static const FNodeClassMetadata& GetNodeInfo();
 
-			static FVertexInterface DeclareVertexInterface();
+			static const FVertexInterface& GetVertexInterface();
 
 			static TUniquePtr<IOperator> CreateOperator(const FCreateOperatorParams& InParams, FBuildErrorArray& OutErrors);
 
@@ -95,7 +95,7 @@ namespace Metasound
 	{
 	}
 
-	FVertexInterface FGainMultiplyOperator::DeclareVertexInterface()
+	const FVertexInterface& FGainMultiplyOperator::GetVertexInterface()
 	{
 		static const FVertexInterface Interface(
 			FInputVertexInterface(
@@ -114,9 +114,10 @@ namespace Metasound
 	{
 		const FGainMultiplyNode& GainMultiplyNode = static_cast<const FGainMultiplyNode&>(InParams.Node);
 		const FDataReferenceCollection& InputCol = InParams.InputDataReferences;
+		const FInputVertexInterface& InputInterface = GetVertexInterface().GetInputInterface();
 
 		FAudioBufferReadRef InputBuffer = InputCol.GetDataReadReferenceOrConstruct<FAudioBuffer>(TEXT("Buffer"), FAudioBuffer(InParams.OperatorSettings.GetNumFramesPerBlock()));
-		FFloatReadRef InputGain = InputCol.GetDataReadReferenceOrConstruct<float>(TEXT("Gain"), 1.0f);
+		FFloatReadRef InputGain = InputCol.GetDataReadReferenceOrConstructWithVertexDefault<float>(InputInterface, TEXT("Gain"));
 
 		return MakeUnique<FGainMultiplyOperator>(InParams.OperatorSettings, InputBuffer, InputGain);
 	}
@@ -141,7 +142,7 @@ namespace Metasound
 			Info.CategoryHierarchy = { LOCTEXT("Metasound_MathCategory", "Math") };
 			Info.Author = PluginAuthor;
 			Info.PromptIfMissing = PluginNodeMissingPrompt;
-			Info.DefaultInterface = DeclareVertexInterface();
+			Info.DefaultInterface = GetVertexInterface();
 
 			return Info;
 		};
diff --git a/Engine/Plugins/Runtime/Metasound/Source/MetasoundStandardNodes/Private/MetasoundPeriodicTriggerNode.cpp b/Engine/Plugins/Runtime/Metasound/Source/MetasoundStandardNodes/Private/MetasoundPeriodicTriggerNode.cpp
--- a/Engine/Plugins/Runtime/Metasound/Source/MetasoundStandardNodes/Private/MetasoundPeriodicTriggerNode.cpp
+++ b/Engine/Plugins/Runtime/Metasound/Source/MetasoundStandardNodes/Private/MetasoundPeriodicTriggerNode.cpp
@@ -23,11 +23,9 @@ namespace Metasound
 
 	class FPeriodicTriggerOperator : public TExecutableOperator<FPeriodicTriggerOperator>
 	{
-		static constexpr float DefaultPeriod = 1.0f;
-
 		public:
 			static const FNodeClassMetadata& GetNodeInfo();
-			static FVertexInterface DeclareVertexInterface();
+			static const FVertexInterface& GetVertexInterface();
 			static TUniquePtr<IOperator> CreateOperator(const FCreateOperatorParams& InParams, FBuildErrorArray& OutErrors);
 
 			static constexpr float MinimumPeriodSeconds = 0.001f;
@@ -125,19 +123,20 @@ namespace Metasound
 	TUniquePtr<IOperator> FPeriodicTriggerOperator::CreateOperator(const FCreateOperatorParams& InParams, FBuildErrorArray& OutErrors)
 	{
 		const FPeriodicTriggerNode& PeriodicTriggerNode = static_cast<const FPeriodicTriggerNode&>(InParams.Node);
+		const FInputVertexInterface& InputInterface = GetVertexInterface().GetInputInterface();
 
-		FFloatTimeReadRef Period = InParams.InputDataReferences.GetDataReadReferenceOrConstruct<FFloatTime>(TEXT("Period"), DefaultPeriod, ETimeResolution::Seconds);
+		FFloatTimeReadRef Period = InParams.InputDataReferences.GetDataReadReferenceOrConstructWithVertexDefault<FFloatTime, float>(InputInterface, TEXT("Period"));
 		FTriggerReadRef TriggerEnable = InParams.InputDataReferences.GetDataReadReferenceOrConstruct<FTrigger>(TEXT("Activate"), InParams.OperatorSettings);
 		FTriggerReadRef TriggerDisable = InParams.InputDataReferences.GetDataReadReferenceOrConstruct<FTrigger>(TEXT("Deactivate"), InParams.OperatorSettings);
 
 		return MakeUnique<FPeriodicTriggerOperator>(InParams.OperatorSettings, TriggerEnable, TriggerDisable, Period);
 	}
 
-	FVertexInterface FPeriodicTriggerOperator::DeclareVertexInterface()
+	const FVertexInterface& FPeriodicTriggerOperator::GetVertexInterface()
 	{
 		static const FVertexInterface Interface(
 			FInputVertexInterface(
-				TInputDataVertexModel<FFloatTime>(TEXT("Period"), LOCTEXT("PeriodTooltip", "The period to trigger in seconds."), DefaultPeriod),
+				TInputDataVertexModel<FFloatTime>(TEXT("Period"), LOCTEXT("PeriodTooltip", "The period to trigger in seconds."), 1.0f),
 				TInputDataVertexModel<FTrigger>(TEXT("Activate"), LOCTEXT("TriggerEnableTooltip", "Enables executing periodic output triggers.")),
 				TInputDataVertexModel<FTrigger>(TEXT("Deactivate"), LOCTEXT("TriggerDisableTooltip", "Disables executing periodic output triggers."))
 			),
@@ -154,14 +153,14 @@ namespace Metasound
 		auto InitNodeInfo = []() -> FNodeClassMetadata
 		{
 			FNodeClassMetadata Info;
-			Info.ClassName = {Metasound::StandardNodes::Namespace, TEXT("PeriodicTrigger"), TEXT("")};
+			Info.ClassName = { Metasound::StandardNodes::Namespace, TEXT("PeriodicTrigger"), TEXT("") };
 			Info.MajorVersion = 1;
 			Info.MinorVersion = 0;
 			Info.DisplayName = LOCTEXT("Metasound_PeriodicTriggerNodeDisplayName", "Periodic Trigger");
 			Info.Description = LOCTEXT("Metasound_PeriodicTriggerNodeDescription", "Emits a trigger periodically based on the period duration given.");
 			Info.Author = PluginAuthor;
 			Info.PromptIfMissing = PluginNodeMissingPrompt;
-			Info.DefaultInterface = DeclareVertexInterface();
+			Info.DefaultInterface = GetVertexInterface();
 
 			return Info;
 		};
