Fix for discontinuity in Wave Player Node

diff --git a/Engine/Plugins/Runtime/Metasound/Source/MetasoundEngine/Private/MetasoundWave.cpp b/Engine/Plugins/Runtime/Metasound/Source/MetasoundEngine/Private/MetasoundWave.cpp
--- a/Engine/Plugins/Runtime/Metasound/Source/MetasoundEngine/Private/MetasoundWave.cpp
+++ b/Engine/Plugins/Runtime/Metasound/Source/MetasoundEngine/Private/MetasoundWave.cpp
@@ -39,11 +39,17 @@ namespace Audio
 
 		// set Circular Buffer capacity
 		CircularDecoderOutputBuffer.SetCapacity(DecodeBlockSizeInSamples * (1.f + FsInToFsOutRatio * MaxPitchShiftRatio));
+
+		TotalNumFramesOutput = 0;
+		TotalNumFramesDecoded = 0;
 		
 		// try to initialize decoders
 		bool bSuccessful = InitializeDecodersInternal(InWave);
 		bDecoderIsDone = !bSuccessful;
 
+		// initialize SRC object (will be re-initialized for pitch shifting)
+		Resampler.Init(Audio::EResamplingMethod::Linear, 1.f / FsInToFsOutRatio, NumChannels);
+
 		return bSuccessful;
 	}
 
@@ -54,19 +60,11 @@ namespace Audio
 		const float TotalInSamplesOutSamplesRatio = PitchShiftRatio * FsInToFsOutRatio;
 
 		const uint32 NumOutputSamples = NumOutputFrames * NumChannels;
-		uint32 NumSamplesToDecode = NumChannels * NumOutputFrames * TotalInSamplesOutSamplesRatio;
-
-		if (NumSamplesToDecode % 2)
-		{
-			++NumSamplesToDecode;
-		}
+		uint32 NumSamplesToDecode = NumChannels * FMath::CeilToInt(NumOutputFrames * TotalInSamplesOutSamplesRatio + 1.f);
 
 		// zero the output
 		FMemory::Memzero(OutputDest, NumOutputSamples * sizeof(float));
 
-		// TODO: calculate the ratio first using SR & PS, then see if the ratio is 1.0f
-		const bool bNeedsSRC = !(FMath::IsNearlyZero(PitchShiftInCents) && FMath::IsNearlyEqual(InputSampleRate, OutputSampleRate));
-
 		// Decode audio until we have enough to satisfy the output (post SRC)
 		while (!bDecoderIsDone && (CircularDecoderOutputBuffer.Num() < NumSamplesToDecode))
 		{
@@ -81,6 +79,9 @@ namespace Audio
 			// push that (interleaved) audio to the (interleaved) circular buffer
 			int32 NumPushed = CircularDecoderOutputBuffer.Push(PreSrcBuffer.GetData(), NumSamplesDecoded);
 			ensure(NumPushed == NumSamplesDecoded); // there will be a discontinuity in the output because our CircularDecoderOutputBuffer was not large enough!
+
+			// for debugging
+			TotalNumFramesDecoded += (NumPushed / NumChannels);
 		}
 
 		// now that we have enough audio decoded, pop off the circular buffer into an interleaved, pre-src temp buffer
@@ -91,7 +92,6 @@ namespace Audio
 		PreSrcBuffer.Reset(NumSamplesToPop);
 		PreSrcBuffer.AddZeroed(NumSamplesToPop);
 
-		// TODO: if we don't need SRC, just pop directly into the output buffer and call it a day
 		// pop into pre-src buffer
 		uint32 NumSamplesPopped = CircularDecoderOutputBuffer.Pop(PreSrcBuffer.GetData(), NumSamplesToPop);
 		ensure(NumSamplesPopped == NumSamplesToPop); // Circular buffer is too small! there will be discontinuities in the output
@@ -99,35 +99,32 @@ namespace Audio
 		// holds result of sample rate conversion
 		float* PostSrcBufferPtr = PreSrcBuffer.GetData(); // assume no SRC
 
-		int32 NumFramesConverted = NumSamplesToPop / NumChannels;
-
-		if (bNeedsSRC)
+		// perform SRC 
+		if (!FMath::IsNearlyEqual(LastPitchShiftCents, PitchShiftInCents))
 		{
-			// perform SRC
 			Resampler.Init(Audio::EResamplingMethod::Linear, 1.f / TotalInSamplesOutSamplesRatio, NumChannels);
-
-			// something is wrong
-			int32 Error = Resampler.ProcessAudio(PreSrcBuffer.GetData(), NumSamplesToPop / NumChannels, false, OutputDest, NumOutputFrames, NumFramesConverted);
-			ensure(!Error);
-
-			if (Error)
-			{
-				bDecoderIsDone = false;
-				Input.Reset();
-				Output.Reset();
-				Decoder.Reset();
-			}
+			LastPitchShiftCents = PitchShiftInCents;
 		}
-		else
+
+		int32 NumFramesConverted;
+		int32 Error = Resampler.ProcessAudio(PreSrcBuffer.GetData(), NumSamplesToPop / NumChannels, false, OutputDest, NumOutputFrames, NumFramesConverted);
+		ensure(!Error);
+
+		if (Error)
 		{
-			FMemory::Memcpy(OutputDest, PreSrcBuffer.GetData(), NumOutputSamples * sizeof(float));
+			bDecoderIsDone = false;
+			Input.Reset();
+			Output.Reset();
+			Decoder.Reset();
 		}
 
+
 		if (NumFramesConverted < (int32)NumOutputFrames)
 		{
 			bDecoderIsDone = true;
 		}
-
+		
+		TotalNumFramesOutput += NumFramesConverted;
 		return NumFramesConverted;
 	}
 
diff --git a/Engine/Plugins/Runtime/Metasound/Source/MetasoundEngine/Public/MetasoundWave.h b/Engine/Plugins/Runtime/Metasound/Source/MetasoundEngine/Public/MetasoundWave.h
--- a/Engine/Plugins/Runtime/Metasound/Source/MetasoundEngine/Public/MetasoundWave.h
+++ b/Engine/Plugins/Runtime/Metasound/Source/MetasoundEngine/Public/MetasoundWave.h
@@ -106,6 +106,10 @@ namespace Audio
 		uint32 DecodeBlockSizeInFrames;
 		uint32 DecodeBlockSizeInSamples;
 		
+		// cached values
+		float LastPitchShiftCents{ 0.f };
+		int32 TotalNumFramesOutput{ 0 };
+		int32 TotalNumFramesDecoded{ 0 };
 
 		bool bDecoderIsDone{ true };
 
