Metasound graph core TArray support

diff --git a/Engine/Plugins/Runtime/Metasound/Source/MetasoundEditor/Private/MetasoundEditorModule.cpp b/Engine/Plugins/Runtime/Metasound/Source/MetasoundEditor/Private/MetasoundEditorModule.cpp
--- a/Engine/Plugins/Runtime/Metasound/Source/MetasoundEditor/Private/MetasoundEditorModule.cpp
+++ b/Engine/Plugins/Runtime/Metasound/Source/MetasoundEditor/Private/MetasoundEditorModule.cpp
@@ -202,9 +202,16 @@ namespace Metasound
 							}
 							break;
 
-							default:
+							
 							case ELiteralType::None:
 							case ELiteralType::Invalid:
+							// TODO: handle array types.
+							case ELiteralType::NoneArray:
+							case ELiteralType::BooleanArray:
+							case ELiteralType::IntegerArray:
+							case ELiteralType::StringArray:
+							case ELiteralType::FloatArray:
+							default:
 							{
 								// Audio types are ubiquitous, so add as subcategory
 								// to be able to stylize connections (i.e. wire color)
@@ -217,7 +224,7 @@ namespace Metasound
 								{
 									PinCategory = FGraphBuilder::PinCategoryAudioFormat;
 								}
-								static_assert(static_cast<int32>(ELiteralType::Invalid) == 7, "Possible missing binding of pin category to primitive type");
+								static_assert(static_cast<int32>(ELiteralType::Invalid) == 12, "Possible missing binding of pin category to primitive type");
 							}
 							break;
 						}
diff --git a/Engine/Plugins/Runtime/Metasound/Source/MetasoundFrontend/Private/MetasoundFrontendDocument.cpp b/Engine/Plugins/Runtime/Metasound/Source/MetasoundFrontend/Private/MetasoundFrontendDocument.cpp
--- a/Engine/Plugins/Runtime/Metasound/Source/MetasoundFrontend/Private/MetasoundFrontendDocument.cpp
+++ b/Engine/Plugins/Runtime/Metasound/Source/MetasoundFrontend/Private/MetasoundFrontendDocument.cpp
@@ -83,13 +83,18 @@ void FMetasoundFrontendLiteral::SetFromLiteral(const Metasound::FLiteral& InLite
 		}
 		break;
 
+		case ELiteralType::None:
 		case ELiteralType::UObjectProxy:
+		case ELiteralType::NoneArray:
+		case ELiteralType::BooleanArray:
+		case ELiteralType::IntegerArray:
+		case ELiteralType::FloatArray:
+		case ELiteralType::StringArray:
 		case ELiteralType::UObjectProxyArray:
-		case ELiteralType::None:
 		case ELiteralType::Invalid:
 		default:
 		{
-			static_assert(static_cast<int32>(ELiteralType::Invalid) == 7, "Possible missing literal type switch coverage");
+			static_assert(static_cast<int32>(ELiteralType::Invalid) == 12, "Possible missing literal type switch coverage");
 		}
 	}
 }
diff --git a/Engine/Plugins/Runtime/Metasound/Source/MetasoundGraphCore/Private/MetasoundNodeInterface.cpp b/Engine/Plugins/Runtime/Metasound/Source/MetasoundGraphCore/Private/MetasoundNodeInterface.cpp
--- a/Engine/Plugins/Runtime/Metasound/Source/MetasoundGraphCore/Private/MetasoundNodeInterface.cpp
+++ b/Engine/Plugins/Runtime/Metasound/Source/MetasoundGraphCore/Private/MetasoundNodeInterface.cpp
@@ -53,7 +53,7 @@ namespace Metasound
 
 		if (InVariant != NAME_None)
 		{
-			return FName(FString::Printf(TEXT("%s.%s"), *ScopedName.ToString(), *InVariant.ToString()));
+			return FName(*FString::Printf(TEXT("%s.%s"), *ScopedName.ToString(), *InVariant.ToString()));
 		}
 		else
 		{
@@ -63,7 +63,7 @@ namespace Metasound
 
 	FName FNodeClassName::FormatScopedName(const FName& InNamespace, const FName& InName)
 	{
-		return FName(FString::Printf(TEXT("%s.%s"), *InNamespace.ToString(), *InName.ToString()));
+		return FName(*FString::Printf(TEXT("%s.%s"), *InNamespace.ToString(), *InName.ToString()));
 	}
 
 	bool operator==(const FOutputDataSource& InLeft, const FOutputDataSource& InRight)
diff --git a/Engine/Plugins/Runtime/Metasound/Source/MetasoundGraphCore/Public/MetasoundDataFactory.h b/Engine/Plugins/Runtime/Metasound/Source/MetasoundGraphCore/Public/MetasoundDataFactory.h
--- a/Engine/Plugins/Runtime/Metasound/Source/MetasoundGraphCore/Public/MetasoundDataFactory.h
+++ b/Engine/Plugins/Runtime/Metasound/Source/MetasoundGraphCore/Public/MetasoundDataFactory.h
@@ -21,27 +21,44 @@ namespace Metasound
 		template<typename DataType, typename... ArgTypes>
 		struct TDataTypeConstructorTraits
 		{
+			// Array construction for TArray types. 
+			static constexpr bool bIsArrayDefaultConstructible = false;
+			static constexpr bool bIsArrayConstructibleWithSettings = false;
+			static constexpr bool bIsArrayConstructibleWithArgs = false;
+			static constexpr bool bIsArrayConstructibleWithSettingsAndArgs = false;
+
+			// DataType Construction
 			static constexpr bool bIsDefaultConstructible = std::is_constructible<DataType>::value;
 			static constexpr bool bIsConstructibleWithSettings = std::is_constructible<DataType, const FOperatorSettings&>::value;
 			static constexpr bool bIsConstructibleWithArgs = std::is_constructible<DataType, ArgTypes...>::value;
 			static constexpr bool bIsConstructibleWithSettingsAndArgs = std::is_constructible<DataType, const FOperatorSettings&, ArgTypes...>::value;
 		};
 
-		/** Partial specialization for TArray<> types. Current implementation of the TArray<>
-		 * exposes a copy constructor which is declared even if the TArray<> element types
-		 * are incompatible. Because the copy constructor is declared, std::is_constructible 
-		 * returns true even if the constructor will not compile. This partial specialization
-		 * corrects the TDataTypeConstructorTraits so that the TArray<> copy constructor is
-		 * not used if the underlying element types are incompatible. 
+		/** Partial specialization for TArray<> types. A TArray can be constructed elementwise with
+		 * an array containing constructor args.
 		 */
 		template<typename ElementType, typename OtherElementType>
 		struct TDataTypeConstructorTraits<TArray<ElementType>, TArray<OtherElementType>>
 		{
+		private:
 			using DataType = TArray<ElementType>;
+			using ElementConstructorTraits = TDataTypeConstructorTraits<ElementType, OtherElementType>;
+
+			// FLiteral::FNone is a special case where an array of TArray<FLiteral::FNone> only determines 
+			// the number of elements that should be default constructed. 
+			static constexpr bool bIsOtherElementNone = std::is_same<FLiteral::FNone, OtherElementType>::value;
+
+		public:
 
+			// Array construction of TArray types.
+			static constexpr bool bIsArrayDefaultConstructible = ElementConstructorTraits::bIsDefaultConstructible && bIsOtherElementNone;
+			static constexpr bool bIsArrayConstructibleWithSettings = ElementConstructorTraits::bIsConstructibleWithSettings && bIsOtherElementNone;
+			static constexpr bool bIsArrayConstructibleWithArgs = ElementConstructorTraits::bIsConstructibleWithArgs;
+			static constexpr bool bIsArrayConstructibleWithSettingsAndArgs = ElementConstructorTraits::bIsConstructibleWithSettingsAndArgs;
+
+			// TArray construction
 			static constexpr bool bIsDefaultConstructible = std::is_constructible<DataType>::value;
 			static constexpr bool bIsConstructibleWithSettings = std::is_constructible<DataType, const FOperatorSettings&>::value;
-			// TArray copy constructor is only compilable if the element types are compatible.
 			static constexpr bool bIsConstructibleWithArgs = std::is_constructible<ElementType, OtherElementType>::value;
 			// TArray does not take operator settings.
 			static constexpr bool bIsConstructibleWithSettingsAndArgs = false;
@@ -57,37 +74,50 @@ namespace Metasound
 		template<typename DataType, typename... Types>
 		struct TDataTypeVariantConstructorTraits;
 
-		// Specialization for single arg.
-		template<typename DataType, typename ArgType>
-		struct TDataTypeVariantConstructorTraits<DataType, ArgType>
-		{
-			using FConstructorTraits = TDataTypeConstructorTraits<DataType, ArgType>;
-
-			public:
-
-			static constexpr bool bIsDefaultConstructible = FConstructorTraits::bIsDefaultConstructible;
-			static constexpr bool bIsConstructibleWithSettings = FConstructorTraits::bIsConstructibleWithSettings;
-			static constexpr bool bIsConstructibleWithArgs = FConstructorTraits::bIsConstructibleWithArgs;
-			static constexpr bool bIsConstructibleWithSettingsAndArgs = FConstructorTraits::bIsConstructibleWithSettingsAndArgs;
-		};
-
 		// Specialization for unrolling parameter packs. 
 		template<typename DataType, typename ArgType, typename... AdditionalTypes>
 		struct TDataTypeVariantConstructorTraits<DataType, ArgType, AdditionalTypes...>
 		{
 			private:
 
+			// Recursive declaration to get constructor traits for the additional types. 
 			using FAdditionalConstructorTraits = TDataTypeVariantConstructorTraits<DataType, AdditionalTypes...>;
+
+			// Constructor traits for this ArgType
 			using FArgConstructorTraits = TDataTypeConstructorTraits<DataType, ArgType>;
 			
 			public:
 
+			static constexpr bool bIsArrayDefaultConstructible = FAdditionalConstructorTraits::bIsArrayDefaultConstructible;
+			static constexpr bool bIsArrayConstructibleWithSettings = FAdditionalConstructorTraits::bIsArrayConstructibleWithSettings;
+			static constexpr bool bIsArrayConstructibleWithArgs = FArgConstructorTraits::bIsArrayConstructibleWithArgs || FAdditionalConstructorTraits::bIsArrayConstructibleWithArgs;
+			static constexpr bool bIsArrayConstructibleWithSettingsAndArgs = FArgConstructorTraits::bIsArrayConstructibleWithSettingsAndArgs || FAdditionalConstructorTraits::bIsArrayConstructibleWithSettingsAndArgs;
+
 			static constexpr bool bIsDefaultConstructible = FAdditionalConstructorTraits::bIsDefaultConstructible;
 			static constexpr bool bIsConstructibleWithSettings = FAdditionalConstructorTraits::bIsConstructibleWithSettings;
 			static constexpr bool bIsConstructibleWithArgs = FArgConstructorTraits::bIsConstructibleWithArgs || FAdditionalConstructorTraits::bIsConstructibleWithArgs;
 			static constexpr bool bIsConstructibleWithSettingsAndArgs = FArgConstructorTraits::bIsConstructibleWithSettingsAndArgs || FAdditionalConstructorTraits::bIsConstructibleWithSettingsAndArgs;
 		};
 
+		// Specialization for single arg. This terminates the variadic parameter unpacking by excluding the recursive declaration of TDataTypeVariantConstructorTraits<>
+		template<typename DataType, typename ArgType>
+		struct TDataTypeVariantConstructorTraits<DataType, ArgType>
+		{
+			// Constructor traits for this arg type.
+			using FConstructorTraits = TDataTypeConstructorTraits<DataType, ArgType>;
+
+			public:
+			static constexpr bool bIsArrayDefaultConstructible = FConstructorTraits::bIsArrayDefaultConstructible;
+			static constexpr bool bIsArrayConstructibleWithSettings = FConstructorTraits::bIsArrayConstructibleWithSettings;
+			static constexpr bool bIsArrayConstructibleWithArgs = FConstructorTraits::bIsArrayConstructibleWithArgs;
+			static constexpr bool bIsArrayConstructibleWithSettingsAndArgs = FConstructorTraits::bIsArrayConstructibleWithSettingsAndArgs;
+
+			static constexpr bool bIsDefaultConstructible = FConstructorTraits::bIsDefaultConstructible;
+			static constexpr bool bIsConstructibleWithSettings = FConstructorTraits::bIsConstructibleWithSettings;
+			static constexpr bool bIsConstructibleWithArgs = FConstructorTraits::bIsConstructibleWithArgs;
+			static constexpr bool bIsConstructibleWithSettingsAndArgs = FConstructorTraits::bIsConstructibleWithSettingsAndArgs;
+		};
+
 		// Specialization for unpacking the types supported by a TVariant. 
 		template<typename DataType, typename FirstVariantType, typename... AdditionalVariantTypes>
 		struct TDataTypeVariantConstructorTraits<DataType, TVariant<FirstVariantType, AdditionalVariantTypes...>>
@@ -98,6 +128,11 @@ namespace Metasound
 			
 			public:
 
+			static constexpr bool bIsArrayDefaultConstructible = FConstructorTraits::bIsArrayDefaultConstructible;
+			static constexpr bool bIsArrayConstructibleWithSettings = FConstructorTraits::bIsArrayConstructibleWithSettings;
+			static constexpr bool bIsArrayConstructibleWithArgs = FConstructorTraits::bIsArrayConstructibleWithArgs;
+			static constexpr bool bIsArrayConstructibleWithSettingsAndArgs = FConstructorTraits::bIsArrayConstructibleWithSettingsAndArgs;
+
 			static constexpr bool bIsDefaultConstructible = FConstructorTraits::bIsDefaultConstructible;
 			static constexpr bool bIsConstructibleWithSettings = FConstructorTraits::bIsConstructibleWithSettings;
 			static constexpr bool bIsConstructibleWithArgs = FConstructorTraits::bIsConstructibleWithArgs;
@@ -320,7 +355,7 @@ namespace Metasound
 			template<
 				typename VariantParseType,
 				typename VariantType,
-				typename std::enable_if< VariantParseType::bCreateWithArgType, int>::type = 0
+				typename std::enable_if< VariantParseType::bCreateWithArg, int>::type = 0
 			>
 			static auto CreateNewFromVariant(const FOperatorSettings& InSettings, const VariantType& InVariant)
 			{
@@ -332,7 +367,7 @@ namespace Metasound
 
 			/** Create a new object using FOperatorSettings and TVariant.
 			 *
-			 * Note: In this implementation, the InVariant object is ignore. FOperatorSettings 
+			 * Note: In this implementation, the InVariant object is ignored. FOperatorSettings 
 			 * are forwarded to CreatorType::CreateNew(...) if it is supported.
 			 *
 			 * @tparam VariantParseType - A type which describes how to parse the InVariant
@@ -349,15 +384,102 @@ namespace Metasound
 				// The constructor must be the default constructor or accept a FOperatorSettings.
 				using TExplicitArgsConstructorForwarding = TExplicitArgsConstructorForwarding<DataType>;
 
+				constexpr bool bExpectsNone = std::is_same<typename VariantParseType::ArgType, FLiteral::FNone>::value;
+				
 				// When constructing an object from a variant, callers expect the variant value
 				// to be used in the constructor. It is an error to ignore the variant value.
 				// This error can be fixed by changing the InVariant object passed in 
 				// or by adding a constructor to the DataType which accepts the VariantType.
-				checkf(false, TEXT("The value passed to the constructor is being ignored")); 
+				checkf(bExpectsNone, TEXT("The value passed to the constructor is being ignored")); 
 
 				return FInternalFactory::template CreateNew<TExplicitArgsConstructorForwarding>(InSettings);
 			}
 
+			/** Create a new object using FOperatorSettings and TVariant.
+			 *
+			 * Note: In this implementation, the InVariant object is an array of args
+			 * where each array element is the constructor args for one resulting 
+			 * array element. Constructor args are forwarded to the element data type and 
+			 * the resulting object is placed in the output array.
+			 *
+			 * @tparam VariantParseType - A type which describes how to parse the InVariant
+			 *                            argument.
+			 * @tparam VariantType - The expected underlying type stored in InVariant.
+			 */
+			template<
+				typename VariantParseType,
+				typename VariantType,
+				typename std::enable_if< VariantParseType::bCreateArrayWithArg, int>::type = 0
+			>
+			static auto CreateNewFromVariant(const FOperatorSettings& InSettings, const VariantType& InVariant)
+			{
+				using FConstructorForwarding = TExplicitArgsConstructorForwarding<DataType>;
+
+				using FArgType = typename VariantParseType::ArgType;
+				using FArgElementType = typename VariantParseType::ArgElementType;
+				using FElementType = typename VariantParseType::ElementType;
+				using FElementConstructorForwarding = TExplicitArgsConstructorForwarding<FElementType>;
+
+				using FElementFactory = TDataFactoryInternal<FElementType, DataFactoryPrivate::TDataTypeCreator<FElementType>>;
+
+				const FArgType& ConstructorValues = InVariant.template Get<FArgType>();
+
+				// Array construct objects
+				DataType OutputContainer;
+				for (const FArgElementType& ConstructorValue : ConstructorValues)
+				{
+					OutputContainer.Add(FElementFactory::template CreateNew<FElementConstructorForwarding>(InSettings, ConstructorValue));
+				}
+
+				// Move construct array.
+				return FInternalFactory::template CreateNew<FConstructorForwarding>(InSettings, MoveTemp(OutputContainer));
+			}
+
+			/** Create a new object using FOperatorSettings and TVariant.
+			 *
+			 * Note: In this implementation, the InVariant object is an array of 
+			 * FLiteral::FNone where each array element simply denotes that a corresponding
+			 * element should be default constructed. 
+			 *
+			 * @tparam VariantParseType - A type which describes how to parse the InVariant
+			 *                            argument.
+			 * @tparam VariantType - The expected underlying type stored in InVariant.
+			 */
+			template<
+				typename VariantParseType,
+				typename VariantType,
+				typename std::enable_if< VariantParseType::bCreateArrayWithoutArg, int>::type = 0
+			>
+			static auto CreateNewFromVariant(const FOperatorSettings& InSettings, const VariantType& InVariant)
+			{
+				using FConstructorForwarding = TExplicitArgsConstructorForwarding<DataType>;
+
+				using FArgType = typename VariantParseType::ArgType;
+				using FArgElementType = typename VariantParseType::ArgElementType;
+				using FElementType = typename VariantParseType::ElementType;
+				using FElementConstructorForwarding = TExplicitArgsConstructorForwarding<FElementType>;
+				using FElementFactory = TDataFactoryInternal<FElementType, DataFactoryPrivate::TDataTypeCreator<FElementType>>;
+
+				constexpr bool bExpectsNone = std::is_same<FArgElementType, FLiteral::FNone>::value;
+				
+				// When constructing an object from a variant, callers expect the variant value
+				// to be used in the constructor. It is an error to ignore the variant value.
+				// This error can be fixed by changing the InVariant object passed in 
+				// or by adding a constructor to the DataType which accepts the VariantType.
+				checkf(bExpectsNone, TEXT("The value passed to the constructor is being ignored")); 
+
+
+				// Retrieve argument array
+				const FArgType& ConstructorValues = InVariant.template Get<FArgType>();
+
+				// Default construct equivalent number of objects
+				DataType OutputContainer;
+				OutputContainer.AddDefaulted(ConstructorValues.Num());
+
+				// Move construct array.
+				return FInternalFactory::template CreateNew<FConstructorForwarding>(InSettings, MoveTemp(OutputContainer));
+			}
+
 			/** Create a new object using FOperatorSettings and TVariant.
 			 *
 			 * Note: In this implementation, the TVariant is parsed as a fallback type.
@@ -371,7 +493,7 @@ namespace Metasound
 			template<
 				typename VariantParseType,
 				typename VariantType,
-				typename std::enable_if< VariantParseType::bCreateWithFallbackArgType, int>::type = 0
+				typename std::enable_if< VariantParseType::bCreateWithFallbackArg, int>::type = 0
 			>
 			static auto CreateNewFromVariant(const FOperatorSettings& InSettings, const VariantType& InVariant)
 			{
@@ -502,26 +624,57 @@ namespace Metasound
 			static constexpr bool bIsConstructibleWithFallbackArg = FVariantPromoter::bIsValidType; 
 		};
 
+		/** TElementType determines the element type of a container. */
+		template<typename ContainterType>
+		struct TElementType
+		{
+			using Type = void;
+		};
+
+		/** TElementType specialization for TArray */
+		template<typename ElementType>
+		struct TElementType<TArray<ElementType>>
+		{
+			using Type = ElementType;
+		};
+
 		/** TDataTypeVariantParsing informs the TDataVariantFactoryInternal on which factory
-		 * method to instatiate.
+		 * method to instantiate.
 		 */
 		template<typename DataType, typename DesiredArgType, typename FirstVariantType, typename... AdditionalVariantTypes>
 		struct TDataTypeVariantParsing
 		{
+		private:
 			using FDesiredConstructorTraits = TDataTypeConstructorTraits<DataType, DesiredArgType>; 
 			using FFallbackHelper = TDataTypeVariantFallbackHelper<DataType, FirstVariantType, AdditionalVariantTypes...>;
 
-			static constexpr bool bCreateWithArgType = FDesiredConstructorTraits::bIsConstructibleWithArgs || FDesiredConstructorTraits::bIsConstructibleWithSettingsAndArgs;
-			static constexpr bool bCreateWithoutArg = !bCreateWithArgType && (FDesiredConstructorTraits::bIsConstructibleWithSettings || FDesiredConstructorTraits::bIsDefaultConstructible);
-			static constexpr bool bCreateWithFallbackArgType = !(bCreateWithArgType || bCreateWithoutArg) && FFallbackHelper::bIsConstructibleWithFallbackArg;
-			static constexpr bool bCannotForwardToConstructor = !(bCreateWithArgType || bCreateWithoutArg || bCreateWithFallbackArgType);
+			// Determine which elementwise construction methods are supported
+			static constexpr bool bCanCreateArrayWithArg = FDesiredConstructorTraits::bIsArrayConstructibleWithArgs || FDesiredConstructorTraits::bIsArrayConstructibleWithSettingsAndArgs;
+			static constexpr bool bCanCreateArrayWithoutArg = FDesiredConstructorTraits::bIsArrayConstructibleWithSettings || FDesiredConstructorTraits::bIsArrayDefaultConstructible;
 
-			typedef DesiredArgType ArgType;
-			typedef typename FFallbackHelper::FallbackType FallbackArgType;
-		};
+			// Determine which construction methods are supported
+			static constexpr bool bCanCreateWithArg = FDesiredConstructorTraits::bIsConstructibleWithArgs || FDesiredConstructorTraits::bIsConstructibleWithSettingsAndArgs;
+			static constexpr bool bCanCreateWithoutArg = FDesiredConstructorTraits::bIsConstructibleWithSettings || FDesiredConstructorTraits::bIsDefaultConstructible;
+			static constexpr bool bCanCreateWithFallbackArg = FFallbackHelper::bIsConstructibleWithFallbackArg;
+
+		public:
+
+			static constexpr bool bCannotForwardToConstructor = !(bCanCreateWithArg || bCanCreateWithoutArg || bCanCreateWithFallbackArg || bCanCreateArrayWithArg || bCanCreateArrayWithoutArg);
 
+			// Determine which construction method to use.
+			static constexpr bool bCreateArrayWithArg = bCanCreateArrayWithArg;
+			static constexpr bool bCreateArrayWithoutArg = !(bCanCreateArrayWithArg) && bCanCreateArrayWithoutArg;
 
+			static constexpr bool bCreateWithArg = !(bCreateArrayWithArg || bCreateArrayWithoutArg) && bCanCreateWithArg;
+			static constexpr bool bCreateWithoutArg = !(bCreateArrayWithArg || bCreateArrayWithoutArg || bCreateWithArg) && bCanCreateWithoutArg;
+			static constexpr bool bCreateWithFallbackArg = !(bCreateArrayWithArg || bCreateArrayWithoutArg || bCreateWithArg || bCreateWithoutArg) && bCanCreateWithFallbackArg;
 
+			// Determine types
+			using ElementType = typename TElementType<DataType>::Type;
+			using ArgElementType = typename TElementType<DesiredArgType>::Type;
+			using ArgType = DesiredArgType;
+			using FallbackArgType = typename FFallbackHelper::FallbackType;
+		};
 	} // End namespace DataFactoryPrivate
 
 	/** Determines whether a DataType supports a constructor which accepts and FOperatorSettings
@@ -533,12 +686,32 @@ namespace Metasound
 	template<typename DataType, typename... ArgTypes>
 	struct TIsParsable
 	{
-		using TConstructorTraits = DataFactoryPrivate::TDataTypeConstructorTraits<DataType, ArgTypes...>;
+		using FConstructorTraits = DataFactoryPrivate::TDataTypeConstructorTraits<DataType, ArgTypes...>;
 
 		/* True if the DataType supports a constructor which accepts and FOperatorSettings 
 		 * with ArgTypes and/or just ArgTypes. False otherwise.
 		 */
-		static constexpr bool Value = TConstructorTraits::bIsConstructibleWithArgs || TConstructorTraits::bIsConstructibleWithSettingsAndArgs;
+		static constexpr bool Value = FConstructorTraits::bIsConstructibleWithArgs || FConstructorTraits::bIsConstructibleWithSettingsAndArgs || FConstructorTraits::bIsArrayConstructibleWithArgs || FConstructorTraits::bIsArrayConstructibleWithSettingsAndArgs;
+	};
+
+	/* Specialization for FLiteral::FNone arg type. */
+	template<typename DataType>
+	struct TIsParsable<DataType, FLiteral::FNone>
+	{
+		using FConstructorTraits = DataFactoryPrivate::TDataTypeConstructorTraits<DataType>;
+
+		/* True if the DataType supports a constructor which accepts and FOperatorSettings or default constructor.  */
+		static constexpr bool Value = FConstructorTraits::bIsConstructibleWithArgs || FConstructorTraits::bIsConstructibleWithSettingsAndArgs || FConstructorTraits::bIsArrayConstructibleWithArgs || FConstructorTraits::bIsArrayConstructibleWithSettingsAndArgs;
+	};
+
+	/* Specialization for TArray data type and TArray<FLiteral::FNone> arg type. */
+	template<typename DataElementType>
+	struct TIsParsable<TArray<DataElementType>, TArray<FLiteral::FNone>>
+	{
+		using FElementConstructorTraits = DataFactoryPrivate::TDataTypeConstructorTraits<DataElementType>;
+
+		/* True if the DataElementType supports a default constructor or constructor which accepts and FOperatorSettings. */
+		static constexpr bool Value = FElementConstructorTraits::bIsConstructibleWithArgs || FElementConstructorTraits::bIsConstructibleWithSettingsAndArgs ;
 	};
 
 	/** Determines whether a DataType supports construction using the given literal.
@@ -550,7 +723,15 @@ namespace Metasound
 	{
 		using TVariantConstructorTraits = DataFactoryPrivate::TDataTypeVariantConstructorTraits<DataType, FLiteral::FVariantType>;
 
-		static constexpr bool bIsParsableFromAnyLiteralType = TVariantConstructorTraits::bIsDefaultConstructible || TVariantConstructorTraits::bIsConstructibleWithSettings || TVariantConstructorTraits::bIsConstructibleWithArgs || TVariantConstructorTraits::bIsConstructibleWithSettingsAndArgs;
+		static constexpr bool bIsParsableFromAnyLiteralType = 
+			TVariantConstructorTraits::bIsDefaultConstructible || 
+			TVariantConstructorTraits::bIsConstructibleWithSettings || 
+			TVariantConstructorTraits::bIsConstructibleWithArgs || 
+			TVariantConstructorTraits::bIsConstructibleWithSettingsAndArgs || 
+			TVariantConstructorTraits::bIsArrayDefaultConstructible || 
+			TVariantConstructorTraits::bIsArrayConstructibleWithSettings || 
+			TVariantConstructorTraits::bIsArrayConstructibleWithArgs || 
+			TVariantConstructorTraits::bIsArrayConstructibleWithSettingsAndArgs;
 
 		/** Determines if a constructor for the DataType exists which accepts 
 		 * an FOperatorSettings with the literals constructor arg type, and/or one that
@@ -566,6 +747,10 @@ namespace Metasound
 		{
 			switch (InLiteral.GetType())
 			{
+				case ELiteralType::None:
+
+					return TIsParsable<DataType>::Value;
+
 				case ELiteralType::Boolean:
 
 					return TIsParsable<DataType, bool>::Value;
@@ -586,17 +771,34 @@ namespace Metasound
 
 					return TIsParsable<DataType, Audio::IProxyDataPtr>::Value;
 
-				case ELiteralType::UObjectProxyArray:
+				case ELiteralType::NoneArray:
 
-					return TIsParsable<DataType, TArray<Audio::IProxyDataPtr>>::Value;
+					return TIsParsable<DataType, TArray<FLiteral::FNone>>::Value; 
 
-				case ELiteralType::None:
+				case ELiteralType::BooleanArray:
 
-					return TIsParsable<DataType>::Value;
+					return TIsParsable<DataType, TArray<bool>>::Value;
+
+				case ELiteralType::IntegerArray:
+
+					return TIsParsable<DataType, TArray<int32>>::Value;
+
+				case ELiteralType::FloatArray:
+
+					return TIsParsable<DataType, TArray<float>>::Value;
+
+				case ELiteralType::StringArray:
+
+					return TIsParsable<DataType, TArray<FString>>::Value;
+
+				case ELiteralType::UObjectProxyArray:
+
+					return TIsParsable<DataType, TArray<Audio::IProxyDataPtr>>::Value;
 
 				default:
 
 					checkNoEntry();
+					static_assert(static_cast<int32>(ELiteralType::Invalid) == 12, "Possible omitted ELiteralType value.");
 					return false;
 			}
 		}
@@ -772,7 +974,7 @@ namespace Metasound
 
 				case ELiteralType::Integer:
 
-					return CreateExplicitArgs<int>(InSettings, InLiteral.Value);
+					return CreateExplicitArgs<int32>(InSettings, InLiteral.Value);
 
 				case ELiteralType::Float:
 
@@ -786,13 +988,33 @@ namespace Metasound
 
 					return CreateExplicitArgs<Audio::IProxyDataPtr>(InSettings, InLiteral.Value);
 
+				case ELiteralType::NoneArray:
+
+					return CreateExplicitArgs<TArray<FLiteral::FNone>>(InSettings, InLiteral.Value);
+
+				case ELiteralType::BooleanArray:
+
+					return CreateExplicitArgs<TArray<bool>>(InSettings, InLiteral.Value);
+
+				case ELiteralType::IntegerArray:
+
+					return CreateExplicitArgs<TArray<int>>(InSettings, InLiteral.Value);
+
+				case ELiteralType::FloatArray:
+
+					return CreateExplicitArgs<TArray<float>>(InSettings, InLiteral.Value);
+
+				case ELiteralType::StringArray:
+
+					return CreateExplicitArgs<TArray<FString>>(InSettings, InLiteral.Value);
+
 				case ELiteralType::UObjectProxyArray:
 
 					return CreateExplicitArgs<TArray<Audio::IProxyDataPtr>>(InSettings, InLiteral.Value);
 
 				case ELiteralType::None:
 
-					return CreateExplicitArgs<void>(InSettings, InLiteral.Value);
+					return CreateExplicitArgs<FLiteral::FNone>(InSettings, InLiteral.Value);
 
 				case ELiteralType::Invalid:
 				default:
diff --git a/Engine/Plugins/Runtime/Metasound/Source/MetasoundGraphCore/Public/MetasoundLiteral.h b/Engine/Plugins/Runtime/Metasound/Source/MetasoundGraphCore/Public/MetasoundLiteral.h
--- a/Engine/Plugins/Runtime/Metasound/Source/MetasoundGraphCore/Public/MetasoundLiteral.h
+++ b/Engine/Plugins/Runtime/Metasound/Source/MetasoundGraphCore/Public/MetasoundLiteral.h
@@ -16,99 +16,21 @@
 {
 	enum class ELiteralType : uint8
 	{
+		None, // If the literal is None, TType(const FOperatorSettings&) or TType() will be invoked.
 		Boolean,
 		Integer,
 		Float,
 		String,
 		UObjectProxy,
+		NoneArray, // If this is set, TType(const FOperatorSettings&) or TType() will be invoked for each element in the array. 
+		BooleanArray,
+		IntegerArray,
+		FloatArray,
+		StringArray,
 		UObjectProxyArray,
-		None, // If this is set, we will invoke TType(const FOperatorSettings&) If that constructor exists, or the default constructor if not.
 		Invalid,
 	};
 
-	namespace MetasoundLiteralIntrinsics
-	{
-		// Default template for converting a template parameter to a literal argument type
-		template<typename... ArgTypes>
-		ELiteralType GetLiteralArgTypeFromDecayed()
-		{
-			return ELiteralType::Invalid;
-		}
 
-		template<>
-		FORCEINLINE ELiteralType GetLiteralArgTypeFromDecayed<bool>()
-		{
-			return ELiteralType::Boolean;
-		}
-
-		template<>
-		FORCEINLINE ELiteralType GetLiteralArgTypeFromDecayed<int32>()
-		{
-			return ELiteralType::Integer;
-		}
-
-		template<>
-		FORCEINLINE ELiteralType GetLiteralArgTypeFromDecayed<float>()
-		{
-			return ELiteralType::Float;
-		}
-
-		template<>
-		FORCEINLINE ELiteralType GetLiteralArgTypeFromDecayed<FString>()
-		{
-			return ELiteralType::String;
-		}
-
-		template<>
-		FORCEINLINE ELiteralType GetLiteralArgTypeFromDecayed<Audio::IProxyDataPtr>()
-		{
-			return ELiteralType::UObjectProxy;
-		}
-
-		template<>
-		FORCEINLINE ELiteralType GetLiteralArgTypeFromDecayed<TArray<Audio::IProxyDataPtr>>()
-		{
-			return ELiteralType::UObjectProxyArray;
-		}
-
-		template<>
-		FORCEINLINE ELiteralType GetLiteralArgTypeFromDecayed<>()
-		{
-			return ELiteralType::None;
-		}
-	}
-
-
-	/** Provides literal type information for a given type. 
-	 *
-	 * @tparam ArgType - A C++ type.
-	 */
-	template<typename... ArgTypes>
-	struct TLiteralTypeInfo
-	{
-		/** Returns the associated ELiteralType for the C++ type provided in the TLiteralTypeInfo<Type> */
-		static const ELiteralType GetLiteralArgTypeEnum()
-		{
-			// Use decayed version of template arg to remove references and cv qualifiers. 
-			return MetasoundLiteralIntrinsics::GetLiteralArgTypeFromDecayed<typename std::decay<ArgTypes...>::type>();
-		}
-	};
-
-	/** Provides literal type information for a given type. 
-	 *
-	 * @tparam ArgType - A C++ type.
-	 */
-	template<>
-	struct TLiteralTypeInfo<>
-	{
-		/** Returns the associated ELiteralType for the C++ type provided in the TLiteralTypeInfo<Type> */
-		static const ELiteralType GetLiteralArgTypeEnum()
-		{
-			// Use decayed version of template arg to remove references and cv qualifiers. 
-			return MetasoundLiteralIntrinsics::GetLiteralArgTypeFromDecayed<>();
-		}
-	};
-
-
 	/**
 	 * FLiteral 
@@ -137,55 +59,91 @@
 	 */
 	struct FLiteral
 	{
-		struct FInvalidLiteralState {};
-		struct FEmptyLiteralState {};
+		struct FInvalid {};
 
+		/* FNone is used in scenarios where an object is constructed with an FLiteral
+		 * where the expected object constructor is the default constructor, or a
+		 * constructor which accepts an FOperatorSettings.
+		 */
+		struct FNone {};
+
 		FLiteral(const FLiteral& Other) = delete;
 		FLiteral(FLiteral&& Other) = default;
 		FLiteral& operator=(FLiteral&& Other) = default;
 
-		using FVariantType = TVariant<FInvalidLiteralState, FEmptyLiteralState, bool, int, float, FString, Audio::IProxyDataPtr, TArray<Audio::IProxyDataPtr>>;
+		using FVariantType = TVariant<
+			FNone, bool, int32, float, FString, Audio::IProxyDataPtr, // Single value types
+			TArray<FNone>, TArray<bool>, TArray<int32>, TArray<float>, TArray<FString>, TArray<Audio::IProxyDataPtr>, // Array of values types
+			FInvalid
+		>;
 
 		FVariantType Value;
 
 #if METASOUND_DEBUG_LITERALS
 		private:
+
 		FString DebugString;
 
 		void InitDebugString()
 		{
+			// GetIndex returns the index of the current type set in the TVariant.
+			// The index refers to the order of the template parameters to TVariant.
 			switch (Value.GetIndex())
 			{
 				case 0:
-					DebugString = TEXT("INVALID");
+					DebugString = TEXT("NONE");
 				break;
 
 				case 1:
-					DebugString = TEXT("EMPTY");
-				break;
-
-				case 2:
 					DebugString = FString::Printf(TEXT("Bool: %s"), Value.Get<bool>() ? TEXT("true") : TEXT("false"));
 				break;
 
-				case 3:
+				case 2:
 					DebugString = FString::Printf(TEXT("Int32: %d"), Value.Get<int32>());
 				break;
 
-				case 4:
+				case 3:
 					DebugString = FString::Printf(TEXT("Float: %f"), Value.Get<float>());
 				break;
 
-				case 5:
+				case 4:
 					DebugString = FString::Printf(TEXT("String: %s"), *Value.Get<FString>());
 				break;
 
-				case 6:
+				case 5:
 					DebugString = FString::Printf(TEXT("Object: %f"), *Value.Get<Audio::IProxyDataPtr>()->GetProxyTypeName().ToString());
 				break;
 
+				case 6:
+					DebugString = TEXT("TArray<NONE>");
+				break;
+
 				case 7:
+					DebugString = FString::Printf(TEXT("TArray<Bool>"));
+				break;
+
+				case 8:
+					DebugString = FString::Printf(TEXT("TArray<int32>"));
+				break;
+
+				case 9:
+					DebugString = FString::Printf(TEXT("TArray<float>"));
+				break;
+
+				case 10:
+					DebugString = FString::Printf(TEXT("TArray<FString>"));
+				break;
+
+				case 11:
+					DebugString = FString::Printf(TEXT("TArray<UObject>"));
+				break;
+
+				case 12:
+					DebugString = TEXT("INVALID");
+				break;
+
 				default:
+					static_assert(TVariantSize<FVariantType>::Value == 13, "Possible missing FVariantType case coverage");
 					checkNoEntry();
 			}
 		}
@@ -197,7 +155,7 @@
 		// builds an invalid FLiteral.
 		static FLiteral CreateInvalid()
 		{
-			return FLiteral(FInvalidLiteralState());
+			return FLiteral(FInvalid());
 		}
 
 		bool IsValid() const
@@ -210,44 +168,61 @@
 			switch (Value.GetIndex())
 			{
 				case 0:
-					return ELiteralType::Invalid;
-
-				case 1:
 					return ELiteralType::None;
 
-				case 2:
+				case 1:
 					return ELiteralType::Boolean;
 
-				case 3:
+				case 2:
 					return ELiteralType::Integer;
 
-				case 4:
+				case 3:
 					return ELiteralType::Float;
 
-				case 5:
+				case 4:
 					return ELiteralType::String;
 
-				case 6:
+				case 5:
 					return ELiteralType::UObjectProxy;
 
+				case 6:
+					return ELiteralType::NoneArray;
+
 				case 7:
+					return ELiteralType::BooleanArray;
+
+				case 8:
+					return ELiteralType::IntegerArray;
+
+				case 9:
+					return ELiteralType::FloatArray;
+
+				case 10:
+					return ELiteralType::StringArray;
+
+				case 11:
 					return ELiteralType::UObjectProxyArray;
 
+				case 12:
 				default:
-					checkNoEntry();
+					static_assert(TVariantSize<FVariantType>::Value == 13, "Possible missing FVariantType case coverage");
 					return ELiteralType::Invalid;
 			}
 		}
 
 		FLiteral Clone() const
 		{
+			// TODO: If Clone() is supported after reworking vertices,
+			// wrap the Audio::IProxyData in a object that manages virtual
+			// copy construction.  Then replace this switch with simple
+			// TVariant copy constructor.
 			switch (GetType())
 			{
 			case ELiteralType::Invalid:
 				return CreateInvalid();
 
 			case ELiteralType::None:
-				return FLiteral(FEmptyLiteralState());
+				return FLiteral(FNone());
 
 			case ELiteralType::Boolean:
 				return FLiteral(Value.Get<bool>());
@@ -262,12 +237,39 @@
 				return FLiteral(Value.Get<FString>());
 
 			case ELiteralType::UObjectProxy:
-			{
-				TUniquePtr<Audio::IProxyDataPtr> THing;
-				return FLiteral(Value.Get<Audio::IProxyDataPtr>()->Clone());
-			}
+				if (const Audio::IProxyDataPtr& ProxyPtr = Value.Get<Audio::IProxyDataPtr>())
+				{
+					return FLiteral(ProxyPtr->Clone());
+				}
+				return CreateInvalid();
 
+			case ELiteralType::NoneArray:
+				return FLiteral(Value.Get<TArray<FNone>>());
+
+			case ELiteralType::BooleanArray:
+				return FLiteral(Value.Get<TArray<bool>>());
+
+			case ELiteralType::FloatArray:
+				return FLiteral(Value.Get<TArray<float>>());
+
+			case ELiteralType::IntegerArray:
+				return FLiteral(Value.Get<TArray<int32>>());
+
+			case ELiteralType::StringArray:
+				return FLiteral(Value.Get<TArray<FString>>());
+
 			case ELiteralType::UObjectProxyArray:
+				{
+					TArray<Audio::IProxyDataPtr> ProxyPtrArrayCopy;
+					for (const Audio::IProxyDataPtr& Ptr : Value.Get<TArray<Audio::IProxyDataPtr>>())
+					{
+						if (Ptr.IsValid())
+						{
+							ProxyPtrArrayCopy.Add(Ptr->Clone());
+						}
+					}
+					return FLiteral(MoveTemp(ProxyPtrArrayCopy));
+				}
 				return CreateInvalid();
 
 			default:
@@ -296,13 +298,128 @@
 			>
 		FLiteral(ArgTypes&&... Args)
 		{
-			Value.Set<FEmptyLiteralState>(FEmptyLiteralState());
+			Value.Set<FNone>(FNone());
 #if METASOUND_DEBUG_LITERALS
 			InitDebugString();
 #endif
 		}
+
 	};
 
+	namespace MetasoundLiteralIntrinsics
+	{
+		// Default template for converting a template parameter to a literal argument type
+		template<typename... ArgTypes>
+		ELiteralType GetLiteralArgTypeFromDecayed()
+		{
+			return ELiteralType::Invalid;
+		}
+
+		template<>
+		FORCEINLINE ELiteralType GetLiteralArgTypeFromDecayed<>()
+		{
+			return ELiteralType::None;
+		}
+
+		template<>
+		FORCEINLINE ELiteralType GetLiteralArgTypeFromDecayed<bool>()
+		{
+			return ELiteralType::Boolean;
+		}
+
+		template<>
+		FORCEINLINE ELiteralType GetLiteralArgTypeFromDecayed<int32>()
+		{
+			return ELiteralType::Integer;
+		}
+
+		template<>
+		FORCEINLINE ELiteralType GetLiteralArgTypeFromDecayed<float>()
+		{
+			return ELiteralType::Float;
+		}
+
+		template<>
+		FORCEINLINE ELiteralType GetLiteralArgTypeFromDecayed<FString>()
+		{
+			return ELiteralType::String;
+		}
+
+		template<>
+		FORCEINLINE ELiteralType GetLiteralArgTypeFromDecayed<Audio::IProxyDataPtr>()
+		{
+			return ELiteralType::UObjectProxy;
+		}
+
+		template<>
+		FORCEINLINE ELiteralType GetLiteralArgTypeFromDecayed<TArray<FLiteral::FNone>>()
+		{
+			return ELiteralType::NoneArray;
+		}
+
+		template<>
+		FORCEINLINE ELiteralType GetLiteralArgTypeFromDecayed<TArray<bool>>()
+		{
+			return ELiteralType::BooleanArray;
+		}
+
+		template<>
+		FORCEINLINE ELiteralType GetLiteralArgTypeFromDecayed<TArray<int32>>()
+		{
+			return ELiteralType::IntegerArray;
+		}
+
+		template<>
+		FORCEINLINE ELiteralType GetLiteralArgTypeFromDecayed<TArray<float>>()
+		{
+			return ELiteralType::FloatArray;
+		}
+
+		template<>
+		FORCEINLINE ELiteralType GetLiteralArgTypeFromDecayed<TArray<FString>>()
+		{
+			return ELiteralType::StringArray;
+		}
+
+		template<>
+		FORCEINLINE ELiteralType GetLiteralArgTypeFromDecayed<TArray<Audio::IProxyDataPtr>>()
+		{
+			return ELiteralType::UObjectProxyArray;
+		}
+	}
+
+
+	/** Provides literal type information for a given type. 
+	 *
+	 * @tparam ArgType - A C++ type.
+	 */
+	template<typename... ArgTypes>
+	struct TLiteralTypeInfo
+	{
+		/** Returns the associated ELiteralType for the C++ type provided in the TLiteralTypeInfo<Type> */
+		static const ELiteralType GetLiteralArgTypeEnum()
+		{
+			// Use decayed version of template arg to remove references and cv qualifiers. 
+			return MetasoundLiteralIntrinsics::GetLiteralArgTypeFromDecayed<typename std::decay<ArgTypes...>::type>();
+		}
+	};
+
+	/** Provides literal type information for a given type. 
+	 *
+	 * @tparam ArgType - A C++ type.
+	 */
+	template<>
+	struct TLiteralTypeInfo<>
+	{
+		/** Returns the associated ELiteralType for the C++ type provided in the TLiteralTypeInfo<Type> */
+		static const ELiteralType GetLiteralArgTypeEnum()
+		{
+			// Use decayed version of template arg to remove references and cv qualifiers. 
+			return MetasoundLiteralIntrinsics::GetLiteralArgTypeFromDecayed<>();
+		}
+	};
+
+
 	// This struct is passed to FInputNodeConstructorCallback and FOutputNodeConstructorCallback.
 	struct FInputNodeConstructorParams
 	{
