CIS fix missing file in submit

diff --git a/Engine/Plugins/Runtime/Metasound/Source/MetasoundFrontend/Private/MetasoundFrontendStandardController.cpp b/Engine/Plugins/Runtime/Metasound/Source/MetasoundFrontend/Private/MetasoundFrontendStandardController.cpp
--- a/Engine/Plugins/Runtime/Metasound/Source/MetasoundFrontend/Private/MetasoundFrontendStandardController.cpp
+++ b/Engine/Plugins/Runtime/Metasound/Source/MetasoundFrontend/Private/MetasoundFrontendStandardController.cpp
@@ -2355,7 +2355,7 @@ namespace Metasound
 				case EMetasoundFrontendLiteralType::None:
 				default:
 				{
-					static_assert(static_cast<int32>(ELiteralType::Invalid) == 7, "Possible missing ELiteralType case coverage");
+					static_assert(static_cast<int32>(EMetasoundFrontendLiteralType::Invalid) == 7, "Possible missing ELiteralType case coverage");
 					return false;
 				}
 				break;
@@ -2391,19 +2391,18 @@ namespace Metasound
 				}
 				break;
 
-				// TODO: Remove.
-				case ELiteralType::UObjectProxyArray:
-				{
-					return false;
-				}
-				break;
-
+				case ELiteralType::None:
 				case ELiteralType::UObjectProxy:	// UObject literals do not currently support default data type references
+				case ELiteralType::NoneArray:
+				case ELiteralType::BooleanArray:
+				case ELiteralType::IntegerArray:
+				case ELiteralType::FloatArray:
+				case ELiteralType::StringArray:
+				case ELiteralType::UObjectProxyArray:
 				case ELiteralType::Invalid:
-				case ELiteralType::None:
 				default:
 				{
-					static_assert(static_cast<int32>(ELiteralType::Invalid) == 7, "Possible missing ELiteralType case coverage");
+					static_assert(static_cast<int32>(ELiteralType::Invalid) == 12, "Possible missing ELiteralType case coverage");
 					return false;
 				}
 				break;
@@ -2626,14 +2625,14 @@ namespace Metasound
 		{
 			if (GraphClassPtr.IsValid())
 			{
-				auto IsConnectionFromThisNode = [&](const FMetasoundFrontendEdge& ConDesc) { return ConDesc.FromNodeID == InDesc.ID; };
+				auto IsEdgeForThisNode = [&](const FMetasoundFrontendEdge& ConDesc) { return (ConDesc.FromNodeID == InDesc.ID) || (ConDesc.ToNodeID == InDesc.ID); };
 
 				// Remove any reference connections
-				GraphClassPtr->Graph.Edges.RemoveAll(IsConnectionFromThisNode);
+				int32 NumRemoved = GraphClassPtr->Graph.Edges.RemoveAll(IsEdgeForThisNode);
 
 				auto IsNodeWithID = [&](const FMetasoundFrontendNode& Desc) { return InDesc.ID == Desc.ID; };
 
-				int32 NumRemoved = GraphClassPtr->Graph.Nodes.RemoveAll(IsNodeWithID);
+				NumRemoved += GraphClassPtr->Graph.Nodes.RemoveAll(IsNodeWithID);
 
 				OwningDocument->RemoveUnreferencedDependencies();
 
