Metasound Graph Editor Updates
- Add context search by pin type
- Add new node autoconnection to drag pin
- Fix bug with newly added input nodes snapping to origin

diff --git a/Engine/Plugins/Runtime/Metasound/Source/MetasoundEditor/Private/MetasoundEditorGraphBuilder.cpp b/Engine/Plugins/Runtime/Metasound/Source/MetasoundEditor/Private/MetasoundEditorGraphBuilder.cpp
--- a/Engine/Plugins/Runtime/Metasound/Source/MetasoundEditor/Private/MetasoundEditorGraphBuilder.cpp
+++ b/Engine/Plugins/Runtime/Metasound/Source/MetasoundEditor/Private/MetasoundEditorGraphBuilder.cpp
@@ -39,6 +39,9 @@ namespace Metasound
 
 		const FName FGraphBuilder::PinSubCategoryTime = "time";
 
+		const FText FGraphBuilder::ConvertMenuName = LOCTEXT("MetasoundConversionsMenu", "Conversions");
+		const FText FGraphBuilder::FunctionMenuName = LOCTEXT("MetasoundFunctionsMenu", "Functions");
+
 		namespace GraphBuilderPrivate
 		{
 			void DeleteNode(UObject& InMetasound, Frontend::FNodeHandle InNodeHandle)
@@ -160,9 +163,65 @@ namespace Metasound
 			return NewOutputName;
 		}
 
+		Frontend::FInputHandle FGraphBuilder::GetInputHandleFromPin(const UEdGraphPin* InPin)
+		{
+			using namespace Frontend;
+
+			if (InPin && ensure(InPin->Direction == EGPD_Input))
+			{
+				if (UMetasoundEditorGraphNode* EdNode = CastChecked<UMetasoundEditorGraphNode>(InPin->GetOwningNode()))
+				{
+					FNodeHandle NodeHandle = EdNode->GetNodeHandle();
+					if (NodeHandle->IsValid())
+					{
+						TArray<FInputHandle> Inputs = NodeHandle->GetInputsWithVertexName(InPin->GetName());
+						if (ensure(Inputs.Num() == 1))
+						{
+							return Inputs[0];
+						}
+					}
+				}
+			}
+
+			return IInputController::GetInvalidHandle();
+		}
+
+		Frontend::FConstInputHandle FGraphBuilder::GetConstInputHandleFromPin(const UEdGraphPin* InPin)
+		{
+			return GetInputHandleFromPin(InPin);
+		}
+
+		Frontend::FOutputHandle FGraphBuilder::GetOutputHandleFromPin(const UEdGraphPin* InPin)
+		{
+			using namespace Frontend;
+
+			if (InPin && ensure(InPin->Direction == EGPD_Output))
+			{
+				if (UMetasoundEditorGraphNode* EdNode = CastChecked<UMetasoundEditorGraphNode>(InPin->GetOwningNode()))
+				{
+					FNodeHandle NodeHandle = EdNode->GetNodeHandle();
+					if (NodeHandle->IsValid())
+					{
+						TArray<FOutputHandle> Outputs = NodeHandle->GetOutputsWithVertexName(InPin->GetName());
+						if (ensure(Outputs.Num() == 1))
+						{
+							return Outputs[0];
+						}
+					}
+				}
+			}
+
+			return IOutputController::GetInvalidHandle();
+		}
+
+		Frontend::FConstOutputHandle FGraphBuilder::GetConstOutputHandleFromPin(const UEdGraphPin* InPin)
+		{
+			return GetOutputHandleFromPin(InPin);
+		}
+
 		UEdGraphNode* FGraphBuilder::AddInput(UObject& InMetasound, const FString& InName, const FName InTypeName, const FMetasoundFrontendNodeStyle& InNodeStyle, const FText& InToolTip, bool bInSelectNewNode)
 		{
-			Frontend::FNodeHandle NodeHandle = AddInputNodeHandle(InMetasound, InName, InTypeName, InNodeStyle, &InToolTip);
+			Frontend::FNodeHandle NodeHandle = AddInputNodeHandle(InMetasound, InName, InTypeName, InNodeStyle, InToolTip);
 			return AddNode(InMetasound, NodeHandle, bInSelectNewNode);
 		}
 
@@ -172,9 +231,9 @@ namespace Metasound
 
 			const FScopedTransaction Transaction(LOCTEXT("SetMetasoundGraphNode", "Add/Update Metasound Literal Input"));
 
-			const FString& InInputName = InInputPin.GetName();
+			const FString InInputName = InInputPin.GetName();
 			TArray<FInputHandle> InputHandles = InNodeHandle->GetInputsWithVertexName(InInputName);
-			if (!ensureAlways(InputHandles.Num() == 1))
+			if (!ensure(InputHandles.Num() == 1))
 			{
 				return;
 			}
@@ -257,7 +316,7 @@ namespace Metasound
 				case ELiteralType::Invalid:
 				default:
 				{
-					ensureAlwaysMsgf(false, TEXT("Failed to set input node default: Literal type not supported"));
+					ensureMsgf(false, TEXT("Failed to set input node default: Literal type not supported"));
 					return;
 				}
 				break;
@@ -270,7 +329,7 @@ namespace Metasound
 				Style.Display.Visibility = EMetasoundFrontendNodeStyleDisplayVisibility::Hidden;
 
 				const FString NewInputName = GenerateUniqueInputName(InMetasound, "LiteralInput");
-				Frontend::FNodeHandle NewInputNode = AddInputNodeHandle(InMetasound, NewInputName, TypeName, Style, nullptr /* ToolTip */, &PinDataTypeDefaultLiteral);
+				Frontend::FNodeHandle NewInputNode = AddInputNodeHandle(InMetasound, NewInputName, TypeName, Style, FText::GetEmpty(), &PinDataTypeDefaultLiteral);
 				NewInputNode->SetNodeStyle(Style);
 
 				Frontend::FGraphHandle GraphHandle = InNodeHandle->GetOwningGraph();
@@ -278,7 +337,7 @@ namespace Metasound
 
 				const TArray<FGuid> PointIDs = GraphHandle->GetDefaultIDsForInputVertex(NewInputNode->GetNodeName());
 				// TODO: Support multiple literal point ids for arrays
-				if (ensureAlways(PointIDs.Num() == 1))
+				if (ensure(PointIDs.Num() == 1))
 				{
 					InNodeHandle->GetClassDescription();
 					const TArray<FMetasoundFrontendVertexLiteral>& Defaults = InputHandle->GetDefaults();
@@ -291,12 +350,12 @@ namespace Metasound
 				}
 
 				TArray<Metasound::Frontend::FOutputHandle> OutputHandles = NewInputNode->GetOutputs();
-				if (ensureAlways(OutputHandles.Num() == 1))
+				if (ensure(OutputHandles.Num() == 1))
 				{
 					OutputHandle = OutputHandles[0];
 				}
 
-				ensureAlways(InputHandle->Connect(*OutputHandle));
+				ensure(InputHandle->Connect(*OutputHandle));
 
 				InMetasound.PostEditChange();
 				InMetasound.MarkPackageDirty();
@@ -311,7 +370,7 @@ namespace Metasound
 					const TArray<FGuid> PointIDs = GraphHandle->GetDefaultIDsForInputVertex(InputNode->GetNodeName());
 
 					// TODO: Support multiple literal point ids for arrays
-					if (ensureAlways(PointIDs.Num() == 1))
+					if (ensure(PointIDs.Num() == 1))
 					{
 						GraphHandle->SetDefaultInputToFrontendLiteral(InputNode->GetNodeName(), PointIDs[0], TypeName, PinDataTypeDefaultLiteral);
 					}
@@ -322,7 +381,7 @@ namespace Metasound
 			}
 		}
 
-		Frontend::FNodeHandle FGraphBuilder::AddInputNodeHandle(UObject& InMetasound, const FString& InName, const FName InTypeName, const FMetasoundFrontendNodeStyle& InNodeStyle, const FText* InToolTip, const FMetasoundFrontendLiteral* InDefaultValue)
+		Frontend::FNodeHandle FGraphBuilder::AddInputNodeHandle(UObject& InMetasound, const FString& InName, const FName InTypeName, const FMetasoundFrontendNodeStyle& InNodeStyle, const FText& InToolTip, const FMetasoundFrontendLiteral* InDefaultValue)
 		{
 			FMetasoundAssetBase* MetasoundAsset = IMetasoundUObjectRegistry::Get().GetObjectAsAssetBase(&InMetasound);
 			check(MetasoundAsset);
@@ -334,12 +393,9 @@ namespace Metasound
 			Description.Name = InName;
 			Description.TypeName = InTypeName;
 
-			if (InToolTip)
-			{
-				Description.Metadata.Description = *InToolTip;
-			}
+			Description.Metadata.Description = InToolTip;
 
-			FGuid PointID = GraphHandle->GetNewPointID();
+			const FGuid PointID = GraphHandle->GetNewPointID();
 
 			Description.PointIDs.Add(PointID);
 			FMetasoundFrontendVertexLiteral DefaultVertex;
@@ -347,7 +403,9 @@ namespace Metasound
 			Description.Defaults.Add(DefaultVertex);
 
 			Frontend::FNodeHandle NodeHandle = GraphHandle->AddInputVertex(Description);
-			if (ensureAlways(NodeHandle->IsValid()))
+			NodeHandle->SetNodeStyle(InNodeStyle);
+
+			if (ensure(NodeHandle->IsValid()))
 			{
 				GraphHandle->SetInputDisplayName(InName, FText::FromString(InName));
 
@@ -394,27 +452,33 @@ namespace Metasound
 			return NodeHandle;
 		}
 
-		bool FGraphBuilder::ConnectNodes(UEdGraphPin& InInputPin, UEdGraphPin& InOutputPin)
+		bool FGraphBuilder::ConnectNodes(UEdGraphPin& InInputPin, UEdGraphPin& InOutputPin, bool bConnectEdPins)
 		{
 			using namespace Metasound::Frontend;
 
-			UMetasoundEditorGraphNode* InputGraphNode = CastChecked<UMetasoundEditorGraphNode>(InInputPin.GetOwningNode());
-			FNodeHandle InputNodeHandle = InputGraphNode->GetNodeHandle();
-			TArray<FInputHandle> InputHandles = InputNodeHandle->GetInputsWithVertexName(InInputPin.GetName());
-
-			UMetasoundEditorGraphNode* OutputGraphNode = CastChecked<UMetasoundEditorGraphNode>(InOutputPin.GetOwningNode());
-			FNodeHandle OutputNodeHandle = OutputGraphNode->GetNodeHandle();
-			TArray<FOutputHandle> OutputHandles = OutputNodeHandle->GetOutputsWithVertexName(InOutputPin.GetName());
+			// When true, will recursively call back into this function
+			// from the schema if the editor pins are successfully connected
+			if (bConnectEdPins)
+			{
+				const UEdGraphSchema* Schema = InInputPin.GetSchema();
+				if (ensure(Schema))
+				{
+					return Schema->TryCreateConnection(&InInputPin, &InOutputPin);
+				}
+				else
+				{
+					return false;
+				}
+			}
 
-			if (!ensure(InputHandles.Num() == 1 && OutputHandles.Num() == 1))
+			FInputHandle InputHandle = GetInputHandleFromPin(&InInputPin);
+			FOutputHandle OutputHandle = GetOutputHandleFromPin(&InOutputPin);
+			if (!InputHandle->IsValid() || !OutputHandle->IsValid())
 			{
 				InInputPin.BreakLinkTo(&InOutputPin);
 				return false;
 			}
 
-			FInputHandle InputHandle = InputHandles[0];
-			FOutputHandle OutputHandle = OutputHandles[0];
-
 			FOutputHandle ExistingOutput = InputHandle->GetCurrentlyConnectedOutput();
 			if (ExistingOutput->IsValid())
 			{
@@ -422,7 +486,10 @@ namespace Metasound
 				const FMetasoundFrontendNodeStyle& NodeStyle = NodeHandle->GetNodeStyle();
 				if (NodeHandle->IsValid() && NodeStyle.Display.Visibility == EMetasoundFrontendNodeStyleDisplayVisibility::Hidden)
 				{
-					GraphBuilderPrivate::DeleteNode(InputGraphNode->GetMetasoundChecked(), NodeHandle);
+					if (UMetasoundEditorGraph* Graph = Cast<UMetasoundEditorGraph>(InInputPin.GetOwningNode()->GetGraph()))
+					{
+						GraphBuilderPrivate::DeleteNode(Graph->GetMetasoundChecked(), NodeHandle);
+					}
 				}
 			}
 
@@ -725,7 +792,7 @@ namespace Metasound
 			FEdGraphPinType PinType("MetasoundNode", NAME_None, nullptr, EPinContainerType::None, false, FEdGraphTerminalType());
 			UEdGraphPin* NewPin = InEditorNode.CreatePin(EGPD_Input, PinType, FName(*InInputHandle->GetName()));
 
-			if (ensureAlways(NewPin))
+			if (ensure(NewPin))
 			{
 				NewPin->PinToolTip = InInputHandle->GetTooltip().ToString();
 				NewPin->PinType = EditorModule.FindDataType(InInputHandle->GetDataType()).PinType;
@@ -745,7 +812,7 @@ namespace Metasound
 			FEdGraphPinType PinType("MetasoundNode", NAME_None, nullptr, EPinContainerType::None, false, FEdGraphTerminalType());
 			UEdGraphPin* NewPin = InEditorNode.CreatePin(EGPD_Output, PinType, FName(*InOutputHandle->GetName()));
 
-			if (ensureAlways(NewPin))
+			if (ensure(NewPin))
 			{
 				NewPin->PinToolTip = InOutputHandle->GetTooltip().ToString();
 				NewPin->PinType = EditorModule.FindDataType(InOutputHandle->GetDataType()).PinType;
diff --git a/Engine/Plugins/Runtime/Metasound/Source/MetasoundEditor/Private/MetasoundEditorGraphBuilder.h b/Engine/Plugins/Runtime/Metasound/Source/MetasoundEditor/Private/MetasoundEditorGraphBuilder.h
--- a/Engine/Plugins/Runtime/Metasound/Source/MetasoundEditor/Private/MetasoundEditorGraphBuilder.h
+++ b/Engine/Plugins/Runtime/Metasound/Source/MetasoundEditor/Private/MetasoundEditorGraphBuilder.h
@@ -42,6 +42,9 @@ namespace Metasound
 			// Custom pin-related styles for non-literal types (ex. wire color, pin heads, etc.)
 			static const FName PinSubCategoryTime; // Time type
 
+			static const FText ConvertMenuName;
+			static const FText FunctionMenuName;
+
 			// Adds a node to the editor graph that corresponds to the provided node handle.
 			static UEdGraphNode* AddNode(UObject& InMetasound, Frontend::FNodeHandle& InNodeHandle, bool bInSelectNewNode = true);
 
@@ -52,7 +55,8 @@ namespace Metasound
 			static Frontend::FNodeHandle AddNodeHandle(UObject& InMetasound, const Frontend::FNodeClassInfo& InClassInfo, const FMetasoundFrontendNodeStyle& InNodeStyle);
 
 			// Attempts to connect graph nodes together.  Returns true if succeeded, breaks pin link and returns false if failed.
-			static bool ConnectNodes(UEdGraphPin& InInputPin, UEdGraphPin& InOutputPin);
+			// If bConnectEdPins is set, will attempt to connect the Editor Graph representation of the pins.
+			static bool ConnectNodes(UEdGraphPin& InInputPin, UEdGraphPin& InOutputPin, bool bConnectEdPins);
 
 			static FString GenerateUniqueInputName(const UObject& InMetasound, const FName InBaseName);
 
@@ -62,6 +66,16 @@ namespace Metasound
 
 			static TArray<FString> GetDataTypeNameCategories(const FName& InDataTypeName);
 
+			// Get the input handle from an input pin.  Ensures pin is an input pin.
+			// TODO: use IDs to connect rather than names. Likely need an UMetasoundEditorGraphPin
+			static Frontend::FInputHandle GetInputHandleFromPin(const UEdGraphPin* InPin);
+			static Frontend::FConstInputHandle GetConstInputHandleFromPin(const UEdGraphPin* InPin);
+
+			// Get the output handle from an output pin.  Ensures pin is an output pin.
+			// TODO: use IDs to connect rather than names. Likely need an UMetasoundEditorGraphPin
+			static Frontend::FOutputHandle GetOutputHandleFromPin(const UEdGraphPin* InPin);
+			static Frontend::FConstOutputHandle GetConstOutputHandleFromPin(const UEdGraphPin* InPin);
+
 			static UEdGraphNode* AddInput(UObject& InMetasound, const FString& InName, const FName InTypeName, const FMetasoundFrontendNodeStyle& InNodeStyle, const FText& InToolTip, bool bInSelectNewNode = true);
 
 			// Adds or updates a hidden input that is set to the default literal value.  If the literal input is added, value defaults to
@@ -70,7 +84,7 @@ namespace Metasound
 			static void AddOrUpdateLiteralInput(UObject& InMetasound, Frontend::FNodeHandle InNodeHandle, UEdGraphPin& InInputPin);
 
 			// Adds an input node handle with the given class info
-			static Frontend::FNodeHandle AddInputNodeHandle(UObject& InMetasound, const FString& InName, const FName InTypeName, const FMetasoundFrontendNodeStyle& InNodeStyle, const FText* InToolTip = nullptr, const FMetasoundFrontendLiteral* InDefaultValue = nullptr);
+			static Frontend::FNodeHandle AddInputNodeHandle(UObject& InMetasound, const FString& InName, const FName InTypeName, const FMetasoundFrontendNodeStyle& InNodeStyle, const FText& InToolTip, const FMetasoundFrontendLiteral* InDefaultValue = nullptr);
 
 			static UEdGraphNode* AddOutput(UObject& InMetasound, const FString& InName, const FName InTypeName, const FMetasoundFrontendNodeStyle& InNodeStyle, const FText& InToolTip, bool bInSelectNewNode = true);
 
diff --git a/Engine/Plugins/Runtime/Metasound/Source/MetasoundEditor/Private/MetasoundEditorGraphSchema.cpp b/Engine/Plugins/Runtime/Metasound/Source/MetasoundEditor/Private/MetasoundEditorGraphSchema.cpp
--- a/Engine/Plugins/Runtime/Metasound/Source/MetasoundEditor/Private/MetasoundEditorGraphSchema.cpp
+++ b/Engine/Plugins/Runtime/Metasound/Source/MetasoundEditor/Private/MetasoundEditorGraphSchema.cpp
@@ -35,6 +35,56 @@ namespace Metasound
 {
 	namespace Editor
 	{
+		bool TryConnectNewNodeToPin(UEdGraphNode& NewGraphNode, UEdGraphPin* FromPin)
+		{
+			using namespace Metasound::Frontend;
+
+			if (!FromPin)
+			{
+				return false;
+			}
+
+			if (FromPin->Direction == EGPD_Input)
+			{
+				FInputHandle InputHandle = FGraphBuilder::GetInputHandleFromPin(FromPin);
+				for (UEdGraphPin* Pin : NewGraphNode.Pins)
+				{
+					if (Pin->Direction == EGPD_Output)
+					{
+						FOutputHandle OutputHandle = FGraphBuilder::GetOutputHandleFromPin(Pin);
+						if (OutputHandle->IsValid() && OutputHandle->GetDataType() == InputHandle->GetDataType())
+						{
+							if (ensure(FGraphBuilder::ConnectNodes(*FromPin, *Pin, true /* bConnectEdPins */)))
+							{
+								return true;
+							}
+						}
+					}
+				}
+			}
+
+			if (FromPin->Direction == EGPD_Output)
+			{
+				FOutputHandle OutputHandle = FGraphBuilder::GetOutputHandleFromPin(FromPin);
+				for (UEdGraphPin* Pin : NewGraphNode.Pins)
+				{
+					if (Pin->Direction == EGPD_Input)
+					{
+						FInputHandle InputHandle = FGraphBuilder::GetInputHandleFromPin(Pin);
+						if (InputHandle->IsValid() && InputHandle->GetDataType() == OutputHandle->GetDataType())
+						{
+							if (ensure(FGraphBuilder::ConnectNodes(*Pin, *FromPin, true /* bConnectEdPins */)))
+							{
+								return true;
+							}
+						}
+					}
+				}
+			}
+
+			return false;
+		}
+
 		TSharedPtr<FEditor> GetEditorForGraph(const UObject& Metasound)
 		{
 			TSharedPtr<IToolkit> FoundAssetEditor = FToolkitManager::Get().FindEditorForAsset(CastChecked<const UObject>(&Metasound));
@@ -104,6 +154,40 @@ namespace Metasound
 			return Settings->DefaultPinTypeColor;
 		}
 
+		template <typename TAction>
+		void GetDataTypeActions(FGraphActionMenuBuilder& ActionMenuBuilder, FDataTypeFilterFunction InFilter, const FText& InCategory, const FText& InTooltipFormat, bool bShowSelectedActions)
+		{
+			using namespace Metasound::Editor;
+			using namespace Metasound::Frontend;
+
+			static const FText MenuJoinFormat = LOCTEXT("MetasoundActionsFormatSubCategory", "{0}|{1}");
+
+			IMetasoundEditorModule& EditorModule = FModuleManager::GetModuleChecked<IMetasoundEditorModule>("MetasoundEditor");
+			EditorModule.IterateDataTypes([&](const FEditorDataType& DataType)
+			{
+				if (InFilter && !InFilter(DataType))
+				{
+					return;
+				}
+
+				const FName DataTypeName = DataType.RegistryInfo.DataTypeName;
+				const FText DataTypeDisplayName = FText::FromString(FGraphBuilder::GetDataTypeDisplayName(DataTypeName));
+				const FText DataTypeTextName = FText::FromName(DataTypeName);
+				const TArray<FString> Categories = FGraphBuilder::GetDataTypeNameCategories(DataTypeName);
+				const FText CategoriesText = FText::FromString(FString::Join(Categories, TEXT("|")));
+				TSharedPtr<TAction> NewNodeAction = MakeShared<TAction>
+				(
+					FText::Format(MenuJoinFormat, InCategory, CategoriesText),
+					DataTypeDisplayName,
+					DataTypeName,
+					FText::Format(InTooltipFormat, DataTypeTextName),
+					0
+				);
+
+				ActionMenuBuilder.AddAction(NewNodeAction);
+			});
+		}
+
 		FConnectionDrawingPolicy* FGraphConnectionDrawingPolicyFactory::CreateConnectionPolicy(
 			const UEdGraphSchema* InSchema,
 			int32 InBackLayerID,
@@ -188,10 +272,13 @@ UEdGraphNode* FMetasoundGraphSchemaAction_NewNode::PerformAction(UEdGraph* Paren
 
 	FMetasoundFrontendNodeStyle Style;
 	Style.Display.Location = Location;
+	if (UEdGraphNode* NewGraphNode = FGraphBuilder::AddNode(ParentMetasound, NodeClassInfo, Style))
+	{
+		TryConnectNewNodeToPin(*NewGraphNode, FromPin);
+		return NewGraphNode;
+	}
 
-	UEdGraphNode* NewGraphNode = FGraphBuilder::AddNode(ParentMetasound, NodeClassInfo, Style);
-
-	return NewGraphNode;
+	return nullptr;
 }
 
 FMetasoundGraphSchemaAction_NewInput::FMetasoundGraphSchemaAction_NewInput(FText InNodeCategory, FText InDisplayName, FName InTypeName, FText InToolTip, const int32 InGrouping)
@@ -216,8 +303,13 @@ UEdGraphNode* FMetasoundGraphSchemaAction_NewInput::PerformAction(UEdGraph* Pare
 	FMetasoundFrontendNodeStyle Style;
 	Style.Display.Location = Location;
 
-	UEdGraphNode* NewGraphNode = FGraphBuilder::AddInput(ParentMetasound, NewNodeName, NodeTypeName, Style, FText::GetEmpty());
-	return NewGraphNode;
+	if (UEdGraphNode* NewGraphNode = FGraphBuilder::AddInput(ParentMetasound, NewNodeName, NodeTypeName, Style, FText::GetEmpty()))
+	{
+		TryConnectNewNodeToPin(*NewGraphNode, FromPin);
+		return NewGraphNode;
+	}
+
+	return nullptr;
 }
 
 FMetasoundGraphSchemaAction_NewOutput::FMetasoundGraphSchemaAction_NewOutput(FText InNodeCategory, FText InDisplayName, FName InTypeName, FText InToolTip, const int32 InGrouping)
@@ -242,13 +334,13 @@ UEdGraphNode* FMetasoundGraphSchemaAction_NewOutput::PerformAction(UEdGraph* Par
 	FMetasoundFrontendNodeStyle Style;
 	Style.Display.Location = Location;
 
-	UEdGraphNode* NewGraphNode = FGraphBuilder::AddOutput(ParentMetasound, NewNodeName, NodeTypeName, Style, FText::GetEmpty());
-	return NewGraphNode;
-}
+	if (UEdGraphNode* NewGraphNode = FGraphBuilder::AddOutput(ParentMetasound, NewNodeName, NodeTypeName, Style, FText::GetEmpty()))
+	{
+		TryConnectNewNodeToPin(*NewGraphNode, FromPin);
+		return NewGraphNode;
+	}
 
-void FMetasoundGraphSchemaAction_NewNode::ConnectToSelectedNodes(UMetasoundEditorGraphNode* NewGraphNode, UEdGraph* ParentGraph) const
-{
-	// TODO: Implement
+	return nullptr;
 }
 
 UEdGraphNode* FMetasoundGraphSchemaAction_NewFromSelected::PerformAction(UEdGraph* ParentGraph, UEdGraphPin* FromPin, const FVector2D Location, bool bSelectNewNode/* = true*/)
@@ -303,26 +395,68 @@ bool UMetasoundEditorGraphSchema::ConnectionCausesLoop(const UEdGraphPin* InputP
 
 void UMetasoundEditorGraphSchema::GetPaletteActions(FGraphActionMenuBuilder& ActionMenuBuilder) const
 {
-	GetAllMetasoundActions(ActionMenuBuilder, false /* bShowSelectedActions */);
 	GetCommentAction(ActionMenuBuilder);
+	GetFunctionActions(ActionMenuBuilder);
+	GetDataTypeInputNodeActions(ActionMenuBuilder);
+	GetDataTypeOutputNodeActions(ActionMenuBuilder);
 }
 
 void UMetasoundEditorGraphSchema::GetGraphContextActions(FGraphContextMenuBuilder& ContextMenuBuilder) const
 {
+	using namespace Metasound;
 	using namespace Metasound::Editor;
 
-	GetAllMetasoundActions(ContextMenuBuilder, true /* bShowSelectedActions */);
-	GetCommentAction(ContextMenuBuilder, ContextMenuBuilder.CurrentGraph);
+	FActionClassFilters ClassFilters;
+	if (ContextMenuBuilder.FromPin)
+	{
+		if (ContextMenuBuilder.FromPin->Direction == EGPD_Input)
+		{
+			Frontend::FConstInputHandle InputHandle = Editor::FGraphBuilder::GetConstInputHandleFromPin(ContextMenuBuilder.FromPin);
+			ClassFilters.OutputFilterFunction = [InputHandle](const FMetasoundFrontendClassOutput& InOutput)
+			{
+				return InOutput.TypeName == InputHandle->GetDataType();
+			};
+
+			// Show only input nodes as output nodes can only connected if FromPin is input
+			GetDataTypeInputNodeActions(ContextMenuBuilder, [InputHandle](const Editor::FEditorDataType& InDataType)
+			{
+				return InDataType.RegistryInfo.DataTypeName == InputHandle->GetDataType();
+			});
+		}
 
-	if (!ContextMenuBuilder.FromPin)
+		if (ContextMenuBuilder.FromPin->Direction == EGPD_Output)
+		{
+			Frontend::FConstOutputHandle OutputHandle = Editor::FGraphBuilder::GetConstOutputHandleFromPin(ContextMenuBuilder.FromPin);
+			ClassFilters.InputFilterFunction = [OutputHandle](const FMetasoundFrontendClassInput& InInput)
+			{
+				return InInput.TypeName == OutputHandle->GetDataType();
+			};
+
+			FDataTypeFilterFunction DataTypeFilter = [OutputHandle](const Editor::FEditorDataType& InDataType)
+			{
+				return InDataType.RegistryInfo.DataTypeName == OutputHandle->GetDataType();
+			};
+
+			// Show only output nodes as input nodes can only connected if FromPin is output
+			GetDataTypeOutputNodeActions(ContextMenuBuilder, DataTypeFilter);
+		}
+	}
+	else
 	{
-		TSharedPtr<FEditor> MetasoundEditor = GetEditorForGraph(*ContextMenuBuilder.CurrentGraph);
+		TSharedPtr<Editor::FEditor> MetasoundEditor = Editor::GetEditorForGraph(*ContextMenuBuilder.CurrentGraph);
 		if (MetasoundEditor.IsValid() && MetasoundEditor->CanPasteNodes())
 		{
 			TSharedPtr<FMetasoundGraphSchemaAction_Paste> NewAction = MakeShared<FMetasoundGraphSchemaAction_Paste>(FText::GetEmpty(), LOCTEXT("PasteHereAction", "Paste here"), FText::GetEmpty(), 0);
 			ContextMenuBuilder.AddAction(NewAction);
 		}
+
+		GetCommentAction(ContextMenuBuilder, ContextMenuBuilder.CurrentGraph);
+		GetDataTypeInputNodeActions(ContextMenuBuilder);
+		GetDataTypeOutputNodeActions(ContextMenuBuilder);
 	}
+
+	GetFunctionActions(ContextMenuBuilder, ClassFilters);
+	GetConversionActions(ContextMenuBuilder, ClassFilters);
 }
 
 void UMetasoundEditorGraphSchema::GetContextMenuActions(class UToolMenu* Menu, class UGraphNodeContextMenuContext* Context) const
@@ -363,6 +497,8 @@ void UMetasoundEditorGraphSchema::CreateDefaultNodesForGraph(UEdGraph& Graph) co
 
 const FPinConnectionResponse UMetasoundEditorGraphSchema::CanCreateConnection(const UEdGraphPin* PinA, const UEdGraphPin* PinB) const
 {
+	using namespace Metasound;
+
 	// Make sure the pins are not on the same node
 	if (PinA->GetOwningNode() == PinB->GetOwningNode())
 	{
@@ -388,25 +524,14 @@ const FPinConnectionResponse UMetasoundEditorGraphSchema::CanCreateConnection(co
 		return FPinConnectionResponse(CONNECT_RESPONSE_DISALLOW, LOCTEXT("ConnectionTypeIncorrect", "Connection pin types do not match"));
 	}
 
-	// TODO: use IDs to connect rather than names. Likely needs a UMetasoundEditorGraphPin
-	//int32 InputPointID = InputPin->GetMetasoundPointID();
-	//int32 InputPointID = OutputPin->GetMetasoundPointID();
-
-	UMetasoundEditorGraphNode* InputGraphNode = CastChecked<UMetasoundEditorGraphNode>(InputPin->GetOwningNode());
-	Metasound::Frontend::FNodeHandle InputNodeHandle = InputGraphNode->GetNodeHandle();
-	TArray<Metasound::Frontend::FInputHandle> InputHandles = InputNodeHandle->GetInputsWithVertexName(InputPin->GetName());
-
-	UMetasoundEditorGraphNode* OutputGraphNode = CastChecked<UMetasoundEditorGraphNode>(OutputPin->GetOwningNode());
-	Metasound::Frontend::FNodeHandle OutputNodeHandle = OutputGraphNode->GetNodeHandle();
-	TArray<Metasound::Frontend::FOutputHandle> OutputHandles = OutputNodeHandle->GetOutputsWithVertexName(OutputPin->GetName());
+	Frontend::FInputHandle InputHandle = Editor::FGraphBuilder::GetInputHandleFromPin(InputPin);
+	Frontend::FOutputHandle OutputHandle = Editor::FGraphBuilder::GetOutputHandleFromPin(OutputPin);
 
-	if (ensure((InputHandles.Num() == 1) && (OutputHandles.Num() == 1)))
+	if (InputHandle->IsValid() && OutputHandle->IsValid())
 	{
-		Metasound::Frontend::FInputHandle InputHandle = InputHandles[0];
-		Metasound::Frontend::FOutputHandle OutputHandle = OutputHandles[0];
 		// TODO: Implement YesWithConverterNode to provide conversion options
-		Metasound::Frontend::FConnectability Connectability = InputHandle->CanConnectTo(*OutputHandle);
-		if (Connectability.Connectable != Metasound::Frontend::FConnectability::EConnectable::Yes)
+		Frontend::FConnectability Connectability = InputHandle->CanConnectTo(*OutputHandle);
+		if (Connectability.Connectable != Frontend::FConnectability::EConnectable::Yes)
 		{
 			const FName InputType = InputHandle->GetDataType();
 			const FName OutputType = OutputHandle->GetDataType();
@@ -442,7 +567,7 @@ bool UMetasoundEditorGraphSchema::TryCreateConnection(UEdGraphPin* PinA, UEdGrap
 {
 	using namespace Metasound::Frontend;
 
-	if (!ensureAlways(PinA && PinB))
+	if (!ensure(PinA && PinB))
 	{
 		return false;
 	}
@@ -454,7 +579,7 @@ bool UMetasoundEditorGraphSchema::TryCreateConnection(UEdGraphPin* PinA, UEdGrap
 		return false;
 	}
 
-	if (!ensureAlways(InputPin && OutputPin))
+	if (!ensure(InputPin && OutputPin))
 	{
 		return false;
 	}
@@ -464,7 +589,7 @@ bool UMetasoundEditorGraphSchema::TryCreateConnection(UEdGraphPin* PinA, UEdGrap
 	bool bModified = UEdGraphSchema::TryCreateConnection(PinA, PinB);
 	if (bModified)
 	{
-		bModified = Metasound::Editor::FGraphBuilder::ConnectNodes(*InputPin, *OutputPin);
+		bModified = Metasound::Editor::FGraphBuilder::ConnectNodes(*InputPin, *OutputPin, false);
 	}
 
 	return bModified;
@@ -605,81 +730,101 @@ void UMetasoundEditorGraphSchema::DroppedAssetsOnNode(const TArray<FAssetData>&
 	// TODO: Implement for metasound references
 }
 
-void UMetasoundEditorGraphSchema::GetAllMetasoundActions(FGraphActionMenuBuilder& ActionMenuBuilder, bool bShowSelectedActions) const
+void UMetasoundEditorGraphSchema::GetConversionActions(FGraphActionMenuBuilder& ActionMenuBuilder, Metasound::Editor::FActionClassFilters InFilters, bool bShowSelectedActions) const
 {
-	using namespace Metasound::Editor;
-	using namespace Metasound::Frontend;
+	using namespace Metasound;
+
+	const FText MenuJoinFormat = LOCTEXT("MetasoundActionsFormatSubCategory", "{0}|{1}");
+	const TArray<Frontend::FNodeClassInfo> ClassInfos = Frontend::GetAllAvailableNodeClasses();
+	for (const Frontend::FNodeClassInfo& ClassInfo : ClassInfos)
+	{
+		const FMetasoundFrontendClass ClassDescription = Frontend::GenerateClassDescription(ClassInfo);
+		if (InFilters.InputFilterFunction && !ClassDescription.Interface.Inputs.ContainsByPredicate(InFilters.InputFilterFunction))
+		{
+			continue;
+		}
+
+		if (InFilters.OutputFilterFunction && !ClassDescription.Interface.Outputs.ContainsByPredicate(InFilters.OutputFilterFunction))
+		{
+			continue;
+		}
 
+		const FMetasoundFrontendClassMetadata Metadata = ClassDescription.Metadata;
+		const FText Tooltip = Metadata.Author.IsEmpty()
+			? Metadata.Description
+			: FText::Format(LOCTEXT("MetasoundTooltipAuthorFormat", "{0}\nAuthor: {1}"), Metadata.Description, Metadata.Author);
+
+		if (!Metadata.CategoryHierarchy.IsEmpty() && !Metadata.CategoryHierarchy[0].CompareTo(Editor::FGraphBuilder::ConvertMenuName))
+		{
+			TSharedPtr<FMetasoundGraphSchemaAction_NewNode> NewNodeAction = MakeShared<FMetasoundGraphSchemaAction_NewNode>
+			(
+				Editor::FGraphBuilder::ConvertMenuName,
+				Metadata.DisplayName,
+				Tooltip,
+				0
+			);
+
+			NewNodeAction->NodeClassInfo = ClassInfo;
+			ActionMenuBuilder.AddAction(NewNodeAction);
+
+		}
+	}
+}
+
+void UMetasoundEditorGraphSchema::GetDataTypeInputNodeActions(FGraphActionMenuBuilder& ActionMenuBuilder, Metasound::Editor::FDataTypeFilterFunction InFilter, bool bShowSelectedActions) const
+{
 	const FText InputMenuName = LOCTEXT("MetasoundActionsInputsMenu", "Inputs");
+	const FText InputTooltipFormat = LOCTEXT("MetasoundTooltipAddInputFormat", "Adds an input of type {0} to the Metasound");
+	Metasound::Editor::GetDataTypeActions<FMetasoundGraphSchemaAction_NewInput>(ActionMenuBuilder, InFilter, InputMenuName, InputTooltipFormat, bShowSelectedActions);
+}
+
+void UMetasoundEditorGraphSchema::GetDataTypeOutputNodeActions(FGraphActionMenuBuilder& ActionMenuBuilder, Metasound::Editor::FDataTypeFilterFunction InFilter, bool bShowSelectedActions) const
+{
 	const FText OutputMenuName = LOCTEXT("MetasoundActionsOutputsMenu", "Outputs");
+	const FText OutputTooltipFormat = LOCTEXT("MetasoundTooltipAddOutputFormat", "Adds an output of type {0} to the Metasound");
+	Metasound::Editor::GetDataTypeActions<FMetasoundGraphSchemaAction_NewOutput>(ActionMenuBuilder, InFilter, OutputMenuName, OutputTooltipFormat, bShowSelectedActions);
+}
 
-	const FText MenuJoinFormat = LOCTEXT("MetasoundActionsFormatSubCategory", "{0}|{1}");
+void UMetasoundEditorGraphSchema::GetFunctionActions(FGraphActionMenuBuilder& ActionMenuBuilder, Metasound::Editor::FActionClassFilters InFilters, bool bShowSelectedActions) const
+{
+	using namespace Metasound::Editor;
+	using namespace Metasound::Frontend;
 
-	IMetasoundEditorModule& EditorModule = FModuleManager::GetModuleChecked<IMetasoundEditorModule>("MetasoundEditor");
-	EditorModule.IterateDataTypes([MenuJoinFormat, InputMenuName, OutputMenuName, InMenuBuilder = &ActionMenuBuilder](const FEditorDataType& DataType)
-	{
-		const FName DataTypeName = DataType.RegistryInfo.DataTypeName;
-		const FText DataTypeDisplayName = FText::FromString(FGraphBuilder::GetDataTypeDisplayName(DataTypeName));
-		const FText DataTypeTextName = FText::FromName(DataTypeName);
-
-		const TArray<FString> Categories = FGraphBuilder::GetDataTypeNameCategories(DataTypeName);
-		const FText CategoriesText = FText::FromString(FString::Join(Categories, TEXT("|")));
-
-
-		TSharedPtr<FMetasoundGraphSchemaAction_NewInput> AddInputNodeAction = MakeShared<FMetasoundGraphSchemaAction_NewInput>
-		(
-			FText::Format(MenuJoinFormat, InputMenuName, CategoriesText),
-			DataTypeDisplayName,
-			DataTypeName,
-			FText::Format(LOCTEXT("MetasoundTooltipAddInputFormat", "Adds an input of type {0} to the Metasound"), DataTypeTextName),
-			0
-		);
-		InMenuBuilder->AddAction(AddInputNodeAction);
-
-		TSharedPtr<FMetasoundGraphSchemaAction_NewOutput> AddOutputNodeAction = MakeShared<FMetasoundGraphSchemaAction_NewOutput>
-		(
-			FText::Format(MenuJoinFormat, OutputMenuName, CategoriesText),
-			DataTypeDisplayName,
-			DataTypeName,
-			FText::Format(LOCTEXT("MetasoundTooltipAddOutputFormat", "Adds an output of type {0} to the Metasound"), DataTypeTextName),
-			0
-		);
-		InMenuBuilder->AddAction(AddOutputNodeAction);
-	});
-
-	const FText NodeMenuName = LOCTEXT("MetasoundAddFunctionsMenu", "Functions");
-	const FText ConvertMenuName = LOCTEXT("MetasoundConverionsMenu", "Conversions");
+	const FText MenuJoinFormat = LOCTEXT("MetasoundActionsFormatSubCategory", "{0}|{1}");
 
 	const TArray<FNodeClassInfo> ClassInfos = GetAllAvailableNodeClasses();
 	for (const FNodeClassInfo& ClassInfo : ClassInfos)
 	{
 		const FMetasoundFrontendClass ClassDescription = GenerateClassDescription(ClassInfo);
+		if (InFilters.InputFilterFunction && !ClassDescription.Interface.Inputs.ContainsByPredicate(InFilters.InputFilterFunction))
+		{
+			continue;
+		}
+
+		if (InFilters.OutputFilterFunction && !ClassDescription.Interface.Outputs.ContainsByPredicate(InFilters.OutputFilterFunction))
+		{
+			continue;
+		}
+
 		const FMetasoundFrontendClassMetadata Metadata = ClassDescription.Metadata;
 		const FText Tooltip = Metadata.Author.IsEmpty()
 			? Metadata.Description
 			: FText::Format(LOCTEXT("MetasoundTooltipAuthorFormat", "{0}\nAuthor: {1}"), Metadata.Description, Metadata.Author);
 
-		const FText* MenuName = &NodeMenuName;
-		FText CategoriesText = FText::GetEmpty();
-		if (!Metadata.CategoryHierarchy.IsEmpty() && !Metadata.CategoryHierarchy[0].CompareTo(ConvertMenuName))
-		{
-			MenuName = &ConvertMenuName;
-		}
-		else
+		if (Metadata.CategoryHierarchy.IsEmpty() || Metadata.CategoryHierarchy[0].CompareTo(FGraphBuilder::ConvertMenuName))
 		{
-			CategoriesText = FText::Join(LOCTEXT("MetasoundActionsCategoryDelim", "|"), Metadata.CategoryHierarchy);
+			const FText CategoriesText = FText::Join(LOCTEXT("MetasoundActionsCategoryDelim", "|"), Metadata.CategoryHierarchy);
+			TSharedPtr<FMetasoundGraphSchemaAction_NewNode> NewNodeAction = MakeShared<FMetasoundGraphSchemaAction_NewNode>
+			(
+				FText::Format(MenuJoinFormat, FGraphBuilder::FunctionMenuName, CategoriesText),
+				Metadata.DisplayName,
+				Tooltip,
+				0
+			);
+
+			NewNodeAction->NodeClassInfo = ClassInfo;
+			ActionMenuBuilder.AddAction(NewNodeAction);
 		}
-
-		TSharedPtr<FMetasoundGraphSchemaAction_NewNode> NewNodeAction = MakeShared<FMetasoundGraphSchemaAction_NewNode>
-		(
-			FText::Format(MenuJoinFormat, *MenuName, CategoriesText),
-			Metadata.DisplayName,
-			Tooltip,
-			0
-		);
-
-		NewNodeAction->NodeClassInfo = ClassInfo;
-		ActionMenuBuilder.AddAction(NewNodeAction);
 	}
 }
 
diff --git a/Engine/Plugins/Runtime/Metasound/Source/MetasoundEditor/Public/MetasoundEditorGraphSchema.h b/Engine/Plugins/Runtime/Metasound/Source/MetasoundEditor/Public/MetasoundEditorGraphSchema.h
--- a/Engine/Plugins/Runtime/Metasound/Source/MetasoundEditor/Public/MetasoundEditorGraphSchema.h
+++ b/Engine/Plugins/Runtime/Metasound/Source/MetasoundEditor/Public/MetasoundEditorGraphSchema.h
@@ -7,7 +7,9 @@
 #include "CoreMinimal.h"
 #include "EdGraph/EdGraphSchema.h"
 #include "EdGraphUtilities.h"
+#include "MetasoundEditorModule.h"
 #include "MetasoundFrontend.h"
+#include "Templates/Function.h"
 #include "UObject/ObjectMacros.h"
 
 #include "MetasoundEditorGraphSchema.generated.h"
@@ -24,6 +26,16 @@ namespace Metasound
 {
 	namespace Editor
 	{
+		using FInputFilterFunction = TFunction<bool(const FMetasoundFrontendClassInput&)>;
+		using FOutputFilterFunction = TFunction<bool(const FMetasoundFrontendClassOutput&)>;
+		using FDataTypeFilterFunction = TFunction<bool(const FEditorDataType&)>;
+
+		struct FActionClassFilters
+		{
+			FInputFilterFunction InputFilterFunction;
+			FOutputFilterFunction OutputFilterFunction;
+		};
+
 		struct FGraphConnectionDrawingPolicyFactory : public FGraphPanelPinConnectionFactory
 		{
 		public:
@@ -139,10 +151,6 @@ struct METASOUNDEDITOR_API FMetasoundGraphSchemaAction_NewNode : public FEdGraph
 	//~ Begin FEdGraphSchemaAction Interface
 	virtual UEdGraphNode* PerformAction(UEdGraph* ParentGraph, UEdGraphPin* FromPin, const FVector2D Location, bool bSelectNewNode = true) override;
 	//~ End FEdGraphSchemaAction Interface
-
-private:
-	/** Connects new node to output of selected nodes */
-	void ConnectToSelectedNodes(UMetasoundEditorGraphNode* NewGraphNode, UEdGraph* ParentGraph) const;
 };
 
 /** Action to add nodes to the graph based on selected objects*/
@@ -232,8 +240,12 @@ class UMetasoundEditorGraphSchema : public UEdGraphSchema
 	//~ End EdGraphSchema Interface
 
 private:
-	/** Adds actions for creating every type of GraphNode */
-	void GetAllMetasoundActions(FGraphActionMenuBuilder& ActionMenuBuilder, bool bShowSelectedActions) const;
+	/** Adds actions for creating actions associated with graph DataTypes */
+	void GetConversionActions(FGraphActionMenuBuilder& ActionMenuBuilder, Metasound::Editor::FActionClassFilters InFilters = Metasound::Editor::FActionClassFilters(), bool bShowSelectedActions = true) const;
+	void GetFunctionActions(FGraphActionMenuBuilder& ActionMenuBuilder, Metasound::Editor::FActionClassFilters InFilters = Metasound::Editor::FActionClassFilters(), bool bShowSelectedActions = true) const;
+
+	void GetDataTypeInputNodeActions(FGraphActionMenuBuilder& ActionMenuBuilder, Metasound::Editor::FDataTypeFilterFunction InFilter = Metasound::Editor::FDataTypeFilterFunction(), bool bShowSelectedActions = true) const;
+	void GetDataTypeOutputNodeActions(FGraphActionMenuBuilder& ActionMenuBuilder, Metasound::Editor::FDataTypeFilterFunction InFilter = Metasound::Editor::FDataTypeFilterFunction(), bool bShowSelectedActions = true) const;
 
 	/** Adds action for creating a comment */
 	void GetCommentAction(FGraphActionMenuBuilder& ActionMenuBuilder, const UEdGraph* CurrentGraph = nullptr) const;
diff --git a/Engine/Plugins/Runtime/Metasound/Source/MetasoundFrontend/Public/MetasoundAutoConverterNode.h b/Engine/Plugins/Runtime/Metasound/Source/MetasoundFrontend/Public/MetasoundAutoConverterNode.h
--- a/Engine/Plugins/Runtime/Metasound/Source/MetasoundFrontend/Public/MetasoundAutoConverterNode.h
+++ b/Engine/Plugins/Runtime/Metasound/Source/MetasoundFrontend/Public/MetasoundAutoConverterNode.h
@@ -62,10 +62,10 @@ namespace Metasound
 				DisplayStyle.bShowOutputNames = false;
 
 				FNodeClassMetadata Info;
-				Info.ClassName = {TEXT("Convert"), GetMetasoundDataTypeName<ToDataType>(), GetMetasoundDataTypeName<FromDataType>()}; 
+				Info.ClassName = { TEXT("Convert"), GetMetasoundDataTypeName<ToDataType>(), GetMetasoundDataTypeName<FromDataType>() };
 				Info.MajorVersion = 1;
 				Info.MinorVersion = 0;
-				Info.DisplayName = FText::Format(LOCTEXT("Metasound_AutoConverterNodeDisplayNameFormat", "Convert {0} to {1}"), FText::FromName(GetMetasoundDataTypeName<FromDataType>()), FText::FromName(GetMetasoundDataTypeName<ToDataType>()));
+				Info.DisplayName = FText::Format(LOCTEXT("Metasound_AutoConverterNodeDisplayNameFormat", "{0} to {1}"), FText::FromName(GetMetasoundDataTypeName<FromDataType>()), FText::FromName(GetMetasoundDataTypeName<ToDataType>()));
 				Info.Description = LOCTEXT("Metasound_ConverterNodeDescription", "Converts between two different data types.");
 				Info.Author = PluginAuthor;
 				Info.DisplayStyle = DisplayStyle;
