Metasound Enum support

diff --git a/Engine/Plugins/Runtime/Metasound/Source/MetasoundEditor/Private/MetasoundLiteralDescriptionDetailCustomization.cpp b/Engine/Plugins/Runtime/Metasound/Source/MetasoundEditor/Private/MetasoundLiteralDescriptionDetailCustomization.cpp
--- a/Engine/Plugins/Runtime/Metasound/Source/MetasoundEditor/Private/MetasoundLiteralDescriptionDetailCustomization.cpp
+++ b/Engine/Plugins/Runtime/Metasound/Source/MetasoundEditor/Private/MetasoundLiteralDescriptionDetailCustomization.cpp
@@ -20,6 +20,7 @@
 #include "Templates/SharedPointer.h"
 #include "Widgets/Text/STextBlock.h"
 #include "PropertyCustomizationHelpers.h"
+#include "Slate/Public/Widgets/SToolTip.h"
 
 #define LOCTEXT_NAMESPACE "MetasoundEditor"
 
@@ -131,6 +132,57 @@ void FMetasoundFrontendLiteralDetailCustomization::CustomizeChildren(TSharedRef<
 		TSharedPtr<IPropertyHandle> PropertyHandle;
 		FText DisplayName;
 		UClass* OwningDataTypeUClass;
+		FName OwningDataTypeName;
+
+		TSharedRef<SWidget> CreateEnumWidget(const TSharedPtr<const Metasound::Frontend::IEnumDataTypeInterface>& EnumInterface) const
+		{
+			auto RemoveNameSpace = [](const FString& InString) -> FString
+			{ 
+				FString Namespace, EnumName;
+				InString.Split(TEXT("::"), &Namespace, &EnumName);
+				return EnumName;
+			};
+			auto GetAll = [EnumInterface, RemoveNameSpace](TArray<TSharedPtr<FString>>& OutStrings, TArray<TSharedPtr<SToolTip>>& OutTooltips, TArray<bool>&)
+			{
+				for (const Metasound::Frontend::IEnumDataTypeInterface::FGenericInt32Entry& i : EnumInterface->GetAllEntries())
+				{
+					OutTooltips.Emplace(SNew(SToolTip).Text(i.Tooltip));
+
+					// Trim the namespace off for display.
+					OutStrings.Emplace(MakeShared<FString>(RemoveNameSpace(i.Name.GetPlainNameString())));
+				}
+			};
+			auto GetValue = [EnumInterface, Prop = PropertyHandle, RemoveNameSpace]() -> FString
+			{
+				int32 IntValue;
+				if (Prop->GetValue(IntValue))
+				{
+					if (TOptional<FName> Result = EnumInterface->ToName(IntValue))
+					{
+						return RemoveNameSpace(Result->GetPlainNameString());
+					}
+					UE_LOG(LogTemp, Warning, TEXT("Failed to Get Valid Value for Property '%s' with Value of '%d'"),
+						*GetNameSafe(Prop->GetProperty()), IntValue);
+				}
+				return {};
+			};
+			auto SelectedValue = [EnumInterface, Prop = PropertyHandle](const FString& InSelected)
+			{
+				FString FullyQualifiedName = FString::Printf(TEXT("%s::%s"), *EnumInterface->GetNamespace().GetPlainNameString(), *InSelected);
+				TOptional<int32> Result = EnumInterface->ToValue(*FullyQualifiedName);
+				if (ensure(Result))
+				{
+					Prop->SetValue(*Result);
+				}
+			};
+
+			return PropertyCustomizationHelpers::MakePropertyComboBox(
+				nullptr,
+				FOnGetPropertyComboBoxStrings::CreateLambda(GetAll),
+				FOnGetPropertyComboBoxValue::CreateLambda(GetValue),
+				FOnPropertyComboBoxValueSelected::CreateLambda(SelectedValue)
+			);
+		}
 
 		TSharedRef<SWidget> CreatePropertyValueWidget(bool bDisplayDefaultPropertyButtons = true) const
 		{
@@ -182,6 +234,18 @@ void FMetasoundFrontendLiteralDetailCustomization::CustomizeChildren(TSharedRef<
 					.DisplayThumbnail(true)
 					.NewAssetFactories(FactoriesToUse);
 			}
+			else if (LiteralType == EMetasoundFrontendLiteralType::Integer)
+			{
+				// Check if this is an Enum.
+				if (TSharedPtr<const Metasound::Frontend::IEnumDataTypeInterface> EnumInterface = 
+					FMetasoundFrontendRegistryContainer::Get()->GetEnumInterfaceForDataType(OwningDataTypeName))
+				{ 
+					return CreateEnumWidget(EnumInterface);
+				}
+
+				// Fall back to regular widget if its not.
+				return PropertyHandle->CreatePropertyValueWidget();
+			}
 			else if (LiteralType == EMetasoundFrontendLiteralType::UObjectArray)
 			{
 				// TODO: Implement.
@@ -203,7 +267,8 @@ void FMetasoundFrontendLiteralDetailCustomization::CustomizeChildren(TSharedRef<
 			TypeEnumHandle,
 			StructPropertyHandle->GetChildHandle(GET_MEMBER_NAME_CHECKED(FMetasoundFrontendLiteral, AsBool)),
 			LOCTEXT("Metasound_LiteralDisplayNameBool", "Bool"),
-			nullptr
+			nullptr,
+			OwningDataTypeName
 		},
 
 		FTypeEditorInfo
@@ -212,7 +277,8 @@ void FMetasoundFrontendLiteralDetailCustomization::CustomizeChildren(TSharedRef<
 			TypeEnumHandle,
 			StructPropertyHandle->GetChildHandle(GET_MEMBER_NAME_CHECKED(FMetasoundFrontendLiteral, AsInteger)),
 			LOCTEXT("Metasound_LiteralDisplayNameInteger", "Int32"),
-			nullptr
+			nullptr,
+			OwningDataTypeName
 		},
 
 		FTypeEditorInfo
@@ -221,7 +287,8 @@ void FMetasoundFrontendLiteralDetailCustomization::CustomizeChildren(TSharedRef<
 			TypeEnumHandle,
 			StructPropertyHandle->GetChildHandle(GET_MEMBER_NAME_CHECKED(FMetasoundFrontendLiteral, AsFloat)),
 			LOCTEXT("Metasound_LiteralDisplayNameFloat", "Float"),
-			nullptr
+			nullptr,
+			OwningDataTypeName
 		},
 
 		FTypeEditorInfo
@@ -230,7 +297,8 @@ void FMetasoundFrontendLiteralDetailCustomization::CustomizeChildren(TSharedRef<
 			TypeEnumHandle,
 			StructPropertyHandle->GetChildHandle(GET_MEMBER_NAME_CHECKED(FMetasoundFrontendLiteral, AsString)),
 			LOCTEXT("Metasound_LiteralDisplayNameString", "String"),
-			nullptr
+			nullptr,
+			OwningDataTypeName
 		},
 
 		FTypeEditorInfo
@@ -239,7 +307,8 @@ void FMetasoundFrontendLiteralDetailCustomization::CustomizeChildren(TSharedRef<
 			TypeEnumHandle,
 			StructPropertyHandle->GetChildHandle(GET_MEMBER_NAME_CHECKED(FMetasoundFrontendLiteral, AsUObject)),
 			LOCTEXT("Metasound_LiteralDisplayNameUObject", "UObject"),
-			OwningDataTypeClass
+			OwningDataTypeClass,
+			OwningDataTypeName
 		},
 
 		FTypeEditorInfo
@@ -248,8 +317,9 @@ void FMetasoundFrontendLiteralDetailCustomization::CustomizeChildren(TSharedRef<
 			TypeEnumHandle,
 			StructPropertyHandle->GetChildHandle(GET_MEMBER_NAME_CHECKED(FMetasoundFrontendLiteral, AsUObjectArray)),
 			LOCTEXT("Metasound_LiteralDisplayNameUObjectArray", "UObjectArray"),
-			OwningDataTypeClass
-		},
+			OwningDataTypeClass,
+			OwningDataTypeName
+		}
 	};
 
 	for (const FTypeEditorInfo& Info : TypePropertyInfo)
diff --git a/Engine/Plugins/Runtime/Metasound/Source/MetasoundFrontend/Private/MetasoundFrontendRegistries.cpp b/Engine/Plugins/Runtime/Metasound/Source/MetasoundFrontend/Private/MetasoundFrontendRegistries.cpp
--- a/Engine/Plugins/Runtime/Metasound/Source/MetasoundFrontend/Private/MetasoundFrontendRegistries.cpp
+++ b/Engine/Plugins/Runtime/Metasound/Source/MetasoundFrontend/Private/MetasoundFrontendRegistries.cpp
@@ -209,6 +209,11 @@ Metasound::ELiteralType FMetasoundFrontendRegistryContainer::GetDesiredLiteralTy
 
 	const FDataTypeRegistryElement& DataTypeInfo = DataTypeRegistry[InDataType];
 	
+	if (DataTypeInfo.Info.bIsEnum)
+	{
+		return Metasound::ELiteralType::Integer;
+	}
+
 	// If there's a designated preferred literal type for this datatype, use that.
 	if (DataTypeInfo.Info.PreferredLiteralType != Metasound::ELiteralType::None)
 	{
@@ -333,6 +338,16 @@ bool FMetasoundFrontendRegistryContainer::RegisterDataType(const ::Metasound::FD
 	}
 }
 
+bool FMetasoundFrontendRegistryContainer::RegisterEnumDataInterface(FName InDataType, TSharedPtr<IEnumDataTypeInterface>&& InInterface)
+{
+	if (FDataTypeRegistryElement* Element = DataTypeRegistry.Find(InDataType))
+	{
+		Element->EnumInterface = MoveTemp(InInterface);
+		return true;
+	}
+	return false;
+}
+
 bool FMetasoundFrontendRegistryContainer::RegisterExternalNode(Metasound::FCreateMetasoundNodeFunction&& InCreateNode, Metasound::FCreateMetasoundFrontendClassFunction&& InCreateDescription)
 {
 	FNodeRegistryElement RegistryElement = FNodeRegistryElement(MoveTemp(InCreateNode), MoveTemp(InCreateDescription));
@@ -485,3 +500,14 @@ bool FMetasoundFrontendRegistryContainer::GetInfoForDataType(FName InDataType, M
 	}
 }
 
+
+TSharedPtr<const Metasound::Frontend::IEnumDataTypeInterface>
+FMetasoundFrontendRegistryContainer::GetEnumInterfaceForDataType(FName InDataType) const
+{
+	if (const FDataTypeRegistryElement* Element = DataTypeRegistry.Find(InDataType))
+	{
+		return Element->EnumInterface;
+	}
+	return nullptr;
+}
+
diff --git a/Engine/Plugins/Runtime/Metasound/Source/MetasoundFrontend/Public/MetasoundDataTypeRegistrationMacro.h b/Engine/Plugins/Runtime/Metasound/Source/MetasoundFrontend/Public/MetasoundDataTypeRegistrationMacro.h
--- a/Engine/Plugins/Runtime/Metasound/Source/MetasoundFrontend/Public/MetasoundDataTypeRegistrationMacro.h
+++ b/Engine/Plugins/Runtime/Metasound/Source/MetasoundFrontend/Public/MetasoundDataTypeRegistrationMacro.h
@@ -18,6 +18,7 @@
 #include "MetasoundReceiveNode.h"
 #include "MetasoundRouter.h"
 #include "MetasoundSendNode.h"
+#include "MetasoundEnum.h"
 
 #include <type_traits>
 
@@ -43,7 +44,7 @@ namespace Metasound
 	{
 		return nullptr;
 	}
-
+	
 	// Helper utility to test if we can transmit a datatype between a send and a receive node.
 	template <typename TDataType>
 	struct TIsTransmittable
@@ -123,7 +124,56 @@ namespace Metasound
 		AttemptToRegisterConverter<TDataType, float>();
 		AttemptToRegisterConverter<TDataType, FString>();
 	}
+	
+	// SFINAE for enum types.
+	template<typename TDataType, typename std::enable_if<TEnumTraits<TDataType>::bIsEnum, bool>::type = true>
+	bool RegisterEnumDataTypeWithFrontend()
+	{
+		using InnerType = typename TEnumTraits<TDataType>::InnerType;
+		using FStringHelper = TEnumStringHelper<InnerType>;
 
+		struct FEnumHandler : Metasound::Frontend::IEnumDataTypeInterface
+		{
+			TArray<FName> GetAllNames() const override
+			{
+				return FStringHelper::GetAllNames();
+			}
+			FName GetNamespace() const override
+			{
+				return FStringHelper::GetNamespace();
+			}
+			TOptional<FName> ToName(int32 InEnumValue) const override
+			{
+				return FStringHelper::ToName(static_cast<InnerType>(InEnumValue));
+			}
+			TOptional<int32> ToValue(FName InName) const override
+			{
+				if (TOptional<InnerType> Result = FStringHelper::FromName(InName))
+				{
+					return static_cast<int32>(*Result);
+				}
+				return {};
+			}
+			TArray<FGenericInt32Entry> GetAllEntries() const override
+			{
+				// Convert to int32 representation 
+				TArray<FGenericInt32Entry> IntEntries;
+				for (const TEnumEntry<InnerType>& i : FStringHelper::GetAllEntries())
+				{					
+					IntEntries.Emplace(FGenericInt32Entry{ static_cast<int32>(i.Value), i.Name, i.Tooltip });
+				}
+				return IntEntries;
+			}
+		};
+
+		return FMetasoundFrontendRegistryContainer::Get()->RegisterEnumDataInterface(
+			GetMetasoundDataTypeName<TDataType>(), MakeShared<FEnumHandler>());
+	}
+	
+	// SFINAE stub for non-enum types. 
+	template<typename TDataType, typename std::enable_if<!TEnumTraits<TDataType>::bIsEnum, bool>::type = true>
+	bool RegisterEnumDataTypeWithFrontend() { return false; }
+	
 	template<typename TDataType, ELiteralType PreferredArgType = ELiteralType::None, typename UClassToUse = UObject>
 	bool RegisterDataTypeWithFrontend()
 	{
@@ -221,7 +271,8 @@ namespace Metasound
 		RegistryInfo.bIsStringParsable = TIsParsable<TDataType, FString>::Value;
 		RegistryInfo.bIsProxyParsable = TIsParsable<TDataType, const Audio::IProxyDataPtr&>::Value;
 		RegistryInfo.bIsProxyArrayParsable = TIsParsable<TDataType, const TArray<Audio::IProxyDataPtr>& >::Value;
-		RegistryInfo.bIsDefaultParsable = TIsParsable<TDataType>::Value;
+		RegistryInfo.bIsDefaultParsable = TIsParsable<TDataType>::Value;		
+		RegistryInfo.bIsEnum = TEnumTraits<TDataType>::bIsEnum;
 
 		RegistryInfo.bIsTransmittable = TIsTransmittable<TDataType>::Value;
 		
@@ -229,13 +280,23 @@ namespace Metasound
 
 		bool bSucceeded = FMetasoundFrontendRegistryContainer::Get()->RegisterDataType(RegistryInfo, MoveTemp(Callbacks));
 		ensureAlwaysMsgf(bSucceeded, TEXT("Failed to register data type %s in the node registry!"), *GetMetasoundDataTypeString<TDataType>());
-		
+
+		// If its an enum, register its data interface AFTER we've registered the data type.
+		if (TEnumTraits<TDataType>::bIsEnum)
+		{
+			if (!ensure(RegisterEnumDataTypeWithFrontend<TDataType>()))
+			{
+				return false;
+			}
+		}
+
 		RegisterConverterNodes<TDataType>();
 		AttemptToRegisterSendAndReceiveNodes<TDataType>();
 		
 		return bSucceeded;
 	}
 
+
 	template<typename DataType>
 	struct TMetasoundDataTypeRegistration
 	{
diff --git a/Engine/Plugins/Runtime/Metasound/Source/MetasoundFrontend/Public/MetasoundEnum.h b/Engine/Plugins/Runtime/Metasound/Source/MetasoundFrontend/Public/MetasoundEnum.h
new file mode 100644
--- /dev/null
+++ b/Engine/Plugins/Runtime/Metasound/Source/MetasoundFrontend/Public/MetasoundEnum.h
@@ -0,0 +1,252 @@
+// Copyright Epic Games, Inc. All Rights Reserved.
+
+#pragma once
+
+#include "MetasoundDataReference.h"
+#include <type_traits>
+
+// NOTE: Metasound Enum types are defined outside of Engine so can't use the UENUM type reflection here.
+// Basic reflection is provides using template specialization and Macros defined below.
+
+// Example usage:
+//
+// 1. Declare an Enum class. 
+// enum class EMyOtherTestEnum : int32
+// {
+//		Alpha = 500,
+//		Beta = -666,
+//		Gamma = 333
+// };
+
+// 2. Declare its wrapper types using the DECLARE_METASOUND_ENUM macro.
+// DECLARE_METASOUND_ENUM(EMyOtherTestEnum, EMyOtherTestEnum::Gamma, METASOUNDSTANDARDNODES_API, FMyOtherTestEnumTypeInfo, FMyOtherTestEnumReadRef, FMyOtherTestEnumWriteRef)
+
+// 3. Define it using the BEGIN/ENTRY/END macros
+// DEFINE_METASOUND_ENUM_BEGIN(EMyOtherTestEnum)
+// 		DEFINE_METASOUND_ENUM_ENTRY(EMyOtherTestEnum::Alpha, LOCTEXT("AlphaDescription", "Alpha tooltip")),
+// 		DEFINE_METASOUND_ENUM_ENTRY(EMyOtherTestEnum::Beta, LOCTEXT("BetaDescription", "Beta tooltip")),
+// 		DEFINE_METASOUND_ENUM_ENTRY(EMyOtherTestEnum::Gamma, LOCTEXT("GammaDescription", "Gamma tooltip"))
+// DEFINE_METASOUND_ENUM_END()
+
+namespace Metasound
+{
+	// Struct to hold each of the entries of the Enum.
+	template<typename T>
+	struct TEnumEntry
+	{
+		T Value;
+		FName Name;
+		FText Tooltip;	// TODO: Remove this for runtime.
+	};
+
+	/** CRTP base class for Enum String Helper type.
+	 *  Provides common code for all specializations.
+	 */ 
+	template<typename Derived, typename EnumType>
+	struct TEnumStringHelperBase
+	{
+		// Give a enum value e.g. 'EMyEnum::One', convert that to a Name (if its valid)
+		static TOptional<FName> ToName(EnumType InValue)
+		{
+			for (const TEnumEntry<EnumType>& i : Derived::GetAllEntries())
+			{
+				if (i.Value == InValue)
+				{
+					return i.Name;
+				}
+			}
+			return {};
+		}
+		// Give a Name "EMyEnum::One", convert that to a Enum Value
+		static TOptional<EnumType> FromName(const FName	InName)
+		{
+			for (const TEnumEntry<EnumType>& i : Derived::GetAllEntries())
+			{
+				if (i.Name == InName)
+				{
+					return i.Value;
+				}
+			}
+			return {};
+		}
+		// Return all possible names.
+		static TArray<FName> GetAllNames()
+		{
+			TArray<FName> Names;
+			for (const TEnumEntry<EnumType> i : Derived::GetAllEntries())
+			{
+				Names.Emplace(i.Name);
+			}
+			return Names;
+		}	
+	};
+
+	/** Metasound Enum String Helper
+	 */
+	template<typename T>
+	struct METASOUNDFRONTEND_API TEnumStringHelper : TEnumStringHelperBase<TEnumStringHelper<T>, T>
+	{
+		static_assert(TIsEnum<T>::Value, "Please define a specialization of this class. The DECLARE_METASOUND_ENUM macros will do this for you");
+	};
+
+	/** Metasound Enum Wrapper
+	 */
+	template<typename EnumType, EnumType DefaultValue>
+	class TEnum final
+	{
+	public:
+		using InnerType = EnumType;
+		static constexpr EnumType Default = DefaultValue;
+
+		static_assert(TIsEnum<EnumType>::Value, "Expecting an Enum type");
+		static_assert(TIsSame<typename std::underlying_type<EnumType>::type, int32>::Value, "We serialize to int32, so limit to that for now");
+
+		// Default.
+		explicit TEnum(EnumType InValue = DefaultValue)
+		{
+			// Try and convert to validate this is a valid value.
+			TOptional<FName> Converted = ToName(InValue);
+			if (ensure(Converted))
+			{
+				EnumValue = InValue;
+			}
+		}
+
+		// From Int32 (this is the common path from a Literal).
+		explicit TEnum(int32 InIntValue)
+			: TEnum(static_cast<EnumType>(InIntValue))
+		{
+		}
+
+		// From Name
+		explicit TEnum(FName InValueName)		
+		{
+			// Try and convert from Name to Value and ensure that we succeed
+			TOptional<EnumType> Converted = NameToEnum(InValueName);
+			if (ensure(Converted))
+			{
+				EnumValue = *Converted;
+			}
+		}
+
+		// Slow, construct from FString to FName.
+		explicit TEnum(const FString& InString)
+			: TEnum(FName(*InString))
+		{
+		}
+
+		EnumType Get() const
+		{
+			return EnumValue;
+		}
+
+		int32 ToInt() const
+		{
+			return static_cast<int32>(EnumValue);
+		}
+
+		// Convert to its FName (if possible).
+		TOptional<FName> ToName() const 
+		{
+			return ToName(EnumValue);
+		}
+
+		// Conversion operator to automatically convert this to its underlying enum type.
+		operator EnumType() const
+		{
+			return EnumValue;
+		}
+
+		// Convert from EnumValue to FName (if possible).
+		static TOptional<FName> ToName(EnumType InValue)
+		{
+			return TEnumStringHelper<EnumType>::ToName(InValue);
+		}
+
+		// Convert from Name to EnumValue (if possible).
+		static TOptional<EnumType> NameToEnum(FName InValue)
+		{
+			return TEnumStringHelper<EnumType>::FromName(InValue);
+		}
+
+		// Return all possible Names
+		static TArray<FName> GetAllNames()
+		{
+			return TEnumStringHelper<EnumType>::GetAllNames();
+		}
+				
+	private:
+		// Keep the type in its fully typed form for debugging.
+		EnumType EnumValue = DefaultValue;
+	};
+
+	template<typename T>
+	struct TEnumTraits
+	{
+		static constexpr bool bIsEnum = false;
+		using InnerType = int32;
+		static constexpr int32 DefaultValue = 0;
+	};
+
+	template<typename T, T D>
+	struct TEnumTraits<TEnum<T, D>>
+	{
+		static constexpr bool bIsEnum = true;
+		using InnerType = T;
+		static constexpr T DefaultValue = D;		
+	};	
+}
+
+// Helper macros: for defining/declaring a new Enum to be used in Metasounds.
+// * Declares/defines a typedef FEnum<YourEnum> wrapper around the enum that you pass it
+// * Declares/defines a specialization of TEnumStringHelper which is used to validate and convert between values/names.
+
+// These helps generate a couple of boiler plate functions:
+//  - GetNamespace() which will return the name of the Enum "EMyEnum" etc. 
+//  - GetAllEntries() which returns all posible entries in the enum. 
+
+/** DECLARE_METASOUND_ENUM
+ * @param ENUMNAME - The typename of your EnumType you want to use for Metasounds. e.g. My
+ * @param DEFAULT - A fully qualified default Enum value. e.g. EMyEnum::One
+ * @param API - The module API this is declared inside e.g. METASOUNDSTANDARDNODES_API 
+ * @param TYPEINFO - The name of the TypeInfo type you want to define e.g. FMyEnumTypeInfo
+ * @param READREF - The name of the Read Reference type you want to define. e.g FMyEnumReadRef
+ * @param WRITEREF -The name of the Write Reference type you want to define e.g. FMyEnumWriteRef
+ */
+#define DECLARE_METASOUND_ENUM(ENUMNAME, DEFAULT, API, TYPEINFO, READREF, WRITEREF)\
+	using FEnum##ENUMNAME = Metasound::TEnum<ENUMNAME, DEFAULT>; \
+	DECLARE_METASOUND_DATA_REFERENCE_TYPES(FEnum##ENUMNAME, API, TYPEINFO, READREF, WRITEREF);\
+	template<> struct API Metasound::TEnumStringHelper<ENUMNAME> : Metasound::TEnumStringHelperBase<Metasound::TEnumStringHelper<ENUMNAME>, ENUMNAME>\
+	{\
+		static FName GetNamespace()\
+		{\
+			static const FName ThisName { TEXT(#ENUMNAME) };\
+			return ThisName;\
+		}\
+		static TArrayView<const Metasound::TEnumEntry<ENUMNAME>> GetAllEntries();\
+	};
+
+/** DEFINE_METASOUND_ENUM_BEGIN
+ * @param ENUMNAME - The typename of your EnumType you want to use for Metasounds. e.g. EMyType
+ */
+#define DEFINE_METASOUND_ENUM_BEGIN(ENUMNAME)\
+	REGISTER_METASOUND_DATATYPE(FEnum##ENUMNAME, "Enum:" #ENUMNAME, ::Metasound::ELiteralType::Integer);\
+	TArrayView<const Metasound::TEnumEntry<ENUMNAME>> Metasound::TEnumStringHelper<ENUMNAME>::GetAllEntries()\
+	{\
+		static const Metasound::TEnumEntry<ENUMNAME> Entries[] = {
+
+/** DEFINE_METASOUND_ENUM_ENTRY - defines a single Enum Entry
+  * @param ENTRY - Fully Qualified Name of Entry of the Enum. (e.g. EMyType::One)
+  * @param TOOLTIP - FText 
+  */
+	#define DEFINE_METASOUND_ENUM_ENTRY(ENTRY, TOOLTIP) { ENTRY, TEXT(#ENTRY), TOOLTIP }
+
+/** DEFINE_METASOUND_ENUM_END - macro which ends the function body of the GetAllEntries function
+  * @param ENTRY - Fully Qualified Name of Entry of the Enum. (e.g. EMyType::One)
+  * @param TOOLTIP - Non-wide Quoted String use for the tooltip.
+  */
+#define DEFINE_METASOUND_ENUM_END() \
+		};\
+		return Entries;\
+	};
+
diff --git a/Engine/Plugins/Runtime/Metasound/Source/MetasoundFrontend/Public/MetasoundFrontendRegistries.h b/Engine/Plugins/Runtime/Metasound/Source/MetasoundFrontend/Public/MetasoundFrontendRegistries.h
--- a/Engine/Plugins/Runtime/Metasound/Source/MetasoundFrontend/Public/MetasoundFrontendRegistries.h
+++ b/Engine/Plugins/Runtime/Metasound/Source/MetasoundFrontend/Public/MetasoundFrontendRegistries.h
@@ -9,6 +9,7 @@
 #include "MetasoundLiteral.h"
 #include "MetasoundNodeInterface.h"
 #include "MetasoundOperatorInterface.h"
+#include "MetasoundEnum.h"
 
 namespace Metasound
 {
@@ -40,6 +41,9 @@ namespace Metasound
 		bool bIsFloatParsable = false;
 		bool bIsStringParsable = false;
 
+		// Is an TEnum wrapped enum
+		bool bIsEnum = false;
+		
 		// these are used for using UObjects, or arrays of UObjects.
 		bool bIsProxyParsable = false;
 		bool bIsProxyArrayParsable = false;
@@ -130,6 +134,20 @@ namespace Metasound
 			TArray<FConverterNodeInfo> PotentialConverterNodes;
 		};
 
+		struct METASOUNDFRONTEND_API IEnumDataTypeInterface
+		{
+			using FGenericInt32Entry = TEnumEntry<int32>;
+
+			virtual FName GetNamespace() const = 0;
+			virtual TArray<FGenericInt32Entry> GetAllEntries() const = 0;
+
+			virtual TArray<FName> GetAllNames() const = 0;
+
+			virtual TOptional<FName> ToName(int32 InEnumValue) const = 0;
+			virtual TOptional<int32> ToValue(FName InName) const = 0;
+
+			virtual ~IEnumDataTypeInterface() = default;
+		};
 	}
 
 	struct FDataTypeConstructorCallbacks
@@ -164,6 +182,7 @@ class METASOUNDFRONTEND_API FMetasoundFrontendRegistryContainer
 	using FDataTypeRegistryInfo = Metasound::FDataTypeRegistryInfo;
 	using FDataTypeConstructorCallbacks = ::Metasound::FDataTypeConstructorCallbacks;
 	using FNodeClassMetadata = Metasound::FNodeClassMetadata;
+	using IEnumDataTypeInterface = Metasound::Frontend::IEnumDataTypeInterface;
 
 public:
 	static FMetasoundFrontendRegistryContainer* Get();
@@ -208,6 +227,8 @@ public:
 
 	bool RegisterDataType(const FDataTypeRegistryInfo& InDataInfo, const FDataTypeConstructorCallbacks& InCallbacks);
 
+	bool RegisterEnumDataInterface(FName InDataType, TSharedPtr<IEnumDataTypeInterface>&& InInterface);
+
 	/** Register external node with the frontend.
 	 *
 	 * @param InCreateNode - Function for creating node from FNodeInitData.
@@ -238,6 +259,8 @@ public:
 	// @returns false if InDataType wasn't found in the registry. 
 	bool GetInfoForDataType(FName InDataType, FDataTypeRegistryInfo& OutInfo);
 
+	TSharedPtr<const Metasound::Frontend::IEnumDataTypeInterface> GetEnumInterfaceForDataType(FName InDataType) const;
+
 private:
 	FMetasoundFrontendRegistryContainer();
 
@@ -264,6 +287,8 @@ private:
 		Metasound::FDataTypeConstructorCallbacks Callbacks;
 
 		Metasound::FDataTypeRegistryInfo Info;
+
+		TSharedPtr<const Metasound::Frontend::IEnumDataTypeInterface> EnumInterface;
 	};
 
 	TMap<FName, FDataTypeRegistryElement> DataTypeRegistry;
