Frontend Array Support

diff --git a/Engine/Plugins/Runtime/Metasound/Source/MetasoundEditor/Private/MetasoundLiteralDescriptionDetailCustomization.cpp b/Engine/Plugins/Runtime/Metasound/Source/MetasoundEditor/Private/MetasoundLiteralDescriptionDetailCustomization.cpp
--- a/Engine/Plugins/Runtime/Metasound/Source/MetasoundEditor/Private/MetasoundLiteralDescriptionDetailCustomization.cpp
+++ b/Engine/Plugins/Runtime/Metasound/Source/MetasoundEditor/Private/MetasoundLiteralDescriptionDetailCustomization.cpp
@@ -11,6 +11,7 @@
 #include "Layout/Visibility.h"
 #include "MetasoundFrontend.h"
 #include "MetasoundFrontendDocument.h"
+#include "MetasoundFrontendLiteral.h"
 #include "MetasoundFrontendRegistries.h"
 #include "Misc/AssertionMacros.h"
 #include "Misc/Attribute.h"
@@ -24,200 +25,388 @@
 
 #define LOCTEXT_NAMESPACE "MetasoundEditor"
 
+#ifndef GET_STRUCT_NAME_CHECKED
+#define GET_STRUCT_NAME_CHECKED(StructName) \
+		((void)sizeof(StructName), TEXT(#StructName))
+#endif
+
 namespace Metasound
 {
 	namespace Editor
 	{
-		TSharedPtr<IPropertyHandle> GetLiteralHandleForType(TSharedRef<IPropertyHandle> StructPropertyHandle, EMetasoundFrontendLiteralType LiteralType)
+		// Returns true if the property handle points to a property of a given CPP type.
+		bool IsPropertyHandleOfType(const TSharedPtr<IPropertyHandle>& InHandle, const FString& InTypeName)
 		{
-			switch (LiteralType)
+			if (InHandle.IsValid())
 			{
-			case EMetasoundFrontendLiteralType::Bool:
+				if (const FProperty* Property = InHandle->GetProperty())
+				{
+					return Property->GetCPPType() == InTypeName;
+				}
+			}
+
+			return false;
+		}
+
+		// Returns true if the TypeEnumHandle points to a EMeatsoundFrontendLiteralType property of equal value.
+		bool IsLiteralTypeActive(TSharedRef<IPropertyHandle> TypeEnumHandle, EMetasoundFrontendLiteralType LiteralType)
+		{
+			uint8 Value = static_cast<uint8>(EMetasoundFrontendLiteralType::None);
+			TypeEnumHandle->GetValue(Value);
+
+			return Value == static_cast<uint8>(LiteralType);
+		}
+
+		// Sets the literal type on a IPropertyHandle pointing to a FMetasoundFrontendLiteral
+		void SetActiveLiteralType(EMetasoundFrontendLiteralType InType, TSharedPtr<IPropertyHandle> InLiteralPropertyHandle)
+		{
+			check(IsPropertyHandleOfType(InLiteralPropertyHandle, GET_STRUCT_NAME_CHECKED(FMetasoundFrontendLiteral)));
+
+			if (TSharedPtr<IPropertyHandle> TypeHandle= InLiteralPropertyHandle->GetChildHandle(GET_MEMBER_NAME_CHECKED(FMetasoundFrontendLiteral, Type)))
 			{
-				return StructPropertyHandle->GetChildHandle(GET_MEMBER_NAME_CHECKED(FMetasoundFrontendLiteral, AsBool));
+				TypeHandle->SetValue(static_cast<uint8>(InType));
 			}
+		}
+
+		// Returns the Metasound Data Type name for a literal by inspecting the 
+		// parent properties.
+		FName GetOwningMetasoundDataTypeName(TSharedRef<IPropertyHandle> InMetasoundFrontendLiteral)
+		{
+			FName OwningDataTypeName;
+
+			// The association between a literal and a data type is not maintained on
+			// a literal. The owning property must be inspected To find the data type name 
+			// associated with the literal. There are no assurances that a literal will
+			// be held on a property which contains a data type name, but the instances
+			// of literals on a meatsound document USTRUCT have a hardcoded struct member
+			// layout. This layout can be exploited to find the data type when a FMetasoundFrontendLiteral 
+			// is a member of one of the various USTRUCTS in a metasound document.
 
-			case EMetasoundFrontendLiteralType::Float:
+			if (OwningDataTypeName.IsNone())
 			{
-				return StructPropertyHandle->GetChildHandle(GET_MEMBER_NAME_CHECKED(FMetasoundFrontendLiteral, AsFloat));
+				TSharedPtr<IPropertyHandle> OwningClassInputProperty = InMetasoundFrontendLiteral->GetParentHandle()->GetParentHandle()->GetParentHandle();
+				if (IsPropertyHandleOfType(OwningClassInputProperty, GET_STRUCT_NAME_CHECKED(FMetasoundFrontendClassInput)))
+				{
+					// Literal is owned by a FMetasoundFrontendClassInput
+					TSharedPtr<IPropertyHandle> DataTypeNameHandle = OwningClassInputProperty->GetChildHandle(GET_MEMBER_NAME_CHECKED(FMetasoundFrontendClassInput, TypeName));
+					DataTypeNameHandle->GetValue(OwningDataTypeName);
+				}
 			}
 
-			case EMetasoundFrontendLiteralType::Integer:
+			if (OwningDataTypeName.IsNone())
 			{
-				return StructPropertyHandle->GetChildHandle(GET_MEMBER_NAME_CHECKED(FMetasoundFrontendLiteral, AsInteger));
+				TSharedPtr<IPropertyHandle> OwningNodeProperty = InMetasoundFrontendLiteral->GetParentHandle()->GetParentHandle()->GetParentHandle();
+				if (IsPropertyHandleOfType(OwningNodeProperty, GET_STRUCT_NAME_CHECKED(FMetasoundFrontendNode)))
+				{
+					// Literal is owned by a FMetasoundFrontendNode. The literal will need
+					// to be matched to the node input using the PointID in order to get the 
+					// associated data type.
+					checkNoEntry();
+
+					// TODO: implement me. 
+				}
 			}
 
-			case EMetasoundFrontendLiteralType::String:
+			return OwningDataTypeName;
+		}
+
+		// ILiteralPropertyAdapter creates a custom widget for a FMetasoundFrontendLiteral.
+		// This interface can be subclasses to create different widgets under differing 
+		// scenarios.
+		class ILiteralPropertyAdapter
+		{
+			public:
+				ILiteralPropertyAdapter() = default;
+				virtual ~ILiteralPropertyAdapter() = default;
+
+				// Create a widget for a given FMetasoundFrontendLiteral property handle.
+				virtual TSharedRef<SWidget> CreatePropertyValueWidget(TSharedPtr<IPropertyHandle> InLiteralPropertyHandle, bool bDisplayDefaultPropertyButtons=false) const = 0;
+
+				// Return the literal type associated with this property adapater.
+				virtual EMetasoundFrontendLiteralType GetLiteralType() const = 0;
+
+				// Return localized tooltip text for a given FMetasoundFrontendLiteral property handle.
+				virtual const FText& GetToolTipText(TSharedPtr<IPropertyHandle> InLiteralPropertyHandle) const = 0;
+		};
+
+		// FLiteralPropertyBaseAdapter creates a default widget for the underlying
+		// property determined by EMetasoundFrontendLiteralType.
+		class FLiteralPropertyBaseAdapter : public ILiteralPropertyAdapter
+		{
+		public:
+
+			FLiteralPropertyBaseAdapter(EMetasoundFrontendLiteralType InPreferredLiteralType)
+			: PreferredLiteralType(InPreferredLiteralType)
 			{
-				return StructPropertyHandle->GetChildHandle(GET_MEMBER_NAME_CHECKED(FMetasoundFrontendLiteral, AsString));
 			}
 
-			case EMetasoundFrontendLiteralType::UObject:
+			virtual ~FLiteralPropertyBaseAdapter() = default;
+
+			// This sets the FMetasoundFrontendLiteral's active property to `GetLiteralType()` before creating a widget.
+			TSharedRef<SWidget> CreatePropertyValueWidget(TSharedPtr<IPropertyHandle> InLiteralPropertyHandle, bool bDisplayDefaultPropertyButtons = false) const override
 			{
-				return StructPropertyHandle->GetChildHandle(GET_MEMBER_NAME_CHECKED(FMetasoundFrontendLiteral, AsUObject));
+				check(IsPropertyHandleOfType(InLiteralPropertyHandle, GET_STRUCT_NAME_CHECKED(FMetasoundFrontendLiteral)));
+
+				// Set the FMetasoundFrontendLiteral::Type member if invalid // TODO: only do this if invalid.
+				Metasound::Editor::SetActiveLiteralType(GetLiteralType(), InLiteralPropertyHandle);
+
+				return CreatePropertyValueWidgetInternal(InLiteralPropertyHandle, bDisplayDefaultPropertyButtons);
 			}
 
-			case EMetasoundFrontendLiteralType::UObjectArray:
+			EMetasoundFrontendLiteralType GetLiteralType() const override
 			{
-				return StructPropertyHandle->GetChildHandle(GET_MEMBER_NAME_CHECKED(FMetasoundFrontendLiteral, AsUObjectArray));
+				return PreferredLiteralType;
 			}
 
-			case EMetasoundFrontendLiteralType::None:
-			default:
-				static_assert(static_cast<int32>(EMetasoundFrontendLiteralType::Invalid) == 7, "Possible missing case coverage for EMetasoundFrontendLiteralType");
-				return TSharedPtr<IPropertyHandle>();
+			const FText& GetToolTipText(TSharedPtr<IPropertyHandle> InLiteralPropertyHandle) const override
+			{
+				check(IsPropertyHandleOfType(InLiteralPropertyHandle, GET_STRUCT_NAME_CHECKED(FMetasoundFrontendLiteral)));
+
+				if (TSharedPtr<IPropertyHandle> Property = GetValueHandle(InLiteralPropertyHandle))
+				{
+					Property->GetToolTipText();
+				}
+				return FText::GetEmpty();
 			}
-		}
 
-		TSharedPtr<IPropertyHandle> GetActiveLiteralHandle(TSharedRef<IPropertyHandle> StructPropertyHandle, TSharedRef<IPropertyHandle> TypeEnumHandle)
-		{
-			uint8 Value = static_cast<uint8>(EMetasoundFrontendLiteralType::None);
-			if (TypeEnumHandle->GetValue(Value) == FPropertyAccess::Result::Success)
+		protected:
+
+			// Derived classes can override this method to produce a custom widget.
+			virtual TSharedRef<SWidget>CreatePropertyValueWidgetInternal(TSharedPtr<IPropertyHandle> InLiteralPropertyHandle, bool bDisplayDefaultPropertyButtons) const
 			{
-				const EMetasoundFrontendLiteralType LiteralType = static_cast<EMetasoundFrontendLiteralType>(Value);
-				return GetLiteralHandleForType(StructPropertyHandle, LiteralType);
+				if (TSharedPtr<IPropertyHandle> ValueProperty = GetValueHandle(InLiteralPropertyHandle))
+				{
+					return ValueProperty->CreatePropertyValueWidget(bDisplayDefaultPropertyButtons);
+				}
+				return SNullWidget::NullWidget;
 			}
 
-			return TSharedPtr<IPropertyHandle>();
-		}
+			// If `InPropertyHandle` is an array, this will return the first element in the array.
+			// If the array does not contain any elements, one will be added. 
+			TSharedPtr<IPropertyHandle> GetOrAddFirstArrayElement(TSharedPtr<IPropertyHandle> InPropertyHandle) const
+			{
+				if (InPropertyHandle.IsValid())
+				{	
+					TSharedPtr<IPropertyHandleArray> PropertyArray = InPropertyHandle->AsArray();
+					if (ensure(PropertyArray))
+					{
+						uint32 NumChildren = 0;
+						if (FPropertyAccess::Result::Success == PropertyArray->GetNumElements(NumChildren))
+						{
+							if (NumChildren < 1)
+							{
+								PropertyArray->AddItem();
+								if (FPropertyAccess::Result::Success != PropertyArray->GetNumElements(NumChildren))
+								{
+									return TSharedPtr<IPropertyHandle>();
+								}
+							}
 
-		bool IsLiteralTypeActive(TSharedRef<IPropertyHandle> TypeEnumHandle, EMetasoundFrontendLiteralType LiteralType)
-		{
-			uint8 Value = static_cast<uint8>(EMetasoundFrontendLiteralType::None);
-			TypeEnumHandle->GetValue(Value);
+							if (NumChildren > 0)
+							{
+								return PropertyArray->GetElement(0);
+							}
+						}
+					}
+				}
+				return TSharedPtr<IPropertyHandle>();
+			}
 
-			return Value == static_cast<uint8>(LiteralType);
-		}
-	} // namespace Editor
-} // namespace Metasound
+			// Returns the child IPropertyHandle from a FMetasoundFrontendLiteral for a specific EMetasoundFrontendLiteralType	
+			// The returned handle points to the actual stored value (e.g. float, TArray<float>, int32, UObject*)
+			TSharedPtr<IPropertyHandle> GetValueHandle(TSharedPtr<IPropertyHandle> InLiteralPropertyHandle) const
+			{
+				// Check that the property is the correct type.
+				switch (GetLiteralType())
+				{
+					case EMetasoundFrontendLiteralType::BooleanArray:
+					{
+						return InLiteralPropertyHandle->GetChildHandle(GET_MEMBER_NAME_CHECKED(FMetasoundFrontendLiteral, AsBoolean));
+					}
 
+					case EMetasoundFrontendLiteralType::Boolean:
+					{
+						return GetOrAddFirstArrayElement(InLiteralPropertyHandle->GetChildHandle(GET_MEMBER_NAME_CHECKED(FMetasoundFrontendLiteral, AsBoolean)));
+					}
 
-void FMetasoundFrontendLiteralDetailCustomization::CustomizeHeader(TSharedRef<IPropertyHandle> StructPropertyHandle, FDetailWidgetRow& HeaderRow, IPropertyTypeCustomizationUtils& StructCustomizationUtils)
-{
-}
+					case EMetasoundFrontendLiteralType::FloatArray:
+					{
+						return InLiteralPropertyHandle->GetChildHandle(GET_MEMBER_NAME_CHECKED(FMetasoundFrontendLiteral, AsFloat));
+					}
 
-void FMetasoundFrontendLiteralDetailCustomization::CustomizeChildren(TSharedRef<IPropertyHandle> StructPropertyHandle, IDetailChildrenBuilder& ChildBuilder, IPropertyTypeCustomizationUtils& StructCustomizationUtils)
-{
-	TSharedRef<IPropertyHandle> TypeEnumHandle = StructPropertyHandle->GetChildHandle(GET_MEMBER_NAME_CHECKED(FMetasoundFrontendLiteral, Type)).ToSharedRef();
+					case EMetasoundFrontendLiteralType::Float:
+					{
+						return GetOrAddFirstArrayElement(InLiteralPropertyHandle->GetChildHandle(GET_MEMBER_NAME_CHECKED(FMetasoundFrontendLiteral, AsFloat)));
+					}
 
+					case EMetasoundFrontendLiteralType::IntegerArray:
+					{
+						return InLiteralPropertyHandle->GetChildHandle(GET_MEMBER_NAME_CHECKED(FMetasoundFrontendLiteral, AsInteger));
+					}
 
-	FName OwningDataTypeName;
-	UClass* OwningDataTypeClass = nullptr;
-	EMetasoundFrontendLiteralType PreferredLiteralType = EMetasoundFrontendLiteralType::None;
+					case EMetasoundFrontendLiteralType::Integer:
+					{
+						return GetOrAddFirstArrayElement(InLiteralPropertyHandle->GetChildHandle(GET_MEMBER_NAME_CHECKED(FMetasoundFrontendLiteral, AsInteger)));
+					}
 
-	// Grab the preferred uclass for the owning data type.
-	TSharedPtr<IPropertyHandle> InputDescription = StructPropertyHandle->GetParentHandle()->GetParentHandle()->GetParentHandle();
+					case EMetasoundFrontendLiteralType::StringArray:
+					{
+						return InLiteralPropertyHandle->GetChildHandle(GET_MEMBER_NAME_CHECKED(FMetasoundFrontendLiteral, AsString));
+					}
 
-	// TODO: Ensure that StructPropertyHandle is an FMetasoundInputDescription.
-	if (InputDescription)
-	{
-		TSharedPtr<IPropertyHandle> DataTypeNameHandle = InputDescription->GetChildHandle(GET_MEMBER_NAME_CHECKED(FMetasoundFrontendClassInput, TypeName));
-		if (DataTypeNameHandle)
-		{
-			
-			DataTypeNameHandle->GetValue(OwningDataTypeName);
+					case EMetasoundFrontendLiteralType::String:
+					{
+						return GetOrAddFirstArrayElement(InLiteralPropertyHandle->GetChildHandle(GET_MEMBER_NAME_CHECKED(FMetasoundFrontendLiteral, AsString)));
+					}
 
-			FMetasoundFrontendRegistryContainer* Registry = FMetasoundFrontendRegistryContainer::Get();
-			OwningDataTypeClass = Registry->GetLiteralUClassForDataType(OwningDataTypeName);
-			PreferredLiteralType = Metasound::Frontend::GetMetasoundLiteralType(Registry->GetDesiredLiteralTypeForDataType(OwningDataTypeName));
-		}
-	}
+					case EMetasoundFrontendLiteralType::UObjectArray:
+					{
+						return InLiteralPropertyHandle->GetChildHandle(GET_MEMBER_NAME_CHECKED(FMetasoundFrontendLiteral, AsUObject));
+					}
 
-	
+					case EMetasoundFrontendLiteralType::UObject:
+					{
+						return GetOrAddFirstArrayElement(InLiteralPropertyHandle->GetChildHandle(GET_MEMBER_NAME_CHECKED(FMetasoundFrontendLiteral, AsUObject)));
+					}
 
-	struct FTypeEditorInfo
-	{
-		EMetasoundFrontendLiteralType LiteralType;
-		TSharedRef<IPropertyHandle> TypeEnumHandle;
-		TSharedPtr<IPropertyHandle> PropertyHandle;
-		FText DisplayName;
-		UClass* OwningDataTypeUClass;
-		FName OwningDataTypeName;
-
-		TSharedRef<SWidget> CreateEnumWidget(const TSharedPtr<const Metasound::Frontend::IEnumDataTypeInterface>& EnumInterface) const
+					case EMetasoundFrontendLiteralType::NoneArray:
+					{
+						return InLiteralPropertyHandle->GetChildHandle(GET_MEMBER_NAME_CHECKED(FMetasoundFrontendLiteral, AsNumDefault));
+					}
+					case EMetasoundFrontendLiteralType::None:
+					default:
+						static_assert(static_cast<int32>(EMetasoundFrontendLiteralType::Invalid) == 12, "Possible missing case coverage for EMetasoundFrontendLiteralType");
+						return TSharedPtr<IPropertyHandle>();
+				}
+			}
+
+		private:
+
+			EMetasoundFrontendLiteralType PreferredLiteralType;
+		};
+
+		// Special literal property adapater for Metasound enums.
+		class FLiteralEnumPropertyAdapter : public FLiteralPropertyBaseAdapter
 		{
-			auto RemoveNameSpace = [](const FString& InString) -> FString
-			{ 
-				FString Namespace, EnumName;
-				InString.Split(TEXT("::"), &Namespace, &EnumName);
-				return EnumName;
-			};
-			auto GetAll = [EnumInterface, RemoveNameSpace](TArray<TSharedPtr<FString>>& OutStrings, TArray<TSharedPtr<SToolTip>>& OutTooltips, TArray<bool>&)
+		public:
+			FLiteralEnumPropertyAdapter(EMetasoundFrontendLiteralType InLiteralType, TSharedRef<const Frontend::IEnumDataTypeInterface> InEnumInterface)
+			: FLiteralPropertyBaseAdapter(InLiteralType)
+			, EnumInterface(InEnumInterface)
 			{
-				for (const Metasound::Frontend::IEnumDataTypeInterface::FGenericInt32Entry& i : EnumInterface->GetAllEntries())
-				{
-					OutTooltips.Emplace(SNew(SToolTip).Text(i.Tooltip));
+			}
 
-					// Trim the namespace off for display.
-					OutStrings.Emplace(MakeShared<FString>(RemoveNameSpace(i.Name.GetPlainNameString())));
-				}
-			};
-			auto GetValue = [EnumInterface, Prop = PropertyHandle, RemoveNameSpace]() -> FString
+			virtual ~FLiteralEnumPropertyAdapter() = default;
+
+		protected:
+			// Create enum widget.
+			TSharedRef<SWidget> CreatePropertyValueWidgetInternal(TSharedPtr<IPropertyHandle> InLiteralPropertyHandle, bool bDisplayDefaultPropertyButtons = false) const override
 			{
-				int32 IntValue;
-				if (Prop->GetValue(IntValue))
+				if (TSharedPtr<IPropertyHandle> PropertyHandle = GetValueHandle(InLiteralPropertyHandle))
 				{
-					if (TOptional<FName> Result = EnumInterface->ToName(IntValue))
+					auto RemoveNamespace = [](const FString& InString) -> FString
+					{ 
+						FString EnumName;
+						if (InString.Split(TEXT("::"), nullptr, &EnumName))
+						{
+							return EnumName;
+						}
+						return InString;
+					};
+
+					auto GetAll = [EnumInterface = EnumInterface, RemoveNamespace](TArray<TSharedPtr<FString>>& OutStrings, TArray<TSharedPtr<SToolTip>>& OutTooltips, TArray<bool>&)
 					{
-						return RemoveNameSpace(Result->GetPlainNameString());
-					}
-					UE_LOG(LogTemp, Warning, TEXT("Failed to Get Valid Value for Property '%s' with Value of '%d'"),
-						*GetNameSafe(Prop->GetProperty()), IntValue);
-				}
-				return {};
-			};
-			auto SelectedValue = [EnumInterface, Prop = PropertyHandle](const FString& InSelected)
-			{
-				FString FullyQualifiedName = FString::Printf(TEXT("%s::%s"), *EnumInterface->GetNamespace().GetPlainNameString(), *InSelected);
-				TOptional<int32> Result = EnumInterface->ToValue(*FullyQualifiedName);
-				if (ensure(Result))
-				{
-					Prop->SetValue(*Result);
+						for (const Metasound::Frontend::IEnumDataTypeInterface::FGenericInt32Entry& i : EnumInterface->GetAllEntries())
+						{
+							OutTooltips.Emplace(SNew(SToolTip).Text(i.Tooltip));
+
+							// Trim the namespace off for display.
+							OutStrings.Emplace(MakeShared<FString>(RemoveNamespace(i.Name.GetPlainNameString())));
+						}
+					};
+
+					auto GetValue = [EnumInterface = EnumInterface, PropertyHandle = PropertyHandle, RemoveNamespace]() -> FString
+					{
+						int32 IntValue;
+						if (PropertyHandle->GetValue(IntValue))
+						{
+							if (TOptional<FName> Result = EnumInterface->ToName(IntValue))
+							{
+								return RemoveNamespace(Result->GetPlainNameString());
+							}
+							UE_LOG(LogTemp, Warning, TEXT("Failed to Get Valid Value for Property '%s' with Value of '%d'"), *GetNameSafe(PropertyHandle->GetProperty()), IntValue);
+						}
+						return {};
+					};
+
+					auto SelectedValue = [EnumInterface = EnumInterface, PropertyHandle = PropertyHandle](const FString& InSelected)
+					{
+						FString FullyQualifiedName = FString::Printf(TEXT("%s::%s"), *EnumInterface->GetNamespace().GetPlainNameString(), *InSelected);
+						TOptional<int32> Result = EnumInterface->ToValue(*FullyQualifiedName);
+						if (ensure(Result))
+						{
+							PropertyHandle->SetValue(*Result);
+						}
+					};
+
+					return PropertyCustomizationHelpers::MakePropertyComboBox(
+						nullptr,
+						FOnGetPropertyComboBoxStrings::CreateLambda(GetAll),
+						FOnGetPropertyComboBoxValue::CreateLambda(GetValue),
+						FOnPropertyComboBoxValueSelected::CreateLambda(SelectedValue)
+					);
 				}
-			};
-
-			return PropertyCustomizationHelpers::MakePropertyComboBox(
-				nullptr,
-				FOnGetPropertyComboBoxStrings::CreateLambda(GetAll),
-				FOnGetPropertyComboBoxValue::CreateLambda(GetValue),
-				FOnPropertyComboBoxValueSelected::CreateLambda(SelectedValue)
-			);
-		}
 
-		TSharedRef<SWidget> CreatePropertyValueWidget(bool bDisplayDefaultPropertyButtons = true) const
+				return SNullWidget::NullWidget;
+			}
+
+		private:
+			TSharedRef<const Frontend::IEnumDataTypeInterface> EnumInterface;
+		};
+
+
+		// Special property adapter for UObject
+		class FLiteralUObjectPropertyAdapter : public FLiteralPropertyBaseAdapter
 		{
-			if (LiteralType == EMetasoundFrontendLiteralType::UObject)
+		public:
+			FLiteralUObjectPropertyAdapter(UClass* InValidUClass)
+			: FLiteralPropertyBaseAdapter(EMetasoundFrontendLiteralType::UObject)
+			, ValidUClass(InValidUClass)
+			{
+			}
+
+			virtual ~FLiteralUObjectPropertyAdapter() = default;
+
+		protected:
+			TSharedRef<SWidget> CreatePropertyValueWidgetInternal(TSharedPtr<IPropertyHandle> InLiteralPropertyHandle, bool bDisplayDefaultPropertyButtons = false) const override
 			{
-				auto ValidateAsset = [CachedDataTypeUClass = OwningDataTypeUClass](const UObject* InObject, FText& OutReason)  -> bool
+				TSharedPtr<IPropertyHandle> PropertyHandle = GetValueHandle(InLiteralPropertyHandle);
+
+				auto ValidateAsset = [ValidUClass = this->ValidUClass](const UObject* InObject, FText& OutReason)  -> bool
 				{
 					UClass* SelectedObjectClass = InObject->GetClass();
-					return CachedDataTypeUClass && SelectedObjectClass && SelectedObjectClass->IsChildOf(CachedDataTypeUClass);
+					return ValidUClass && SelectedObjectClass && SelectedObjectClass->IsChildOf(ValidUClass);
 				};
 
-				auto CommitAsset = [CachedLiteralType = LiteralType, CachedPropertyHandle = PropertyHandle, CachedEnumHandle = TypeEnumHandle, CachedDataTypeUClass = OwningDataTypeUClass](const FAssetData& InAssetData) -> void
+				auto CommitAsset = [PropertyHandle = PropertyHandle, ValidUClass = this->ValidUClass](const FAssetData& InAssetData) -> void
 				{
 					// if we've hit this code, the presumption is that the datatype for this parameter has already defined a corresponding UClass that can be used to set it.
-					ensureAlways(CachedDataTypeUClass && CachedDataTypeUClass != UObject::StaticClass());
+					ensureAlways(ValidUClass && ValidUClass != UObject::StaticClass());
 
 					UObject* InObject = InAssetData.GetAsset();
-					CachedPropertyHandle->SetValue(InObject);
-
-					// If this asset selector was set to no asset, clear out the literal description.
-					if (InObject)
-					{
-						CachedEnumHandle->SetValue(static_cast<uint8>(EMetasoundFrontendLiteralType::UObject));
-					}
-					else
-					{
-						CachedEnumHandle->SetValue(static_cast<uint8>(EMetasoundFrontendLiteralType::None));
-					}
+					PropertyHandle->SetValue(InObject);
 				};
 
-				auto GetAssetPath = [CachedPropertyHandle = PropertyHandle]() -> FString
+				auto GetAssetPath = [PropertyHandle = PropertyHandle]() -> FString
 				{
 					UObject* Obj = nullptr;
-					CachedPropertyHandle->GetValue(Obj);
-					return  Obj != nullptr ? Obj->GetPathName() : FString();
+
+					PropertyHandle->GetValue(Obj);
+
+					if (nullptr != Obj)
+					{
+						return Obj->GetPathName();
+					}
+					return FString();
 				};
 
 				// TODO: get the UFactory corresponding to this datatype's UClass.
@@ -225,7 +414,7 @@ void FMetasoundFrontendLiteralDetailCustomization::CustomizeChildren(TSharedRef<
 
 				return SNew(SObjectPropertyEntryBox)
 					.ObjectPath_Lambda(GetAssetPath)
-					.AllowedClass(OwningDataTypeUClass)
+					.AllowedClass(ValidUClass)
 					// .OnShouldSetAsset_Lambda(ValidateAsset)
 					.OnObjectChanged_Lambda(CommitAsset)
 					.AllowClear(false)
@@ -234,137 +423,115 @@ void FMetasoundFrontendLiteralDetailCustomization::CustomizeChildren(TSharedRef<
 					.DisplayThumbnail(true)
 					.NewAssetFactories(FactoriesToUse);
 			}
-			else if (LiteralType == EMetasoundFrontendLiteralType::Integer)
-			{
-				// Check if this is an Enum.
-				if (TSharedPtr<const Metasound::Frontend::IEnumDataTypeInterface> EnumInterface = 
-					FMetasoundFrontendRegistryContainer::Get()->GetEnumInterfaceForDataType(OwningDataTypeName))
-				{ 
-					return CreateEnumWidget(EnumInterface);
-				}
 
-				// Fall back to regular widget if its not.
-				return PropertyHandle->CreatePropertyValueWidget();
-			}
-			else if (LiteralType == EMetasoundFrontendLiteralType::UObjectArray)
+		private:
+			UClass* ValidUClass;
+		};
+
+		// Return the literal property adapter associated with the Metasound data type.
+		TSharedPtr<const ILiteralPropertyAdapter> GetLiteralPropertyAdapter(const FName& InDataTypeName)
+		{
+			static TMap<FName, TSharedPtr<ILiteralPropertyAdapter>> AdapterMap;
+
+			// Return an adapter if one already exists for data type
+			TSharedPtr<ILiteralPropertyAdapter> Adapter;
+			if (TSharedPtr<ILiteralPropertyAdapter>* PtrToExistingAdapter = AdapterMap.Find(InDataTypeName))
 			{
-				// TODO: Implement.
-				return SNullWidget::NullWidget;
+				Adapter = *PtrToExistingAdapter;
 			}
-			else
+
+			if ((!Adapter.IsValid()) && (!InDataTypeName.IsNone()))
 			{
-				return PropertyHandle->CreatePropertyValueWidget();
+				// There was no existing adapter, so query the registry to see which adapter to instantiate. 
+				if (FMetasoundFrontendRegistryContainer* Registry = FMetasoundFrontendRegistryContainer::Get())
+				{
+					FDataTypeRegistryInfo DataTypeInfo;
+					if (ensure(Registry->GetInfoForDataType(InDataTypeName, DataTypeInfo)))
+					{
+						const EMetasoundFrontendLiteralType LiteralType = Frontend::GetMetasoundFrontendLiteralType(DataTypeInfo.PreferredLiteralType);
+
+						if (EMetasoundFrontendLiteralType::UObject == LiteralType)
+						{
+							Adapter = MakeShared<FLiteralUObjectPropertyAdapter>(DataTypeInfo.ProxyGeneratorClass);
+						}
+						else if (EMetasoundFrontendLiteralType::UObjectArray == LiteralType)
+						{
+							// TODO: implement me.
+							ensure(false);
+						}
+						else if (DataTypeInfo.bIsEnum)
+						{
+							TSharedPtr<const Frontend::IEnumDataTypeInterface> EnumInterface = Registry->GetEnumInterfaceForDataType(InDataTypeName);
+							if (ensure(EnumInterface))
+							{
+								Adapter = MakeShared<FLiteralEnumPropertyAdapter>(LiteralType, EnumInterface.ToSharedRef());
+							}
+						}
+						else
+						{
+							Adapter = MakeShared<FLiteralPropertyBaseAdapter>(LiteralType);
+						}
+
+						if (Adapter.IsValid())
+						{
+							AdapterMap.Add(InDataTypeName, Adapter);
+						}
+					}
+				}
 			}
-		}
-	};
 
-	static_assert(static_cast<int32>(EMetasoundFrontendLiteralType::Invalid) == 7, "Possible missing property descriptor for literal customization display");
-	const TArray<FTypeEditorInfo> TypePropertyInfo =
-	{
-		FTypeEditorInfo
-		{
-			EMetasoundFrontendLiteralType::Bool,
-			TypeEnumHandle,
-			StructPropertyHandle->GetChildHandle(GET_MEMBER_NAME_CHECKED(FMetasoundFrontendLiteral, AsBool)),
-			LOCTEXT("Metasound_LiteralDisplayNameBool", "Bool"),
-			nullptr,
-			OwningDataTypeName
-		},
-
-		FTypeEditorInfo
-		{
-			EMetasoundFrontendLiteralType::Integer,
-			TypeEnumHandle,
-			StructPropertyHandle->GetChildHandle(GET_MEMBER_NAME_CHECKED(FMetasoundFrontendLiteral, AsInteger)),
-			LOCTEXT("Metasound_LiteralDisplayNameInteger", "Int32"),
-			nullptr,
-			OwningDataTypeName
-		},
-
-		FTypeEditorInfo
-		{
-			EMetasoundFrontendLiteralType::Float,
-			TypeEnumHandle,
-			StructPropertyHandle->GetChildHandle(GET_MEMBER_NAME_CHECKED(FMetasoundFrontendLiteral, AsFloat)),
-			LOCTEXT("Metasound_LiteralDisplayNameFloat", "Float"),
-			nullptr,
-			OwningDataTypeName
-		},
-
-		FTypeEditorInfo
-		{
-			EMetasoundFrontendLiteralType::String,
-			TypeEnumHandle,
-			StructPropertyHandle->GetChildHandle(GET_MEMBER_NAME_CHECKED(FMetasoundFrontendLiteral, AsString)),
-			LOCTEXT("Metasound_LiteralDisplayNameString", "String"),
-			nullptr,
-			OwningDataTypeName
-		},
-
-		FTypeEditorInfo
-		{
-			EMetasoundFrontendLiteralType::UObject,
-			TypeEnumHandle,
-			StructPropertyHandle->GetChildHandle(GET_MEMBER_NAME_CHECKED(FMetasoundFrontendLiteral, AsUObject)),
-			LOCTEXT("Metasound_LiteralDisplayNameUObject", "UObject"),
-			OwningDataTypeClass,
-			OwningDataTypeName
-		},
-
-		FTypeEditorInfo
-		{
-			EMetasoundFrontendLiteralType::UObjectArray,
-			TypeEnumHandle,
-			StructPropertyHandle->GetChildHandle(GET_MEMBER_NAME_CHECKED(FMetasoundFrontendLiteral, AsUObjectArray)),
-			LOCTEXT("Metasound_LiteralDisplayNameUObjectArray", "UObjectArray"),
-			OwningDataTypeClass,
-			OwningDataTypeName
+			return Adapter;
 		}
-	};
 
-	for (const FTypeEditorInfo& Info : TypePropertyInfo)
+	} // namespace Editor
+} // namespace Metasound
+
+
+void FMetasoundFrontendLiteralDetailCustomization::CustomizeHeader(TSharedRef<IPropertyHandle> StructPropertyHandle, FDetailWidgetRow& HeaderRow, IPropertyTypeCustomizationUtils& StructCustomizationUtils)
+{
+}
+
+void FMetasoundFrontendLiteralDetailCustomization::CustomizeChildren(TSharedRef<IPropertyHandle> StructPropertyHandle, IDetailChildrenBuilder& ChildBuilder, IPropertyTypeCustomizationUtils& StructCustomizationUtils)
+{
+	// Get the datatype associated with value
+	FName DataTypeName = Metasound::Editor::GetOwningMetasoundDataTypeName(StructPropertyHandle);
+
+	// Get the literal property adapter associated with the datatype.
+	if (TSharedPtr<const Metasound::Editor::ILiteralPropertyAdapter> Adapter = Metasound::Editor::GetLiteralPropertyAdapter(DataTypeName))
 	{
+		// Create customization.
 		ChildBuilder.AddCustomRow(StructPropertyHandle->GetPropertyDisplayName())
 			.NameContent()
 			[
 				SNew(SHorizontalBox)
 				+ SHorizontalBox::Slot()
-					.FillWidth(0.5f)
-					.Padding(1.0f, 0.0f, 0.0f, 0.0f)
-					.VAlign(VAlign_Center)
-					.HAlign(HAlign_Right)
-					[
-						SNew(STextBlock)
-						.Font(IDetailLayoutBuilder::GetDetailFont())
-						.Text(LOCTEXT("Metasound_LiteralDefault", "Default"))
-						.ToolTipText(TAttribute<FText>::Create([StructPropertyHandle, TypeEnumHandle]()
+				.FillWidth(0.5f)
+				.Padding(1.0f, 0.0f, 0.0f, 0.0f)
+				.VAlign(VAlign_Center)
+				.HAlign(HAlign_Right)
+				[
+					SNew(STextBlock)
+					.Font(IDetailLayoutBuilder::GetDetailFont())
+					.Text(LOCTEXT("Metasound_LiteralDefault", "Default"))
+					.ToolTipText(TAttribute<FText>::Create([=]()
 						{
-							TSharedPtr<IPropertyHandle> VisibleHandle = Metasound::Editor::GetActiveLiteralHandle(StructPropertyHandle, TypeEnumHandle);
-							if (VisibleHandle.IsValid())
-							{
-								return VisibleHandle->GetToolTipText();
-							}
-							return FText::GetEmpty();
-						}))
-					]
+							return Adapter->GetToolTipText(StructPropertyHandle);
+						})
+					)
+				]
 			]
 			.ValueContent()
 			.MinDesiredWidth(120.0f)
 			[
 				SNew(SHorizontalBox)
 				+ SHorizontalBox::Slot()
-					.FillWidth(0.5f)
-					.Padding(1.0f, 0.0f, 0.0f, 0.0f)
-					.VAlign(VAlign_Center)
-					[
-						Info.CreatePropertyValueWidget()
-					]
-			]
-			.Visibility(TAttribute<EVisibility>::Create([PreferredLiteralType, LiteralType = Info.LiteralType]()
-			{
-				const bool bIsLiteralActive = LiteralType == PreferredLiteralType;
-				return bIsLiteralActive ? EVisibility::Visible : EVisibility::Hidden;
-			}));
+				.FillWidth(0.5f)
+				.Padding(1.0f, 0.0f, 0.0f, 0.0f)
+				.VAlign(VAlign_Center)
+				[
+					Adapter->CreatePropertyValueWidget(StructPropertyHandle)
+				]
+			];
 	}
 }
 #undef LOCTEXT_NAMESPACE // MetasoundEditor
diff --git a/Engine/Plugins/Runtime/Metasound/Source/MetasoundEngine/Private/MetasoundSource.cpp b/Engine/Plugins/Runtime/Metasound/Source/MetasoundEngine/Private/MetasoundSource.cpp
--- a/Engine/Plugins/Runtime/Metasound/Source/MetasoundEngine/Private/MetasoundSource.cpp
+++ b/Engine/Plugins/Runtime/Metasound/Source/MetasoundEngine/Private/MetasoundSource.cpp
@@ -137,18 +137,18 @@ ISoundGeneratorPtr UMetasoundSource::CreateSoundGenerator(const FSoundGeneratorI
 	TArray<IOperatorBuilder::FBuildErrorPtr> BuildErrors;
 	TUniquePtr<IOperator> Operator = RootGraph->BuildOperator(InSettings, Environment, BuildErrors);
 
-	if (!Operator.IsValid())
+	// Log build errors
+	for (const IOperatorBuilder::FBuildErrorPtr& Error : BuildErrors)
 	{
-		// Log build errors that resulted in a null operator.
-		UE_LOG(LogMetasound, Error, TEXT("Failed to build Metasound operator from graph in MetasoundSource [%s]"), *GetName());
-		for (const IOperatorBuilder::FBuildErrorPtr& Error : BuildErrors)
+		if (Error.IsValid())
 		{
-			if (Error.IsValid())
-			{
-				UE_LOG(LogMetasound, Warning, TEXT("MetasoundSource [%s] build error [%s] \"%s\""), *GetName(), *(Error->GetErrorType().ToString()), *(Error->GetErrorDescription().ToString()));
-			}
+			UE_LOG(LogMetasound, Warning, TEXT("MetasoundSource [%s] build error [%s] \"%s\""), *GetName(), *(Error->GetErrorType().ToString()), *(Error->GetErrorDescription().ToString()));
 		}
 	}
+	if (!Operator.IsValid())
+	{
+		UE_LOG(LogMetasound, Error, TEXT("Failed to build Metasound operator from graph in MetasoundSource [%s]"), *GetName());
+	}
 	else
 	{
 		FDataReferenceCollection Outputs = Operator->GetOutputs();
diff --git a/Engine/Plugins/Runtime/Metasound/Source/MetasoundEngine/Private/MetasoundWavePlayerNode.cpp b/Engine/Plugins/Runtime/Metasound/Source/MetasoundEngine/Private/MetasoundWavePlayerNode.cpp
--- a/Engine/Plugins/Runtime/Metasound/Source/MetasoundEngine/Private/MetasoundWavePlayerNode.cpp
+++ b/Engine/Plugins/Runtime/Metasound/Source/MetasoundEngine/Private/MetasoundWavePlayerNode.cpp
@@ -157,50 +157,58 @@ namespace Metasound
 
 		void ExecuteInternal(int32 StartFrame, int32 EndFrame)
 		{
-			// shouldn't be calling this function if we don't have access to a valid SoundWave
-			ensure(Wave->IsSoundWaveValid());
-			ensure(Wave->SoundWaveProxy->GetNumChannels() <= 2); // only support mono or stereo inputs
-			ensure(Decoder.CanGenerateAudio());
+			bool bCanDecodeWave = Wave->IsSoundWaveValid() && (Wave->SoundWaveProxy->GetNumChannels() <= 2); // only support mono or stereo inputs
 
 			// note: output is hard-coded to stereo (dual-mono)
 			float* FinalOutputLeft = AudioBufferL->GetData() + StartFrame;
 			float* FinalOutputRight = AudioBufferR->GetData() + StartFrame;
+			const int32 NumOutputFrames = (EndFrame - StartFrame);
 
-			const int32 NumInputChannels = Wave->SoundWaveProxy->GetNumChannels();
-			const bool bNeedsUpmix = (NumInputChannels == 1);
-			const bool bNeedsDeinterleave = !bNeedsUpmix;
+			if (bCanDecodeWave)
+			{
+				ensure(Decoder.CanGenerateAudio());
 
-			const int32 NumOutputFrames = (EndFrame - StartFrame);
-			const int32 NumSamplesToGenerate = NumOutputFrames * NumInputChannels; // (stereo output)
 
-			PostSrcBuffer.Reset(NumSamplesToGenerate);
-			PostSrcBuffer.AddZeroed(NumSamplesToGenerate);
-			float* PostSrcBufferPtr = PostSrcBuffer.GetData();
+				const int32 NumInputChannels = Wave->SoundWaveProxy->GetNumChannels();
+				const bool bNeedsUpmix = (NumInputChannels == 1);
+				const bool bNeedsDeinterleave = !bNeedsUpmix;
 
-			int32 NumFramesDecoded = Decoder.GenerateAudio(PostSrcBufferPtr, NumOutputFrames, *PitchShiftCents); // TODO: pitch shift
+				const int32 NumSamplesToGenerate = NumOutputFrames * NumInputChannels; // (stereo output)
 
-			// TODO: handle decoder having completed during it's decode
-			if (!Decoder.CanGenerateAudio() ||  (NumFramesDecoded < NumOutputFrames))
-			{
-				bIsPlaying = false;
-				TrigggerOnDone->TriggerFrame(StartFrame + NumFramesDecoded);
-			}
+				PostSrcBuffer.Reset(NumSamplesToGenerate);
+				PostSrcBuffer.AddZeroed(NumSamplesToGenerate);
+				float* PostSrcBufferPtr = PostSrcBuffer.GetData();
 
-			if (bNeedsUpmix)
-			{
-				// TODO: attenuate by -3 dB?
-				// copy to Left & Right output buffers
-				FMemory::Memcpy(FinalOutputLeft, PostSrcBufferPtr, sizeof(float) * NumOutputFrames);
-				FMemory::Memcpy(FinalOutputRight, PostSrcBufferPtr, sizeof(float) * NumOutputFrames);
-			}
-			else if (bNeedsDeinterleave)
-			{
-				for (int32 i = 0; i < NumOutputFrames; ++i)
+				int32 NumFramesDecoded = Decoder.GenerateAudio(PostSrcBufferPtr, NumOutputFrames, *PitchShiftCents); // TODO: pitch shift
+
+				// TODO: handle decoder having completed during it's decode
+				if (!Decoder.CanGenerateAudio() ||  (NumFramesDecoded < NumOutputFrames))
 				{
-					// de-interleave each stereo frame into output buffers
-					FinalOutputLeft[i] = PostSrcBufferPtr[(i << 1)];
-					FinalOutputRight[i] = PostSrcBufferPtr[(i << 1) + 1];
+					bIsPlaying = false;
+					TrigggerOnDone->TriggerFrame(StartFrame + NumFramesDecoded);
 				}
+
+				if (bNeedsUpmix)
+				{
+					// TODO: attenuate by -3 dB?
+					// copy to Left & Right output buffers
+					FMemory::Memcpy(FinalOutputLeft, PostSrcBufferPtr, sizeof(float) * NumOutputFrames);
+					FMemory::Memcpy(FinalOutputRight, PostSrcBufferPtr, sizeof(float) * NumOutputFrames);
+				}
+				else if (bNeedsDeinterleave)
+				{
+					for (int32 i = 0; i < NumOutputFrames; ++i)
+					{
+						// de-interleave each stereo frame into output buffers
+						FinalOutputLeft[i] = PostSrcBufferPtr[(i << 1)];
+						FinalOutputRight[i] = PostSrcBufferPtr[(i << 1) + 1];
+					}
+				}
+			}
+			else
+			{
+				FMemory::Memzero(FinalOutputLeft, sizeof(float) * NumOutputFrames);
+				FMemory::Memzero(FinalOutputRight, sizeof(float) * NumOutputFrames);
 			}
 		}
 
diff --git a/Engine/Plugins/Runtime/Metasound/Source/MetasoundFrontend/Private/MetasoundFrontend.cpp b/Engine/Plugins/Runtime/Metasound/Source/MetasoundFrontend/Private/MetasoundFrontend.cpp
--- a/Engine/Plugins/Runtime/Metasound/Source/MetasoundFrontend/Private/MetasoundFrontend.cpp
+++ b/Engine/Plugins/Runtime/Metasound/Source/MetasoundFrontend/Private/MetasoundFrontend.cpp
@@ -136,42 +136,6 @@ namespace Metasound
 			}
 		}
 
-		EMetasoundFrontendLiteralType GetMetasoundLiteralType(Metasound::ELiteralType InLiteralType)
-		{
-			switch (InLiteralType)
-			{
-				case Metasound::ELiteralType::Boolean:
-				{
-					return EMetasoundFrontendLiteralType::Bool;
-				}
-				case Metasound::ELiteralType::Integer:
-				{
-					return EMetasoundFrontendLiteralType::Integer;
-				}
-				case Metasound::ELiteralType::Float:
-				{
-					return EMetasoundFrontendLiteralType::Float;
-				}
-				case Metasound::ELiteralType::String:
-				{
-					return EMetasoundFrontendLiteralType::String;
-				}
-				case Metasound::ELiteralType::UObjectProxy:
-				{
-					return EMetasoundFrontendLiteralType::UObject;
-				}
-				case Metasound::ELiteralType::UObjectProxyArray:
-				{
-					return EMetasoundFrontendLiteralType::UObjectArray;
-				}
-				case Metasound::ELiteralType::None:
-				default:
-				{
-					return EMetasoundFrontendLiteralType::None;
-				}
-			}
-		}
-
 		TArray<FName> GetAllAvailableDataTypes()
 		{
 			return FMetasoundFrontendRegistryContainer::Get()->GetAllValidDataTypes();
@@ -229,6 +193,8 @@ REGISTER_METASOUND_DATATYPE(float, "Float", ::Metasound::ELiteralType::Float)
 REGISTER_METASOUND_DATATYPE(double, "Double", ::Metasound::ELiteralType::Float)
 REGISTER_METASOUND_DATATYPE(FString, "String", ::Metasound::ELiteralType::String)
 
+REGISTER_METASOUND_DATATYPE(TArray<float>, "FloatArray", ::Metasound::ELiteralType::FloatArray)
+
 REGISTER_METASOUND_DATATYPE(Metasound::FAudioBuffer, "Audio:Buffer")
 REGISTER_METASOUND_DATATYPE(Metasound::FSendAddress, "Transmission:Address")
 
diff --git a/Engine/Plugins/Runtime/Metasound/Source/MetasoundFrontend/Private/MetasoundFrontendBaseClasses.cpp b/Engine/Plugins/Runtime/Metasound/Source/MetasoundFrontend/Private/MetasoundFrontendBaseClasses.cpp
--- a/Engine/Plugins/Runtime/Metasound/Source/MetasoundFrontend/Private/MetasoundFrontendBaseClasses.cpp
+++ b/Engine/Plugins/Runtime/Metasound/Source/MetasoundFrontend/Private/MetasoundFrontendBaseClasses.cpp
@@ -31,273 +31,6 @@ namespace Metasound
 			return FMetasoundFrontendRegistryContainer::Get()->ConstructExternalNode(InNodeType, InNodeHash, InInitData);
 		}
 
-#if 0
-		void SetLiteralDescription(FMetasoundFrontendLiteral& OutDescription, bool InValue)
-		{
-			OutDescription.LiteralType = EMetasoundFrontendLiteralType::Bool;
-			OutDescription.AsBool = InValue;
-			OutDescription.AsInteger = 0;
-			OutDescription.AsFloat = 0.0f;
-			OutDescription.AsString.Empty();
-			OutDescription.AsUObject = nullptr;
-			OutDescription.AsUObjectArray.Empty();
-		}
-
-		void SetLiteralDescription(FMetasoundFrontendLiteral& OutDescription, int32 InValue)
-		{
-			OutDescription.LiteralType = EMetasoundFrontendLiteralType::Integer;
-			OutDescription.AsBool = false;
-			OutDescription.AsInteger = InValue;
-			OutDescription.AsFloat = 0.0f;
-			OutDescription.AsString.Empty();
-			OutDescription.AsUObject = nullptr;
-			OutDescription.AsUObjectArray.Empty();
-		}
-
-		void SetLiteralDescription(FMetasoundFrontendLiteral& OutDescription, float InValue)
-		{
-			OutDescription.LiteralType = EMetasoundFrontendLiteralType::Float;
-			OutDescription.AsBool = false;
-			OutDescription.AsInteger = 0;
-			OutDescription.AsFloat = InValue;
-			OutDescription.AsString.Empty();
-			OutDescription.AsUObject = nullptr;
-			OutDescription.AsUObjectArray.Empty();
-		}
-
-		void SetLiteralDescription(FMetasoundFrontendLiteral& OutDescription, const FString& InValue)
-		{
-			OutDescription.LiteralType = EMetasoundFrontendLiteralType::String;
-			OutDescription.AsBool = false;
-			OutDescription.AsInteger = 0;
-			OutDescription.AsFloat = 0.0f;
-			OutDescription.AsString = InValue;
-			OutDescription.AsUObject = nullptr;
-			OutDescription.AsUObjectArray.Empty();
-		}
-
-		void SetLiteralDescription(FMetasoundFrontendLiteral& OutDescription, UObject* InValue)
-		{
-			OutDescription.LiteralType = EMetasoundFrontendLiteralType::UObject;
-			OutDescription.AsBool = false;
-			OutDescription.AsInteger = 0;
-			OutDescription.AsFloat = 0.0f;
-			OutDescription.AsString.Empty();
-			OutDescription.AsUObject = InValue;
-			OutDescription.AsUObjectArray.Empty();
-
-		}
-
-		void SetLiteralDescription(FMetasoundFrontendLiteral& OutDescription, const TArray<UObject*>& InValue)
-		{
-			OutDescription.LiteralType = EMetasoundFrontendLiteralType::UObjectArray;
-			OutDescription.AsBool = false;
-			OutDescription.AsInteger = 0;
-			OutDescription.AsFloat = 0.0f;
-			OutDescription.AsString.Empty();
-			OutDescription.AsUObject = nullptr;
-			OutDescription.AsUObjectArray = InValue;
-		}
-
-		void ClearLiteralDescription(FMetasoundFrontendLiteral& OutDescription)
-		{
-			OutDescription.LiteralType = EMetasoundFrontendLiteralType::None;
-			OutDescription.AsBool = false;
-			OutDescription.AsInteger = 0;
-			OutDescription.AsFloat = 0.0f;
-			OutDescription.AsString.Empty();
-			OutDescription.AsUObject = nullptr;
-			OutDescription.AsUObjectArray.Empty();
-		}
-#endif
-
-		FLiteral GetLiteralParamForDataType(FName InDataType, const FMetasoundFrontendLiteral& InDescription)
-		{
-			EMetasoundFrontendLiteralType LiteralType = InDescription.Type;
-			FMetasoundFrontendRegistryContainer* Registry = FMetasoundFrontendRegistryContainer::Get();
-
-			switch (LiteralType)
-			{
-				case EMetasoundFrontendLiteralType::None:
-				{
-					return GetDefaultParamForDataType(InDataType);
-				}
-				case EMetasoundFrontendLiteralType::Bool:
-				{
-					if (!Registry->DoesDataTypeSupportLiteralType(InDataType, ELiteralType::Boolean))
-					{
-						return FLiteral::CreateInvalid();
-					}
-					else
-					{
-						return FLiteral(InDescription.AsBool);
-					}
-				}
-				case EMetasoundFrontendLiteralType::Float:
-				{
-					if (!Registry->DoesDataTypeSupportLiteralType(InDataType, ELiteralType::Float))
-					{
-						return FLiteral::CreateInvalid();
-					}
-					else
-					{
-						return FLiteral(InDescription.AsFloat);
-					}
-				}
-				case EMetasoundFrontendLiteralType::Integer:
-				{
-					if (!Registry->DoesDataTypeSupportLiteralType(InDataType, ELiteralType::Integer))
-					{
-						return FLiteral::CreateInvalid();
-					}
-					else
-					{
-						return FLiteral(InDescription.AsInteger);
-					}
-				}
-				case EMetasoundFrontendLiteralType::String:
-				{
-					if (!Registry->DoesDataTypeSupportLiteralType(InDataType, ELiteralType::String))
-					{
-						return FLiteral::CreateInvalid();
-					}
-					else
-					{
-						return FLiteral(InDescription.AsString);
-					}
-				}
-				case EMetasoundFrontendLiteralType::UObject:
-				{
-					if (!Registry->DoesDataTypeSupportLiteralType(InDataType, ELiteralType::UObjectProxy))
-					{
-						return FLiteral::CreateInvalid();
-					}
-					else
-					{
-						return Registry->GenerateLiteralForUObject(InDataType, InDescription.AsUObject);
-					}
-				}
-				case EMetasoundFrontendLiteralType::UObjectArray:
-				{
-					if (!Registry->DoesDataTypeSupportLiteralType(InDataType, ELiteralType::UObjectProxyArray))
-					{
-						return FLiteral::CreateInvalid();
-					}
-					else
-					{
-						return Registry->GenerateLiteralForUObjectArray(InDataType, InDescription.AsUObjectArray);
-					}
-				}
-				case EMetasoundFrontendLiteralType::Invalid:
-				default:
-				{
-					return FLiteral::CreateInvalid();
-				}
-			}
-		}
-
-		bool DoesDataTypeSupportLiteralType(FName InDataType, EMetasoundFrontendLiteralType InLiteralType)
-		{
-			switch (InLiteralType)
-			{
-			case EMetasoundFrontendLiteralType::None:
-				return DoesDataTypeSupportLiteralType(InDataType, ELiteralType::None);
-			case EMetasoundFrontendLiteralType::Bool:
-				return DoesDataTypeSupportLiteralType(InDataType, ELiteralType::Boolean);
-			case EMetasoundFrontendLiteralType::Float:
-				return DoesDataTypeSupportLiteralType(InDataType, ELiteralType::Float);
-			case EMetasoundFrontendLiteralType::Integer:
-				return DoesDataTypeSupportLiteralType(InDataType, ELiteralType::Integer);
-			case EMetasoundFrontendLiteralType::String:
-				return DoesDataTypeSupportLiteralType(InDataType, ELiteralType::String);
-			case EMetasoundFrontendLiteralType::UObject:
-				return DoesDataTypeSupportLiteralType(InDataType, ELiteralType::UObjectProxy);
-			case EMetasoundFrontendLiteralType::UObjectArray:
-				return DoesDataTypeSupportLiteralType(InDataType, ELiteralType::UObjectProxyArray);
-			case EMetasoundFrontendLiteralType::Invalid:
-			default:
-				return false;
-			}
-		}
-
-		bool DoesDataTypeSupportLiteralType(FName InDataType, ELiteralType InLiteralType)
-		{
-			FMetasoundFrontendRegistryContainer* Registry = FMetasoundFrontendRegistryContainer::Get();
-
-			return Registry->DoesDataTypeSupportLiteralType(InDataType, InLiteralType);
-		}
-
-		Metasound::FLiteral GetLiteralParam(const FMetasoundFrontendLiteral& InDescription)
-		{
-			EMetasoundFrontendLiteralType LiteralType = InDescription.Type;
-
-			switch (LiteralType)
-			{
-				case EMetasoundFrontendLiteralType::Bool:
-				{
-					return FLiteral(InDescription.AsBool);
-				}
-				case EMetasoundFrontendLiteralType::Float:
-				{
-					return FLiteral(InDescription.AsFloat);
-				}
-				case EMetasoundFrontendLiteralType::Integer:
-				{
-					return FLiteral(InDescription.AsInteger);
-				}
-				case EMetasoundFrontendLiteralType::String:
-				{
-					return FLiteral(InDescription.AsString);
-				}
-				case EMetasoundFrontendLiteralType::UObject:
-				case EMetasoundFrontendLiteralType::UObjectArray:
-				case EMetasoundFrontendLiteralType::None:
-				case EMetasoundFrontendLiteralType::Invalid:
-				default:
-				{
-					return FLiteral::CreateInvalid();
-				}
-			}
-		}
-
-		Metasound::FLiteral GetDefaultParamForDataType(FName InDataType)
-		{
-			FMetasoundFrontendRegistryContainer* Registry = FMetasoundFrontendRegistryContainer::Get();
-
-			ELiteralType DesiredArgType = Registry->GetDesiredLiteralTypeForDataType(InDataType);
-
-			switch (DesiredArgType)
-			{
-				case Metasound::ELiteralType::Boolean:
-				{
-					return FLiteral(false);
-				}
-				case Metasound::ELiteralType::Integer:
-				{
-					return FLiteral(0);
-				}
-				case Metasound::ELiteralType::Float:
-				{
-					return FLiteral(0.0f);
-				}
-				case Metasound::ELiteralType::String:
-				{
-					return FLiteral(FString());
-				}
-				case Metasound::ELiteralType::UObjectProxy:
-				case Metasound::ELiteralType::UObjectProxyArray:
-				case Metasound::ELiteralType::None:
-				{
-					return FLiteral();
-				}
-				case Metasound::ELiteralType::Invalid:
-				default:
-				{
-					return FLiteral::CreateInvalid();
-				}
-			}
-
-		}
 
 		ITransactable::ITransactable(uint32 InUndoLimit, TWeakObjectPtr<UObject> InOwningAsset /* = nullptr */)
 			: UndoLimit(InUndoLimit)
diff --git a/Engine/Plugins/Runtime/Metasound/Source/MetasoundFrontend/Private/MetasoundFrontendDocument.cpp b/Engine/Plugins/Runtime/Metasound/Source/MetasoundFrontend/Private/MetasoundFrontendDocument.cpp
--- a/Engine/Plugins/Runtime/Metasound/Source/MetasoundFrontend/Private/MetasoundFrontendDocument.cpp
+++ b/Engine/Plugins/Runtime/Metasound/Source/MetasoundFrontend/Private/MetasoundFrontendDocument.cpp
@@ -4,153 +4,13 @@
 
 #include "MetasoundFrontend.h"
 #include "MetasoundFrontendRegistries.h"
+#include "MetasoundLog.h"
 
 namespace Metasound
 {
 	const FGuid FrontendInvalidID = FGuid();
 }
 
-void FMetasoundFrontendLiteral::Set(bool InValue)
-{
-	Clear();
-	AsBool = InValue;
-	Type = EMetasoundFrontendLiteralType::Bool;
-}
-
-void FMetasoundFrontendLiteral::Set(int32 InValue)
-{
-	Clear();
-	AsInteger = InValue;
-	Type = EMetasoundFrontendLiteralType::Integer;
-}
-
-void FMetasoundFrontendLiteral::Set(float InValue)
-{
-	Clear();
-	AsFloat = InValue;
-	Type = EMetasoundFrontendLiteralType::Float;
-}
-
-void FMetasoundFrontendLiteral::Set(const FString& InValue)
-{
-	Clear();
-	AsString = InValue;
-	Type = EMetasoundFrontendLiteralType::String;
-}
-
-void FMetasoundFrontendLiteral::Set(UObject* InValue)
-{
-	Clear();
-	AsUObject = InValue;
-	Type = EMetasoundFrontendLiteralType::UObject;
-}
-
-void FMetasoundFrontendLiteral::Set(const TArray<UObject*>& InValue)
-{
-	Clear();
-	AsUObjectArray = InValue;
-	Type = EMetasoundFrontendLiteralType::UObjectArray;
-}
-
-void FMetasoundFrontendLiteral::SetFromLiteral(const Metasound::FLiteral& InLiteral)
-{
-	using namespace Metasound;
-
-	Clear();
-	switch (InLiteral.GetType())
-	{
-		case ELiteralType::Boolean:
-		{
-			Set(InLiteral.Value.Get<bool>());
-		}
-		break;
-
-		case ELiteralType::Float:
-		{
-			Set(InLiteral.Value.Get<float>());
-		}
-		break;
-
-		case ELiteralType::Integer:
-		{
-			Set(InLiteral.Value.Get<int32>());
-		}
-		break;
-
-		case ELiteralType::String:
-		{
-			Set(InLiteral.Value.Get<FString>());
-		}
-		break;
-
-		case ELiteralType::None:
-		case ELiteralType::UObjectProxy:
-		case ELiteralType::NoneArray:
-		case ELiteralType::BooleanArray:
-		case ELiteralType::IntegerArray:
-		case ELiteralType::FloatArray:
-		case ELiteralType::StringArray:
-		case ELiteralType::UObjectProxyArray:
-		case ELiteralType::Invalid:
-		default:
-		{
-			static_assert(static_cast<int32>(ELiteralType::Invalid) == 12, "Possible missing literal type switch coverage");
-		}
-	}
-}
-
-FString FMetasoundFrontendLiteral::ToString() const
-{
-	switch (Type)
-	{
-		case EMetasoundFrontendLiteralType::Bool:
-		return FString::Printf(TEXT("%s"), AsBool ? TEXT("true") : TEXT("false"));
-
-		case EMetasoundFrontendLiteralType::Float:
-		return FString::Printf(TEXT("%f"), AsFloat);
-
-		case EMetasoundFrontendLiteralType::Integer:
-		return FString::Printf(TEXT("%d"), AsInteger);
-
-		case EMetasoundFrontendLiteralType::String:
-		return AsString;
-
-		case EMetasoundFrontendLiteralType::UObject:
-		{
-			if (AsUObject)
-			{
-				return AsUObject->GetFullName();
-			}
-			else
-			{
-				return FString();
-			}
-		}
-
-		case EMetasoundFrontendLiteralType::UObjectArray:
-		case EMetasoundFrontendLiteralType::None:
-		case EMetasoundFrontendLiteralType::Invalid:
-		default:
-		{
-			static_assert(static_cast<int32>(EMetasoundFrontendLiteralType::Invalid) == 7, "Possible missing literal type switch coverage");
-		}
-
-		return FString();
-	}
-}
-
-void FMetasoundFrontendLiteral::Clear()
-{
-	Type = EMetasoundFrontendLiteralType::None;
-
-	AsBool = false;
-	AsInteger = 0;
-	AsFloat = 0.0f;
-	AsString.Empty();
-	AsUObject = nullptr;
-	AsUObjectArray.Empty();
-}
-
 FMetasoundFrontendNodeInterface::FMetasoundFrontendNodeInterface(const FMetasoundFrontendClassInterface& InClassInterface)
 {
 	for (const FMetasoundFrontendClassInput& Input : InClassInterface.Inputs)
@@ -244,7 +104,7 @@ FMetasoundFrontendClassMetadata::FMetasoundFrontendClassMetadata(const Metasound
 FMetasoundFrontendClassInput::FMetasoundFrontendClassInput(const FMetasoundFrontendClassVertex& InOther)
 :	FMetasoundFrontendClassVertex(InOther)
 {
-	EMetasoundFrontendLiteralType DefaultType = Metasound::Frontend::GetMetasoundLiteralType(FMetasoundFrontendRegistryContainer::Get()->GetDesiredLiteralTypeForDataType(InOther.TypeName));
+	EMetasoundFrontendLiteralType DefaultType = Metasound::Frontend::GetMetasoundFrontendLiteralType(FMetasoundFrontendRegistryContainer::Get()->GetDesiredLiteralTypeForDataType(InOther.TypeName));
 
 	for (const FGuid& PointID : PointIDs)
 	{
diff --git a/Engine/Plugins/Runtime/Metasound/Source/MetasoundFrontend/Private/MetasoundFrontendGraph.cpp b/Engine/Plugins/Runtime/Metasound/Source/MetasoundFrontend/Private/MetasoundFrontendGraph.cpp
--- a/Engine/Plugins/Runtime/Metasound/Source/MetasoundFrontend/Private/MetasoundFrontendGraph.cpp
+++ b/Engine/Plugins/Runtime/Metasound/Source/MetasoundFrontend/Private/MetasoundFrontendGraph.cpp
@@ -157,7 +157,7 @@ namespace Metasound
 
 				if (IsLiteralParsableByDataType)
 				{
-					FLiteral Literal = Frontend::GetLiteralParamForDataType(InputVertex.TypeName, VertexLiteral->Value);
+					FLiteral Literal = VertexLiteral->Value.ToLiteral(InputVertex.TypeName);
 
 					FInputNodeConstructorParams InitParams =
 					{
diff --git a/Engine/Plugins/Runtime/Metasound/Source/MetasoundFrontend/Private/MetasoundFrontendLiteral.cpp b/Engine/Plugins/Runtime/Metasound/Source/MetasoundFrontend/Private/MetasoundFrontendLiteral.cpp
new file mode 100644
--- /dev/null
+++ b/Engine/Plugins/Runtime/Metasound/Source/MetasoundFrontend/Private/MetasoundFrontendLiteral.cpp
@@ -0,0 +1,533 @@
+// Copyright Epic Games, Inc. All Rights Reserved.
+
+#include "MetasoundFrontendLiteral.h"
+
+#include "CoreMinimal.h"
+#include "MetasoundLog.h"
+#include "MetasoundFrontendRegistries.h"
+
+FString LexToString(Metasound::FLiteral::FNone InValue)
+{
+	return FString(TEXT(""));
+}
+
+namespace MetasoundFrontendLiteralPrivate
+{
+	// Default size of string for string builder stack allocation. Heap allocation is used if the string exceeds this size. 
+	constexpr int32 ArrayStringAllocatedBytes = 256;
+
+	// Defaults to calling LexToString for most literal data types.
+	template<typename Type>
+	struct TLiteralValueToStringHelper
+	{
+		static FString Convert(const Type& InType)
+		{
+			return LexToString(InType);
+		}
+	};
+
+	// String conversion specialization for UObject*
+	template<>
+	struct TLiteralValueToStringHelper<UObject*>
+	{
+		static FString Convert(const UObject* InType)
+		{
+			if (nullptr == InType)
+			{
+				return FString("nullptr");
+			}
+			return InType->GetFullName();
+		}
+	};
+
+	// String conversion specialization for TArray<>
+	template<typename ElementType>
+	struct TLiteralValueToStringHelper<TArray<ElementType>>
+	{
+		using FArrayType = TArray<ElementType>;
+		
+		static FString Convert(const FArrayType& InArray)
+		{
+			TStringBuilder<ArrayStringAllocatedBytes> Builder;
+			Builder << TEXT("[");
+
+			for (int32 i = 0; i < InArray.Num(); i++)
+			{
+				if (i> 0)
+				{
+					Builder << TEXT(", ");
+				}
+				Builder << TLiteralValueToStringHelper<ElementType>::Convert(InArray[i]);
+			}
+
+			return FString(Builder);
+		}
+	};
+
+	template<typename Type>
+	FString TLiteralValueToString(const Type& InType)
+	{
+		// Uses a helper struct to allow for partial specialization. 
+		// (c++ does not allow partial specialization of functions)
+		return TLiteralValueToStringHelper<Type>::Convert(InType);
+	}
+}
+
+bool FMetasoundFrontendLiteral::IsArray() const
+{
+	switch (Type)
+	{
+		case EMetasoundFrontendLiteralType::NoneArray:
+		case EMetasoundFrontendLiteralType::BooleanArray:
+		case EMetasoundFrontendLiteralType::IntegerArray:
+		case EMetasoundFrontendLiteralType::FloatArray:
+		case EMetasoundFrontendLiteralType::StringArray:
+		case EMetasoundFrontendLiteralType::UObjectArray:
+			return true;
+
+		case EMetasoundFrontendLiteralType::None:
+		case EMetasoundFrontendLiteralType::Boolean:
+		case EMetasoundFrontendLiteralType::Integer:
+		case EMetasoundFrontendLiteralType::Float:
+		case EMetasoundFrontendLiteralType::String:
+		case EMetasoundFrontendLiteralType::UObject:
+		case EMetasoundFrontendLiteralType::Invalid:
+			return false;
+			
+		default:
+			static_assert(static_cast<int32>(EMetasoundFrontendLiteralType::Invalid) == 12, "Possible missing EMetasoundFrontendLiteralType case coverage");
+			{
+				checkNoEntry();
+				return false;
+			}
+	}
+}
+
+void FMetasoundFrontendLiteral::Set(FDefault InValue)
+{
+	Empty();
+
+	Type = EMetasoundFrontendLiteralType::None;
+
+	AsNumDefault = 1;
+}
+
+void FMetasoundFrontendLiteral::Set(const FDefaultArray& InValue)
+{
+	Empty();
+
+	Type = EMetasoundFrontendLiteralType::NoneArray;
+
+	AsNumDefault = FMath::Max(0, InValue.Num);
+}
+
+void FMetasoundFrontendLiteral::Set(bool InValue)
+{
+	Empty();
+
+	Type = EMetasoundFrontendLiteralType::Boolean;
+
+	AsBoolean.Add(InValue);
+}
+
+void FMetasoundFrontendLiteral::Set(const TArray<bool>& InValue)
+{
+	Empty();
+
+	Type = EMetasoundFrontendLiteralType::BooleanArray;
+
+	AsBoolean = InValue;
+}
+
+void FMetasoundFrontendLiteral::Set(int32 InValue)
+{
+	Empty();
+
+	Type = EMetasoundFrontendLiteralType::Integer;
+
+	AsInteger.Add(InValue);
+}
+
+void FMetasoundFrontendLiteral::Set(const TArray<int32>& InValue)
+{
+	Empty();
+
+	Type = EMetasoundFrontendLiteralType::IntegerArray;
+
+	AsInteger = InValue;
+}
+
+void FMetasoundFrontendLiteral::Set(float InValue)
+{
+	Empty();
+
+	Type = EMetasoundFrontendLiteralType::Float;
+
+	AsFloat.Add(InValue);
+}
+
+void FMetasoundFrontendLiteral::Set(const TArray<float>& InValue)
+{
+	Empty();
+
+	Type = EMetasoundFrontendLiteralType::FloatArray;
+
+	AsFloat = InValue;
+}
+
+void FMetasoundFrontendLiteral::Set(const FString& InValue)
+{
+	Empty();
+
+	Type = EMetasoundFrontendLiteralType::String;
+
+	AsString.Add(InValue);
+}
+
+void FMetasoundFrontendLiteral::Set(const TArray<FString>& InValue)
+{
+	Empty();
+
+	Type = EMetasoundFrontendLiteralType::StringArray;
+
+	AsString = InValue;
+}
+
+void FMetasoundFrontendLiteral::Set(UObject* InValue)
+{
+	Empty();
+
+	Type = EMetasoundFrontendLiteralType::UObject;
+
+	AsUObject.Add(InValue);
+}
+
+void FMetasoundFrontendLiteral::Set(const TArray<UObject*>& InValue)
+{
+	Empty();
+
+	Type = EMetasoundFrontendLiteralType::UObjectArray;
+
+	AsUObject = InValue;
+}
+
+void FMetasoundFrontendLiteral::SetFromLiteral(const Metasound::FLiteral& InLiteral)
+{
+	using namespace Metasound;
+
+	Clear();
+	switch (InLiteral.GetType())
+	{
+		case ELiteralType::None:
+
+			Set(FMetasoundFrontendLiteral::FDefault{});
+
+			break;
+
+		case ELiteralType::Boolean:
+
+			Set(InLiteral.Value.Get<bool>());
+
+			break;
+
+		case ELiteralType::Float:
+
+			Set(InLiteral.Value.Get<float>());
+
+			break;
+
+		case ELiteralType::Integer:
+		
+			Set(InLiteral.Value.Get<int32>());
+	
+			break;
+
+		case ELiteralType::String:
+		
+			Set(InLiteral.Value.Get<FString>());
+		
+			break;
+
+		case ELiteralType::UObjectProxy:
+			{
+				UE_LOG(LogMetasound, Error, TEXT("Cannot set UObjectProxy from Metasound::FLiteral"));
+				Set(static_cast<UObject*>(nullptr));
+			}
+
+			break;
+
+		case ELiteralType::NoneArray:
+			{
+				int32 Num = InLiteral.Value.Get<TArray<FLiteral::FNone>>().Num(); 
+				Set(FMetasoundFrontendLiteral::FDefaultArray{ Num });
+			}
+
+			break;
+
+		case ELiteralType::BooleanArray:
+
+			Set(InLiteral.Value.Get<TArray<bool>>());
+
+			break;
+
+		case ELiteralType::IntegerArray:
+
+			Set(InLiteral.Value.Get<TArray<int32>>());
+
+			break;
+
+		case ELiteralType::FloatArray:
+
+			Set(InLiteral.Value.Get<TArray<float>>());
+
+			break;
+
+		case ELiteralType::StringArray:
+
+			Set(InLiteral.Value.Get<TArray<FString>>());
+
+			break;
+
+		case ELiteralType::UObjectProxyArray:
+			{
+				UE_LOG(LogMetasound, Error, TEXT("Cannot set UObjectProxy from Metasound::FLiteral"));
+				Set(TArray<UObject*>());
+			}
+
+			break;
+
+		case ELiteralType::Invalid:
+		default:
+		{
+			static_assert(static_cast<int32>(ELiteralType::Invalid) == 12, "Possible missing literal type switch coverage");
+		}
+	}
+}
+
+Metasound::FLiteral FMetasoundFrontendLiteral::ToLiteral(const FName& InMetasoundDataTypeName) const
+{
+	using namespace Metasound;
+	using namespace Metasound::Frontend;
+
+	FLiteral Literal = FLiteral::CreateInvalid();
+
+	const bool bIsTypeSupported = DoesDataTypeSupportLiteralType(InMetasoundDataTypeName, Type);
+
+	if (!bIsTypeSupported)
+	{
+		UE_LOG(LogMetasound, Error, TEXT("Reverting to default literal type for data type. Failed to create supported Metasound::FLiteral for data type [Name:%s] with FMetasoundFrontendLiteral [Literal:%s]"), *InMetasoundDataTypeName.ToString(), *ToString());
+
+		Literal = GetDefaultParamForDataType(InMetasoundDataTypeName);
+	}
+	else if (EMetasoundFrontendLiteralType::UObject == Type)
+	{
+		if (ensure(AsUObject.Num() > 0))
+		{
+			// UObject proxies must go through the registry. The registry contains the information
+			// needed to convert UObjects to proxies. 
+			FMetasoundFrontendRegistryContainer* Registry = FMetasoundFrontendRegistryContainer::Get();
+			if (ensure(Registry))
+			{
+				Literal = Registry->GenerateLiteralForUObject(InMetasoundDataTypeName, AsUObject[0]);
+			}
+		}
+	}
+	else if (EMetasoundFrontendLiteralType::UObjectArray == Type)
+	{
+		// UObject proxies must go through the registry. The registry contains the information
+		// needed to convert UObjects to proxies. 
+		FMetasoundFrontendRegistryContainer* Registry = FMetasoundFrontendRegistryContainer::Get();
+		if (ensure(Registry))
+		{
+			Literal = Registry->GenerateLiteralForUObjectArray(InMetasoundDataTypeName, AsUObject);
+		}
+	}
+	else 
+	{
+		// Use default conversions for core literal types.
+		Literal = ToLiteralNoProxy();
+	}
+
+	// The support for the data type should be the same whether we pass in an
+	// Metasound::ELiteralType or a EMetasoundFrontendLiteralType
+	check(bIsTypeSupported == DoesDataTypeSupportLiteralType(InMetasoundDataTypeName, Literal.GetType()));
+
+	return Literal;
+}
+
+Metasound::FLiteral FMetasoundFrontendLiteral::ToLiteralNoProxy() const
+{
+	using namespace Metasound;
+
+	if ((EMetasoundFrontendLiteralType::UObject == Type) || (EMetasoundFrontendLiteralType::UObjectArray == Type))
+	{
+		UE_LOG(LogMetasound, Error, TEXT("Cannot convert FMetasoundFrontendLiteral to Metasound::Literal without data type name [Literal:%s]"), *ToString());
+		return FLiteral::CreateInvalid();
+	}
+
+	switch (Type)
+	{
+		case EMetasoundFrontendLiteralType::None:
+			return FLiteral(FLiteral::FNone{});
+
+		case EMetasoundFrontendLiteralType::Boolean:
+			return FLiteral(AsBoolean[0]);
+
+		case EMetasoundFrontendLiteralType::Integer:
+			return FLiteral(AsInteger[0]);
+			
+		case EMetasoundFrontendLiteralType::Float:
+			return FLiteral(AsFloat[0]);
+
+		case EMetasoundFrontendLiteralType::String:
+			return FLiteral(AsString[0]);
+
+		case EMetasoundFrontendLiteralType::NoneArray:
+		{
+			TArray<FLiteral::FNone> Nones;
+			Nones.SetNum(AsNumDefault);
+			return FLiteral(Nones);
+		}
+
+		case EMetasoundFrontendLiteralType::BooleanArray:
+			return FLiteral(AsBoolean);
+
+		case EMetasoundFrontendLiteralType::IntegerArray:
+			return FLiteral(AsInteger);
+
+		case EMetasoundFrontendLiteralType::FloatArray:
+			return FLiteral(AsFloat);
+
+		case EMetasoundFrontendLiteralType::StringArray:
+			return FLiteral(AsString);
+
+		case EMetasoundFrontendLiteralType::UObject:
+		case EMetasoundFrontendLiteralType::UObjectArray:
+		case EMetasoundFrontendLiteralType::Invalid:
+			return FLiteral::CreateInvalid();
+
+		default:
+			static_assert(static_cast<int32>(EMetasoundFrontendLiteralType::Invalid) == 12, "Possible missing literal type switch coverage");
+			return FLiteral::CreateInvalid();
+	}
+}
+
+FString FMetasoundFrontendLiteral::ToString() const
+{
+	using namespace Metasound;
+	using namespace MetasoundFrontendLiteralPrivate;
+
+	switch (Type)
+	{
+		case EMetasoundFrontendLiteralType::None:
+			return TLiteralValueToString(FLiteral::FNone{});
+
+		case EMetasoundFrontendLiteralType::Boolean:
+			return TLiteralValueToString(AsBoolean[0]);
+
+		case EMetasoundFrontendLiteralType::Integer:
+			return TLiteralValueToString(AsInteger[0]);
+
+		case EMetasoundFrontendLiteralType::Float:
+			return TLiteralValueToString(AsFloat[0]);
+
+		case EMetasoundFrontendLiteralType::String:
+			return TLiteralValueToString(AsString[0]);
+
+		case EMetasoundFrontendLiteralType::UObject:
+			return TLiteralValueToString(AsUObject[0]);
+
+		case EMetasoundFrontendLiteralType::NoneArray:
+		{
+			TArray<FLiteral::FNone> Nones;
+			Nones.SetNum(AsNumDefault);
+			return TLiteralValueToString(Nones);
+		}
+
+		case EMetasoundFrontendLiteralType::BooleanArray:
+			return TLiteralValueToString(AsBoolean);
+
+		case EMetasoundFrontendLiteralType::IntegerArray:
+			return TLiteralValueToString(AsInteger);
+
+		case EMetasoundFrontendLiteralType::FloatArray:
+			return TLiteralValueToString(AsFloat);
+
+		case EMetasoundFrontendLiteralType::StringArray:
+			return TLiteralValueToString(AsString);
+
+		case EMetasoundFrontendLiteralType::UObjectArray:
+			return TLiteralValueToString(AsUObject);
+
+		case EMetasoundFrontendLiteralType::Invalid:
+		default:
+			static_assert(static_cast<int32>(EMetasoundFrontendLiteralType::Invalid) == 12, "Possible missing literal type switch coverage");
+			return FString();
+	}
+}
+
+void FMetasoundFrontendLiteral::Clear()
+{
+	Empty();
+	Type = EMetasoundFrontendLiteralType::None;
+	AsNumDefault = 1;
+}
+
+void FMetasoundFrontendLiteral::Empty()
+{
+	AsNumDefault = 0;
+	AsBoolean.Empty();
+	AsInteger.Empty();
+	AsFloat.Empty();
+	AsString.Empty();
+	AsUObject.Empty();
+}
+
+namespace Metasound
+{
+	namespace Frontend
+	{
+		// Convenience function to convert Metasound::ELiteralType to EMetasoundFrontendLiteralType, the UEnum used for metasound documents.
+		EMetasoundFrontendLiteralType GetMetasoundFrontendLiteralType(ELiteralType InLiteralType)
+		{
+			static_assert(static_cast<uint8>(EMetasoundFrontendLiteralType::Invalid) == static_cast<uint8>(ELiteralType::Invalid), "Possible literal type enum mismatch");
+			return static_cast<EMetasoundFrontendLiteralType>(InLiteralType);
+		}
+
+		// Convenience function to convert Metasound::ELiteralType to EMetasoundFrontendLiteralType, the UEnum used for metasound documents.
+		ELiteralType GetMetasoundLiteralType(EMetasoundFrontendLiteralType InLiteralType)
+		{
+			static_assert(static_cast<uint8>(EMetasoundFrontendLiteralType::Invalid) == static_cast<uint8>(ELiteralType::Invalid), "Possible literal type enum mismatch");
+			return static_cast<ELiteralType>(InLiteralType);
+		}
+
+		bool DoesDataTypeSupportLiteralType(FName InDataType, EMetasoundFrontendLiteralType InFrontendLiteralType)
+		{
+			ELiteralType CoreLiteralType = GetMetasoundLiteralType(InFrontendLiteralType);
+			return DoesDataTypeSupportLiteralType(InDataType, CoreLiteralType);
+		}
+
+		bool DoesDataTypeSupportLiteralType(FName InDataType, ELiteralType InLiteralType)
+		{
+			FMetasoundFrontendRegistryContainer* Registry = FMetasoundFrontendRegistryContainer::Get();
+			if (ensure(Registry))
+			{
+				return Registry->DoesDataTypeSupportLiteralType(InDataType, InLiteralType);
+			}
+			return false;
+		}
+
+		Metasound::FLiteral GetDefaultParamForDataType(FName InDataType)
+		{
+			FMetasoundFrontendRegistryContainer* Registry = FMetasoundFrontendRegistryContainer::Get();
+			if (ensure(Registry))
+			{
+				ELiteralType DesiredArgType = Registry->GetDesiredLiteralTypeForDataType(InDataType);
+
+				FLiteral::GetDefaultForType(DesiredArgType);
+			}
+			return FLiteral::CreateInvalid();
+		}
+	}
+}
+
diff --git a/Engine/Plugins/Runtime/Metasound/Source/MetasoundFrontend/Private/MetasoundFrontendStandardController.cpp b/Engine/Plugins/Runtime/Metasound/Source/MetasoundFrontend/Private/MetasoundFrontendStandardController.cpp
--- a/Engine/Plugins/Runtime/Metasound/Source/MetasoundFrontend/Private/MetasoundFrontendStandardController.cpp
+++ b/Engine/Plugins/Runtime/Metasound/Source/MetasoundFrontend/Private/MetasoundFrontendStandardController.cpp
@@ -2312,101 +2312,42 @@ namespace Metasound
 
 		bool FGraphController::SetDefaultInputToFrontendLiteral(const FString& InInputName, FGuid InPointID, FName InDataTypeName, const FMetasoundFrontendLiteral& InLiteral)
 		{
-			switch (InLiteral.Type)
+			if (FMetasoundFrontendClassInput* Desc = FindInputDescriptionWithName(InInputName))
 			{
-				case EMetasoundFrontendLiteralType::Bool:
-				{
-					return SetDefaultInputToLiteralInternal(InInputName, InPointID, InLiteral.AsBool);
-				}
-				break;
-
-				case EMetasoundFrontendLiteralType::Float:
-				{
-					return SetDefaultInputToLiteralInternal(InInputName, InPointID, InLiteral.AsFloat);
-				}
-				break;
-
-				case EMetasoundFrontendLiteralType::Integer:
-				{
-					return SetDefaultInputToLiteralInternal(InInputName, InPointID, InLiteral.AsInteger);
-				}
-				break;
-
-				case EMetasoundFrontendLiteralType::String:
-				{
-					return SetDefaultInputToLiteralInternal(InInputName, InPointID, InLiteral.AsString);
-				}
-				break;
+				auto IsLiteralWithSamePointID = [&](const FMetasoundFrontendVertexLiteral& InVertexLiteral) 
+				{ 
+					return InVertexLiteral.PointID == InPointID; 
+				};
 
-				case EMetasoundFrontendLiteralType::UObject:
+				FMetasoundFrontendVertexLiteral* VertexLiteral = Desc->Defaults.FindByPredicate(IsLiteralWithSamePointID);
+				if (nullptr == VertexLiteral)
 				{
-					return SetDefaultInputToLiteralInternal(InInputName, InPointID, InLiteral.AsUObject);
+					VertexLiteral = &Desc->Defaults.AddDefaulted_GetRef();
+					VertexLiteral->PointID = InPointID;
 				}
-				break;
 
-				// TODO: Remove as array support will not be specialized to object type
-				case EMetasoundFrontendLiteralType::UObjectArray:
+				if (ensure(DoesDataTypeSupportLiteralType(Desc->TypeName, InLiteral.Type)))
 				{
-					return false;
+					VertexLiteral->Value = InLiteral;
+					return true;
 				}
-				break;
-
-				case EMetasoundFrontendLiteralType::Invalid:
-				case EMetasoundFrontendLiteralType::None:
-				default:
+				else
 				{
-					static_assert(static_cast<int32>(EMetasoundFrontendLiteralType::Invalid) == 7, "Possible missing ELiteralType case coverage");
-					return false;
+					VertexLiteral->Value.Clear();
 				}
-				break;
 			}
+
+			return false;
 		}
 
 		bool FGraphController::SetDefaultInputToTypeDefaultLiteral(const FString& InInputName, FGuid InPointID, FName InDataTypeName)
 		{
-			Metasound::FLiteral LiteralParam = Frontend::GetDefaultParamForDataType(InDataTypeName);
-			switch (LiteralParam.GetType())
-			{
-				case ELiteralType::Boolean:
-				{
-					return SetDefaultInputToLiteralInternal(InInputName, InPointID, LiteralParam.Value.Get<bool>());
-				}
-				break;
-
-				case ELiteralType::Float:
-				{
-					return SetDefaultInputToLiteralInternal(InInputName, InPointID, LiteralParam.Value.Get<float>());
-				}
-				break;
+			FMetasoundFrontendLiteral FrontendLiteral;
 
-				case ELiteralType::Integer:
-				{
-					return SetDefaultInputToLiteralInternal(InInputName, InPointID, LiteralParam.Value.Get<int32>());
-				}
-				break;
+			Metasound::FLiteral Literal = Frontend::GetDefaultParamForDataType(InDataTypeName);
+			FrontendLiteral.SetFromLiteral(Literal);
 
-				case ELiteralType::String:
-				{
-					return SetDefaultInputToLiteralInternal(InInputName, InPointID, LiteralParam.Value.Get<FString>());
-				}
-				break;
-
-				case ELiteralType::None:
-				case ELiteralType::UObjectProxy:	// UObject literals do not currently support default data type references
-				case ELiteralType::NoneArray:
-				case ELiteralType::BooleanArray:
-				case ELiteralType::IntegerArray:
-				case ELiteralType::FloatArray:
-				case ELiteralType::StringArray:
-				case ELiteralType::UObjectProxyArray:
-				case ELiteralType::Invalid:
-				default:
-				{
-					static_assert(static_cast<int32>(ELiteralType::Invalid) == 12, "Possible missing ELiteralType case coverage");
-					return false;
-				}
-				break;
-			}
+			return SetDefaultInputToFrontendLiteral(InInputName, InPointID, InDataTypeName, FrontendLiteral);
 		}
 
 		// Set the display name for the input with the given name
diff --git a/Engine/Plugins/Runtime/Metasound/Source/MetasoundFrontend/Private/MetasoundFrontendStandardController.h b/Engine/Plugins/Runtime/Metasound/Source/MetasoundFrontend/Private/MetasoundFrontendStandardController.h
--- a/Engine/Plugins/Runtime/Metasound/Source/MetasoundFrontend/Private/MetasoundFrontendStandardController.h
+++ b/Engine/Plugins/Runtime/Metasound/Source/MetasoundFrontend/Private/MetasoundFrontendStandardController.h
@@ -600,60 +600,6 @@ namespace Metasound
 		/** FGraphController represents a Metasound graph class. */
 		class METASOUNDFRONTEND_API FGraphController : public IGraphController
 		{
-			// Util for setting template to literal based on template argument
-			template<typename ArgType>
-			bool SetDefaultInputToLiteralInternal(const FString& InInputName, FGuid InPointID, ArgType InValue)
-			{
-				if (FMetasoundFrontendClassInput* Desc = FindInputDescriptionWithName(InInputName))
-				{
-					auto IsLiteralWithSamePointID = [&](const FMetasoundFrontendVertexLiteral& InVertexLiteral) 
-					{ 
-						return InVertexLiteral.PointID == InPointID; 
-					};
-
-					FMetasoundFrontendVertexLiteral* VertexLiteral = Desc->Defaults.FindByPredicate(IsLiteralWithSamePointID);
-					if (nullptr == VertexLiteral)
-					{
-						VertexLiteral = &Desc->Defaults.AddDefaulted_GetRef();
-						VertexLiteral->PointID = InPointID;
-					}
-
-					if (ensure(FMetasoundFrontendRegistryContainer::Get()->DoesDataTypeSupportLiteralType<ArgType>(Desc->TypeName)))
-					{
-						VertexLiteral->Value.Set(InValue);
-						return true;
-					}
-				}
-
-				return false;
-			}
-
-			template<>
-			bool SetDefaultInputToLiteralInternal<UObject*>(const FString& InInputName, FGuid InPointID, UObject* InValue)
-			{
-				if (FMetasoundFrontendClassInput* Desc = FindInputDescriptionWithName(InInputName))
-				{
-					auto IsLiteralWithSamePointID = [&](const FMetasoundFrontendVertexLiteral& InVertexLiteral)
-					{
-						return InVertexLiteral.PointID == InPointID;
-					};
-
-					FMetasoundFrontendVertexLiteral* VertexLiteral = Desc->Defaults.FindByPredicate(IsLiteralWithSamePointID);
-					if (nullptr == VertexLiteral)
-					{
-						VertexLiteral = &Desc->Defaults.AddDefaulted_GetRef();
-						VertexLiteral->PointID = InPointID;
-					}
-
-					if (ensure(FMetasoundFrontendRegistryContainer::Get()->DoesDataTypeSupportLiteralType(Desc->TypeName, Metasound::ELiteralType::UObjectProxy)))
-					{
-						VertexLiteral->Value.Set(InValue);
-						return true;
-					}
-				}
-
-				return false;
-			}
 
 			// Private token only allows members or friends to call constructor.
 			enum EPrivateToken { Token };
diff --git a/Engine/Plugins/Runtime/Metasound/Source/MetasoundFrontend/Public/MetasoundFrontend.h b/Engine/Plugins/Runtime/Metasound/Source/MetasoundFrontend/Public/MetasoundFrontend.h
--- a/Engine/Plugins/Runtime/Metasound/Source/MetasoundFrontend/Public/MetasoundFrontend.h
+++ b/Engine/Plugins/Runtime/Metasound/Source/MetasoundFrontend/Public/MetasoundFrontend.h
@@ -91,8 +91,6 @@ namespace Metasound
 			return GenerateClassDescription<NodeType>(InitData);
 		}
 
-		// Convenience function to convert Metasound::ELiteralType to EMetasoundFrontendLiteralType, the UEnum used for metasound documents.
-		METASOUNDFRONTEND_API EMetasoundFrontendLiteralType GetMetasoundLiteralType(Metasound::ELiteralType InLiteralType);
 
 		template<typename DataType>
 		FName GetDataTypeName()
diff --git a/Engine/Plugins/Runtime/Metasound/Source/MetasoundFrontend/Public/MetasoundFrontendBaseClasses.h b/Engine/Plugins/Runtime/Metasound/Source/MetasoundFrontend/Public/MetasoundFrontendBaseClasses.h
--- a/Engine/Plugins/Runtime/Metasound/Source/MetasoundFrontend/Public/MetasoundFrontendBaseClasses.h
+++ b/Engine/Plugins/Runtime/Metasound/Source/MetasoundFrontend/Public/MetasoundFrontendBaseClasses.h
@@ -28,23 +28,6 @@ namespace Metasound
 		// @returns nullptr if the type given wasn't found.
 		METASOUNDFRONTEND_API TUniquePtr<INode> ConstructExternalNode(const FName& InNodeClassFullName, uint32 InNodeHash, const FNodeInitData& InInitData);
 
-
-		// Utility functions for building a ::Metasound::FDataInitParam corresponding to a literal.
-
-		// Return the literal description parsed into a init param. 
-		// @Returns an invalid init param if the data type couldn't be found, or if the literal type was incompatible with the data type.
-		METASOUNDFRONTEND_API FLiteral GetLiteralParamForDataType(FName InDataType, const FMetasoundFrontendLiteral& InDescription);
-
-		METASOUNDFRONTEND_API bool DoesDataTypeSupportLiteralType(FName InDataType, EMetasoundFrontendLiteralType InLiteralType);
-		METASOUNDFRONTEND_API bool DoesDataTypeSupportLiteralType(FName InDataType, ELiteralType InLiteralType);
-		
-		// Returns a literal param without any type checking.
-		METASOUNDFRONTEND_API FLiteral GetLiteralParam(const FMetasoundFrontendLiteral& InDescription);
-
-		// Returns the defaulted version of a literal param for the given data type.
-		// @Returns an invalid init param if the data type couldn't be found.
-		METASOUNDFRONTEND_API FLiteral GetDefaultParamForDataType(FName InDataType);
-
 		// Utility base class for classes that want to support naive multi-level undo/redo.
 		class METASOUNDFRONTEND_API ITransactable : public TSharedFromThis<ITransactable>
 		{
diff --git a/Engine/Plugins/Runtime/Metasound/Source/MetasoundFrontend/Public/MetasoundFrontendDocument.h b/Engine/Plugins/Runtime/Metasound/Source/MetasoundFrontend/Public/MetasoundFrontendDocument.h
--- a/Engine/Plugins/Runtime/Metasound/Source/MetasoundFrontend/Public/MetasoundFrontendDocument.h
+++ b/Engine/Plugins/Runtime/Metasound/Source/MetasoundFrontend/Public/MetasoundFrontendDocument.h
@@ -5,6 +5,7 @@
 #include "CoreMinimal.h"
 #include "Internationalization/Text.h"
 #include "MetasoundAccessPtr.h"
+#include "MetasoundFrontendLiteral.h"
 #include "MetasoundNodeInterface.h"
 #include "Misc/Guid.h"
 
@@ -70,61 +71,6 @@ struct FMetasoundFrontendVersion
 	FMetasoundFrontendVersionNumber Number;
 };
 
-// The type of a given literal for an input value.
-UENUM()
-enum class EMetasoundFrontendLiteralType : uint8
-{
-	None,
-	Bool,
-	Float,
-	Integer,
-	String,
-	UObject,
-	UObjectArray,
-	Invalid UMETA(Hidden)
-};
-
-// Represents the serialized version of variant literal types. 
-USTRUCT()
-struct METASOUNDFRONTEND_API FMetasoundFrontendLiteral
-{
-	GENERATED_BODY()
-
-	// The set type of this literal.
-	UPROPERTY(EditAnywhere, Category = Customized)
-	EMetasoundFrontendLiteralType Type = EMetasoundFrontendLiteralType::None;
-
-	UPROPERTY(EditAnywhere, Category = Customized)
-	bool AsBool = false;
-
-	UPROPERTY(EditAnywhere, Category = Customized)
-	int32 AsInteger = 0;
-
-	UPROPERTY(EditAnywhere, Category = Customized)
-	float AsFloat = 0.f;
-
-	UPROPERTY(EditAnywhere, Category = Customized)
-	FString AsString = TEXT("");
-
-	UPROPERTY(EditAnywhere, Category = Customized)
-	UObject* AsUObject = nullptr;
-
-	UPROPERTY(EditAnywhere, Category = Customized)
-	TArray<UObject*> AsUObjectArray;
-
-	void Set(bool InValue);
-	void Set(int32 InValue);
-	void Set(float InValue);
-	void Set(const FString& InValue);
-	void Set(UObject* InValue);
-	void Set(const TArray<UObject*>& InValue);
-
-	void SetFromLiteral(const Metasound::FLiteral& InLiteral);
-
-	FString ToString() const;
-
-	void Clear();
-};
 
 // An FMetasoundFrontendVertex provides a named connection point of a node.
 USTRUCT() 
diff --git a/Engine/Plugins/Runtime/Metasound/Source/MetasoundFrontend/Public/MetasoundFrontendLiteral.h b/Engine/Plugins/Runtime/Metasound/Source/MetasoundFrontend/Public/MetasoundFrontendLiteral.h
new file mode 100644
--- /dev/null
+++ b/Engine/Plugins/Runtime/Metasound/Source/MetasoundFrontend/Public/MetasoundFrontendLiteral.h
@@ -0,0 +1,132 @@
+// Copyright Epic Games, Inc. All Rights Reserved.
+
+#pragma once
+
+#include "CoreMinimal.h"
+#include "MetasoundLiteral.h"
+
+#include <type_traits>
+
+#include "MetasoundFrontendLiteral.generated.h"
+
+// The type of a given literal for an input value.
+//
+// The EMetasoundFrontendLiteralType's are matched to Metasound::ELiteralType`s 
+// by giving them the same value. This supports easy conversion from one type to
+// another.
+UENUM()
+enum class EMetasoundFrontendLiteralType : uint8
+{
+
+	None = static_cast<uint8>(Metasound::ELiteralType::None), //< A value of None expresses that an object being constructed with a literal should be default constructed.
+	Boolean = static_cast<uint8>(Metasound::ELiteralType::Boolean),
+	Integer = static_cast<uint8>(Metasound::ELiteralType::Integer),
+	Float = static_cast<uint8>(Metasound::ELiteralType::Float),
+	String = static_cast<uint8>(Metasound::ELiteralType::String),
+	UObject = static_cast<uint8>(Metasound::ELiteralType::UObjectProxy),
+
+	NoneArray = static_cast<uint8>(Metasound::ELiteralType::NoneArray), //< A NoneArray expresses the number of objects to be default constructed.
+	BooleanArray = static_cast<uint8>(Metasound::ELiteralType::BooleanArray),
+	IntegerArray = static_cast<uint8>(Metasound::ELiteralType::IntegerArray),
+	FloatArray = static_cast<uint8>(Metasound::ELiteralType::FloatArray),
+	StringArray = static_cast<uint8>(Metasound::ELiteralType::StringArray),
+	UObjectArray = static_cast<uint8>(Metasound::ELiteralType::UObjectProxyArray),
+
+
+	Invalid UMETA(Hidden)
+};
+
+// Check that the static_cast<>s above are using the correct type.
+static_assert(std::is_same<uint8, std::underlying_type<EMetasoundFrontendLiteralType>::type>::value, "Update type in static_cast<TYPE> from Metasound::ELiteralType to EMetasoundFrontendLiteralType in EMetasoundFrontendLiteralType declaration.");
+static_assert(std::is_same<std::underlying_type<Metasound::ELiteralType>::type, std::underlying_type<EMetasoundFrontendLiteralType>::type>::value, "EMetasoundFrontendLiteralType and Metasound::ELiteralType must have matching underlying types to support conversion.");
+
+
+// Represents the serialized version of variant literal types. 
+USTRUCT()
+struct METASOUNDFRONTEND_API FMetasoundFrontendLiteral
+{
+	GENERATED_BODY()
+
+	struct FDefault {};
+	struct FDefaultArray { int32 Num = 0; };
+
+	// The set type of this literal.
+	UPROPERTY(EditAnywhere, Category = Customized)
+	EMetasoundFrontendLiteralType Type = EMetasoundFrontendLiteralType::None;
+
+	UPROPERTY(EditAnywhere, Category = Customized)
+	int32 AsNumDefault;
+
+	UPROPERTY(EditAnywhere, Category = Customized)
+	TArray<bool> AsBoolean;
+
+	UPROPERTY(EditAnywhere, Category = Customized)
+	TArray<int32> AsInteger;
+
+	UPROPERTY(EditAnywhere, Category = Customized)
+	TArray<float> AsFloat;
+
+	UPROPERTY(EditAnywhere, Category = Customized)
+	TArray<FString> AsString;
+
+	UPROPERTY(EditAnywhere, Category = Customized)
+	TArray<UObject*> AsUObject;
+
+	// Returns true if the stored Type is an array type.
+	bool IsArray() const;
+
+	void Set(FDefault InValue);
+	void Set(const FDefaultArray& InValue);
+	void Set(bool InValue);
+	void Set(const TArray<bool>& InValue);
+	void Set(int32 InValue);
+	void Set(const TArray<int32>& InValue);
+	void Set(float InValue);
+	void Set(const TArray<float>& InValue);
+	void Set(const FString& InValue);
+	void Set(const TArray<FString>& InValue);
+	void Set(UObject* InValue);
+	void Set(const TArray<UObject*>& InValue);
+
+	void SetFromLiteral(const Metasound::FLiteral& InLiteral);
+
+	// Return the literal description parsed into a init param. 
+	// @Returns an invalid init param if the data type couldn't be found, or if the literal type was incompatible with the data type.
+	Metasound::FLiteral ToLiteral(const FName& InMetasoundDataType) const;
+
+	// Does not do type checking or handle proxy
+	Metasound::FLiteral ToLiteralNoProxy() const;
+
+	// Convert the value to a string for printing. 
+	FString ToString() const;
+
+	// Remove any stored data and set to an invalid state.
+	void Clear();
+
+private:
+	// Remove all values. 
+	void Empty();
+};
+
+namespace Metasound
+{
+	namespace Frontend
+	{
+		// Returns true if the data type can be constructed using the literal type.
+		METASOUNDFRONTEND_API bool DoesDataTypeSupportLiteralType(FName InDataType, EMetasoundFrontendLiteralType InLiteralType);
+
+		// Returns true if the data type can be constructed using the literal type.
+		METASOUNDFRONTEND_API bool DoesDataTypeSupportLiteralType(FName InDataType, ELiteralType InLiteralType);
+
+		// Returns the defaulted version of a literal param for the given data type.
+		// @Returns an invalid init param if the data type couldn't be found.
+		METASOUNDFRONTEND_API FLiteral GetDefaultParamForDataType(FName InDataType);
+
+		// Convenience function to convert Metasound::ELiteralType to EMetasoundFrontendLiteralType.
+		METASOUNDFRONTEND_API EMetasoundFrontendLiteralType GetMetasoundFrontendLiteralType(ELiteralType InLiteralType);
+
+		// Convenience function to convert EMetasoundFrontendLiteralType to Metasound::ELiteralType.
+		METASOUNDFRONTEND_API ELiteralType GetMetasoundLiteralType(EMetasoundFrontendLiteralType InLiteralType);
+	}
+}
+
diff --git a/Engine/Plugins/Runtime/Metasound/Source/MetasoundFrontend/Public/MetasoundPrimitives.h b/Engine/Plugins/Runtime/Metasound/Source/MetasoundFrontend/Public/MetasoundPrimitives.h
--- a/Engine/Plugins/Runtime/Metasound/Source/MetasoundFrontend/Public/MetasoundPrimitives.h
+++ b/Engine/Plugins/Runtime/Metasound/Source/MetasoundFrontend/Public/MetasoundPrimitives.h
@@ -18,5 +18,7 @@ namespace Metasound
 	DECLARE_METASOUND_DATA_REFERENCE_TYPES(double, METASOUNDFRONTEND_API, FDoubleTypeInfo, FDoubleReadRef, FDoubleWriteRef);
 
 	DECLARE_METASOUND_DATA_REFERENCE_TYPES(FString, METASOUNDFRONTEND_API, FStringTypeInfo, FStringReadRef, FStringWriteRef);
+
+	DECLARE_METASOUND_DATA_REFERENCE_TYPES(TArray<float>, METASOUNDFRONTEND_API, FFloatArrayTypeInfo, FFloatArrayReadRef, FFloatArrayWriteRef);
 }
 
diff --git a/Engine/Plugins/Runtime/Metasound/Source/MetasoundGraphCore/Private/MetasoundLiteral.cpp b/Engine/Plugins/Runtime/Metasound/Source/MetasoundGraphCore/Private/MetasoundLiteral.cpp
new file mode 100644
--- /dev/null
+++ b/Engine/Plugins/Runtime/Metasound/Source/MetasoundGraphCore/Private/MetasoundLiteral.cpp
@@ -0,0 +1,260 @@
+// Copyright Epic Games, Inc. All Rights Reserved.
+
+
+#include "MetasoundLiteral.h"
+
+#include "IAudioProxyInitializer.h"
+#include "Misc/TVariant.h"
+#include <type_traits>
+
+namespace Metasound
+{
+#if METASOUND_DEBUG_LITERALS
+	void FLiteral::InitDebugString()
+	{
+		switch (GetType())
+		{
+			case ELiteralType::None:
+				DebugString = TEXT("NONE");
+			break;
+
+			case ELiteralType::Boolean:
+				DebugString = FString::Printf(TEXT("Boolean: %s"), Value.Get<bool>() ? TEXT("true") : TEXT("false"));
+			break;
+
+			case ELiteralType::Integer:
+				DebugString = FString::Printf(TEXT("Int32: %d"), Value.Get<int32>());
+			break;
+
+			case ELiteralType::Float:
+				DebugString = FString::Printf(TEXT("Float: %f"), Value.Get<float>());
+			break;
+
+			case ELiteralType::String:
+				DebugString = FString::Printf(TEXT("String: %s"), *Value.Get<FString>());
+			break;
+
+			case ELiteralType::UObjectProxy:
+			{
+				FString ProxyType = TEXT("nullptr");
+				if (Value.Get<Audio::IProxyDataPtr>().IsValid())
+				{
+					ProxyType = Value.Get<Audio::IProxyDataPtr>()->GetProxyTypeName().ToString();
+				}
+				DebugString = FString::Printf(TEXT("Audio::IProxyDataPtr: %s"), *ProxyType);
+			}
+			break;
+
+			case ELiteralType::NoneArray:
+				DebugString = TEXT("TArray<NONE>");
+			break;
+
+			case ELiteralType::BooleanArray:
+				DebugString = FString::Printf(TEXT("TArray<Boolean>"));
+			break;
+
+			case ELiteralType::IntegerArray:
+				DebugString = FString::Printf(TEXT("TArray<int32>"));
+			break;
+
+			case ELiteralType::FloatArray:
+				DebugString = FString::Printf(TEXT("TArray<float>"));
+			break;
+
+			case ELiteralType::StringArray:
+				DebugString = FString::Printf(TEXT("TArray<FString>"));
+			break;
+
+			case ELiteralType::UObjectProxyArray:
+				DebugString = FString::Printf(TEXT("TArray<Audio::IProxyDataPtr>"));
+			break;
+
+			case ELiteralType::Invalid:
+				DebugString = TEXT("INVALID");
+			break;
+
+			default:
+				static_assert(static_cast<int32>(ELiteralType::Invalid) == 12, "Possible missing ELiteralType case coverage");
+				checkNoEntry();
+		}
+	}
+#endif // METASOUND_DEBUG_LITERALS
+
+
+	// builds an invalid FLiteral.
+	FLiteral FLiteral::CreateInvalid()
+	{
+		return FLiteral(FInvalid());
+	}
+
+	FLiteral FLiteral::GetDefaultForType(ELiteralType InType)
+	{
+		switch (InType)
+		{
+			case ELiteralType::None:
+				return FLiteral(TLiteralTypeInfo<FLiteral::FNone>::GetDefaultValue());
+
+			case ELiteralType::Boolean:
+				return FLiteral(TLiteralTypeInfo<bool>::GetDefaultValue());
+
+			case ELiteralType::Integer:
+				return FLiteral(TLiteralTypeInfo<int32>::GetDefaultValue());
+
+			case ELiteralType::Float:
+				return FLiteral(TLiteralTypeInfo<float>::GetDefaultValue());
+
+			case ELiteralType::String:
+				return FLiteral(TLiteralTypeInfo<FString>::GetDefaultValue());
+
+			case ELiteralType::UObjectProxy:
+				return FLiteral(TLiteralTypeInfo<Audio::IProxyDataPtr>::GetDefaultValue());
+			
+			case ELiteralType::NoneArray:
+				return FLiteral(TLiteralTypeInfo<TArray<FLiteral::FNone>>::GetDefaultValue());
+			
+			case ELiteralType::BooleanArray:
+				return FLiteral(TLiteralTypeInfo<TArray<bool>>::GetDefaultValue());
+			
+			case ELiteralType::IntegerArray:
+				return FLiteral(TLiteralTypeInfo<TArray<int32>>::GetDefaultValue());
+			
+			case ELiteralType::FloatArray:
+				return FLiteral(TLiteralTypeInfo<TArray<float>>::GetDefaultValue());
+
+			case ELiteralType::StringArray:
+				return FLiteral(TLiteralTypeInfo<TArray<FString>>::GetDefaultValue());
+
+			case ELiteralType::UObjectProxyArray:
+				return FLiteral(TLiteralTypeInfo<TArray<Audio::IProxyDataPtr>>::GetDefaultValue());
+
+			case ELiteralType::Invalid:
+			default:
+				static_assert(static_cast<int32>(ELiteralType::Invalid) == 12, "Possible missing ELiteralType case coverage");
+				return FLiteral::CreateInvalid();
+
+		}
+	}
+
+	bool FLiteral::IsValid() const
+	{
+		return GetType() != ELiteralType::Invalid;
+	}
+
+	ELiteralType FLiteral::GetType() const
+	{
+		switch (Value.GetIndex())
+		{
+			case 0:
+				return ELiteralType::None;
+
+			case 1:
+				return ELiteralType::Boolean;
+
+			case 2:
+				return ELiteralType::Integer;
+
+			case 3:
+				return ELiteralType::Float;
+
+			case 4:
+				return ELiteralType::String;
+
+			case 5:
+				return ELiteralType::UObjectProxy;
+
+			case 6:
+				return ELiteralType::NoneArray;
+
+			case 7:
+				return ELiteralType::BooleanArray;
+
+			case 8:
+				return ELiteralType::IntegerArray;
+
+			case 9:
+				return ELiteralType::FloatArray;
+
+			case 10:
+				return ELiteralType::StringArray;
+
+			case 11:
+				return ELiteralType::UObjectProxyArray;
+
+			case 12:
+			default:
+				static_assert(TVariantSize<FVariantType>::Value == 13, "Possible missing FVariantType case coverage");
+				return ELiteralType::Invalid;
+		}
+	}
+
+	FLiteral FLiteral::Clone() const
+	{
+		// TODO: If Clone() is supported after reworking vertices,
+		// wrap the Audio::IProxyData in a object that manages virtual
+		// copy construction.  Then replace this switch with simple
+		// FLiteral copy constructor.
+		switch (GetType())
+		{
+		case ELiteralType::Invalid:
+			return CreateInvalid();
+
+		case ELiteralType::None:
+			return FLiteral(FNone());
+
+		case ELiteralType::Boolean:
+			return FLiteral(Value.Get<bool>());
+
+		case ELiteralType::Float:
+			return FLiteral(Value.Get<float>());
+
+		case ELiteralType::Integer:
+			return FLiteral(Value.Get<int32>());
+
+		case ELiteralType::String:
+			return FLiteral(Value.Get<FString>());
+
+		case ELiteralType::UObjectProxy:
+			if (const Audio::IProxyDataPtr& ProxyPtr = Value.Get<Audio::IProxyDataPtr>())
+			{
+				return FLiteral(ProxyPtr->Clone());
+			}
+			return CreateInvalid();
+
+		case ELiteralType::NoneArray:
+			return FLiteral(Value.Get<TArray<FNone>>());
+
+		case ELiteralType::BooleanArray:
+			return FLiteral(Value.Get<TArray<bool>>());
+
+		case ELiteralType::FloatArray:
+			return FLiteral(Value.Get<TArray<float>>());
+
+		case ELiteralType::IntegerArray:
+			return FLiteral(Value.Get<TArray<int32>>());
+
+		case ELiteralType::StringArray:
+			return FLiteral(Value.Get<TArray<FString>>());
+
+		case ELiteralType::UObjectProxyArray:
+			{
+				TArray<Audio::IProxyDataPtr> ProxyPtrArrayCopy;
+				for (const Audio::IProxyDataPtr& Ptr : Value.Get<TArray<Audio::IProxyDataPtr>>())
+				{
+					if (Ptr.IsValid())
+					{
+						ProxyPtrArrayCopy.Add(Ptr->Clone());
+					}
+				}
+				return FLiteral(MoveTemp(ProxyPtrArrayCopy));
+			}
+			return CreateInvalid();
+
+		default:
+			checkNoEntry();
+			return CreateInvalid();
+		}
+	}
+
+}
+
+
diff --git a/Engine/Plugins/Runtime/Metasound/Source/MetasoundGraphCore/Public/MetasoundFacade.h b/Engine/Plugins/Runtime/Metasound/Source/MetasoundGraphCore/Public/MetasoundFacade.h
--- a/Engine/Plugins/Runtime/Metasound/Source/MetasoundGraphCore/Public/MetasoundFacade.h
+++ b/Engine/Plugins/Runtime/Metasound/Source/MetasoundGraphCore/Public/MetasoundFacade.h
@@ -40,7 +40,7 @@ namespace Metasound
 			template<typename T, T> 
 			struct Helper;
 
-			// Check for "static const FNodeClassMetadata& U::CreateOperator(const FCreateOperatorParams& Inparams, FBuildErrorArray& OutErrors)"
+			// Check for "static const FNodeClassMetadata& U::GetNodeInfo()"
 			template<typename T>
 			static uint8 Check(Helper<const FNodeClassMetadata&(*)(), &T::GetNodeInfo>*);
 
diff --git a/Engine/Plugins/Runtime/Metasound/Source/MetasoundGraphCore/Public/MetasoundLiteral.h b/Engine/Plugins/Runtime/Metasound/Source/MetasoundGraphCore/Public/MetasoundLiteral.h
--- a/Engine/Plugins/Runtime/Metasound/Source/MetasoundGraphCore/Public/MetasoundLiteral.h
+++ b/Engine/Plugins/Runtime/Metasound/Source/MetasoundGraphCore/Public/MetasoundLiteral.h
@@ -5,39 +5,45 @@
 #include "CoreMinimal.h"
 
 #include "IAudioProxyInitializer.h"
-#include "MetasoundOperatorSettings.h" 
-#include "Misc/Guid.h"
 #include "Misc/TVariant.h"
 #include <type_traits>
 
 #define METASOUND_DEBUG_LITERALS 1
 
+// Forward Declare
+struct FGuid;
+class FString;
+
+
 namespace Metasound
 {
+	/** ELiteralType describes the format of the values held in the literal. */
 	enum class ELiteralType : uint8
 	{
-		None, // If the literal is None, TType(const FOperatorSettings&) or TType() will be invoked.
-		Boolean,
+		None, //< If the literal is None, TType(const FOperatorSettings&) or TType() will be invoked.
+		Boolean, 
 		Integer,
 		Float,
 		String,
 		UObjectProxy,
-		NoneArray, // If this is set, TType(const FOperatorSettings&) or TType() will be invoked for each element in the array. 
+		NoneArray, //< For NoneArray types,  TType(const FOperatorSettings&) or TType() will be invoked for each element in the array. 
 		BooleanArray,
 		IntegerArray,
 		FloatArray,
 		StringArray,
 		UObjectProxyArray,
-		Invalid,
+		Invalid  //< The literal is in an invalid state and cannot be used to construct an object.
 	};
 
 
 	/**
-	 * FLiteral 
-	 * Convenience wrapper for safely invoking the correct constructor.
-	 * This can be used for int32, float, bool, FString, Audio::IProxyData, or TArray<Audio::IProxyData>
+	 * FLiteral represents a constant value in a Metasound graph and is primarily
+	 * used to invoke the correct constructor of a Metasound data type. To be constructed
+	 * using a FLiteral, the Metasound data type must support a constructor which 
+	 * accepts one of the types supported by a FLiteral. The Metasound data type may 
+	 * optionally accept a `const FOperatorSettings&` in addition.
 	 *
-	 * For example:
+	 * Example:
 	 * // Somewhere before DECLARE_METASOUND_DATA_REFERENCE_TYPES(FAudioBuffer...), ParseFrom<FAudioBuffer>(int32) is defined
 	 * // which means this data type can be created from an integer.
 	 * Metasound::FAudioBuffer ::Metasound::ParseFrom<Metasound::FAudioBuffer>(int32 InNumFrames, const ::Metasound::FOperatorSettings&)
@@ -57,8 +63,9 @@ namespace Metasound
 	 *     //...
 	 * }
 	 */
-	struct FLiteral
+	struct METASOUNDGRAPHCORE_API FLiteral
 	{
+	public:
 		struct FInvalid {};
 
 		/* FNone is used in scenarios where an object is constructed with an FLiteral
@@ -67,217 +74,12 @@ namespace Metasound
 		 */
 		struct FNone {};
 
-		FLiteral(const FLiteral& Other) = delete;
-		FLiteral(FLiteral&& Other) = default;
-		FLiteral& operator=(FLiteral&& Other) = default;
-
 		using FVariantType = TVariant<
 			FNone, bool, int32, float, FString, Audio::IProxyDataPtr, // Single value types
 			TArray<FNone>, TArray<bool>, TArray<int32>, TArray<float>, TArray<FString>, TArray<Audio::IProxyDataPtr>, // Array of values types
 			FInvalid
 		>;
 
-		FVariantType Value;
-
-#if METASOUND_DEBUG_LITERALS
-		private:
-
-		FString DebugString;
-
-		void InitDebugString()
-		{
-			// GetIndex returns the index of the current type set in the TVariant.
-			// The index refers to the order of the template parameters to TVariant.
-			switch (Value.GetIndex())
-			{
-				case 0:
-					DebugString = TEXT("NONE");
-				break;
-
-				case 1:
-					DebugString = FString::Printf(TEXT("Bool: %s"), Value.Get<bool>() ? TEXT("true") : TEXT("false"));
-				break;
-
-				case 2:
-					DebugString = FString::Printf(TEXT("Int32: %d"), Value.Get<int32>());
-				break;
-
-				case 3:
-					DebugString = FString::Printf(TEXT("Float: %f"), Value.Get<float>());
-				break;
-
-				case 4:
-					DebugString = FString::Printf(TEXT("String: %s"), *Value.Get<FString>());
-				break;
-
-				case 5:
-					DebugString = FString::Printf(TEXT("Object: %f"), *Value.Get<Audio::IProxyDataPtr>()->GetProxyTypeName().ToString());
-				break;
-
-				case 6:
-					DebugString = TEXT("TArray<NONE>");
-				break;
-
-				case 7:
-					DebugString = FString::Printf(TEXT("TArray<Bool>"));
-				break;
-
-				case 8:
-					DebugString = FString::Printf(TEXT("TArray<int32>"));
-				break;
-
-				case 9:
-					DebugString = FString::Printf(TEXT("TArray<float>"));
-				break;
-
-				case 10:
-					DebugString = FString::Printf(TEXT("TArray<FString>"));
-				break;
-
-				case 11:
-					DebugString = FString::Printf(TEXT("TArray<UObject>"));
-				break;
-
-				case 12:
-					DebugString = TEXT("INVALID");
-				break;
-
-				default:
-					static_assert(TVariantSize<FVariantType>::Value == 13, "Possible missing FVariantType case coverage");
-					checkNoEntry();
-			}
-		}
-
-		public:
-#endif // METAOUND_DEBUG_LITERALS
-
-
-		// builds an invalid FLiteral.
-		static FLiteral CreateInvalid()
-		{
-			return FLiteral(FInvalid());
-		}
-
-		bool IsValid() const
-		{
-			return GetType() != ELiteralType::Invalid;
-		}
-
-		ELiteralType GetType() const
-		{
-			switch (Value.GetIndex())
-			{
-				case 0:
-					return ELiteralType::None;
-
-				case 1:
-					return ELiteralType::Boolean;
-
-				case 2:
-					return ELiteralType::Integer;
-
-				case 3:
-					return ELiteralType::Float;
-
-				case 4:
-					return ELiteralType::String;
-
-				case 5:
-					return ELiteralType::UObjectProxy;
-
-				case 6:
-					return ELiteralType::NoneArray;
-
-				case 7:
-					return ELiteralType::BooleanArray;
-
-				case 8:
-					return ELiteralType::IntegerArray;
-
-				case 9:
-					return ELiteralType::FloatArray;
-
-				case 10:
-					return ELiteralType::StringArray;
-
-				case 11:
-					return ELiteralType::UObjectProxyArray;
-
-				case 12:
-				default:
-					static_assert(TVariantSize<FVariantType>::Value == 13, "Possible missing FVariantType case coverage");
-					return ELiteralType::Invalid;
-			}
-		}
-
-		FLiteral Clone() const
-		{
-			// TODO: If Clone() is supported after reworking vertices,
-			// wrap the Audio::IProxyData in a object that manages virtual
-			// copy construction.  Then replace this switch with simple
-			// TVariant copy constructor.
-			switch (GetType())
-			{
-			case ELiteralType::Invalid:
-				return CreateInvalid();
-
-			case ELiteralType::None:
-				return FLiteral(FNone());
-
-			case ELiteralType::Boolean:
-				return FLiteral(Value.Get<bool>());
-
-			case ELiteralType::Float:
-				return FLiteral(Value.Get<float>());
-
-			case ELiteralType::Integer:
-				return FLiteral(Value.Get<int32>());
-
-			case ELiteralType::String:
-				return FLiteral(Value.Get<FString>());
-
-			case ELiteralType::UObjectProxy:
-				if (const Audio::IProxyDataPtr& ProxyPtr = Value.Get<Audio::IProxyDataPtr>())
-				{
-					return FLiteral(ProxyPtr->Clone());
-				}
-				return CreateInvalid();
-
-			case ELiteralType::NoneArray:
-				return FLiteral(Value.Get<TArray<FNone>>());
-
-			case ELiteralType::BooleanArray:
-				return FLiteral(Value.Get<TArray<bool>>());
-
-			case ELiteralType::FloatArray:
-				return FLiteral(Value.Get<TArray<float>>());
-
-			case ELiteralType::IntegerArray:
-				return FLiteral(Value.Get<TArray<int32>>());
-
-			case ELiteralType::StringArray:
-				return FLiteral(Value.Get<TArray<FString>>());
-
-			case ELiteralType::UObjectProxyArray:
-				{
-					TArray<Audio::IProxyDataPtr> ProxyPtrArrayCopy;
-					for (const Audio::IProxyDataPtr& Ptr : Value.Get<TArray<Audio::IProxyDataPtr>>())
-					{
-						if (Ptr.IsValid())
-						{
-							ProxyPtrArrayCopy.Add(Ptr->Clone());
-						}
-					}
-					return FLiteral(MoveTemp(ProxyPtrArrayCopy));
-				}
-				return CreateInvalid();
-
-			default:
-				checkNoEntry();
-				return CreateInvalid();
-			}
-		}
-
 		/** Construct a literal param with a single argument. */
 		template<
 			typename... ArgTypes,
@@ -304,6 +106,40 @@ namespace Metasound
 #endif
 		}
 
+		FLiteral(const FLiteral& Other) = delete;
+		FLiteral(FLiteral&& Other) = default;
+		FLiteral& operator=(FLiteral&& Other) = default;
+
+		FVariantType Value;
+
+		// builds an invalid FLiteral.
+		static FLiteral CreateInvalid();
+		static FLiteral GetDefaultForType(ELiteralType InType);
+
+		bool IsValid() const;
+
+		ELiteralType GetType() const;
+
+		FLiteral Clone() const;
+
+
+		template<typename ArgType>
+		void Set(ArgType&& Arg)
+		{
+			Value.Set<typename std::decay<ArgType>::type>(Forward<ArgType>(Arg));
+#if METASOUND_DEBUG_LITERALS
+			InitDebugString();
+#endif
+		}
+
+#if METASOUND_DEBUG_LITERALS
+	private:
+
+		FString DebugString;
+
+		void InitDebugString();
+
+#endif // METAOUND_DEBUG_LITERALS
 	};
 
 	namespace MetasoundLiteralIntrinsics
@@ -318,6 +154,7 @@ namespace Metasound
 		template<>
 		FORCEINLINE ELiteralType GetLiteralArgTypeFromDecayed<>()
 		{
+			checkNoEntry();
 			return ELiteralType::None;
 		}
 
@@ -386,6 +223,66 @@ namespace Metasound
 		{
 			return ELiteralType::UObjectProxyArray;
 		}
+
+		template<typename ... ArgTypes>
+		struct TLiteralDefaultValueFromDecayed {};
+
+		template<typename ArgType>
+		struct TLiteralDefaultValueFromDecayed<ArgType>
+		{
+			static const ArgType& GetValue()
+			{
+				static const ArgType Value = static_cast<ArgType>(0);
+				return Value;
+			}
+		};
+
+		template<>
+		struct TLiteralDefaultValueFromDecayed<FLiteral::FNone>
+		{
+			static FLiteral::FNone GetValue()
+			{
+				return FLiteral::FNone{};
+			}
+		};
+
+		template<>
+		struct TLiteralDefaultValueFromDecayed<FString>
+		{
+			static const FString& GetValue()
+			{
+				static const FString Value = TEXT("");
+				return Value;
+			}
+		};
+
+		template<typename ElementType>
+		struct TLiteralDefaultValueFromDecayed<TArray<ElementType>>
+		{
+			static const TArray<ElementType>& GetValue()
+			{
+				static const TArray<ElementType> Value;
+				return Value;
+			}
+		};
+
+		template<>
+		struct TLiteralDefaultValueFromDecayed<Audio::IProxyDataPtr>
+		{
+			static Audio::IProxyDataPtr GetValue()
+			{
+				return Audio::IProxyDataPtr(nullptr);
+			}
+		};
+
+		template<>
+		struct TLiteralDefaultValueFromDecayed<TArray<Audio::IProxyDataPtr>>
+		{
+			static TArray<Audio::IProxyDataPtr> GetValue()
+			{
+				return TArray<Audio::IProxyDataPtr>();
+			}
+		};
 	}
 
 
@@ -394,13 +291,22 @@ namespace Metasound
 	 * @tparam ArgType - A C++ type.
 	 */
 	template<typename... ArgTypes>
-	struct TLiteralTypeInfo
+	struct TLiteralTypeInfo{};
+
+	template<typename ArgType>
+	struct TLiteralTypeInfo<ArgType>
 	{
+		using FDecayedType = typename std::decay<ArgType>::type;
 		/** Returns the associated ELiteralType for the C++ type provided in the TLiteralTypeInfo<Type> */
 		static const ELiteralType GetLiteralArgTypeEnum()
 		{
 			// Use decayed version of template arg to remove references and cv qualifiers. 
-			return MetasoundLiteralIntrinsics::GetLiteralArgTypeFromDecayed<typename std::decay<ArgTypes...>::type>();
+			return MetasoundLiteralIntrinsics::GetLiteralArgTypeFromDecayed<FDecayedType>();
+		}
+
+		static auto GetDefaultValue()
+		{
+			return MetasoundLiteralIntrinsics::TLiteralDefaultValueFromDecayed<FDecayedType>::GetValue();
 		}
 	};
 
@@ -417,6 +323,11 @@ namespace Metasound
 			// Use decayed version of template arg to remove references and cv qualifiers. 
 			return MetasoundLiteralIntrinsics::GetLiteralArgTypeFromDecayed<>();
 		}
+
+		static FLiteral::FNone GetDefaultValue()
+		{
+			return FLiteral::FNone{};
+		}
 	};
 
 
