Metasound Enum Editor

diff --git a/Engine/Plugins/Runtime/Metasound/Source/MetasoundEditor/Private/MetasoundLiteralDescriptionDetailCustomization.cpp b/Engine/Plugins/Runtime/Metasound/Source/MetasoundEditor/Private/MetasoundLiteralDescriptionDetailCustomization.cpp
--- a/Engine/Plugins/Runtime/Metasound/Source/MetasoundEditor/Private/MetasoundLiteralDescriptionDetailCustomization.cpp
+++ b/Engine/Plugins/Runtime/Metasound/Source/MetasoundEditor/Private/MetasoundLiteralDescriptionDetailCustomization.cpp
@@ -301,50 +301,58 @@ namespace Metasound
 			// Create enum widget.
 			TSharedRef<SWidget> CreatePropertyValueWidgetInternal(TSharedPtr<IPropertyHandle> InLiteralPropertyHandle, bool bDisplayDefaultPropertyButtons = false) const override
 			{
+				using namespace Metasound::Frontend;
+				
 				if (TSharedPtr<IPropertyHandle> PropertyHandle = GetValueHandle(InLiteralPropertyHandle))
-				{
-					auto RemoveNamespace = [](const FString& InString) -> FString
-					{ 
-						FString EnumName;
-						if (InString.Split(TEXT("::"), nullptr, &EnumName))
+				{				
+					auto GetAll = [Interface = EnumInterface](TArray<TSharedPtr<FString>>& OutStrings, TArray<TSharedPtr<SToolTip>>& OutTooltips, TArray<bool>&)
+					{
+						for (const IEnumDataTypeInterface::FGenericInt32Entry& i : Interface->GetAllEntries())
 						{
-							return EnumName;
+							OutTooltips.Emplace(SNew(SToolTip).Text(i.Tooltip));
+							OutStrings.Emplace(MakeShared<FString>(i.DisplayName.ToString()));
 						}
-						return InString;
 					};
-
-					auto GetAll = [EnumInterface = EnumInterface, RemoveNamespace](TArray<TSharedPtr<FString>>& OutStrings, TArray<TSharedPtr<SToolTip>>& OutTooltips, TArray<bool>&)
+					auto GetValue = [Interface = EnumInterface, Prop = PropertyHandle]() -> FString
 					{
-						for (const Metasound::Frontend::IEnumDataTypeInterface::FGenericInt32Entry& i : EnumInterface->GetAllEntries())
+						int32 IntValue;
+						if (Prop->GetValue(IntValue) != FPropertyAccess::Success)
 						{
-							OutTooltips.Emplace(SNew(SToolTip).Text(i.Tooltip));
-
-							// Trim the namespace off for display.
-							OutStrings.Emplace(MakeShared<FString>(RemoveNamespace(i.Name.GetPlainNameString())));
+							IntValue = Interface->GetDefaultValue();
+							UE_LOG(LogTemp, Warning, TEXT("Failed ready int Property '%s', defaulting."), *GetNameSafe(Prop->GetProperty()));
 						}
-					};
+						if (TOptional<IEnumDataTypeInterface::FGenericInt32Entry> Result = Interface->FindByValue(IntValue))
+						{
+							return Result->DisplayName.ToString();
+						}
+						UE_LOG(LogTemp, Warning, TEXT("Failed to resolve int value '%d' to a valid enum value for enum '%s'"),
+							IntValue, *Interface->GetNamespace().ToString());
 
-					auto GetValue = [EnumInterface = EnumInterface, PropertyHandle = PropertyHandle, RemoveNamespace]() -> FString
+						// Return default (should always succeed as we can't have empty Enums and we must have a default).
+						return Interface->FindByValue(Interface->GetDefaultValue())->DisplayName.ToString();
+					};
+					auto SelectedValue = [Interface = EnumInterface, Prop = PropertyHandle](const FString& InSelected)
 					{
-						int32 IntValue;
-						if (PropertyHandle->GetValue(IntValue))
+						TOptional<IEnumDataTypeInterface::FGenericInt32Entry> Found =
+							Interface->FindEntryBy([TextSelected = FText::FromString(InSelected)](const IEnumDataTypeInterface::FGenericInt32Entry& i)
+						{ return i.DisplayName.EqualTo(TextSelected); });
+
+						if (Found)
 						{
-							if (TOptional<FName> Result = EnumInterface->ToName(IntValue))
+							// Only save the changes if its different and we can read the old value to check that.
+							int32 CurrentValue;
+							bool bReadCurrentValue = Prop->GetValue(CurrentValue) == FPropertyAccess::Success;
+							if ((bReadCurrentValue && CurrentValue != Found->Value) || !bReadCurrentValue)
 							{
-								return RemoveNamespace(Result->GetPlainNameString());
+								ensure(Prop->SetValue(Found->Value) == FPropertyAccess::Success);
 							}
-							UE_LOG(LogTemp, Warning, TEXT("Failed to Get Valid Value for Property '%s' with Value of '%d'"), *GetNameSafe(PropertyHandle->GetProperty()), IntValue);
 						}
-						return {};
-					};
-
-					auto SelectedValue = [EnumInterface = EnumInterface, PropertyHandle = PropertyHandle](const FString& InSelected)
-					{
-						FString FullyQualifiedName = FString::Printf(TEXT("%s::%s"), *EnumInterface->GetNamespace().GetPlainNameString(), *InSelected);
-						TOptional<int32> Result = EnumInterface->ToValue(*FullyQualifiedName);
-						if (ensure(Result))
+						else
 						{
-							PropertyHandle->SetValue(*Result);
+							UE_LOG(LogTemp, Warning, TEXT("Failed to Set Valid Value for Property '%s' with Value of '%s', writing default."),
+								*GetNameSafe(Prop->GetProperty()), *InSelected);
+
+							ensure(Prop->SetValue(Interface->GetDefaultValue()) == FPropertyAccess::Success);
 						}
 					};
 
diff --git a/Engine/Plugins/Runtime/Metasound/Source/MetasoundEditor/Private/SMetasoundEnumPin.cpp b/Engine/Plugins/Runtime/Metasound/Source/MetasoundEditor/Private/SMetasoundEnumPin.cpp
new file mode 100644
--- /dev/null
+++ b/Engine/Plugins/Runtime/Metasound/Source/MetasoundEditor/Private/SMetasoundEnumPin.cpp
@@ -0,0 +1,110 @@
+// Copyright Epic Games, Inc. All Rights Reserved.
+
+#include "SMetasoundEnumPin.h"
+#include "MetasoundEditorGraphNode.h"
+#include "ScopedTransaction.h"
+
+void SMetasoundEnumPin::Construct(const FArguments& InArgs, UEdGraphPin* InGraphPinObj)
+{
+	SGraphPin::Construct(SGraphPin::FArguments(), InGraphPinObj);
+}
+
+TSharedRef<SWidget> SMetasoundEnumPin::GetDefaultValueWidget()
+{
+	//Get list of enum indexes
+	TArray< TSharedPtr<int32> > ComboItems;
+	GenerateComboBoxIndexes(ComboItems);
+
+	//Create widget
+	return SAssignNew(ComboBox, SPinComboBox)
+		.ComboItemList(ComboItems)
+		.VisibleText(this, &SMetasoundEnumPin::OnGetText)
+		.OnSelectionChanged(this, &SMetasoundEnumPin::ComboBoxSelectionChanged)
+		.IsEnabled(this, &SGraphPin::GetDefaultValueIsEditable)
+		.Visibility(this, &SGraphPin::GetDefaultValueVisibility)
+		.OnGetDisplayName(this, &SMetasoundEnumPin::OnGetFriendlyName)
+		.OnGetTooltip(this, &SMetasoundEnumPin::OnGetTooltip);
+}
+
+TSharedPtr<const Metasound::Frontend::IEnumDataTypeInterface> 
+SMetasoundEnumPin::FindEnumInterfaceFromPin(UEdGraphPin* InPin) 
+{
+	auto MetasoundEditorNode = Cast<UMetasoundEditorGraphNode>(InPin->GetOwningNode());
+	Metasound::Frontend::FNodeHandle NodeHandle = MetasoundEditorNode->GetNodeHandle();
+	auto Inputs = NodeHandle->GetInputsWithVertexName(InPin->GetName());
+	if (Inputs.Num() > 0)
+	{
+		FName DataType = Inputs[0]->GetDataType();
+		return FMetasoundFrontendRegistryContainer::Get()->GetEnumInterfaceForDataType(DataType);
+	}
+	return nullptr;
+}
+
+FString SMetasoundEnumPin::OnGetText() const
+{
+	using namespace Metasound::Frontend;
+
+	TSharedPtr<const IEnumDataTypeInterface> EnumInterface = FindEnumInterfaceFromPin(GraphPinObj);
+	check(EnumInterface.IsValid());
+
+	int32 SelectedValue = FCString::Atoi(*GraphPinObj->GetDefaultAsString());	// Enums are currently serialized as ints (the value of the enum).
+	if (TOptional<IEnumDataTypeInterface::FGenericInt32Entry> Result = EnumInterface->FindByValue(SelectedValue))
+	{
+		return Result->DisplayName.ToString();
+	}
+	return TEXT("");
+}
+
+void SMetasoundEnumPin::GenerateComboBoxIndexes(TArray<TSharedPtr<int32>>& OutComboBoxIndexes)
+{
+	using namespace Metasound::Frontend;
+	TSharedPtr<const IEnumDataTypeInterface> EnumInterface = FindEnumInterfaceFromPin(GraphPinObj);
+	check(EnumInterface.IsValid());
+
+	const TArray<IEnumDataTypeInterface::FGenericInt32Entry>& Entries = EnumInterface->GetAllEntries();
+	for (int32 i = 0; i < Entries.Num(); ++i)
+	{
+		OutComboBoxIndexes.Add(MakeShared<int32>(i));
+	}
+}
+
+void SMetasoundEnumPin::ComboBoxSelectionChanged(TSharedPtr<int32> NewSelection, ESelectInfo::Type SelectInfo)
+{
+	using namespace Metasound::Frontend;
+	TSharedPtr<const IEnumDataTypeInterface> EnumInterface = FindEnumInterfaceFromPin(GraphPinObj);
+	check(EnumInterface.IsValid());
+
+	const TArray<IEnumDataTypeInterface::FGenericInt32Entry>& Entries = EnumInterface->GetAllEntries();
+
+	if (NewSelection.IsValid() && Entries.IsValidIndex(*NewSelection))
+	{
+		int32 EnumValue = Entries[*NewSelection].Value;
+		FString EnumValueString = FString::FromInt(EnumValue);
+		if (GraphPinObj->GetDefaultAsString() != EnumValueString)
+		{
+			const FScopedTransaction Transaction(NSLOCTEXT("GraphEditor", "ChangeEnumPinValue", "Change Enum Pin Value"));
+			GraphPinObj->Modify();
+
+			//Set new selection
+			GraphPinObj->GetSchema()->TrySetDefaultValue(*GraphPinObj, EnumValueString);
+		}
+	}
+}
+
+FText SMetasoundEnumPin::OnGetFriendlyName(int32 EnumIndex)
+{
+	TSharedPtr<const Metasound::Frontend::IEnumDataTypeInterface> Interface = FindEnumInterfaceFromPin(GraphPinObj);
+	check(Interface.IsValid());
+	const auto& Entries = Interface->GetAllEntries();
+	check(Entries.IsValidIndex(EnumIndex));
+	return Entries[EnumIndex].DisplayName;
+}
+
+FText SMetasoundEnumPin::OnGetTooltip(int32 EnumIndex)
+{
+	TSharedPtr<const Metasound::Frontend::IEnumDataTypeInterface> Interface = FindEnumInterfaceFromPin(GraphPinObj);
+	check(Interface.IsValid());
+	const auto& Entries = Interface->GetAllEntries();
+	check(Entries.IsValidIndex(EnumIndex));
+	return Entries[EnumIndex].Tooltip;
+}
diff --git a/Engine/Plugins/Runtime/Metasound/Source/MetasoundEditor/Private/SMetasoundEnumPin.h b/Engine/Plugins/Runtime/Metasound/Source/MetasoundEditor/Private/SMetasoundEnumPin.h
new file mode 100644
--- /dev/null
+++ b/Engine/Plugins/Runtime/Metasound/Source/MetasoundEditor/Private/SMetasoundEnumPin.h
@@ -0,0 +1,43 @@
+// Copyright Epic Games, Inc. All Rights Reserved.
+
+#pragma once
+
+#include "CoreMinimal.h"
+#include "Widgets/DeclarativeSyntaxSupport.h"
+#include "Widgets/SWidget.h"
+#include "SGraphPin.h"
+#include "SGraphPinComboBox.h"
+
+namespace Metasound
+{
+	namespace Frontend
+	{
+		struct IEnumDataTypeInterface;
+	}
+}
+
+class SMetasoundEnumPin : public SGraphPin
+{
+	SLATE_BEGIN_ARGS(SMetasoundEnumPin) {}
+	SLATE_END_ARGS()
+
+	void Construct(const FArguments& InArgs, UEdGraphPin* InGraphPinObj);
+
+	static TSharedPtr<const Metasound::Frontend::IEnumDataTypeInterface> FindEnumInterfaceFromPin(UEdGraphPin* InPin);
+
+protected:	
+	TSharedPtr<class SPinComboBox>	ComboBox;
+
+	TSharedRef<SWidget>	GetDefaultValueWidget();
+	
+	FString OnGetText() const;
+
+	void GenerateComboBoxIndexes(TArray<TSharedPtr<int32>>& OutComboBoxIndexes);
+
+	void ComboBoxSelectionChanged(TSharedPtr<int32> NewSelection, ESelectInfo::Type SelectInfo);
+
+	FText OnGetFriendlyName(int32 EnumIndex);
+
+	FText OnGetTooltip(int32 EnumIndex);
+};
+
diff --git a/Engine/Plugins/Runtime/Metasound/Source/MetasoundEditor/Private/SMetasoundGraphNode.cpp b/Engine/Plugins/Runtime/Metasound/Source/MetasoundEditor/Private/SMetasoundGraphNode.cpp
--- a/Engine/Plugins/Runtime/Metasound/Source/MetasoundEditor/Private/SMetasoundGraphNode.cpp
+++ b/Engine/Plugins/Runtime/Metasound/Source/MetasoundEditor/Private/SMetasoundGraphNode.cpp
@@ -27,6 +27,10 @@
 #include "Widgets/Images/SImage.h"
 #include "Widgets/SBoxPanel.h"
 #include "Widgets/SOverlay.h"
+#include "MetasoundFrontendRegistries.h"
+#include "PropertyCustomizationHelpers.h"
+#include "SGraphPinComboBox.h"
+#include "SMetasoundEnumPin.h"
 
 #define LOCTEXT_NAMESPACE "MetasoundGraphNode"
 
@@ -97,6 +101,10 @@ TSharedPtr<SGraphPin> SMetasoundGraphNode::CreatePinWidget(UEdGraphPin* InPin) c
 
 		if (InPin->PinType.PinCategory == FGraphBuilder::PinCategoryInt32)
 		{
+			if (SMetasoundEnumPin::FindEnumInterfaceFromPin(InPin))
+			{
+				return SNew(SMetasoundEnumPin, InPin);
+			}
 			return SNew(SGraphPinInteger, InPin);
 		}
 
diff --git a/Engine/Plugins/Runtime/Metasound/Source/MetasoundFrontend/Public/MetasoundDataTypeRegistrationMacro.h b/Engine/Plugins/Runtime/Metasound/Source/MetasoundFrontend/Public/MetasoundDataTypeRegistrationMacro.h
--- a/Engine/Plugins/Runtime/Metasound/Source/MetasoundFrontend/Public/MetasoundDataTypeRegistrationMacro.h
+++ b/Engine/Plugins/Runtime/Metasound/Source/MetasoundFrontend/Public/MetasoundDataTypeRegistrationMacro.h
@@ -134,34 +134,28 @@ namespace Metasound
 
 		struct FEnumHandler : Metasound::Frontend::IEnumDataTypeInterface
 		{
-			TArray<FName> GetAllNames() const override
-			{
-				return FStringHelper::GetAllNames();
-			}
 			FName GetNamespace() const override
 			{
 				return FStringHelper::GetNamespace();
 			}
-			TOptional<FName> ToName(int32 InEnumValue) const override
+			int32 GetDefaultValue() const override
 			{
-				return FStringHelper::ToName(static_cast<InnerType>(InEnumValue));
+				return static_cast<int32>(TEnumTraits<TDataType>::DefaultValue);
 			}
-			TOptional<int32> ToValue(FName InName) const override
+			const TArray<FGenericInt32Entry>& GetAllEntries() const override
 			{
-				if (TOptional<InnerType> Result = FStringHelper::FromName(InName))
+				auto BuildIntEntries = []()
 				{
-					return static_cast<int32>(*Result);
-				}
-				return {};
-			}
-			TArray<FGenericInt32Entry> GetAllEntries() const override
-			{
-				// Convert to int32 representation 
-				TArray<FGenericInt32Entry> IntEntries;
-				for (const TEnumEntry<InnerType>& i : FStringHelper::GetAllEntries())
-				{					
-					IntEntries.Emplace(FGenericInt32Entry{ static_cast<int32>(i.Value), i.Name, i.Tooltip });
-				}
+					// Convert to int32 representation 
+					TArray<FGenericInt32Entry> IntEntries;
+					IntEntries.Reserve(FStringHelper::GetAllEntries().Num());
+					for (const TEnumEntry<InnerType>& i : FStringHelper::GetAllEntries())
+					{
+						IntEntries.Emplace(i);
+					}
+					return IntEntries;
+				};
+				static const TArray<FGenericInt32Entry> IntEntries = BuildIntEntries();
 				return IntEntries;
 			}
 		};
diff --git a/Engine/Plugins/Runtime/Metasound/Source/MetasoundFrontend/Public/MetasoundEnum.h b/Engine/Plugins/Runtime/Metasound/Source/MetasoundFrontend/Public/MetasoundEnum.h
--- a/Engine/Plugins/Runtime/Metasound/Source/MetasoundFrontend/Public/MetasoundEnum.h
+++ b/Engine/Plugins/Runtime/Metasound/Source/MetasoundFrontend/Public/MetasoundEnum.h
@@ -23,9 +23,9 @@
 
 // 3. Define it using the BEGIN/ENTRY/END macros
 // DEFINE_METASOUND_ENUM_BEGIN(EMyOtherTestEnum)
-// 		DEFINE_METASOUND_ENUM_ENTRY(EMyOtherTestEnum::Alpha, LOCTEXT("AlphaDescription", "Alpha tooltip")),
-// 		DEFINE_METASOUND_ENUM_ENTRY(EMyOtherTestEnum::Beta, LOCTEXT("BetaDescription", "Beta tooltip")),
-// 		DEFINE_METASOUND_ENUM_ENTRY(EMyOtherTestEnum::Gamma, LOCTEXT("GammaDescription", "Gamma tooltip"))
+// 		DEFINE_METASOUND_ENUM_ENTRY(EMyOtherTestEnum::Alpha, LOCTEXT("AlphaDescription", "Alpha"), LOCTEXT("AlphaDescriptionTT", "Alpha tooltip")),
+// 		DEFINE_METASOUND_ENUM_ENTRY(EMyOtherTestEnum::Beta, LOCTEXT("BetaDescription", "Beta"), LOCTEXT("BetaDescriptioTT", "Beta tooltip")),
+// 		DEFINE_METASOUND_ENUM_ENTRY(EMyOtherTestEnum::Gamma, LOCTEXT("GammaDescription", "Gamma", LOCTEXT("GammaDescriptionTT", "Gamma tooltip"))
 // DEFINE_METASOUND_ENUM_END()
 
 namespace Metasound
@@ -36,7 +36,14 @@ namespace Metasound
 	{
 		T Value;
 		FName Name;
-		FText Tooltip;	// TODO: Remove this for runtime.
+		FText DisplayName;		// TODO: Remove this from runtime.
+		FText Tooltip;			// TODO: Remove this from runtime.
+
+		// Allow implicit conversion to int32 entry
+		operator TEnumEntry<int32>() const
+		{
+			return TEnumEntry<int32>{ static_cast<int32>(Value), Name, DisplayName, Tooltip };
+		}	
 	};
 
 	/** CRTP base class for Enum String Helper type.
@@ -73,7 +80,7 @@ namespace Metasound
 		static TArray<FName> GetAllNames()
 		{
 			TArray<FName> Names;
-			for (const TEnumEntry<EnumType> i : Derived::GetAllEntries())
+			for (const TEnumEntry<EnumType>& i : Derived::GetAllEntries())
 			{
 				Names.Emplace(i.Name);
 			}
@@ -99,7 +106,10 @@ namespace Metasound
 		static constexpr EnumType Default = DefaultValue;
 
 		static_assert(TIsEnum<EnumType>::Value, "Expecting an Enum type");
-		static_assert(TIsSame<typename std::underlying_type<EnumType>::type, int32>::Value, "We serialize to int32, so limit to that for now");
+		
+		// Allow Non-enum class Enum or Enum class that's are derived from int32
+		static_assert(TIsSame<typename std::underlying_type<EnumType>::type, int32>::Value || 
+			!TIsEnumClass<EnumType>::Value, "We serialize to int32, so limit to that for now");
 
 		// Default.
 		explicit TEnum(EnumType InValue = DefaultValue)
@@ -206,16 +216,17 @@ namespace Metasound
 //  - GetAllEntries() which returns all posible entries in the enum. 
 
 /** DECLARE_METASOUND_ENUM
- * @param ENUMNAME - The typename of your EnumType you want to use for Metasounds. e.g. My
+ * @param ENUMNAME - The typename of your EnumType you want to use for Metasounds. e.g. MyEnum
  * @param DEFAULT - A fully qualified default Enum value. e.g. EMyEnum::One
  * @param API - The module API this is declared inside e.g. METASOUNDSTANDARDNODES_API 
+ * @param ENUMTYPEDEF - The name of the TEnum<YourType> wrapper type
  * @param TYPEINFO - The name of the TypeInfo type you want to define e.g. FMyEnumTypeInfo
  * @param READREF - The name of the Read Reference type you want to define. e.g FMyEnumReadRef
  * @param WRITEREF -The name of the Write Reference type you want to define e.g. FMyEnumWriteRef
  */
-#define DECLARE_METASOUND_ENUM(ENUMNAME, DEFAULT, API, TYPEINFO, READREF, WRITEREF)\
-	using FEnum##ENUMNAME = Metasound::TEnum<ENUMNAME, DEFAULT>; \
-	DECLARE_METASOUND_DATA_REFERENCE_TYPES(FEnum##ENUMNAME, API, TYPEINFO, READREF, WRITEREF);\
+#define DECLARE_METASOUND_ENUM(ENUMNAME, DEFAULT, API, ENUMTYPEDEF, TYPEINFO, READREF, WRITEREF)\
+	using ENUMTYPEDEF = Metasound::TEnum<ENUMNAME, DEFAULT>; \
+	DECLARE_METASOUND_DATA_REFERENCE_TYPES(ENUMTYPEDEF, API, TYPEINFO, READREF, WRITEREF);\
 	template<> struct API Metasound::TEnumStringHelper<ENUMNAME> : Metasound::TEnumStringHelperBase<Metasound::TEnumStringHelper<ENUMNAME>, ENUMNAME>\
 	{\
 		static FName GetNamespace()\
@@ -227,24 +238,25 @@ namespace Metasound
 	};
 
 /** DEFINE_METASOUND_ENUM_BEGIN
- * @param ENUMNAME - The typename of your EnumType you want to use for Metasounds. e.g. EMyType
+ * @param ENUMNAME - The typename of your raw EnumType you want to use for Metasounds. e.g. EMyType
+ * @param ENUMTYPEDEF - The name of the TEnum<YourType> wrapper type
+ * @param DATATYPENAMESTRING - The string that will the data type name "Enum:<string>" e.g. "MyEnum"
  */
-#define DEFINE_METASOUND_ENUM_BEGIN(ENUMNAME)\
-	REGISTER_METASOUND_DATATYPE(FEnum##ENUMNAME, "Enum:" #ENUMNAME, ::Metasound::ELiteralType::Integer);\
+#define DEFINE_METASOUND_ENUM_BEGIN(ENUMNAME,ENUMTYPEDEF,DATATYPENAMESTRING)\
+	REGISTER_METASOUND_DATATYPE(ENUMTYPEDEF, "Enum:" DATATYPENAMESTRING, ::Metasound::ELiteralType::Integer);\
 	TArrayView<const Metasound::TEnumEntry<ENUMNAME>> Metasound::TEnumStringHelper<ENUMNAME>::GetAllEntries()\
 	{\
 		static const Metasound::TEnumEntry<ENUMNAME> Entries[] = {
 
 /** DEFINE_METASOUND_ENUM_ENTRY - defines a single Enum Entry
   * @param ENTRY - Fully Qualified Name of Entry of the Enum. (e.g. EMyType::One)
-  * @param TOOLTIP - FText 
+  * @param DISPLAYNAME - FText Presented to User
+  * @param TOOLTIP - FText displayed as a tooltip
   */
-	#define DEFINE_METASOUND_ENUM_ENTRY(ENTRY, TOOLTIP) { ENTRY, TEXT(#ENTRY), TOOLTIP }
+	#define DEFINE_METASOUND_ENUM_ENTRY(ENTRY, DISPLAYNAME, TOOLTIP) { ENTRY, TEXT(#ENTRY), DISPLAYNAME, TOOLTIP }
 
 /** DEFINE_METASOUND_ENUM_END - macro which ends the function body of the GetAllEntries function
-  * @param ENTRY - Fully Qualified Name of Entry of the Enum. (e.g. EMyType::One)
-  * @param TOOLTIP - Non-wide Quoted String use for the tooltip.
-  */
+   */
 #define DEFINE_METASOUND_ENUM_END() \
 		};\
 		return Entries;\
diff --git a/Engine/Plugins/Runtime/Metasound/Source/MetasoundFrontend/Public/MetasoundFrontendRegistries.h b/Engine/Plugins/Runtime/Metasound/Source/MetasoundFrontend/Public/MetasoundFrontendRegistries.h
--- a/Engine/Plugins/Runtime/Metasound/Source/MetasoundFrontend/Public/MetasoundFrontendRegistries.h
+++ b/Engine/Plugins/Runtime/Metasound/Source/MetasoundFrontend/Public/MetasoundFrontendRegistries.h
@@ -138,15 +138,45 @@ namespace Metasound
 		{
 			using FGenericInt32Entry = TEnumEntry<int32>;
 
+			virtual ~IEnumDataTypeInterface() = default;
+			virtual const TArray<FGenericInt32Entry>& GetAllEntries() const = 0;
 			virtual FName GetNamespace() const = 0;
-			virtual TArray<FGenericInt32Entry> GetAllEntries() const = 0;
-
-			virtual TArray<FName> GetAllNames() const = 0;
+			virtual int32 GetDefaultValue() const = 0;
 
-			virtual TOptional<FName> ToName(int32 InEnumValue) const = 0;
-			virtual TOptional<int32> ToValue(FName InName) const = 0;
-
-			virtual ~IEnumDataTypeInterface() = default;
+			template<typename Predicate>
+			TOptional<FGenericInt32Entry> FindEntryBy(Predicate Pred) const
+			{
+				TArray<FGenericInt32Entry> Entries = GetAllEntries();
+				if (FGenericInt32Entry* Found = Entries.FindByPredicate(Pred))
+				{
+					return *Found;
+				}
+				return {};
+			}
+			TOptional<FGenericInt32Entry> FindByValue(int32 InEnumValue) const
+			{
+				return FindEntryBy([InEnumValue](const FGenericInt32Entry& i) -> bool { return i.Value == InEnumValue; });
+			}
+			TOptional<FGenericInt32Entry> FindByName(FName InEnumName) const
+			{
+				return FindEntryBy([InEnumName](const FGenericInt32Entry& i) -> bool { return i.Name == InEnumName; });
+			}
+			TOptional<FName> ToName(int32 InEnumValue) const
+			{
+				if(TOptional<FGenericInt32Entry> Result = FindByValue(InEnumValue))
+				{
+					return Result->Name;
+				}
+				return {};
+			}
+			TOptional<int32> ToValue(FName InName) const
+			{
+				if (TOptional<FGenericInt32Entry> Result = FindByName(InName))
+				{
+					return Result->Value;
+				}
+				return {};
+			}
 		};
 	}
 
diff --git a/Engine/Plugins/Runtime/Metasound/Source/MetasoundStandardNodes/Private/MetasoundBasicFilters.cpp b/Engine/Plugins/Runtime/Metasound/Source/MetasoundStandardNodes/Private/MetasoundBasicFilters.cpp
--- a/Engine/Plugins/Runtime/Metasound/Source/MetasoundStandardNodes/Private/MetasoundBasicFilters.cpp
+++ b/Engine/Plugins/Runtime/Metasound/Source/MetasoundStandardNodes/Private/MetasoundBasicFilters.cpp
@@ -8,6 +8,7 @@
 #include "MetasoundNodeRegistrationMacro.h"
 #include "MetasoundStandardNodesNames.h"
 #include "MetasoundOperatorSettings.h"
+#include "MetasoundDataTypeRegistrationMacro.h"
 
 #include "DSP/Filter.h"
 #include "DSP/InterpolatedOnePole.h"
@@ -685,7 +686,22 @@ namespace Metasound
 
 
 #pragma region Biquad Filter
-
+	DECLARE_METASOUND_ENUM(Audio::EBiquadFilter::Type, Audio::EBiquadFilter::Lowpass,
+		METASOUNDSTANDARDNODES_API, FEnumEBiquadFilterType, FEnumBiQuadFilterTypeInfo, FEnumBiQuadFilterReadRef, FEnumBiQuadFilterWriteRef);
+
+	DEFINE_METASOUND_ENUM_BEGIN(Audio::EBiquadFilter::Type, FEnumEBiquadFilterType, "BiquadFilterType")
+		DEFINE_METASOUND_ENUM_ENTRY(Audio::EBiquadFilter::Lowpass, LOCTEXT("LpDescription", "Low Pass"), LOCTEXT("LpDescriptionTT", "Low Pass Biquad Filter")),
+		DEFINE_METASOUND_ENUM_ENTRY(Audio::EBiquadFilter::Highpass, LOCTEXT("HpDescription", "High Pass"), LOCTEXT("HpDescriptionTT", "High Pass Biquad Filter")),
+		DEFINE_METASOUND_ENUM_ENTRY(Audio::EBiquadFilter::Bandpass, LOCTEXT("BpDescription", "Band Pass"), LOCTEXT("BpDescriptionTT", "Band Pass Biquad Filter")),
+		DEFINE_METASOUND_ENUM_ENTRY(Audio::EBiquadFilter::Notch, LOCTEXT("NotchDescription", "Notch "), LOCTEXT("NotchDescriptionTT", "Notch Biquad Filter")),
+		DEFINE_METASOUND_ENUM_ENTRY(Audio::EBiquadFilter::ParametricEQ, LOCTEXT("ParaEqDescription", "Parametric EQ"), LOCTEXT("ParaEqDescriptionTT", "Parametric EQ Biquad Filter")),
+		DEFINE_METASOUND_ENUM_ENTRY(Audio::EBiquadFilter::LowShelf, LOCTEXT("LowShelfDescription", "Low Shelf"), LOCTEXT("LowShelfDescriptionTT", "Low Shelf Biquad Filter")),
+		DEFINE_METASOUND_ENUM_ENTRY(Audio::EBiquadFilter::HighShelf, LOCTEXT("HighShelfDescription", "High Shelf"), LOCTEXT("HighShelfDescriptionTT", "High Shelf Biquad Filter")),
+		DEFINE_METASOUND_ENUM_ENTRY(Audio::EBiquadFilter::AllPass, LOCTEXT("AllPassDescription", "All Pass"), LOCTEXT("AllPassDescriptionTT", "All Pass Biquad Filter")),
+		DEFINE_METASOUND_ENUM_ENTRY(Audio::EBiquadFilter::ButterworthLowPass, LOCTEXT("LowPassButterDescription", "Butterworth Low Pass"), LOCTEXT("LowPassButterDescriptionTT", "Butterworth Low Pass Biquad Filter")),
+		DEFINE_METASOUND_ENUM_ENTRY(Audio::EBiquadFilter::ButterworthHighPass,LOCTEXT("HighPassButterDescription", "Butterworth High Pass"), LOCTEXT("HighPassButterDescriptionTT", "Butterworth High Pass Biquad Filter"))
+	DEFINE_METASOUND_ENUM_END()
+	
 	class FBiquadFilterOperator : public TExecutableOperator<FBiquadFilterOperator>
 	{
 	public:
@@ -701,6 +717,7 @@ namespace Metasound
 			InputDataReferences.AddDataReadReference(TEXT("In"), FAudioBufferReadRef(AudioInput));
 			InputDataReferences.AddDataReadReference(TEXT("Freq."), FFloatReadRef(Frequency));
 			InputDataReferences.AddDataReadReference(TEXT("Bandwidth"), FFloatReadRef(Bandwidth));
+			InputDataReferences.AddDataReadReference(TEXT("Type"), FEnumBiQuadFilterReadRef(FilterType));
 
 			return InputDataReferences;
 		}
@@ -722,6 +739,7 @@ namespace Metasound
 		FAudioBufferReadRef AudioInput;
 		FFloatReadRef Frequency;
 		FFloatReadRef Bandwidth;
+		FEnumBiQuadFilterReadRef FilterType;
 
 		// cached data
 		float LastFrequency;
@@ -745,19 +763,22 @@ namespace Metasound
 			const FOperatorSettings& InSettings,
 			const FAudioBufferReadRef& InAudioInput,
 			const FFloatReadRef& InFrequency,
-			const FFloatReadRef& InBandwidth
+			const FFloatReadRef& InBandwidth,
+			const FEnumBiQuadFilterReadRef& InFilterType
 		)
 			: AudioInput(InAudioInput)
 			, Frequency(InFrequency)
 			, Bandwidth(InBandwidth)
+			, FilterType(InFilterType)
 			, AudioOutput(FAudioBufferWriteRef::CreateNew(InSettings))
 			, BlockSize(InSettings.GetNumFramesPerBlock())
 			, SampleRate(InSettings.GetSampleRate())
+			, LastFilterType(*FilterType)
 		{
 			// verify our buffer sizes:
 			check(AudioOutput->Num() == BlockSize);
 
-			BiquadFilter.Init(SampleRate, 1, Audio::EBiquadFilter::Type::ButterworthLowPass);
+			BiquadFilter.Init(SampleRate, 1, *FilterType);
 		}
 	};
 
@@ -799,7 +820,8 @@ namespace Metasound
 			FInputVertexInterface(
 				TInputDataVertexModel<FAudioBuffer>(TEXT("In"), LOCTEXT("AudioInputTooltip", "Audio Input")),
 				TInputDataVertexModel<float>(TEXT("Freq."), LOCTEXT("FrequencyTooltip", "Controls Filter Cutoff")),
-				TInputDataVertexModel<float>(TEXT("Bandwidth"), LOCTEXT("BandwidthTooltip", "Bandwidth Control"))
+				TInputDataVertexModel<float>(TEXT("Bandwidth"), LOCTEXT("BandwidthTooltip", "Bandwidth Control")),
+				TInputDataVertexModel<FEnumEBiquadFilterType>(TEXT("Type"), LOCTEXT("FilterTypeDescription", "Biquad Filter Type"))
 			),
 			FOutputVertexInterface(
 				TOutputDataVertexModel<FAudioBuffer>(TEXT("Out"), LOCTEXT("FilterOutputToolTip", "Audio Output"))
@@ -819,12 +841,14 @@ namespace Metasound
 		FAudioBufferReadRef AudioIn = InputDataRefs.GetDataReadReferenceOrConstruct<FAudioBuffer>(TEXT("In"), InParams.OperatorSettings);
 		FFloatReadRef FrequencyIn = InputDataRefs.GetDataReadReferenceOrConstruct<float>(TEXT("Freq."));
 		FFloatReadRef BandwidthIn = InputDataRefs.GetDataReadReferenceOrConstruct<float>(TEXT("Bandwidth"));
+		FEnumBiQuadFilterReadRef FilterType = InputDataRefs.GetDataReadReferenceOrConstruct<FEnumEBiquadFilterType>(TEXT("Type"));
 
 		return MakeUnique<FBiquadFilterOperator>(
 			InParams.OperatorSettings
 			, AudioIn
 			, FrequencyIn
 			, BandwidthIn
+			, FilterType
 			);
 	}
 
@@ -840,6 +864,12 @@ namespace Metasound
 			BiquadFilter.SetBandwidth(*Bandwidth);
 		}
 
+		if (*FilterType != LastFilterType)
+		{
+			BiquadFilter.SetType(*FilterType);
+			LastFilterType = *FilterType;
+		}
+
 		BiquadFilter.ProcessAudio(AudioInput->GetData(), AudioInput->Num(), AudioOutput->GetData());
 
 		LastFrequency = *Frequency;
