Fixing typo in wave selector node

diff --git a/Engine/Plugins/Runtime/Metasound/Source/MetasoundEngine/Private/MetasoundWaveSelectorNode.cpp b/Engine/Plugins/Runtime/Metasound/Source/MetasoundEngine/Private/MetasoundWaveSelectorNode.cpp
--- a/Engine/Plugins/Runtime/Metasound/Source/MetasoundEngine/Private/MetasoundWaveSelectorNode.cpp
+++ b/Engine/Plugins/Runtime/Metasound/Source/MetasoundEngine/Private/MetasoundWaveSelectorNode.cpp
@@ -9,7 +9,7 @@
 #include "MetasoundTrigger.h"
 #include "MetasoundWave.h"
 
-#define LOCTEXT_NAMESPACE "MetasoundWaveNode"
+#define LOCTEXT_NAMESPACE "MetasoundWaveSelectorNode"
 
 
 namespace Metasound
