- General Metasound Type Reduction
- Reduced Metasound Time Types to single FTime
- Show AudioBuffer as just "Audio" in UI/UX in preparation for nuking mono/stereo type ("mono" will be implicit for single audio type)
- Removed Frequency type
- Removed exposition of Gain as registered type (left code for potential use in node development)
- Added Nodes to convert accordingly between removed types

diff --git a/Engine/Plugins/Runtime/Metasound/Source/MetasoundEditor/Private/MetasoundEditorGraphBuilder.cpp b/Engine/Plugins/Runtime/Metasound/Source/MetasoundEditor/Private/MetasoundEditorGraphBuilder.cpp
--- a/Engine/Plugins/Runtime/Metasound/Source/MetasoundEditor/Private/MetasoundEditorGraphBuilder.cpp
+++ b/Engine/Plugins/Runtime/Metasound/Source/MetasoundEditor/Private/MetasoundEditorGraphBuilder.cpp
@@ -27,7 +27,7 @@ namespace Metasound
 {
 	namespace Editor
 	{
-		const FName FGraphBuilder::PinCategoryAudioFormat = "audio_format";
+		const FName FGraphBuilder::PinCategoryAudio = "audio";
 		const FName FGraphBuilder::PinCategoryBoolean = "bool";
 		const FName FGraphBuilder::PinCategoryDouble = "double";
 		const FName FGraphBuilder::PinCategoryFloat = "float";
diff --git a/Engine/Plugins/Runtime/Metasound/Source/MetasoundEditor/Private/MetasoundEditorGraphBuilder.h b/Engine/Plugins/Runtime/Metasound/Source/MetasoundEditor/Private/MetasoundEditorGraphBuilder.h
--- a/Engine/Plugins/Runtime/Metasound/Source/MetasoundEditor/Private/MetasoundEditorGraphBuilder.h
+++ b/Engine/Plugins/Runtime/Metasound/Source/MetasoundEditor/Private/MetasoundEditorGraphBuilder.h
@@ -29,7 +29,7 @@ namespace Metasound
 		class FGraphBuilder
 		{
 		public:
-			static const FName PinCategoryAudioFormat;  // Audio formats (ex. Buffers, Mono, Stereo)
+			static const FName PinCategoryAudio;
 			static const FName PinCategoryBoolean;
 			static const FName PinCategoryDouble;
 			static const FName PinCategoryFloat;
diff --git a/Engine/Plugins/Runtime/Metasound/Source/MetasoundEditor/Private/MetasoundEditorGraphSchema.cpp b/Engine/Plugins/Runtime/Metasound/Source/MetasoundEditor/Private/MetasoundEditorGraphSchema.cpp
--- a/Engine/Plugins/Runtime/Metasound/Source/MetasoundEditor/Private/MetasoundEditorGraphSchema.cpp
+++ b/Engine/Plugins/Runtime/Metasound/Source/MetasoundEditor/Private/MetasoundEditorGraphSchema.cpp
@@ -102,7 +102,7 @@ namespace Metasound
 			const UMetasoundEditorSettings* Settings = GetDefault<UMetasoundEditorSettings>();
 			check(Settings);
 
-			if (PinType.PinCategory == FGraphBuilder::PinCategoryAudioFormat)
+			if (PinType.PinCategory == FGraphBuilder::PinCategoryAudio)
 			{
 				return Settings->AudioPinTypeColor;
 			}
@@ -245,7 +245,7 @@ namespace Metasound
 				{
 					OutParams.WireThickness = Settings->DefaultExecutionWireThickness;
 				}
-				else if (InputPin->PinType.PinCategory == FGraphBuilder::PinCategoryAudioFormat)
+				else if (InputPin->PinType.PinCategory == FGraphBuilder::PinCategoryAudio)
 				{
 					OutParams.bDrawBubbles = true;
 				}
diff --git a/Engine/Plugins/Runtime/Metasound/Source/MetasoundEditor/Private/MetasoundEditorModule.cpp b/Engine/Plugins/Runtime/Metasound/Source/MetasoundEditor/Private/MetasoundEditorModule.cpp
--- a/Engine/Plugins/Runtime/Metasound/Source/MetasoundEditor/Private/MetasoundEditorModule.cpp
+++ b/Engine/Plugins/Runtime/Metasound/Source/MetasoundEditor/Private/MetasoundEditorModule.cpp
@@ -213,16 +213,11 @@ namespace Metasound
 							case ELiteralType::FloatArray:
 							default:
 							{
-								// Audio types are ubiquitous, so add as subcategory
-								// to be able to stylize connections (i.e. wire color)
-								if (DataTypeName == "Audio:Buffer"
-									|| DataTypeName == "Audio:Unformatted"
-									|| DataTypeName == "Audio:Mono"
-									|| DataTypeName == "Audio:Stereo"
-									|| DataTypeName == "Audio:Multichannel"
-								)
+								// Audio types are ubiquitous, so added as subcategory
+								// to be able to stylize connections (i.e. wire color & wire animation)
+								if (DataTypeName == "Audio")
 								{
-									PinCategory = FGraphBuilder::PinCategoryAudioFormat;
+									PinCategory = FGraphBuilder::PinCategoryAudio;
 								}
 								static_assert(static_cast<int32>(ELiteralType::Invalid) == 12, "Possible missing binding of pin category to primitive type");
 							}
diff --git a/Engine/Plugins/Runtime/Metasound/Source/MetasoundEditor/Private/SMetasoundGraphNode.cpp b/Engine/Plugins/Runtime/Metasound/Source/MetasoundEditor/Private/SMetasoundGraphNode.cpp
--- a/Engine/Plugins/Runtime/Metasound/Source/MetasoundEditor/Private/SMetasoundGraphNode.cpp
+++ b/Engine/Plugins/Runtime/Metasound/Source/MetasoundEditor/Private/SMetasoundGraphNode.cpp
@@ -79,7 +79,7 @@ TSharedPtr<SGraphPin> SMetasoundGraphNode::CreatePinWidget(UEdGraphPin* InPin) c
 
 	if (const UMetasoundEditorGraphSchema* GraphSchema = Cast<const UMetasoundEditorGraphSchema>(InPin->GetSchema()))
 	{
-		if (InPin->PinType.PinCategory == FGraphBuilder::PinCategoryAudioFormat)
+		if (InPin->PinType.PinCategory == FGraphBuilder::PinCategoryAudio)
 		{
 			return SNew(SGraphPin, InPin);
 		}
diff --git a/Engine/Plugins/Runtime/Metasound/Source/MetasoundEngine/Private/MetasoundSource.cpp b/Engine/Plugins/Runtime/Metasound/Source/MetasoundEngine/Private/MetasoundSource.cpp
--- a/Engine/Plugins/Runtime/Metasound/Source/MetasoundEngine/Private/MetasoundSource.cpp
+++ b/Engine/Plugins/Runtime/Metasound/Source/MetasoundEngine/Private/MetasoundSource.cpp
@@ -13,6 +13,7 @@
 #include "MetasoundGenerator.h"
 #include "MetasoundInstanceTransmitter.h"
 #include "MetasoundLog.h"
+#include "MetasoundOperatorSettings.h"
 #include "MetasoundPrimitives.h"
 #include "MetasoundReceiveNode.h"
 #include "MetasoundTrigger.h"
@@ -36,7 +37,7 @@ UMetasoundSource::UMetasoundSource(const FObjectInitializer& ObjectInitializer)
 	bLooping = true;
 	
 	// todo: ensure that we have a method so that the audio engine can be authoritative over the sample rate the UMetasoundSource runs at.
-	SampleRate = 48000.0f;
+	SampleRate = 48000.f;
 
 }
 
@@ -96,7 +97,7 @@ ISoundGeneratorPtr UMetasoundSource::CreateSoundGenerator(const FSoundGeneratorI
 	VirtualizationMode = EVirtualizationMode::PlayWhenSilent;
 
 	SampleRate = InParams.SampleRate;
-	FOperatorSettings InSettings = GetOperatorSettings(SampleRate);
+	FOperatorSettings InSettings = GetOperatorSettings(static_cast<FSampleRate>(SampleRate));
 	FMetasoundEnvironment Environment;
 
 	// Add audio device ID to environment.
@@ -365,7 +366,7 @@ TArray<FString> UMetasoundSource::GetTransmittableInputVertexNames() const
 	return GraphInputVertexNames;
 }
 
-Metasound::FOperatorSettings UMetasoundSource::GetOperatorSettings(float InSampleRate) const
+Metasound::FOperatorSettings UMetasoundSource::GetOperatorSettings(Metasound::FSampleRate InSampleRate) const
 {
 	constexpr float BlockRate = 100.f; // Metasound graph gets evaluated 100 times per second.
 	return Metasound::FOperatorSettings(InSampleRate, BlockRate);
diff --git a/Engine/Plugins/Runtime/Metasound/Source/MetasoundEngine/Public/MetasoundSource.h b/Engine/Plugins/Runtime/Metasound/Source/MetasoundEngine/Public/MetasoundSource.h
--- a/Engine/Plugins/Runtime/Metasound/Source/MetasoundEngine/Public/MetasoundSource.h
+++ b/Engine/Plugins/Runtime/Metasound/Source/MetasoundEngine/Public/MetasoundSource.h
@@ -6,6 +6,7 @@
 #include "MetasoundFrontend.h"
 #include "MetasoundFrontendDocument.h"
 #include "MetasoundInstanceTransmitter.h"
+#include "MetasoundOperatorSettings.h"
 #include "MetasoundRouter.h"
 
 #include "Sound/SoundWaveProcedural.h"
@@ -144,7 +145,7 @@ private:
 	Metasound::Frontend::FNodeHandle AddInputPinForSendAddress(const Metasound::FMetasoundInstanceTransmitter::FSendInfo& InSendInfo, Metasound::Frontend::FGraphHandle InGraph) const;
 	bool CopyDocumentAndInjectReceiveNodes(uint64 InInstanceID, const FMetasoundFrontendDocument& InSourceDoc, FMetasoundFrontendDocument& OutDestDoc) const;	
 	TArray<FString> GetTransmittableInputVertexNames() const;
-	Metasound::FOperatorSettings GetOperatorSettings(float InSampleRate) const;
+	Metasound::FOperatorSettings GetOperatorSettings(Metasound::FSampleRate InSampleRate) const;
 	Metasound::FSendAddress CreateSendAddress(uint64 InInstanceID, const FString& InVertexName) const;
 	
 
diff --git a/Engine/Plugins/Runtime/Metasound/Source/MetasoundFrontend/Private/MetasoundFrontend.cpp b/Engine/Plugins/Runtime/Metasound/Source/MetasoundFrontend/Private/MetasoundFrontend.cpp
--- a/Engine/Plugins/Runtime/Metasound/Source/MetasoundFrontend/Private/MetasoundFrontend.cpp
+++ b/Engine/Plugins/Runtime/Metasound/Source/MetasoundFrontend/Private/MetasoundFrontend.cpp
@@ -195,7 +195,7 @@ REGISTER_METASOUND_DATATYPE(FString, "String", ::Metasound::ELiteralType::String
 
 REGISTER_METASOUND_DATATYPE(TArray<float>, "FloatArray", ::Metasound::ELiteralType::FloatArray)
 
-REGISTER_METASOUND_DATATYPE(Metasound::FAudioBuffer, "Audio:Buffer")
-REGISTER_METASOUND_DATATYPE(Metasound::FSendAddress, "Transmission:Address")
+REGISTER_METASOUND_DATATYPE(Metasound::FAudioBuffer, "Audio")
+REGISTER_METASOUND_DATATYPE(Metasound::FSendAddress, "Transmission:Address", ::Metasound::ELiteralType::String)
 
 IMPLEMENT_MODULE(FMetasoundFrontendModule, MetasoundFrontend);
diff --git a/Engine/Plugins/Runtime/Metasound/Source/MetasoundGenerator/Private/MetasoundGenerator.cpp b/Engine/Plugins/Runtime/Metasound/Source/MetasoundGenerator/Private/MetasoundGenerator.cpp
--- a/Engine/Plugins/Runtime/Metasound/Source/MetasoundGenerator/Private/MetasoundGenerator.cpp
+++ b/Engine/Plugins/Runtime/Metasound/Source/MetasoundGenerator/Private/MetasoundGenerator.cpp
@@ -2,7 +2,6 @@
 #include "MetasoundGenerator.h"
 
 #include "DSP/Dsp.h"
-#include "MetasoundFrequency.h"
 #include "MetasoundGraph.h"
 #include "MetasoundInputNode.h"
 #include "MetasoundOperatorBuilder.h"
diff --git a/Engine/Plugins/Runtime/Metasound/Source/MetasoundGraphCore/Private/MetasoundOperatorSettings.cpp b/Engine/Plugins/Runtime/Metasound/Source/MetasoundGraphCore/Private/MetasoundOperatorSettings.cpp
--- a/Engine/Plugins/Runtime/Metasound/Source/MetasoundGraphCore/Private/MetasoundOperatorSettings.cpp
+++ b/Engine/Plugins/Runtime/Metasound/Source/MetasoundGraphCore/Private/MetasoundOperatorSettings.cpp
@@ -20,17 +20,17 @@ namespace Metasound
 		static constexpr int32 MaximumNumFramesPerBlock = 16384;
 	}
 
-	FOperatorSettings::FOperatorSettings(float InSampleRate, float InTargetBlockRate)
+	FOperatorSettings::FOperatorSettings(FSampleRate InSampleRate, float InTargetBlockRate)
 	{
 		SetSampleRate(InSampleRate);
 		SetTargetBlockRate(InTargetBlockRate);
 	}
 
-	void FOperatorSettings::SetSampleRate(float InSampleRate)
+	void FOperatorSettings::SetSampleRate(FSampleRate InSampleRate)
 	{
 		using namespace MetasoundOperatorSettingsPrivate;
 
-		if (ensureMsgf(InSampleRate > 0.f, TEXT("Invalid sample rate %f. Defaulting to %f"), InSampleRate, DefaultSampleRate))
+		if (ensureMsgf(InSampleRate > 0, TEXT("Invalid sample rate %d. Defaulting to %d"), InSampleRate, DefaultSampleRate))
 		{
 			SampleRate = InSampleRate;
 		}
diff --git a/Engine/Plugins/Runtime/Metasound/Source/MetasoundGraphCore/Public/MetasoundDataReferenceCollection.h b/Engine/Plugins/Runtime/Metasound/Source/MetasoundGraphCore/Public/MetasoundDataReferenceCollection.h
--- a/Engine/Plugins/Runtime/Metasound/Source/MetasoundGraphCore/Public/MetasoundDataReferenceCollection.h
+++ b/Engine/Plugins/Runtime/Metasound/Source/MetasoundGraphCore/Public/MetasoundDataReferenceCollection.h
@@ -4,6 +4,7 @@
 
 #include "CoreMinimal.h"
 #include "MetasoundDataReference.h"
+#include "MetasoundOperatorSettings.h"
 #include "MetasoundVertex.h"
 
 namespace Metasound
@@ -226,7 +227,7 @@ namespace Metasound
 			 * @return A readable data reference.
 			 */
 			template<typename DataType, typename DefaultDataType = DataType, typename... ConstructorArgTypes>
-			TDataReadReference<DataType> GetDataReadReferenceOrConstructWithVertexDefault(const FInputVertexInterface& InputVertices, const FString& InName, ConstructorArgTypes&&... ConstructorArgs) const
+			TDataReadReference<DataType> GetDataReadReferenceOrConstructWithVertexDefault(const FInputVertexInterface& InputVertices, const FString& InName) const
 			{
 				using FDataRefType = TDataReadReference<DataType>;
 
@@ -237,7 +238,32 @@ namespace Metasound
 				else
 				{
 					const DefaultDataType DefaultValue = InputVertices[InName].GetDefaultValue().Value.Get<DefaultDataType>();
-					return FDataRefType::CreateNew(DefaultValue, Forward<ConstructorArgTypes>(ConstructorArgs)...);
+					return FDataRefType::CreateNew(DefaultValue);
+				}
+			}
+
+			/** Returns a readable data ref from the collection or construct one
+			 * if one is not there with the default provided from the given input vertex collection.
+			 *
+			 * @param InName - Name of readable data.
+			 * @param InputVertices - Collection of input vertices to retrieve default from if reference is not available
+			 * @param ConstructorArgs - Arguments to pass to constructor of TDataReadReference<DataType>
+			 *
+			 * @return A readable data reference.
+			 */
+			template<typename DataType, typename DefaultDataType = DataType, typename... ConstructorArgTypes>
+			TDataReadReference<DataType> GetDataReadReferenceOrConstructWithVertexDefault(const FInputVertexInterface& InputVertices, const FString& InName, const FOperatorSettings& InOperatorSettings) const
+			{
+				using FDataRefType = TDataReadReference<DataType>;
+
+				if (ContainsDataReadReference<DataType>(InName))
+				{
+					return GetDataReadReference<DataType>(InName);
+				}
+				else
+				{
+					const DefaultDataType DefaultValue = InputVertices[InName].GetDefaultValue().Value.Get<DefaultDataType>();
+					return FDataRefType::CreateNew(InOperatorSettings, DefaultValue);
 				}
 			}
 
diff --git a/Engine/Plugins/Runtime/Metasound/Source/MetasoundGraphCore/Public/MetasoundFacade.h b/Engine/Plugins/Runtime/Metasound/Source/MetasoundGraphCore/Public/MetasoundFacade.h
--- a/Engine/Plugins/Runtime/Metasound/Source/MetasoundGraphCore/Public/MetasoundFacade.h
+++ b/Engine/Plugins/Runtime/Metasound/Source/MetasoundGraphCore/Public/MetasoundFacade.h
@@ -60,7 +60,7 @@ namespace Metasound
 	struct TFacadeOperatorClass
 	{
 		// Require that OperatorType is subclass of IOperator
-		static_assert(std::is_base_of<IOperator, OperatorType>::value, "To use the FNodeFacade consturctor, the OperatorType must be derived from IOperator");
+		static_assert(std::is_base_of<IOperator, OperatorType>::value, "To use the FNodeFacade constructor, the OperatorType must be derived from IOperator");
 
 		// Require static TUniquePtr<IOperator> OperatorType::CreateOperator(const FCreateOperatorParams&, TArray<TUniquePtr<IOperatorBuildError>>&) exists.
 		static_assert(TIsFactoryMethodDeclared<OperatorType>::Value, "To use the FNodeFacade constructor, the OperatorType must have the static function \"static TUniquePtr<IOperator> OperatorType::CreateOperator(const FCreateOperatorParams&, TArray<TUniquePtr<IOperatorBuildError>>&)\"");
diff --git a/Engine/Plugins/Runtime/Metasound/Source/MetasoundGraphCore/Public/MetasoundLiteral.h b/Engine/Plugins/Runtime/Metasound/Source/MetasoundGraphCore/Public/MetasoundLiteral.h
--- a/Engine/Plugins/Runtime/Metasound/Source/MetasoundGraphCore/Public/MetasoundLiteral.h
+++ b/Engine/Plugins/Runtime/Metasound/Source/MetasoundGraphCore/Public/MetasoundLiteral.h
@@ -8,7 +8,7 @@
 #include "Misc/TVariant.h"
 #include <type_traits>
 
-#define METASOUND_DEBUG_LITERALS 1
+#define METASOUND_DEBUG_LITERALS 0
 
 // Forward Declare
 struct FGuid;
diff --git a/Engine/Plugins/Runtime/Metasound/Source/MetasoundGraphCore/Public/MetasoundOperatorSettings.h b/Engine/Plugins/Runtime/Metasound/Source/MetasoundGraphCore/Public/MetasoundOperatorSettings.h
--- a/Engine/Plugins/Runtime/Metasound/Source/MetasoundGraphCore/Public/MetasoundOperatorSettings.h
+++ b/Engine/Plugins/Runtime/Metasound/Source/MetasoundGraphCore/Public/MetasoundOperatorSettings.h
@@ -6,6 +6,9 @@
 
 namespace Metasound
 {
+	using FSampleRate = int32;
+	static_assert(TIsIntegral<FSampleRate>::Value, "FSampleRate must be integral type.");
+
 	/** FOperatorSettings
 	 *
 	 * Audio settings for Metasound IOperators including audio sample rate, 
@@ -40,10 +43,10 @@ namespace Metasound
 			 * @param InSampleRate      - Audio sample rate in Hz. 
 			 * @param InTargetBlockRate - The desired block rate in Hz.
 			 */
-			FOperatorSettings(float InSampleRate, float InTargetBlockRate);
+			FOperatorSettings(FSampleRate InSampleRate, float InTargetBlockRate);
 
 			/** Set the audio sample rate in Hz. */
-			void SetSampleRate(float InSampleRate);
+			void SetSampleRate(FSampleRate InSampleRate);
 
 			/** Get the audio sample rate in Hz. */
 			FORCEINLINE float GetSampleRate() const
@@ -87,7 +90,7 @@ namespace Metasound
 
 			int32 RoundToAligned(int32 Alignment, int32 InNum) const;
 
-			float SampleRate = 1.f;
+			FSampleRate SampleRate = 1;
 			
 			float TargetBlockRate = 1.f;
 
diff --git a/Engine/Plugins/Runtime/Metasound/Source/MetasoundStandardNodes/MetasoundStandardNodes.Build.cs b/Engine/Plugins/Runtime/Metasound/Source/MetasoundStandardNodes/MetasoundStandardNodes.Build.cs
--- a/Engine/Plugins/Runtime/Metasound/Source/MetasoundStandardNodes/MetasoundStandardNodes.Build.cs
+++ b/Engine/Plugins/Runtime/Metasound/Source/MetasoundStandardNodes/MetasoundStandardNodes.Build.cs
@@ -9,6 +9,9 @@ namespace UnrealBuildTool.Rules
 	{
 		public MetasoundStandardNodes(ReadOnlyTargetRules Target) : base(Target)
 		{
+			OptimizeCode = CodeOptimization.Never;
+			PCHUsage = ModuleRules.PCHUsageMode.NoSharedPCHs;
+
 			PublicDependencyModuleNames.AddRange
 			(
 				new string[]
diff --git a/Engine/Plugins/Runtime/Metasound/Source/MetasoundStandardNodes/Private/MetasoundADSRNode.cpp b/Engine/Plugins/Runtime/Metasound/Source/MetasoundStandardNodes/Private/MetasoundADSRNode.cpp
--- a/Engine/Plugins/Runtime/Metasound/Source/MetasoundStandardNodes/Private/MetasoundADSRNode.cpp
+++ b/Engine/Plugins/Runtime/Metasound/Source/MetasoundStandardNodes/Private/MetasoundADSRNode.cpp
@@ -8,6 +8,7 @@
 #include "MetasoundExecutableOperator.h"
 #include "MetasoundNodeRegistrationMacro.h"
 #include "MetasoundPrimitives.h"
+#include "MetasoundSampleCounter.h"
 #include "MetasoundStandardNodesNames.h"
 #include "MetasoundTime.h"
 #include "MetasoundTrigger.h"
@@ -21,10 +22,10 @@ namespace Metasound
 		public:
 			struct FADSRReferences
 			{
-				FFloatTimeReadRef Attack;
-				FFloatTimeReadRef Decay;
+				FTimeReadRef Attack;
+				FTimeReadRef Decay;
 				FFloatReadRef SustainLevel;
-				FFloatTimeReadRef Release;
+				FTimeReadRef Release;
 			};
 
 			static const FNodeClassMetadata& GetNodeInfo();
@@ -53,8 +54,8 @@ namespace Metasound
 
 			bool bReleased;
 
-			FSampleTime TimeUntilRelease;
-			FSampleTime TimePerBlock;
+			FSampleCounter SamplesUntilRelease;
+			FSampleCount SamplesPerBlock;
 			FADSRReferences ADSRReferences;
 			FAudioBufferWriteRef EnvelopeBuffer;
 	};
@@ -66,8 +67,8 @@ namespace Metasound
 	,	TriggerDecayComplete(FTriggerWriteRef::CreateNew(InSettings))
 	,	TriggerReleaseComplete(FTriggerWriteRef::CreateNew(InSettings))
 	,	bReleased(false)
-	,	TimeUntilRelease(0, InSettings.GetSampleRate())
-	,	TimePerBlock(InSettings.GetNumFramesPerBlock(), InSettings.GetSampleRate())
+	,	SamplesUntilRelease(0, InSettings.GetSampleRate())
+	,	SamplesPerBlock(InSettings.GetNumFramesPerBlock())
 	,	ADSRReferences(InADSRData)
 	,	EnvelopeBuffer(FAudioBufferWriteRef::CreateNew(InSettings.GetNumFramesPerBlock()))
 	{
@@ -79,9 +80,9 @@ namespace Metasound
 	FDataReferenceCollection FADSROperator::GetInputs() const
 	{
 		FDataReferenceCollection InputDataReferences;
-		InputDataReferences.AddDataReadReference(TEXT("Attack"), FFloatTimeReadRef(ADSRReferences.Attack));
-		InputDataReferences.AddDataReadReference(TEXT("Decay"), FFloatTimeReadRef(ADSRReferences.Decay));
-		InputDataReferences.AddDataReadReference(TEXT("Release"), FFloatTimeReadRef(ADSRReferences.Release));
+		InputDataReferences.AddDataReadReference(TEXT("Attack"), FTimeReadRef(ADSRReferences.Attack));
+		InputDataReferences.AddDataReadReference(TEXT("Decay"), FTimeReadRef(ADSRReferences.Decay));
+		InputDataReferences.AddDataReadReference(TEXT("Release"), FTimeReadRef(ADSRReferences.Release));
 		InputDataReferences.AddDataReadReference(TEXT("Sustain Level"), FFloatReadRef(ADSRReferences.SustainLevel));
 		InputDataReferences.AddDataReadReference(TEXT("Trigger Attack"), FTriggerReadRef(TriggerAttack));
 		InputDataReferences.AddDataReadReference(TEXT("Trigger Release"), FTriggerReadRef(TriggerRelease));
@@ -109,7 +110,7 @@ namespace Metasound
 		int32 DecayCompleteFrame = -1;
 		int32 ReleaseCompleteFrame = -1;
 
-		TimeUntilRelease -= TimePerBlock;
+		SamplesUntilRelease -= SamplesPerBlock;
 
 		auto GenerateLambda = [&] (int32 InFrame)
 		{
@@ -163,9 +164,9 @@ namespace Metasound
 
 	void FADSROperator::Execute()
 	{
-		Envelope.SetAttackTime(ADSRReferences.Attack->GetMilliseconds());
-		Envelope.SetDecayTime(ADSRReferences.Decay->GetMilliseconds());
-		Envelope.SetReleaseTime(ADSRReferences.Release->GetMilliseconds());
+		Envelope.SetAttackTime(FTime::ToMilliseconds(*ADSRReferences.Attack));
+		Envelope.SetDecayTime(FTime::ToMilliseconds(*ADSRReferences.Decay));
+		Envelope.SetReleaseTime(FTime::ToMilliseconds(*ADSRReferences.Release));
 		Envelope.SetSustainGain(*ADSRReferences.SustainLevel);
 
 		TriggerAttackComplete->AdvanceBlock();
@@ -183,14 +184,14 @@ namespace Metasound
 				if (!bReleased)
 				{
 					bReleased = true;
-					TimeUntilRelease = *(ADSRReferences.Release);
+					SamplesUntilRelease.SetNumSamples(*ADSRReferences.Release);
 
-					const int32 ReleaseFrame = TimeUntilRelease.GetNumSamples();
+					const FSampleCount ReleaseFrame = SamplesUntilRelease.GetNumSamples();
 					if (ReleaseFrame < EndFrame)
 					{
-						GenerateEnvelope(StartFrame, ReleaseFrame);
+						GenerateEnvelope(StartFrame, static_cast<int32>(ReleaseFrame));
 						Envelope.Stop();
-						GenerateEnvelope(ReleaseFrame, EndFrame);
+						GenerateEnvelope(static_cast<int32>(ReleaseFrame), EndFrame);
 					}
 					else
 					{
@@ -232,9 +233,9 @@ namespace Metasound
 	{
 		static const FVertexInterface Interface(
 			FInputVertexInterface(
-				TInputDataVertexModel<FFloatTime>(TEXT("Attack"), LOCTEXT("AttackDurationTooltip", "Attack duration."), 0.01f),
-				TInputDataVertexModel<FFloatTime>(TEXT("Decay"), LOCTEXT("DecayDurationTooltip", "Decay duration."), 0.02f),
-				TInputDataVertexModel<FFloatTime>(TEXT("Release"), LOCTEXT("ReleaseDurationTooltip", "Release duration."), 1.0f),
+				TInputDataVertexModel<FTime>(TEXT("Attack"), LOCTEXT("AttackDurationTooltip", "Attack duration."), 0.01f),
+				TInputDataVertexModel<FTime>(TEXT("Decay"), LOCTEXT("DecayDurationTooltip", "Decay duration."), 0.02f),
+				TInputDataVertexModel<FTime>(TEXT("Release"), LOCTEXT("ReleaseDurationTooltip", "Release duration."), 1.0f),
 				TInputDataVertexModel<float>(TEXT("Sustain Level"), LOCTEXT("SustainLevelTooltip", "Sustain level [0.0f, 1.0f]."), 0.7f),
 				TInputDataVertexModel<FTrigger>(TEXT("Trigger Attack"), LOCTEXT("TriggerAttackTooltip", "Trigger the envelope's attack.")),
 				TInputDataVertexModel<FTrigger>(TEXT("Trigger Release"), LOCTEXT("TriggerReleaseTooltip", "Trigger the envelope's release."))
@@ -285,10 +286,10 @@ namespace Metasound
 		// TODO: If none of these are connected, could pre-generate ADSR envelope and return a different operator.
 		FADSRReferences ADSRReferences =
 		{
-			InputCollection.GetDataReadReferenceOrConstructWithVertexDefault<FFloatTime, float>(InputInterface, TEXT("Attack")),
-			InputCollection.GetDataReadReferenceOrConstructWithVertexDefault<FFloatTime, float>(InputInterface, TEXT("Decay")),
+			InputCollection.GetDataReadReferenceOrConstructWithVertexDefault<FTime, float>(InputInterface, TEXT("Attack")),
+			InputCollection.GetDataReadReferenceOrConstructWithVertexDefault<FTime, float>(InputInterface, TEXT("Decay")),
 			InputCollection.GetDataReadReferenceOrConstructWithVertexDefault<float>(InputInterface, TEXT("Sustain Level")),
-			InputCollection.GetDataReadReferenceOrConstructWithVertexDefault<FFloatTime, float>(InputInterface, TEXT("Release")),
+			InputCollection.GetDataReadReferenceOrConstructWithVertexDefault<FTime, float>(InputInterface, TEXT("Release")),
 		};
 
 		return MakeUnique<FADSROperator>(InParams.OperatorSettings, TriggerAttack, TriggerRelease, ADSRReferences);
diff --git a/Engine/Plugins/Runtime/Metasound/Source/MetasoundStandardNodes/Private/MetasoundBasicFilters.cpp b/Engine/Plugins/Runtime/Metasound/Source/MetasoundStandardNodes/Private/MetasoundBasicFilters.cpp
--- a/Engine/Plugins/Runtime/Metasound/Source/MetasoundStandardNodes/Private/MetasoundBasicFilters.cpp
+++ b/Engine/Plugins/Runtime/Metasound/Source/MetasoundStandardNodes/Private/MetasoundBasicFilters.cpp
@@ -1,14 +1,14 @@
 // Copyright Epic Games, Inc. All Rights Reserved.
-
 #include "MetasoundBasicFilters.h"
+
 #include "MetasoundAudioBuffer.h"
+#include "MetasoundDataTypeRegistrationMacro.h"
 #include "MetasoundPrimitives.h"
 #include "MetasoundExecutableOperator.h"
 #include "MetasoundFacade.h"
 #include "MetasoundNodeRegistrationMacro.h"
-#include "MetasoundStandardNodesNames.h"
 #include "MetasoundOperatorSettings.h"
-#include "MetasoundDataTypeRegistrationMacro.h"
+#include "MetasoundStandardNodesNames.h"
 
 #include "DSP/Filter.h"
 #include "DSP/InterpolatedOnePole.h"
@@ -68,7 +68,7 @@ namespace Metasound
 
 		// data
 		const int32 BlockSize;
-		float SampleRate;
+		FSampleRate SampleRate;
 
 		// dsp
 		Audio::FLadderFilter LadderFilter;
@@ -253,7 +253,7 @@ namespace Metasound
 
 		// data
 		const int32 BlockSize;
-		float SampleRate;
+		FSampleRate SampleRate;
 
 		// dsp
 		Audio::FStateVariableFilter StateVariableFilter;
@@ -439,7 +439,7 @@ namespace Metasound
 
 		// data
 		const int32 BlockSize;
-		float SampleRate;
+		FSampleRate SampleRate;
 
 		// dsp
 		Audio::FInterpolatedLPF OnePoleLowPassFilter;
@@ -583,7 +583,7 @@ namespace Metasound
 
 		// data
 		const int32 BlockSize;
-		float SampleRate;
+		FSampleRate SampleRate;
 
 		// dsp
 		Audio::FInterpolatedHPF OnePoleHighPassFilter;
@@ -750,7 +750,7 @@ namespace Metasound
 
 		// data
 		const int32 BlockSize;
-		float SampleRate;
+		FSampleRate SampleRate;
 
 		// dsp
 		Audio::EBiquadFilter::Type LastFilterType;
diff --git a/Engine/Plugins/Runtime/Metasound/Source/MetasoundStandardNodes/Private/MetasoundBufferAddFloatNode.cpp b/Engine/Plugins/Runtime/Metasound/Source/MetasoundStandardNodes/Private/MetasoundBufferAddFloatNode.cpp
deleted file mode 100644
diff --git a/Engine/Plugins/Runtime/Metasound/Source/MetasoundStandardNodes/Private/MetasoundGainMultiplyNode.cpp b/Engine/Plugins/Runtime/Metasound/Source/MetasoundStandardNodes/Private/MetasoundGainMultiplyNode.cpp
deleted file mode 100644
diff --git a/Engine/Plugins/Runtime/Metasound/Source/MetasoundStandardNodes/Private/MetasoundMathNodes.cpp b/Engine/Plugins/Runtime/Metasound/Source/MetasoundStandardNodes/Private/MetasoundMathNodes.cpp
--- a/Engine/Plugins/Runtime/Metasound/Source/MetasoundStandardNodes/Private/MetasoundMathNodes.cpp
+++ b/Engine/Plugins/Runtime/Metasound/Source/MetasoundStandardNodes/Private/MetasoundMathNodes.cpp
@@ -8,35 +8,32 @@ namespace Metasound
 	METASOUND_REGISTER_NODE(FAddFloatNode)
 	METASOUND_REGISTER_NODE(FAddInt32Node)
 	METASOUND_REGISTER_NODE(FAddInt64Node)
-	METASOUND_REGISTER_NODE(FAddFrequencyNode)
-	METASOUND_REGISTER_NODE(FAddFloatTimeNode)
-	METASOUND_REGISTER_NODE(FAddAudioBufferNode)
+	METASOUND_REGISTER_NODE(FAddTimeNode)
+	METASOUND_REGISTER_NODE(FAddAudioNode)
+	METASOUND_REGISTER_NODE(FAddAudioFloatNode)
 
 	METASOUND_REGISTER_NODE(FSubtractFloatNode)
 	METASOUND_REGISTER_NODE(FSubtractInt32Node)
 	METASOUND_REGISTER_NODE(FSubtractInt64Node)
-	METASOUND_REGISTER_NODE(FSubtractFrequencyNode)
-	METASOUND_REGISTER_NODE(FSubtractFloatTimeNode)
-	METASOUND_REGISTER_NODE(FSubtractAudioBufferNode)
+	METASOUND_REGISTER_NODE(FSubtractTimeNode)
+	METASOUND_REGISTER_NODE(FSubtractAudioNode)
 
+	METASOUND_REGISTER_NODE(FMultiplyAudioNode)
+	METASOUND_REGISTER_NODE(FMultiplyAudioFloatNode)
 	METASOUND_REGISTER_NODE(FMultiplyFloatNode)
 	METASOUND_REGISTER_NODE(FMultiplyInt32Node)
 	METASOUND_REGISTER_NODE(FMultiplyInt64Node)
-	METASOUND_REGISTER_NODE(FMultiplyFrequencyNode)
-	METASOUND_REGISTER_NODE(FMultiplyFloatTimeNode)
-	METASOUND_REGISTER_NODE(FMultiplyAudioBufferNode)
+	METASOUND_REGISTER_NODE(FMultiplyTimeFloatNode)
 
 	METASOUND_REGISTER_NODE(FDivideFloatNode)
 	METASOUND_REGISTER_NODE(FDivideInt32Node)
 	METASOUND_REGISTER_NODE(FDivideInt64Node)
-	METASOUND_REGISTER_NODE(FDivideFrequencyNode)
-	METASOUND_REGISTER_NODE(FDivideFloatTimeNode)
+	METASOUND_REGISTER_NODE(FDivideTimeFloatNode)
 
 	METASOUND_REGISTER_NODE(FRandRangeFloatNode)
 	METASOUND_REGISTER_NODE(FRandRangeInt32Node)
 	METASOUND_REGISTER_NODE(FRandRangeInt64Node)
-	METASOUND_REGISTER_NODE(FRandRangeFrequencyNode)
-	METASOUND_REGISTER_NODE(FRandRangeFloatTimeNode)
+	METASOUND_REGISTER_NODE(FRandRangeTimeNode)
 }
 
 #undef LOCTEXT_NAMESPACE // MetasoundMathOpNode
diff --git a/Engine/Plugins/Runtime/Metasound/Source/MetasoundStandardNodes/Private/MetasoundMixerNode.cpp b/Engine/Plugins/Runtime/Metasound/Source/MetasoundStandardNodes/Private/MetasoundMixerNode.cpp
--- a/Engine/Plugins/Runtime/Metasound/Source/MetasoundStandardNodes/Private/MetasoundMixerNode.cpp
+++ b/Engine/Plugins/Runtime/Metasound/Source/MetasoundStandardNodes/Private/MetasoundMixerNode.cpp
@@ -5,8 +5,6 @@
 #include "MetasoundAudioBuffer.h"
 #include "MetasoundExecutableOperator.h"
 #include "MetasoundFacade.h"
-#include "MetasoundFrequency.h"
-#include "MetasoundGain.h"
 #include "MetasoundNodeRegistrationMacro.h"
 #include "MetasoundOperatorSettings.h"
 #include "MetasoundPrimitives.h"
@@ -24,9 +22,9 @@ namespace Metasound
 	public:
 		static TUniquePtr<IOperator> CreateOperator(const FCreateOperatorParams& InParams, FBuildErrorArray& OutErrors);
 		static const FNodeClassMetadata& GetNodeInfo();
-		static FVertexInterface DeclareVertexInterface();	
+		static const FVertexInterface& GetVertexInterface();	
 
-		FMixerOperator(const FOperatorSettings& InSettings, TArray<FAudioBufferReadRef>&& InBuffers, TArray<FGainReadRef>&& InGains);
+		FMixerOperator(const FOperatorSettings& InSettings, TArray<FAudioBufferReadRef>&& InBuffers, TArray<FFloatReadRef>&& InGains);
 
 		virtual FDataReferenceCollection GetInputs() const override;
 		virtual FDataReferenceCollection GetOutputs() const override;
@@ -34,10 +32,10 @@ namespace Metasound
 		void Execute();
 
 	private:		
-		// TODO: Remove when we get dynamic pin support. (if you update this, update DeclareVertexInterface also).
+		// TODO: Remove when we get dynamic pin support. (if you update this, update GetVertexInterface also).
 		static constexpr const int32 NumHardcodedInputs = 4;
 
-		static FString GetInputAudioPinName(int32 InPinIndex) 
+		static FString GetInputAudioPinName(int32 InPinIndex)
 		{
 			return FString::Printf(TEXT("In %d"), InPinIndex + 1);
 		}
@@ -47,13 +45,13 @@ namespace Metasound
 		}
 
 		TArray<FAudioBufferReadRef> AudioBuffersIn;
-		TArray<FGainReadRef> BufferMixGains;
+		TArray<FFloatReadRef> BufferMixGains;
 		FAudioBufferWriteRef AudioBufferOut;
 		
 		static constexpr const TCHAR* AudioOutPinName = TEXT("Out");
 	};
 
-	FMixerOperator::FMixerOperator(const FOperatorSettings& InSettings, TArray<FAudioBufferReadRef>&& InBuffers, TArray<FGainReadRef>&& InGains)
+	FMixerOperator::FMixerOperator(const FOperatorSettings& InSettings, TArray<FAudioBufferReadRef>&& InBuffers, TArray<FFloatReadRef>&& InGains)
 		: AudioBuffersIn(MoveTemp(InBuffers))
 		, BufferMixGains(MoveTemp(InGains))
 		, AudioBufferOut(FAudioBufferWriteRef::CreateNew(InSettings))
@@ -89,25 +87,25 @@ namespace Metasound
 		}
 	}
 
-	FVertexInterface FMixerOperator::DeclareVertexInterface()
+	const FVertexInterface& FMixerOperator::GetVertexInterface()
 	{
 		static const FVertexInterface Interface(
 			FInputVertexInterface(
 				// 1
 				TInputDataVertexModel<FAudioBuffer>(GetInputAudioPinName(0), LOCTEXT("MixerInputDescription1", "Audio Input 1 of the Mixer")),		// Can't have dynamic pins yet, so just hard code 4.
-				TInputDataVertexModel<FGain>(GetInputGainPinName(0), LOCTEXT("MixerGainDescription1", ""), 1.0f),
+				TInputDataVertexModel<float>(GetInputGainPinName(0), LOCTEXT("MixerGainDescription1", ""), 1.0f),
 
 				// 2
 				TInputDataVertexModel<FAudioBuffer>(GetInputAudioPinName(1), LOCTEXT("MixerInputDescription2", "Audio Input 2 of the Mixer")),
-				TInputDataVertexModel<FGain>(GetInputGainPinName(1), LOCTEXT("MixerGainDescription2", ""), 1.0f),
+				TInputDataVertexModel<float>(GetInputGainPinName(1), LOCTEXT("MixerGainDescription2", ""), 1.0f),
 
 				// 3
 				TInputDataVertexModel<FAudioBuffer>(GetInputAudioPinName(2), LOCTEXT("MixerInputDescription3", "Audio Input 3 of the Mixer")),
-				TInputDataVertexModel<FGain>(GetInputGainPinName(2), LOCTEXT("MixerGainDescription3", ""), 1.0f),
+				TInputDataVertexModel<float>(GetInputGainPinName(2), LOCTEXT("MixerGainDescription3", ""), 1.0f),
 
 				// 4
 				TInputDataVertexModel<FAudioBuffer>(GetInputAudioPinName(3), LOCTEXT("MixerInputDescription3", "Audio Input 4 of the Mixer")),	
-				TInputDataVertexModel<FGain>(GetInputGainPinName(3), LOCTEXT("MixerGainDescription4", ""), 1.0f)
+				TInputDataVertexModel<float>(GetInputGainPinName(3), LOCTEXT("MixerGainDescription4", ""), 1.0f)
 			),
 			FOutputVertexInterface(
 				TOutputDataVertexModel<FAudioBuffer>(AudioOutPinName, LOCTEXT("AudioOutTooltip", "Audio Ouput from the mixer"))
@@ -122,14 +120,14 @@ namespace Metasound
 		auto InitNodeInfo = []() -> FNodeClassMetadata
 		{
 			FNodeClassMetadata Info;
-			Info.ClassName = {Metasound::StandardNodes::Namespace, TEXT("Mixer"), Metasound::StandardNodes::AudioVariant};
+			Info.ClassName = { Metasound::StandardNodes::Namespace, TEXT("Mixer"), Metasound::StandardNodes::AudioVariant };
 			Info.MajorVersion = 1;
 			Info.MinorVersion = 0;
 			Info.DisplayName = LOCTEXT("Metasound_MixerNodeDisplayName", "Mixer");
-			Info.Description = LOCTEXT("Metasound_MixerNodeDescription", "Mixes 1 or more input signals together");
+			Info.Description = LOCTEXT("Metasound_MixerNodeDescription", "Mixes input audio together");
 			Info.Author = PluginAuthor;
 			Info.PromptIfMissing = PluginNodeMissingPrompt;
-			Info.DefaultInterface = DeclareVertexInterface();
+			Info.DefaultInterface = GetVertexInterface();
 
 			return Info;
 		};
@@ -139,42 +137,32 @@ namespace Metasound
 		return Info;
 	}
 
-	FMixerNode::FMixerNode(const FString& InName, const FGuid& InInstanceID, float InDefaultMixGainLinear)
-	: FNodeFacade(InName, InInstanceID, TFacadeOperatorClass<FMixerOperator>())
-	, DefaultMixGainLinear(InDefaultMixGainLinear)
-	{
-	}
-
 	FMixerNode::FMixerNode(const FNodeInitData& InInitData)
-	: FMixerNode(InInitData.InstanceName, InInitData.InstanceID, 1.0f)
+		: FNodeFacade(InInitData.InstanceName, InInitData.InstanceID, TFacadeOperatorClass<FMixerOperator>())
 	{
 	}
 
-	float FMixerNode::GetDefaultMixGainLinear() const
-	{
-		return DefaultMixGainLinear;
-	}
-
 	TUniquePtr<IOperator> FMixerOperator::CreateOperator(const FCreateOperatorParams& InParams, FBuildErrorArray& OutErrors)
 	{
 		const FMixerNode& Node = static_cast<const FMixerNode&>(InParams.Node);
 		const FDataReferenceCollection& InputCol = InParams.InputDataReferences;
 		const FOperatorSettings& Settings = InParams.OperatorSettings;
+		const FInputVertexInterface& InputInterface = GetVertexInterface().GetInputInterface();
 
 		// TODO: Remove this string manipulation when we get dynamic pin support.
 		TArray<FAudioBufferReadRef> InputBuffers;
-		TArray<FGainReadRef> InputGains;
+		TArray<FFloatReadRef> InputGains;
 		for (int32 i = 0; i < NumHardcodedInputs; ++i)
-		{			
-			FString AudioPinName = GetInputAudioPinName(i);
-			FString GainPinName = GetInputGainPinName(i);
+		{
+			const FString AudioPinName = GetInputAudioPinName(i);
+			const FString GainPinName = GetInputGainPinName(i);
 			if (InputCol.ContainsDataReadReference<FAudioBuffer>(AudioPinName))
 			{
 				// Only Create buffers if there's something connected to it.
 				InputBuffers.Emplace(InputCol.GetDataReadReference<FAudioBuffer>(AudioPinName));
 				
 				// Make sure for every valid/connected pin, we have a corresponding gain, even if its defaulted.
-				InputGains.Emplace(InputCol.GetDataReadReferenceOrConstruct<FGain>(GainPinName, Node.GetDefaultMixGainLinear()));
+				InputGains.Emplace(InputCol.GetDataReadReferenceOrConstructWithVertexDefault<float>(InputInterface, GainPinName));
 			}
 		}
 		
diff --git a/Engine/Plugins/Runtime/Metasound/Source/MetasoundStandardNodes/Private/MetasoundOscNode.cpp b/Engine/Plugins/Runtime/Metasound/Source/MetasoundStandardNodes/Private/MetasoundOscNode.cpp
--- a/Engine/Plugins/Runtime/Metasound/Source/MetasoundStandardNodes/Private/MetasoundOscNode.cpp
+++ b/Engine/Plugins/Runtime/Metasound/Source/MetasoundStandardNodes/Private/MetasoundOscNode.cpp
@@ -5,7 +5,6 @@
 #include "MetasoundAudioBuffer.h"
 #include "MetasoundExecutableOperator.h"
 #include "MetasoundFacade.h"
-#include "MetasoundFrequency.h"
 #include "MetasoundNodeRegistrationMacro.h"
 #include "MetasoundPrimitives.h"
 #include "MetasoundStandardNodesNames.h"
@@ -23,9 +22,9 @@ namespace Metasound
 
 		static TUniquePtr<IOperator> CreateOperator(const FCreateOperatorParams& InParams, FBuildErrorArray& OutErrors);
 		static const FNodeClassMetadata& GetNodeInfo();
-		static FVertexInterface DeclareVertexInterface();
+		static const FVertexInterface& GetVertexInterface();
 
-		FOscOperator(const FOperatorSettings& InSettings, const FFrequencyReadRef& InFrequency, const FBoolReadRef& InActivate);
+		FOscOperator(const FOperatorSettings& InSettings, const FFloatReadRef& InFrequency, const FBoolReadRef& InActivate);
 
 		virtual FDataReferenceCollection GetInputs() const override;
 
@@ -58,7 +57,7 @@ namespace Metasound
 		float OneOverSampleRate;
 		float Nyquist;
 
-		FFrequencyReadRef Frequency;
+		FFloatReadRef Frequency;
 		FBoolReadRef Enabled;
 		FAudioBufferWriteRef AudioBuffer;
 
@@ -67,7 +66,7 @@ namespace Metasound
 		static constexpr const TCHAR* AudioOutPinName = TEXT("Audio");
 	};
 
-	FOscOperator::FOscOperator(const FOperatorSettings& InSettings, const FFrequencyReadRef& InFrequency, const FBoolReadRef& InEnabled)
+	FOscOperator::FOscOperator(const FOperatorSettings& InSettings, const FFloatReadRef& InFrequency, const FBoolReadRef& InEnabled)
 		: Phase(0.f)
 		, OneOverSampleRate(1.f / InSettings.GetSampleRate())
 		, Nyquist(InSettings.GetSampleRate() / 2.0f)
@@ -81,7 +80,7 @@ namespace Metasound
 	FDataReferenceCollection FOscOperator::GetInputs() const
 	{
 		FDataReferenceCollection InputDataReferences;
-		InputDataReferences.AddDataReadReference(FrequencyPinName, FFrequencyReadRef(Frequency));
+		InputDataReferences.AddDataReadReference(FrequencyPinName, FFloatReadRef(Frequency));
 		InputDataReferences.AddDataReadReference(EnabledPinName, FBoolReadRef(Enabled));
 
 		return InputDataReferences;
@@ -97,7 +96,7 @@ namespace Metasound
 	void FOscOperator::Execute()
 	{
 		// Clamp frequencies into Nyquist range
-		const float Freq = FMath::Clamp(Frequency->GetHertz(), -Nyquist, Nyquist);
+		const float Freq = FMath::Clamp(*Frequency, -Nyquist, Nyquist);
 		const float PhaseInc = Freq * OneOverSampleRate;
 		float* Data = AudioBuffer->GetData();
 
@@ -119,11 +118,11 @@ namespace Metasound
 		}
 	}
 
-	FVertexInterface FOscOperator::DeclareVertexInterface()
+	const FVertexInterface& FOscOperator::GetVertexInterface()
 	{
 		static const FVertexInterface Interface(
 			FInputVertexInterface(
-				TInputDataVertexModel<FFrequency>(FrequencyPinName, LOCTEXT("OscFrequencyDescription", "The frequency of oscillator."), 440.f),
+				TInputDataVertexModel<float>(FrequencyPinName, LOCTEXT("OscFrequencyDescription", "The frequency of oscillator."), 440.f),
 				TInputDataVertexModel<bool>(EnabledPinName, LOCTEXT("OscActivateDescription", "Enable the oscilator."), true)
 			),
 			FOutputVertexInterface(
@@ -144,10 +143,10 @@ namespace Metasound
 			Info.MajorVersion = 1;
 			Info.MinorVersion = 0;
 			Info.DisplayName = LOCTEXT("Metasound_OscNodeDisplayName", "Oscillator");
-			Info.Description = LOCTEXT("Metasound_OscNodeDescription", "Emits an audio signal of a sinusoid.");
+			Info.Description = LOCTEXT("Metasound_OscNodeDescription", "Emits sinusoidal audio.");
 			Info.Author = PluginAuthor;
 			Info.PromptIfMissing = PluginNodeMissingPrompt;
-			Info.DefaultInterface = DeclareVertexInterface();
+			Info.DefaultInterface = GetVertexInterface();
 
 			return Info;
 		};
@@ -157,26 +156,9 @@ namespace Metasound
 		return Info;
 	}
 
-	FOscNode::FOscNode(const FString& InInstanceName, const FGuid& InInstanceID, float InDefaultFrequency, bool bInDefaultEnablement)
-	:	FNodeFacade(InInstanceName, InInstanceID, TFacadeOperatorClass<FOscOperator>())
-	,	DefaultFrequency(InDefaultFrequency)
-	,	bDefaultEnablement(bInDefaultEnablement)
-	{
-	}
-
 	FOscNode::FOscNode(const FNodeInitData& InInitData)
-		: FOscNode(InInitData.InstanceName, InInitData.InstanceID, 440.0f, true)
-	{
-	}
-
-	float FOscNode::GetDefaultFrequency() const
-	{
-		return DefaultFrequency;
-	}
-
-	bool FOscNode::GetDefaultEnablement() const
+		: FNodeFacade(InInitData.InstanceName, InInitData.InstanceID, TFacadeOperatorClass<FOscOperator>())
 	{
-		return bDefaultEnablement;
 	}
 
 	TUniquePtr<IOperator> FOscOperator::CreateOperator(const FCreateOperatorParams& InParams, FBuildErrorArray& OutErrors) 
@@ -184,9 +166,10 @@ namespace Metasound
 		const FOscNode& OscNode = static_cast<const FOscNode&>(InParams.Node);
 		const FDataReferenceCollection& InputCol = InParams.InputDataReferences;
 		const FOperatorSettings& Settings = InParams.OperatorSettings;
+		const FInputVertexInterface& InputInterface = GetVertexInterface().GetInputInterface();
 
-		FFrequencyReadRef Frequency = InputCol.GetDataReadReferenceOrConstruct<FFrequency>(FrequencyPinName, OscNode.GetDefaultFrequency(), EFrequencyResolution::Hertz);
-		FBoolReadRef Enabled = InputCol.GetDataReadReferenceOrConstruct<bool>(EnabledPinName, OscNode.GetDefaultEnablement());
+		FFloatReadRef Frequency = InputCol.GetDataReadReferenceOrConstructWithVertexDefault<float>(InputInterface, FrequencyPinName);
+		FBoolReadRef Enabled = InputCol.GetDataReadReferenceOrConstructWithVertexDefault<bool>(InputInterface, EnabledPinName);
 
 		return MakeUnique<FOscOperator>(InParams.OperatorSettings, Frequency, Enabled);
 	}
diff --git a/Engine/Plugins/Runtime/Metasound/Source/MetasoundStandardNodes/Private/MetasoundPeriodicTriggerNode.cpp b/Engine/Plugins/Runtime/Metasound/Source/MetasoundStandardNodes/Private/MetasoundPeriodicTriggerNode.cpp
--- a/Engine/Plugins/Runtime/Metasound/Source/MetasoundStandardNodes/Private/MetasoundPeriodicTriggerNode.cpp
+++ b/Engine/Plugins/Runtime/Metasound/Source/MetasoundStandardNodes/Private/MetasoundPeriodicTriggerNode.cpp
@@ -13,6 +13,7 @@
 #include "MetasoundNodeRegistrationMacro.h"
 #include "MetasoundOperatorInterface.h"
 #include "MetasoundPrimitives.h"
+#include "MetasoundSampleCounter.h"
 #include "MetasoundStandardNodesNames.h"
 
 #define LOCTEXT_NAMESPACE "MetasoundStandardNodes"
@@ -28,10 +29,7 @@ namespace Metasound
 			static const FVertexInterface& GetVertexInterface();
 			static TUniquePtr<IOperator> CreateOperator(const FCreateOperatorParams& InParams, FBuildErrorArray& OutErrors);
 
-			static constexpr float MinimumPeriodSeconds = 0.001f;
-			static constexpr float MinimumPeriodSamples = 20.f;
-
-			FPeriodicTriggerOperator(const FOperatorSettings& InSettings, const FTriggerReadRef& InTriggerEnable, const FTriggerReadRef& InTriggerDisable, const FFloatTimeReadRef& InPeriod);
+			FPeriodicTriggerOperator(const FOperatorSettings& InSettings, const FTriggerReadRef& InTriggerEnable, const FTriggerReadRef& InTriggerDisable, const FTimeReadRef& InPeriod);
 
 			virtual FDataReferenceCollection GetInputs() const override;
 
@@ -47,29 +45,27 @@ namespace Metasound
 			FTriggerReadRef TriggerEnable;
 			FTriggerReadRef TriggerDisable;
 
-			FFloatTimeReadRef Period;
+			FTimeReadRef Period;
 
-			float ExecuteDurationInSamples;
-			float SampleRate;
-			float SampleCountdown;
+			FSampleCount BlockSize;
+			FSampleCounter SampleCounter;
 	};
 
-	FPeriodicTriggerOperator::FPeriodicTriggerOperator(const FOperatorSettings& InSettings, const FTriggerReadRef& InTriggerEnable, const FTriggerReadRef& InTriggerDisable, const FFloatTimeReadRef& InPeriod)
+	FPeriodicTriggerOperator::FPeriodicTriggerOperator(const FOperatorSettings& InSettings, const FTriggerReadRef& InTriggerEnable, const FTriggerReadRef& InTriggerDisable, const FTimeReadRef& InPeriod)
 	:	bEnabled(false)
 	,	TriggerOut(FTriggerWriteRef::CreateNew(InSettings))
 	,	TriggerEnable(InTriggerEnable)
 	,	TriggerDisable(InTriggerDisable)
 	,	Period(InPeriod)
-	,	ExecuteDurationInSamples(InSettings.GetNumFramesPerBlock())
-	,	SampleRate(InSettings.GetSampleRate())
-	,	SampleCountdown(0.f)
+	,	BlockSize(InSettings.GetNumFramesPerBlock())
+	,	SampleCounter(0, InSettings.GetSampleRate())
 	{
 	}
 
 	FDataReferenceCollection FPeriodicTriggerOperator::GetInputs() const
 	{
 		FDataReferenceCollection InputDataReferences;
-		InputDataReferences.AddDataReadReference(TEXT("Period"), FFloatTimeReadRef(Period));
+		InputDataReferences.AddDataReadReference(TEXT("Period"), FTimeReadRef(Period));
 		InputDataReferences.AddDataReadReference(TEXT("Activate"), FTriggerReadRef(TriggerEnable));
 		InputDataReferences.AddDataReadReference(TEXT("Deactivate"), FTriggerReadRef(TriggerDisable));
 		return InputDataReferences;
@@ -104,18 +100,15 @@ namespace Metasound
 
 		if (bEnabled)
 		{
-			float PeriodInSamples = FMath::Max(Period->GetSeconds(), MinimumPeriodSeconds) * SampleRate;
-
-			PeriodInSamples = FMath::Max(PeriodInSamples, MinimumPeriodSamples);
-
-			while ((SampleCountdown - ExecuteDurationInSamples) <= 0.f)
+			const FSampleCount PeriodInSamples = FSampleCounter::FromTime(*Period, SampleCounter.GetSampleRate()).GetNumSamples();
+			while ((SampleCounter - BlockSize) <= 0)
 			{
-				uint32 Frame = FMath::RoundToInt(SampleCountdown);
-				TriggerOut->TriggerFrame(Frame);
-				SampleCountdown += PeriodInSamples;
+				const int32 StartOffset = static_cast<int32>(SampleCounter.GetNumSamples());
+				TriggerOut->TriggerFrame(StartOffset);
+				SampleCounter += PeriodInSamples;
 			}
 
-			SampleCountdown -= ExecuteDurationInSamples;
+			SampleCounter -= BlockSize;
 		}
 	}
 
@@ -125,7 +118,7 @@ namespace Metasound
 		const FPeriodicTriggerNode& PeriodicTriggerNode = static_cast<const FPeriodicTriggerNode&>(InParams.Node);
 		const FInputVertexInterface& InputInterface = GetVertexInterface().GetInputInterface();
 
-		FFloatTimeReadRef Period = InParams.InputDataReferences.GetDataReadReferenceOrConstructWithVertexDefault<FFloatTime, float>(InputInterface, TEXT("Period"));
+		FTimeReadRef Period = InParams.InputDataReferences.GetDataReadReferenceOrConstructWithVertexDefault<FTime, float>(InputInterface, TEXT("Period"));
 		FTriggerReadRef TriggerEnable = InParams.InputDataReferences.GetDataReadReferenceOrConstruct<FTrigger>(TEXT("Activate"), InParams.OperatorSettings);
 		FTriggerReadRef TriggerDisable = InParams.InputDataReferences.GetDataReadReferenceOrConstruct<FTrigger>(TEXT("Deactivate"), InParams.OperatorSettings);
 
@@ -136,7 +129,7 @@ namespace Metasound
 	{
 		static const FVertexInterface Interface(
 			FInputVertexInterface(
-				TInputDataVertexModel<FFloatTime>(TEXT("Period"), LOCTEXT("PeriodTooltip", "The period to trigger in seconds."), 1.0f),
+				TInputDataVertexModel<FTime>(TEXT("Period"), LOCTEXT("PeriodTooltip", "The period to trigger in seconds."), 1.0f),
 				TInputDataVertexModel<FTrigger>(TEXT("Activate"), LOCTEXT("TriggerEnableTooltip", "Enables executing periodic output triggers.")),
 				TInputDataVertexModel<FTrigger>(TEXT("Deactivate"), LOCTEXT("TriggerDisableTooltip", "Disables executing periodic output triggers."))
 			),
diff --git a/Engine/Plugins/Runtime/Metasound/Source/MetasoundStandardNodes/Private/MetasoundStandardNodesModule.cpp b/Engine/Plugins/Runtime/Metasound/Source/MetasoundStandardNodes/Private/MetasoundStandardNodesModule.cpp
--- a/Engine/Plugins/Runtime/Metasound/Source/MetasoundStandardNodes/Private/MetasoundStandardNodesModule.cpp
+++ b/Engine/Plugins/Runtime/Metasound/Source/MetasoundStandardNodes/Private/MetasoundStandardNodesModule.cpp
@@ -1,7 +1,6 @@
 // Copyright Epic Games, Inc. All Rights Reserved.
 
 #include "CoreMinimal.h"
-#include "MetasoundFrequency.h"
 #include "MetasoundGain.h"
 #include "MetasoundGraphCoreModule.h"
 #include "MetasoundPrimitives.h"
@@ -11,11 +10,7 @@
 #include "Modules/ModuleManager.h"
 
 REGISTER_METASOUND_DATATYPE(Metasound::FTrigger, "Trigger", ::Metasound::ELiteralType::Boolean)
-REGISTER_METASOUND_DATATYPE(Metasound::FFloatTime, "Time", ::Metasound::ELiteralType::Float)
-REGISTER_METASOUND_DATATYPE(Metasound::FDoubleTime, "Time:HighResolution", ::Metasound::ELiteralType::Float)
-REGISTER_METASOUND_DATATYPE(Metasound::FSampleTime, "Time:SampleResolution", ::Metasound::ELiteralType::Float)
-REGISTER_METASOUND_DATATYPE(Metasound::FFrequency, "Frequency", ::Metasound::ELiteralType::Float)
-REGISTER_METASOUND_DATATYPE(Metasound::FGain, "Gain", ::Metasound::ELiteralType::Float)
+REGISTER_METASOUND_DATATYPE(Metasound::FTime, "Time", ::Metasound::ELiteralType::Float)
 
 namespace Metasound 
 {
diff --git a/Engine/Plugins/Runtime/Metasound/Source/MetasoundStandardNodes/Private/MetasoundSwitchGateNode.cpp b/Engine/Plugins/Runtime/Metasound/Source/MetasoundStandardNodes/Private/MetasoundSwitchGateNode.cpp
--- a/Engine/Plugins/Runtime/Metasound/Source/MetasoundStandardNodes/Private/MetasoundSwitchGateNode.cpp
+++ b/Engine/Plugins/Runtime/Metasound/Source/MetasoundStandardNodes/Private/MetasoundSwitchGateNode.cpp
@@ -5,7 +5,6 @@
 #include "MetasoundAudioBuffer.h"
 #include "MetasoundExecutableOperator.h"
 #include "MetasoundFacade.h"
-#include "MetasoundFrequency.h"
 #include "MetasoundNodeRegistrationMacro.h"
 #include "MetasoundPrimitives.h"
 #include "MetasoundStandardNodesNames.h"
diff --git a/Engine/Plugins/Runtime/Metasound/Source/MetasoundStandardNodes/Private/MetasoundTriggerDelayNode.cpp b/Engine/Plugins/Runtime/Metasound/Source/MetasoundStandardNodes/Private/MetasoundTriggerDelayNode.cpp
--- a/Engine/Plugins/Runtime/Metasound/Source/MetasoundStandardNodes/Private/MetasoundTriggerDelayNode.cpp
+++ b/Engine/Plugins/Runtime/Metasound/Source/MetasoundStandardNodes/Private/MetasoundTriggerDelayNode.cpp
@@ -12,6 +12,7 @@
 #include "MetasoundNodeRegistrationMacro.h"
 #include "MetasoundOperatorInterface.h"
 #include "MetasoundPrimitives.h"
+#include "MetasoundSampleCounter.h"
 #include "MetasoundStandardNodesNames.h"
 
 #define LOCTEXT_NAMESPACE "MetasoundStandardNodes"
@@ -27,7 +28,7 @@ namespace Metasound
 			static const FVertexInterface& GetVertexInterface();
 			static TUniquePtr<IOperator> CreateOperator(const FCreateOperatorParams& InParams, FBuildErrorArray& OutErrors);
 
-			FTriggerDelayOperator(const FOperatorSettings& InSettings, const FTriggerReadRef& InTriggerReset, const FTriggerReadRef& InTriggerIn, const FFloatTimeReadRef& InDelay);
+			FTriggerDelayOperator(const FOperatorSettings& InSettings, const FTriggerReadRef& InTriggerReset, const FTriggerReadRef& InTriggerIn, const FTimeReadRef& InDelay);
 
 			virtual FDataReferenceCollection GetInputs() const override;
 			virtual FDataReferenceCollection GetOutputs() const override;
@@ -35,35 +36,33 @@ namespace Metasound
 			void Execute();
 
 		private:
-			int32 SamplesUntilTrigger = -1;
-			int32 LastFrameReset = -1;
+			FSampleCounter NextTriggerCounter;
+			FSampleCount LastFrameReset;
 
 			FTriggerReadRef TriggerIn;
 			FTriggerReadRef TriggerReset;
 			FTriggerWriteRef TriggerOut;
 
-			FFloatTimeReadRef Delay;
+			FTimeReadRef Delay;
 
-			float FramesPerBlock;
-			float SampleRate;
+			FSampleCount FramesPerBlock;
 	};
 
-	FTriggerDelayOperator::FTriggerDelayOperator(const FOperatorSettings& InSettings, const FTriggerReadRef& InTriggerReset, const FTriggerReadRef& InTriggerIn, const FFloatTimeReadRef& InDelay)
-	:	SamplesUntilTrigger(-1)
+	FTriggerDelayOperator::FTriggerDelayOperator(const FOperatorSettings& InSettings, const FTriggerReadRef& InTriggerReset, const FTriggerReadRef& InTriggerIn, const FTimeReadRef& InDelay)
+	:	NextTriggerCounter(-1, InSettings.GetSampleRate())
 	,	LastFrameReset(-1)
 	,	TriggerIn(InTriggerIn)
 	,	TriggerReset(InTriggerReset)
 	,	TriggerOut(FTriggerWriteRef::CreateNew(InSettings))
 	,	Delay(InDelay)
 	,	FramesPerBlock(InSettings.GetNumFramesPerBlock())
-	,	SampleRate(InSettings.GetSampleRate())
 	{
 	}
 
 	FDataReferenceCollection FTriggerDelayOperator::GetInputs() const
 	{
 		FDataReferenceCollection InputDataReferences;
-		InputDataReferences.AddDataReadReference(TEXT("Delay"), FFloatTimeReadRef(Delay));
+		InputDataReferences.AddDataReadReference(TEXT("Delay"), FTimeReadRef(Delay));
 		InputDataReferences.AddDataReadReference(TEXT("In"), FTriggerReadRef(TriggerIn));
 		InputDataReferences.AddDataReadReference(TEXT("Reset"), FTriggerReadRef(TriggerIn));
 		return InputDataReferences;
@@ -83,6 +82,7 @@ namespace Metasound
 		TriggerOut->AdvanceBlock();
 
 		LastFrameReset = -1;
+
 		TriggerReset->ExecuteBlock(
 			[&](int32 StartFrame, int32 EndFrame)
 			{
@@ -101,27 +101,27 @@ namespace Metasound
 			{
 				if (StartFrame > LastFrameReset)
 				{
-					SamplesUntilTrigger = Delay->GetNumSamples(SampleRate);
+					NextTriggerCounter.SetNumSamples(*Delay);
 					LastFrameReset = -1;
 				}
 				else
 				{
-					SamplesUntilTrigger = -1;
+					NextTriggerCounter.SetNumSamples(-1);
 				}
 			}
 		);
 
-		if (SamplesUntilTrigger >= 0)
+		if (NextTriggerCounter >= 0)
 		{
-			int32 SamplesRemaining = SamplesUntilTrigger - FramesPerBlock;
+			FSampleCount SamplesRemaining = NextTriggerCounter.GetNumSamples() - FramesPerBlock;
 			if (SamplesRemaining > 0.0f)
 			{
-				SamplesUntilTrigger -= FramesPerBlock;
+				NextTriggerCounter -= FramesPerBlock;
 			}
 			else
 			{
 				TriggerOut->TriggerFrame(SamplesRemaining + (int32)FramesPerBlock);
-				SamplesUntilTrigger = -1;
+				NextTriggerCounter.SetNumSamples(-1);
 			}
 		}
 	}
@@ -131,7 +131,7 @@ namespace Metasound
 		const FTriggerDelayNode& DelayNode = static_cast<const FTriggerDelayNode&>(InParams.Node);
 
 		const FInputVertexInterface& InputInterface = GetVertexInterface().GetInputInterface();
-		FFloatTimeReadRef Delay = InParams.InputDataReferences.GetDataReadReferenceOrConstructWithVertexDefault<FFloatTime, float>(InputInterface, TEXT("Delay"));
+		FTimeReadRef Delay = InParams.InputDataReferences.GetDataReadReferenceOrConstructWithVertexDefault<FTime, float>(InputInterface, TEXT("Delay"));
 
 		FTriggerReadRef TriggerIn = InParams.InputDataReferences.GetDataReadReferenceOrConstruct<FTrigger>(TEXT("In"), InParams.OperatorSettings);
 		FTriggerReadRef TriggerReset = InParams.InputDataReferences.GetDataReadReferenceOrConstruct<FTrigger>(TEXT("Reset"), InParams.OperatorSettings);
@@ -143,7 +143,7 @@ namespace Metasound
 	{
 		static const FVertexInterface Interface(
 			FInputVertexInterface(
-				TInputDataVertexModel<FFloatTime>(TEXT("Delay"), LOCTEXT("TriggerDelayTooltip", "Time to delay and execute deferred trigger in seconds."), 1.0f),
+				TInputDataVertexModel<FTime>(TEXT("Delay"), LOCTEXT("TriggerDelayTooltip", "Time to delay and execute deferred trigger in seconds."), 1.0f),
 				TInputDataVertexModel<FTrigger>(TEXT("In"), LOCTEXT("TriggerDelayInTooltip", "Triggers delay.")),
 				TInputDataVertexModel<FTrigger>(TEXT("Reset"), LOCTEXT("TriggerDelayResetInTooltip", "Resets the trigger delay, clearing the execution task if pending."))
 			),
diff --git a/Engine/Plugins/Runtime/Metasound/Source/MetasoundStandardNodes/Private/MetasoundTriggerPipeNode.cpp b/Engine/Plugins/Runtime/Metasound/Source/MetasoundStandardNodes/Private/MetasoundTriggerPipeNode.cpp
--- a/Engine/Plugins/Runtime/Metasound/Source/MetasoundStandardNodes/Private/MetasoundTriggerPipeNode.cpp
+++ b/Engine/Plugins/Runtime/Metasound/Source/MetasoundStandardNodes/Private/MetasoundTriggerPipeNode.cpp
@@ -14,6 +14,7 @@
 #include "MetasoundStandardNodesNames.h"
 #include "MetasoundTime.h"
 #include "MetasoundTrigger.h"
+#include "../Public/MetasoundSampleCounter.h"
 
 #define LOCTEXT_NAMESPACE "MetasoundStandardNodes"
 
@@ -29,26 +30,26 @@ namespace Metasound
 		static const FVertexInterface& GetVertexInterface();
 		static TUniquePtr<IOperator> CreateOperator(const FCreateOperatorParams& InParams, FBuildErrorArray& OutErrors);
 
-		FTriggerPipeOperator(const FOperatorSettings& InSettings, const FTriggerReadRef& InTriggerReset, const FTriggerReadRef& InTriggerIn, const FFloatTimeReadRef& InDelay);
+		FTriggerPipeOperator(const FOperatorSettings& InSettings, const FTriggerReadRef& InTriggerReset, const FTriggerReadRef& InTriggerIn, const FTimeReadRef& InDelay);
 
 		virtual FDataReferenceCollection GetInputs() const override;
 		virtual FDataReferenceCollection GetOutputs() const override;
 		void Execute();
 
 	private:
-		TArray<int32> SamplesUntilTrigger;
+		TArray<FSampleCount> SamplesUntilTrigger;
 
 		FTriggerReadRef TriggerIn;
 		FTriggerReadRef TriggerReset;
 		FTriggerWriteRef TriggerOut;
 
-		FFloatTimeReadRef Delay;
+		FTimeReadRef Delay;
 
-		float FramesPerBlock;
-		float SampleRate;
+		FSampleCount FramesPerBlock;
+		FSampleRate SampleRate;
 	};
 
-	FTriggerPipeOperator::FTriggerPipeOperator(const FOperatorSettings& InSettings, const FTriggerReadRef& InTriggerReset, const FTriggerReadRef& InTriggerIn, const FFloatTimeReadRef& InDelay)
+	FTriggerPipeOperator::FTriggerPipeOperator(const FOperatorSettings& InSettings, const FTriggerReadRef& InTriggerReset, const FTriggerReadRef& InTriggerIn, const FTimeReadRef& InDelay)
 	:	TriggerIn(InTriggerIn)
 	,	TriggerReset(InTriggerReset)
 	,	TriggerOut(FTriggerWriteRef::CreateNew(InSettings))
@@ -61,7 +62,7 @@ namespace Metasound
 	FDataReferenceCollection FTriggerPipeOperator::GetInputs() const
 	{
 		FDataReferenceCollection InputDataReferences;
-		InputDataReferences.AddDataReadReference(TEXT("Delay"), FFloatTimeReadRef(Delay));
+		InputDataReferences.AddDataReadReference(TEXT("Delay"), FTimeReadRef(Delay));
 		InputDataReferences.AddDataReadReference(TEXT("In"), FTriggerReadRef(TriggerIn));
 		InputDataReferences.AddDataReadReference(TEXT("Reset"), FTriggerReadRef(TriggerIn));
 		return InputDataReferences;
@@ -86,7 +87,7 @@ namespace Metasound
 			},
 			[this](int32 StartFrame, int32 EndFrame)
 			{
-				SamplesUntilTrigger.AddUnique(Delay->GetNumSamples(SampleRate));
+				SamplesUntilTrigger.AddUnique(Delay->GetSeconds() * SampleRate);
 			}
 		);
 
@@ -98,9 +99,9 @@ namespace Metasound
 			{
 				// Iterate backward and only remove delayed triggers that occur
 				// after the reset trigger's start frame.
-				for (int32 i = SamplesUntilTrigger.Num() - 1; i >= 0; --i)
+				for (FSampleCount i = SamplesUntilTrigger.Num() - 1; i >= 0; --i)
 				{
-					const int32 SamplesRemaining = SamplesUntilTrigger[i] - FramesPerBlock;
+					const FSampleCount SamplesRemaining = SamplesUntilTrigger[i] - FramesPerBlock;
 					if (SamplesRemaining >= StartFrame)
 					{
 						SamplesUntilTrigger.RemoveAtSwap(i);
@@ -109,16 +110,16 @@ namespace Metasound
 			}
 		);
 
-		for (int32 i = SamplesUntilTrigger.Num() - 1; i >= 0; --i)
+		for (FSampleCount i = SamplesUntilTrigger.Num() - 1; i >= 0; --i)
 		{
-			const int32 SamplesRemaining = SamplesUntilTrigger[i] - FramesPerBlock;
+			const FSampleCount SamplesRemaining = SamplesUntilTrigger[i] - FramesPerBlock;
 			if (SamplesRemaining >= 0)
 			{
 				SamplesUntilTrigger[i] -= FramesPerBlock;
 			}
 			else
 			{
-				TriggerOut->TriggerFrame(SamplesRemaining + (int32)FramesPerBlock);
+				TriggerOut->TriggerFrame(SamplesRemaining + static_cast<FSampleCount>(FramesPerBlock));
 				SamplesUntilTrigger.RemoveAtSwap(i);
 			}
 		}
@@ -128,15 +129,9 @@ namespace Metasound
 	{
 		const FTriggerPipeNode& DelayNode = static_cast<const FTriggerPipeNode&>(InParams.Node);
 		const FDataReferenceCollection& InputCollection = InParams.InputDataReferences;
-
-		auto GetOrConstructTime = [&](const FInputVertexInterface& InputVertices, const FString& InputName, ETimeResolution Resolution = ETimeResolution::Seconds)
-		{
-			float DefaultValue = InputVertices[InputName].GetDefaultValue().Value.Get<float>();
-			return InputCollection.GetDataReadReferenceOrConstruct<FFloatTime>(InputName, DefaultValue, Resolution);
-		};
-
 		const FInputVertexInterface& InputInterface = GetVertexInterface().GetInputInterface();
-		FFloatTimeReadRef Delay = GetOrConstructTime(InputInterface, TEXT("Delay"));
+
+		FTimeReadRef Delay = InputCollection.GetDataReadReferenceOrConstructWithVertexDefault<FTime, float>(InputInterface, TEXT("Delay"));
 		FTriggerReadRef TriggerIn = InParams.InputDataReferences.GetDataReadReferenceOrConstruct<FTrigger>(TEXT("In"), InParams.OperatorSettings);
 		FTriggerReadRef TriggerReset = InParams.InputDataReferences.GetDataReadReferenceOrConstruct<FTrigger>(TEXT("Reset"), InParams.OperatorSettings);
 
@@ -147,7 +142,7 @@ namespace Metasound
 	{
 		static const FVertexInterface Interface(
 			FInputVertexInterface(
-				TInputDataVertexModel<FFloatTime>(TEXT("Delay"), LOCTEXT("DelayTooltip", "Time to delay and execute deferred input trigger execution(s) in seconds."), 1.0f),
+				TInputDataVertexModel<FTime>(TEXT("Delay"), LOCTEXT("DelayTooltip", "Time to delay and execute deferred input trigger execution(s) in seconds."), 1.0f),
 				TInputDataVertexModel<FTrigger>(TEXT("In"), LOCTEXT("TriggerInTooltip", "Trigger to execute at a future time by the given delay amount.")),
 				TInputDataVertexModel<FTrigger>(TEXT("Reset"), LOCTEXT("TriggerDelayResetInTooltip", "Resets the trigger delay, clearing any pending execution tasks."))
 			),
diff --git a/Engine/Plugins/Runtime/Metasound/Source/MetasoundStandardNodes/Public/MetasoundBufferAddFloatNode.h b/Engine/Plugins/Runtime/Metasound/Source/MetasoundStandardNodes/Public/MetasoundBufferAddFloatNode.h
deleted file mode 100644
diff --git a/Engine/Plugins/Runtime/Metasound/Source/MetasoundStandardNodes/Public/MetasoundFrequency.h b/Engine/Plugins/Runtime/Metasound/Source/MetasoundStandardNodes/Public/MetasoundFrequency.h
deleted file mode 100644
diff --git a/Engine/Plugins/Runtime/Metasound/Source/MetasoundStandardNodes/Public/MetasoundGain.h b/Engine/Plugins/Runtime/Metasound/Source/MetasoundStandardNodes/Public/MetasoundGain.h
--- a/Engine/Plugins/Runtime/Metasound/Source/MetasoundStandardNodes/Public/MetasoundGain.h
+++ b/Engine/Plugins/Runtime/Metasound/Source/MetasoundStandardNodes/Public/MetasoundGain.h
@@ -119,6 +119,4 @@ namespace Metasound
 		private:
 			float LinearGain = DefaultGain;
 	};
-
-	DECLARE_METASOUND_DATA_REFERENCE_TYPES(FGain, METASOUNDSTANDARDNODES_API, FGainTypeInfo, FGainReadRef, FGainWriteRef)
 }
diff --git a/Engine/Plugins/Runtime/Metasound/Source/MetasoundStandardNodes/Public/MetasoundGainMultiplyNode.h b/Engine/Plugins/Runtime/Metasound/Source/MetasoundStandardNodes/Public/MetasoundGainMultiplyNode.h
deleted file mode 100644
diff --git a/Engine/Plugins/Runtime/Metasound/Source/MetasoundStandardNodes/Public/MetasoundMathNodes.h b/Engine/Plugins/Runtime/Metasound/Source/MetasoundStandardNodes/Public/MetasoundMathNodes.h
--- a/Engine/Plugins/Runtime/Metasound/Source/MetasoundStandardNodes/Public/MetasoundMathNodes.h
+++ b/Engine/Plugins/Runtime/Metasound/Source/MetasoundStandardNodes/Public/MetasoundMathNodes.h
@@ -9,7 +9,6 @@
 #include "MetasoundDataReferenceCollection.h"
 #include "MetasoundExecutableOperator.h"
 #include "MetasoundFacade.h"
-#include "MetasoundFrequency.h"
 #include "MetasoundNode.h"
 #include "MetasoundNodeRegistrationMacro.h"
 #include "MetasoundNodeInterface.h"
@@ -22,59 +21,115 @@
 
 #define LOCTEXT_NAMESPACE "MetasoundMathOpNode"
 
-#define DEFINE_METASOUND_MATHOP(OpName, DataTypeName, DataClass, DisplayName, Description) \
+#define DEFINE_METASOUND_MATHOP(OpName, DataTypeName, DataClass, Description) \
 	class F##OpName##DataTypeName##Node \
-		: public ::Metasound::TMathOpNode<F##OpName##DataTypeName##Node, ::Metasound::TMathOp##OpName<DataClass>, DataClass> \
+		: public ::Metasound::TMathOpNode<F##OpName##DataTypeName##Node, ::Metasound::TMathOp##OpName<DataClass, DataClass>, DataClass, DataClass> \
 	{ \
 	public: \
 		static FNodeClassName GetClassName() { return {::Metasound::StandardNodes::Namespace, TEXT(#OpName), TEXT(#DataTypeName)}; } \
-		static FText GetDisplayName() { return DisplayName; } \
+		static FText GetDisplayName() { return MathOpNames::Get##OpName##DisplayName<DataClass>(); } \
 		static FText GetDescription() { return Description; } \
 		static FName GetImageName() { return "MetasoundEditor.Graph.Node.Math." #OpName; } \
 		F##OpName##DataTypeName##Node(const FNodeInitData& InInitData) : TMathOpNode(InInitData) { } \
 	};
 
+#define DEFINE_METASOUND_OPERAND_TYPED_MATHOP(OpName, DataTypeName, DataClass, OperandTypeName, OperandDataClass, Description) \
+	class F##OpName##DataTypeName##OperandTypeName##Node \
+		: public ::Metasound::TMathOpNode<F##OpName##DataTypeName##OperandTypeName##Node, ::Metasound::TMathOp##OpName<DataClass, OperandDataClass>, DataClass, OperandDataClass> \
+	{ \
+	public: \
+		static FNodeClassName GetClassName() { return {::Metasound::StandardNodes::Namespace, TEXT(#OpName), TEXT(#DataTypeName " by " #OperandTypeName)}; } \
+		static FText GetDisplayName() { return MathOpNames::Get##OpName##DisplayName<DataClass, OperandDataClass>(); } \
+		static FText GetDescription() { return Description; } \
+		static FName GetImageName() { return "MetasoundEditor.Graph.Node.Math." #OpName; } \
+		F##OpName##DataTypeName##OperandTypeName##Node(const FNodeInitData& InInitData) : TMathOpNode(InInitData) { } \
+	};
+
 namespace Metasound
 {
-	namespace MathOpDisplayNames
+	namespace MathOpNames
 	{
+		static const FString PrimaryOperandName = TEXT("PrimaryOperand");
+
+		static const FString AdditionalOperandsName = TEXT("AdditionalOperands");
+
 		template<typename DataType>
 		const FText GetAddDisplayName()
 		{
-			static const FText DisplayName = FText::Format(LOCTEXT("AddNodeDisplayNamePattern", "Add {0}"), FText::FromString(GetMetasoundDataTypeString<DataType>()));
+			static const FText DisplayName = FText::Format(LOCTEXT("AddNodeDisplayNamePattern", "Add ({0})"), FText::FromString(GetMetasoundDataTypeString<DataType>()));
+			return DisplayName;
+		}
+
+		template<typename DataType, typename OperandDataType>
+		const FText GetAddDisplayName()
+		{
+			static const FText DisplayName = FText::Format(
+				LOCTEXT("AddNodeDisplayNamePattern", "Add ({0} to {1})"),
+				FText::FromString(GetMetasoundDataTypeString<OperandDataType>()),
+				FText::FromString(GetMetasoundDataTypeString<DataType>()));
 			return DisplayName;
 		}
 
 		template<typename DataType>
 		const FText GetSubtractDisplayName()
 		{
-			static const FText DisplayName = FText::Format(LOCTEXT("SubtractNodeDisplayNamePattern", "Subtract {0}"), FText::FromString(GetMetasoundDataTypeString<DataType>()));
+			static const FText DisplayName = FText::Format(LOCTEXT("SubtractNodeDisplayNamePattern", "Subtract ({0})"), FText::FromString(GetMetasoundDataTypeString<DataType>()));
+			return DisplayName;
+		}
+
+		template<typename DataType, typename OperandDataType>
+		const FText GetSubtractDisplayName()
+		{
+			static const FText DisplayName = FText::Format(
+				LOCTEXT("SubtractNodeOperandTypedDisplayNamePattern", "Subtract ({0} from {1})"),
+				FText::FromString(GetMetasoundDataTypeString<OperandDataType>()),
+				FText::FromString(GetMetasoundDataTypeString<DataType>()));
 			return DisplayName;
 		}
 
 		template<typename DataType>
 		const FText GetMultiplyDisplayName()
 		{
-			static const FText DisplayName = FText::Format(LOCTEXT("MultiplyNodeDisplayNamePattern", "Multiply {0}"), FText::FromString(GetMetasoundDataTypeString<DataType>()));
+			static const FText DisplayName = FText::Format(LOCTEXT("MultiplyNodeDisplayNamePattern", "Multiply ({0})"), FText::FromString(GetMetasoundDataTypeString<DataType>()));
+			return DisplayName;
+		}
+
+		template<typename DataType, typename OperandDataType>
+		const FText GetMultiplyDisplayName()
+		{
+			static const FText DisplayName = FText::Format(
+				LOCTEXT("MultiplyNodeOperandTypedDisplayNamePattern", "Multiply ({0} by {1})"),
+				FText::FromString(GetMetasoundDataTypeString<DataType>()),
+				FText::FromString(GetMetasoundDataTypeString<OperandDataType>()));
 			return DisplayName;
 		}
 
 		template<typename DataType>
 		const FText GetDivideDisplayName()
 		{
-			static const FText DisplayName = FText::Format(LOCTEXT("DivideNodeDisplayNamePattern", "Divide {0}"), FText::FromString(GetMetasoundDataTypeString<DataType>()));
+			static const FText DisplayName = FText::Format(LOCTEXT("DivideNodeDisplayNamePattern", "Divide ({0})"), FText::FromString(GetMetasoundDataTypeString<DataType>()));
+			return DisplayName;
+		}
+
+		template<typename DataType, typename OperandDataType>
+		const FText GetDivideDisplayName()
+		{
+			static const FText DisplayName = FText::Format(
+				LOCTEXT("DivideNodeOperandTypedDisplayNamePattern", "Divide ({0} by {1})"),
+				FText::FromString(GetMetasoundDataTypeString<DataType>()),
+				FText::FromString(GetMetasoundDataTypeString<OperandDataType>()));
 			return DisplayName;
 		}
 
 		template<typename DataType>
 		const FText GetRandRangeDisplayName()
 		{
-			static const FText DisplayName = FText::Format(LOCTEXT("RandRangeNodeDisplayNamePattern", "RandRange {0}"), FText::FromString(GetMetasoundDataTypeString<DataType>()));
+			static const FText DisplayName = FText::Format(LOCTEXT("RandRangeNodeDisplayNamePattern", "Random ({0})"), FText::FromString(GetMetasoundDataTypeString<DataType>()));
 			return DisplayName;
 		}
 	}
 
-	template <typename TDerivedClass, typename TMathOpClass, typename TDataClass>
+	template <typename TDerivedClass, typename TMathOpClass, typename TDataClass, typename TOperandDataClass>
 	class TMathOpNode : public FNodeFacade
 	{
 	private:
@@ -82,27 +137,16 @@ namespace Metasound
 		{
 		private:
 			using TDataClassReadRef = TDataReadReference<TDataClass>;
+			using TOperandDataClassReadRef = TDataReadReference<TOperandDataClass>;
 			using TDataClassWriteRef = TDataWriteReference<TDataClass>;
 
-			TArray<TDataClassReadRef> OperandRefs;
+			TMathOpClass InstanceData;
+
+			TDataClassReadRef PrimaryOperandRef;
+			TArray<TOperandDataClassReadRef> AdditionalOperandRefs;
 			TDataClassWriteRef ValueRef;
 
 		public:
-			static FVertexInterface DeclareVertexInterface()
-			{
-				static const FVertexInterface Interface(
-					FInputVertexInterface(
-						TInputDataVertexModel<TDataClass>(TEXT("Operand1"), LOCTEXT("MathOpMinTooltip", "First value to operate on.")),
-						TInputDataVertexModel<TDataClass>(TEXT("Operand2"), LOCTEXT("MathOpMaxTooltip", "Second value to operate on."))
-					),
-					FOutputVertexInterface(
-						TOutputDataVertexModel<TDataClass>(TEXT("Out"), LOCTEXT("MathOpOutTooltip", "Math operation result"))
-					)
-				);
-
-				return Interface;
-			}
-
 			static const FNodeClassMetadata& GetNodeInfo()
 			{
 				auto InitNodeInfo = []() -> FNodeClassMetadata
@@ -121,7 +165,7 @@ namespace Metasound
 					Info.Description = TDerivedClass::GetDescription();
 					Info.Author = PluginAuthor;
 					Info.PromptIfMissing = PluginNodeMissingPrompt;
-					Info.DefaultInterface = DeclareVertexInterface();
+					Info.DefaultInterface = TMathOpClass::GetVertexInterface();
 					Info.DisplayStyle = DisplayStyle;
 					Info.CategoryHierarchy = { LOCTEXT("Metasound_MathCategory", "Math") };
 
@@ -134,33 +178,47 @@ namespace Metasound
 
 			void Execute()
 			{
-				TMathOpClass::Calculate(OperandRefs, ValueRef);
+				TMathOpClass::Calculate(InstanceData, PrimaryOperandRef, AdditionalOperandRefs, ValueRef);
 			}
 
 			static TUniquePtr<IOperator> CreateOperator(const FCreateOperatorParams& InParams, FBuildErrorArray& OutErrors)
 			{
 				const TMathOpNode& MathOpNode = static_cast<const TMathOpNode&>(InParams.Node);
 
-				// TODO: Support dynamic number of inputs
-				TDataClassReadRef Op1 = InParams.InputDataReferences.GetDataReadReferenceOrConstruct<TDataClass>(TEXT("Operand1"), TMathOpClass::GetDefault(&InParams.OperatorSettings));
-				TDataClassReadRef Op2 = InParams.InputDataReferences.GetDataReadReferenceOrConstruct<TDataClass>(TEXT("Operand2"), TMathOpClass::GetDefault(&InParams.OperatorSettings));
+				const FInputVertexInterface& InputInterface = MathOpNode.GetVertexInterface().GetInputInterface();
 
-				TArray<TDataClassReadRef> OperandRefs = { Op1, Op2 };
+				FLiteral DefaultValue = FLiteral::CreateInvalid();
+				const FString& PrimaryOperandName = MathOpNames::PrimaryOperandName;
+				if (InputInterface.Contains(PrimaryOperandName))
+				{
+					DefaultValue = InputInterface[PrimaryOperandName].GetDefaultValue().Clone();
+				}
+				TDataClassReadRef PrimaryOperand = InParams.InputDataReferences.GetDataReadReferenceOrConstruct<TDataClass>(MathOpNames::PrimaryOperandName, TMathOpClass::GetDefault(InParams.OperatorSettings, DefaultValue));
+
+				// TODO: Support dynamic number of inputs
+				const FString& OpName = MathOpNames::AdditionalOperandsName;
+				if (InputInterface.Contains(OpName))
+				{
+					DefaultValue = InputInterface[OpName].GetDefaultValue().Clone();
+				}
+				TOperandDataClassReadRef Op1 = InParams.InputDataReferences.GetDataReadReferenceOrConstruct<TOperandDataClass>(OpName, TMathOpClass::GetDefaultOp(InParams.OperatorSettings, DefaultValue));
+				TArray<TOperandDataClassReadRef> AdditionalOperandRefs = { Op1 };
 
-				return MakeUnique<TMathOperator>(InParams.OperatorSettings, OperandRefs);
+				return MakeUnique<TMathOperator>(InParams.OperatorSettings, PrimaryOperand, AdditionalOperandRefs);
 			}
 
-			TMathOperator(const FOperatorSettings& InSettings, const TArray<TDataClassReadRef>& InOperands)
-			:	OperandRefs(InOperands)
-			,	ValueRef(TDataClassWriteRef::CreateNew(TMathOpClass::GetDefault(&InSettings)))
+			TMathOperator(const FOperatorSettings& InSettings, const TDataClassReadRef& InPrimaryOperand, const TArray<TOperandDataClassReadRef>& InAdditionalOperands)
+			:	PrimaryOperandRef(InPrimaryOperand)
+			,	AdditionalOperandRefs(InAdditionalOperands)
+			,	ValueRef(TDataWriteReferenceFactory<TDataClass>::CreateAny(InSettings))
 			{
 			}
 
 			virtual FDataReferenceCollection GetInputs() const override
 			{
 				FDataReferenceCollection InputDataReferences;
-				InputDataReferences.AddDataReadReference(TEXT("Operand1"), OperandRefs[0]);
-				InputDataReferences.AddDataReadReference(TEXT("Operand2"), OperandRefs[1]);
+				InputDataReferences.AddDataReadReference(MathOpNames::PrimaryOperandName, PrimaryOperandRef);
+				InputDataReferences.AddDataReadReference(MathOpNames::AdditionalOperandsName, AdditionalOperandRefs[0]);
 				return InputDataReferences;
 			}
 
@@ -181,111 +239,201 @@ namespace Metasound
 	};
 
 // Standard Math Operations
-	template <typename TDataClass>
+	template <typename TDataClass, typename TOperandDataClass = TDataClass>
 	class TMathOpAdd
 	{
 	public:
 		using TDataClassReadRef = TDataReadReference<TDataClass>;
 		using TDataClassWriteRef = TDataWriteReference<TDataClass>;
 
-		static TDataClass GetDefault(const FOperatorSettings* InSettings)
+		using TOperandDataClassReadRef = TDataReadReference<TOperandDataClass>;
+
+		static const FVertexInterface& GetVertexInterface()
 		{
-			return 0;
+			static const FVertexInterface Interface(
+				FInputVertexInterface(
+					TInputDataVertexModel<TDataClass>(MathOpNames::PrimaryOperandName, LOCTEXT("MathOpAddendInitialTooltip", "Initial attend."), static_cast<TDataClass>(0)),
+					TInputDataVertexModel<TOperandDataClass>(MathOpNames::AdditionalOperandsName, LOCTEXT("MathOpAddendsTooltip", "Additional attend(s)."), static_cast<TOperandDataClass>(0))
+				),
+				FOutputVertexInterface(
+					TOutputDataVertexModel<TDataClass>(TEXT("Out"), LOCTEXT("MathOpOutTooltip", "Math operation result"))
+				)
+			);
+
+			return Interface;
 		}
 
-		static void Calculate(const TArray<TDataClassReadRef>& InOperands, TDataClassWriteRef& OutResult)
+		static TDataClass GetDefault(const FOperatorSettings& InSettings, const FLiteral& InVertexDefault)
 		{
-			if (!ensure(!InOperands.IsEmpty()))
+			return static_cast<TDataClass>(0);
+		}
+
+		static TOperandDataClass GetDefaultOp(const FOperatorSettings& InSettings, const FLiteral& InVertexDefault)
+		{
+			return static_cast<TOperandDataClass>(0);
+		}
+
+		static void Calculate(TMathOpAdd<TDataClass, TOperandDataClass>& InInstanceData, const TDataClassReadRef& InPrimaryOperand, const TArray<TOperandDataClassReadRef>& InAdditionalOperands, TDataClassWriteRef& OutResult)
+		{
+			if (!ensure(!InAdditionalOperands.IsEmpty()))
 			{
 				return;
 			}
 
-			*OutResult = *InOperands[0];
-			for (int32 i = 1; i < InOperands.Num(); ++i)
+			*OutResult = *InPrimaryOperand;
+			for (int32 i = 0; i < InAdditionalOperands.Num(); ++i)
 			{
-				*OutResult += *InOperands[i];
+				*OutResult += *InAdditionalOperands[i];
 			}
 		}
 	};
 
-	template <typename TDataClass>
+	template <typename TDataClass, typename TOperandDataClass = TDataClass>
 	class TMathOpSubtract
 	{
 	public:
 		using TDataClassReadRef = TDataReadReference<TDataClass>;
 		using TDataClassWriteRef = TDataWriteReference<TDataClass>;
 
-		static TDataClass GetDefault(const FOperatorSettings* InSettings)
+		using TOperandDataClassReadRef = TDataReadReference<TOperandDataClass>;
+
+		static const FVertexInterface& GetVertexInterface()
+		{
+			static const FVertexInterface Interface(
+				FInputVertexInterface(
+					TInputDataVertexModel<TDataClass>(MathOpNames::PrimaryOperandName, LOCTEXT("MathOpMinuendTooltip", "Minuend."), static_cast<TDataClass>(0)),
+					TInputDataVertexModel<TOperandDataClass>(MathOpNames::AdditionalOperandsName, LOCTEXT("MathOpSubtrahendsTooltip", "Subtrahend(s)."), static_cast<TOperandDataClass>(0))
+				),
+				FOutputVertexInterface(
+					TOutputDataVertexModel<TDataClass>(TEXT("Out"), LOCTEXT("MathOpOutTooltip", "Subtraction result"))
+				)
+			);
+
+			return Interface;
+		}
+
+		static TDataClass GetDefault(const FOperatorSettings& InSettings, const FLiteral& InVertexDefault)
+		{
+			return static_cast<TDataClass>(0);
+		}
+
+		static TOperandDataClass GetDefaultOp(const FOperatorSettings& InSettings, const FLiteral& InVertexDefault)
 		{
-			return 0;
+			return static_cast<TOperandDataClass>(0);
 		}
 
-		static void Calculate(const TArray<TDataClassReadRef>& InOperands, TDataClassWriteRef& OutResult)
+		static void Calculate(TMathOpSubtract<TDataClass, TOperandDataClass>& InInstanceData, const TDataClassReadRef& InPrimaryOperand, const TArray<TOperandDataClassReadRef>& InAdditionalOperands, TDataClassWriteRef& OutResult)
 		{
-			if (!ensure(!InOperands.IsEmpty()))
+			if (!ensure(!InAdditionalOperands.IsEmpty()))
 			{
 				return;
 			}
 
-			*OutResult = *InOperands[0];
-			for (int32 i = 1; i < InOperands.Num(); ++i)
+			*OutResult = *InPrimaryOperand;
+			for (int32 i = 0; i < InAdditionalOperands.Num(); ++i)
 			{
-				*OutResult -= *InOperands[i];
+				*OutResult -= *InAdditionalOperands[i];
 			}
 		}
 	};
 
-	template <typename TDataClass>
+	template <typename TDataClass, typename TOperandDataClass = TDataClass>
 	class TMathOpMultiply
 	{
 	public:
 		using TDataClassReadRef = TDataReadReference<TDataClass>;
 		using TDataClassWriteRef = TDataWriteReference<TDataClass>;
 
-		static TDataClass GetDefault(const FOperatorSettings* InSettings)
+		using TOperandDataClassReadRef = TDataReadReference<TOperandDataClass>;
+
+		static const FVertexInterface& GetVertexInterface()
+		{
+			static const FVertexInterface Interface(
+				FInputVertexInterface(
+					TInputDataVertexModel<TDataClass>(MathOpNames::PrimaryOperandName, LOCTEXT("MathOpInitMultiplicandTooltip", "Initial multiplicand."), static_cast<TDataClass>(1)),
+					TInputDataVertexModel<TOperandDataClass>(MathOpNames::AdditionalOperandsName, LOCTEXT("MathOpMultiplicandsTooltip", "Additional multiplicand(s)."), static_cast<TOperandDataClass>(1))
+				),
+				FOutputVertexInterface(
+					TOutputDataVertexModel<TDataClass>(TEXT("Out"), LOCTEXT("MultiplicationResultTooltip", "Multiplication result"))
+				)
+			);
+
+			return Interface;
+		}
+
+		static TDataClass GetDefault(const FOperatorSettings& InSettings, const FLiteral& InVertexDefault)
+		{
+			return static_cast<TDataClass>(1);
+		}
+
+		static TOperandDataClass GetDefaultOp(const FOperatorSettings& InSettings, const FLiteral& InVertexDefault)
 		{
-			return 1;
+			return static_cast<TOperandDataClass>(1);
 		}
 
-		static void Calculate(const TArray<TDataClassReadRef>& InOperands, TDataClassWriteRef& OutResult)
+		static void Calculate(TMathOpMultiply<TDataClass, TOperandDataClass>& InInstanceData, const TDataClassReadRef& InPrimaryOperand, const TArray<TOperandDataClassReadRef>& InAdditionalOperands, TDataClassWriteRef& OutResult)
 		{
-			if (!ensure(!InOperands.IsEmpty()))
+			if (!ensure(!InAdditionalOperands.IsEmpty()))
 			{
 				return;
 			}
 
-			*OutResult = *InOperands[0];
-			for (int32 i = 1; i < InOperands.Num(); ++i)
+			*OutResult = *InPrimaryOperand;
+			for (int32 i = 0; i < InAdditionalOperands.Num(); ++i)
 			{
-				*OutResult *= *InOperands[i];
+				*OutResult *= *InAdditionalOperands[i];
 			}
 		}
 	};
 
-	template <typename TDataClass>
+	template <typename TDataClass, typename TOperandDataClass = TDataClass>
 	class TMathOpDivide
 	{
 	public:
 		using TDataClassReadRef = TDataReadReference<TDataClass>;
 		using TDataClassWriteRef = TDataWriteReference<TDataClass>;
 
-		static TDataClass GetDefault(const FOperatorSettings* InSettings)
+		using TOperandDataClassReadRef = TDataReadReference<TOperandDataClass>;
+
+		static const FVertexInterface& GetVertexInterface()
+		{
+			static const FVertexInterface Interface(
+				FInputVertexInterface(
+					TInputDataVertexModel<TDataClass>(MathOpNames::PrimaryOperandName, LOCTEXT("MathOpDividendTooltip", "Dividend."), static_cast<TDataClass>(0)),
+					TInputDataVertexModel<TOperandDataClass>(MathOpNames::AdditionalOperandsName, LOCTEXT("MathOpDivisorsTooltip", "Divisor(s)."), static_cast<TOperandDataClass>(0))
+				),
+				FOutputVertexInterface(
+					TOutputDataVertexModel<TDataClass>(TEXT("Out"), LOCTEXT("MathOpOutTooltip", "Math operation result"))
+				)
+			);
+
+			return Interface;
+		}
+
+		static TDataClass GetDefault(const FOperatorSettings& InSettings, const FLiteral& InVertexDefault)
+		{
+			return static_cast<TDataClass>(1);
+		}
+
+		static TOperandDataClass GetDefaultOp(const FOperatorSettings& InSettings, const FLiteral& InVertexDefault)
 		{
-			return 1;
+			return static_cast<TOperandDataClass>(1);
 		}
 
-		static void Calculate(const TArray<TDataClassReadRef>& InOperands, TDataClassWriteRef& OutResult)
+		static void Calculate(TMathOpDivide& InInstanceData, const TDataClassReadRef& InPrimaryOperand, const TArray<TOperandDataClassReadRef>& InAdditionalOperands, TDataClassWriteRef& OutResult)
 		{
-			if (!ensure(!InOperands.IsEmpty()))
+			*OutResult = *InPrimaryOperand;
+
+			if (!ensure(!InAdditionalOperands.IsEmpty()))
 			{
 				return;
 			}
 
-			*OutResult = *InOperands[0];
-			for (int32 i = 1; i < InOperands.Num(); ++i)
+			*OutResult = *InPrimaryOperand;
+			for (int32 i = 0; i < InAdditionalOperands.Num(); ++i)
 			{
-				const TDataClass& OperandValue = *InOperands[i];
-				if (OperandValue == 0)
+				const TDataClass& OperandValue = *InAdditionalOperands[i];
+				if (OperandValue == static_cast<TDataClass>(0))
 				{
 					// TODO: Error here
 					return;
@@ -301,20 +449,40 @@ namespace Metasound
 	class TMathOpAdd<FAudioBuffer>
 	{
 	public:
-		static FAudioBuffer GetDefault(const FOperatorSettings* InSettings)
+		static const FVertexInterface& GetVertexInterface()
 		{
-			return FAudioBuffer(InSettings ? InSettings->GetNumFramesPerBlock() : 0);
+			static const FVertexInterface Interface(
+				FInputVertexInterface(
+					TInputDataVertexModel<FAudioBuffer>(MathOpNames::PrimaryOperandName, LOCTEXT("MathOpAddendTooltip", "First attend.")),
+					TInputDataVertexModel<FAudioBuffer>(MathOpNames::AdditionalOperandsName, LOCTEXT("MathOpAddendAdditionalTooltip", "Additional attends."))
+				),
+				FOutputVertexInterface(
+					TOutputDataVertexModel<FAudioBuffer>(TEXT("Out"), LOCTEXT("MathOpOutTooltip", "Math operation result"))
+				)
+			);
+
+			return Interface;
 		}
 
-		static void Calculate(const TArray<FAudioBufferReadRef>& InOperands, FAudioBufferWriteRef& OutResult)
+		static FAudioBuffer GetDefault(const FOperatorSettings& InSettings, const FLiteral& InVertexDefault)
 		{
-			if (!ensure(!InOperands.IsEmpty()))
+			return FAudioBuffer(InSettings.GetNumFramesPerBlock());
+		}
+
+		static FAudioBuffer GetDefaultOp(const FOperatorSettings& InSettings, const FLiteral& InVertexDefault)
+		{
+			return GetDefault(InSettings, InVertexDefault);
+		}
+
+		static void Calculate(TMathOpAdd<FAudioBuffer>& InInstanceData, FAudioBufferReadRef& InPrimaryOperand, const TArray<FAudioBufferReadRef>& InAdditionalOperands, FAudioBufferWriteRef& OutResult)
+		{
+			if (!ensure(!InAdditionalOperands.IsEmpty()))
 			{
 				OutResult->Zero();
 				return;
 			}
 
-			const int32 NumSamples = InOperands[0]->Num();
+			const int32 NumSamples = InAdditionalOperands[0]->Num();
 			if (!ensure(NumSamples == OutResult->Num()))
 			{
 				// TODO: Error
@@ -322,11 +490,11 @@ namespace Metasound
 				return;
 			}
 
-			FMemory::Memcpy(OutResult->GetData(), InOperands[0]->GetData(), sizeof(float) * NumSamples);
+			FMemory::Memcpy(OutResult->GetData(), InAdditionalOperands[0]->GetData(), sizeof(float) * NumSamples);
 
-			for (int32 i = 1; i < InOperands.Num(); ++i)
+			for (int32 i = 0; i < InAdditionalOperands.Num(); ++i)
 			{
-				const FAudioBuffer& Operand = *InOperands[i];
+				const FAudioBuffer& Operand = *InAdditionalOperands[i];
 				if (!ensure(NumSamples == Operand.Num()))
 				{
 					// TODO: Error
@@ -334,7 +502,7 @@ namespace Metasound
 					return;
 				}
 
-				const float* OperandData = InOperands[i]->GetData();
+				const float* OperandData = InAdditionalOperands[i]->GetData();
 				float* OutData = OutResult->GetData();
 
 				if (NumSamples % AUDIO_SIMD_FLOAT_ALIGNMENT)
@@ -354,24 +522,183 @@ namespace Metasound
 		}
 	};
 
+	template <>
+	class TMathOpAdd<FTime>
+	{
+	public:
+		static const FVertexInterface& GetVertexInterface()
+		{
+			static const FVertexInterface Interface(
+				FInputVertexInterface(
+					TInputDataVertexModel<FTime>(MathOpNames::PrimaryOperandName, LOCTEXT("MathOpAddendTooltip", "First attend."), 0.0f),
+					TInputDataVertexModel<FTime>(MathOpNames::AdditionalOperandsName, LOCTEXT("MathOpAddendAdditionalTooltip", "Additional attends."), 0.0f)
+				),
+				FOutputVertexInterface(
+					TOutputDataVertexModel<FTime>(TEXT("Out"), LOCTEXT("MathOpOutTooltip", "Math operation result"))
+				)
+			);
+
+			return Interface;
+		}
+
+		static FTime GetDefault(const FOperatorSettings& InSettings, const FLiteral& InVertexDefault)
+		{
+			return FTime(InVertexDefault.Value.Get<float>());
+		}
+
+		static FTime GetDefaultOp(const FOperatorSettings& InSettings, const FLiteral& InVertexDefault)
+		{
+			return GetDefault(InSettings, InVertexDefault);
+		}
+
+		static void Calculate(TMathOpAdd& InInstanceData, const FTimeReadRef& InPrimaryOperand, const TArray<FTimeReadRef>& InAdditionalOperands, FTimeWriteRef& OutResult)
+		{
+			if (!ensure(!InAdditionalOperands.IsEmpty()))
+			{
+				return;
+			}
+
+			*OutResult = *InPrimaryOperand;
+			for (int32 i = 0; i < InAdditionalOperands.Num(); ++i)
+			{
+				*OutResult += *InAdditionalOperands[i];
+			}
+		}
+	};
+
+	template <>
+	class TMathOpAdd<FAudioBuffer, float>
+	{
+	public:
+		static const FVertexInterface& GetVertexInterface()
+		{
+			static const FVertexInterface Interface(
+				FInputVertexInterface(
+					TInputDataVertexModel<FAudioBuffer>(MathOpNames::PrimaryOperandName, LOCTEXT("MathOpAddendAudioTooltip", "Audio Buffer to add offset(s) to.")),
+					TInputDataVertexModel<float>(MathOpNames::AdditionalOperandsName, LOCTEXT("MathOpAddendAdditionalTooltip", "Float attends of which to offset buffer samples."), 0.0f)
+				),
+				FOutputVertexInterface(
+					TOutputDataVertexModel<FAudioBuffer>(TEXT("Out"), LOCTEXT("MathOpAudioFloatAddOutTooltip", "Resulting buffer"))
+				)
+			);
+
+			return Interface;
+		}
+
+		static FAudioBuffer GetDefault(const FOperatorSettings& InSettings, const FLiteral& InVertexDefault)
+		{
+			return FAudioBuffer(InSettings.GetNumFramesPerBlock());
+		}
+
+		static float GetDefaultOp(const FOperatorSettings& InSettings, const FLiteral& InVertexDefault)
+		{
+			return InVertexDefault.Value.Get<float>();
+		}
+
+		static void Calculate(TMathOpAdd& InInstanceData, const FAudioBufferReadRef& InPrimaryOperand, const TArray<FFloatReadRef>& InAdditionalOperands, FAudioBufferWriteRef& OutResult)
+		{
+			FMemory::Memcpy(OutResult->GetData(), InPrimaryOperand->GetData(), sizeof(float) * OutResult->Num());
+
+			const int32 SIMDRemainder = OutResult->Num() % AUDIO_SIMD_FLOAT_ALIGNMENT;
+			const int32 SIMDCount = OutResult->Num() - SIMDRemainder;
+
+			for (int32 i = 0; i < InAdditionalOperands.Num(); ++i)
+			{
+				Audio::AddConstantToBufferInplace(OutResult->GetData(), SIMDCount, *InAdditionalOperands[i]);
+
+				for (int32 j = SIMDCount; j < OutResult->Num(); ++j)
+				{
+					OutResult->GetData()[j] += *InAdditionalOperands[i];
+				}
+			}
+		}
+	};
+
+	// int64 has to be explicitly implemented as they are PODs that are not implemented as a literal base type.
+	template <>
+	class TMathOpAdd<int64>
+	{
+	public:
+		static const FVertexInterface& GetVertexInterface()
+		{
+			static const FVertexInterface Interface(
+				FInputVertexInterface(
+					TInputDataVertexModel<int64>(MathOpNames::PrimaryOperandName, LOCTEXT("MathOpInt64InitialAddendMinuendTooltip", "Initial Addend."), 0),
+					TInputDataVertexModel<int64>(MathOpNames::AdditionalOperandsName, LOCTEXT("MathOpInt64AddendsTooltip", "Additional Addend(s)."), 0)
+				),
+				FOutputVertexInterface(
+					TOutputDataVertexModel<int64>(TEXT("Out"), LOCTEXT("MathOpOutTooltip", "Addition result"))
+				)
+			);
+
+			return Interface;
+		}
+
+		static int32 GetDefault(const FOperatorSettings& InSettings, const FLiteral& InVertexDefault)
+		{
+			return static_cast<int64>(InVertexDefault.Value.Get<int32>());
+		}
+
+		static int32 GetDefaultOp(const FOperatorSettings& InSettings, const FLiteral& InVertexDefault)
+		{
+			return static_cast<int64>(InVertexDefault.Value.Get<int32>());
+		}
+
+		static void Calculate(TMathOpAdd& InInstanceData, const FInt64ReadRef& InPrimaryOperand, const TArray<FInt64ReadRef>& InAdditionalOperands, FInt64WriteRef& OutResult)
+		{
+			*OutResult = *InPrimaryOperand;
+
+			if (!ensure(!InAdditionalOperands.IsEmpty()))
+			{
+				return;
+			}
+
+			for (int32 i = 0; i < InAdditionalOperands.Num(); ++i)
+			{
+				const int64 OperandValue = *InAdditionalOperands[i];
+				*OutResult += OperandValue;
+			}
+		}
+	};
+
 	template <>
 	class TMathOpSubtract<FAudioBuffer>
 	{
 	public:
-		static FAudioBuffer GetDefault(const FOperatorSettings* InSettings)
+		static const FVertexInterface& GetVertexInterface()
+		{
+			static const FVertexInterface Interface(
+				FInputVertexInterface(
+					TInputDataVertexModel<FAudioBuffer>(MathOpNames::PrimaryOperandName, LOCTEXT("MathOpBuffersMinuendTooltip", "Initial buffer to act as minuend.")),
+					TInputDataVertexModel<FAudioBuffer>(MathOpNames::AdditionalOperandsName, LOCTEXT("MathOpSubtractBuffersSubtrahendsTooltip", "Additional buffers to act as subtrahend(s)."))
+				),
+				FOutputVertexInterface(
+					TOutputDataVertexModel<FAudioBuffer>(TEXT("Out"), LOCTEXT("MathOpSubtractBuffersOutTooltip", "Resulting buffer"))
+				)
+			);
+
+			return Interface;
+		}
+
+		static FAudioBuffer GetDefault(const FOperatorSettings& InSettings, const FLiteral& InVertexDefault)
+		{
+			return FAudioBuffer(InSettings.GetNumFramesPerBlock());
+		}
+
+		static FAudioBuffer GetDefaultOp(const FOperatorSettings& InSettings, const FLiteral& InVertexDefault)
 		{
-			return FAudioBuffer(InSettings ? InSettings->GetNumFramesPerBlock() : 0);
+			return GetDefault(InSettings, InVertexDefault);
 		}
 
-		static void Calculate(TArray<FAudioBufferReadRef>& InOperands, FAudioBufferWriteRef& OutResult)
+		static void Calculate(TMathOpSubtract& InInstanceData, FAudioBufferReadRef& InPrimaryOperand, const TArray<FAudioBufferReadRef>& InAdditionalOperands, FAudioBufferWriteRef& OutResult)
 		{
-			if (!ensure(!InOperands.IsEmpty()))
+			if (!ensure(!InAdditionalOperands.IsEmpty()))
 			{
 				OutResult->Zero();
 				return;
 			}
 
-			const int32 NumSamples = InOperands[0]->Num();
+			const int32 NumSamples = InAdditionalOperands[0]->Num();
 			if (!ensure(NumSamples == OutResult->Num()))
 			{
 				// TODO: Error
@@ -379,11 +706,11 @@ namespace Metasound
 				return;
 			}
 
-			FMemory::Memcpy(OutResult->GetData(), InOperands[0]->GetData(), sizeof(float) * NumSamples);
+			FMemory::Memcpy(OutResult->GetData(), InAdditionalOperands[0]->GetData(), sizeof(float) * NumSamples);
 
-			for (int32 i = 1; i < InOperands.Num(); ++i)
+			for (int32 i = 0; i < InAdditionalOperands.Num(); ++i)
 			{
-				const FAudioBuffer& Operand = *InOperands[i];
+				const FAudioBuffer& Operand = *InAdditionalOperands[i];
 				if (!ensure(NumSamples == Operand.Num()))
 				{
 					// TODO: Error
@@ -391,7 +718,7 @@ namespace Metasound
 					return;
 				}
 
-				const float* OperandData = InOperands[i]->GetData();
+				const float* OperandData = InAdditionalOperands[i]->GetData();
 				float* OutData = OutResult->GetData();
 
 				if (NumSamples % AUDIO_SIMD_FLOAT_ALIGNMENT)
@@ -409,24 +736,135 @@ namespace Metasound
 		}
 	};
 
+	template <>
+	class TMathOpSubtract<FTime>
+	{
+	public:
+		static const FVertexInterface& GetVertexInterface()
+		{
+			static const FVertexInterface Interface(
+				FInputVertexInterface(
+					TInputDataVertexModel<FTime>(MathOpNames::PrimaryOperandName, LOCTEXT("MathOpTimeMinuendTooltip", "Time minuend."), 0.0f),
+					TInputDataVertexModel<FTime>(MathOpNames::AdditionalOperandsName, LOCTEXT("MathOpTimeSubtrahendsTooltip", "Time subtrahends."), 0.0f)
+				),
+				FOutputVertexInterface(
+					TOutputDataVertexModel<FTime>(TEXT("Out"), LOCTEXT("MathOpTimeSubtractOutTooltip", "Resulting time value"))
+				)
+			);
+
+			return Interface;
+		}
+
+		static FTime GetDefault(const FOperatorSettings& InSettings, const FLiteral& InVertexDefault)
+		{
+			return FTime(InVertexDefault.Value.Get<float>());
+		}
+
+		static FTime GetDefaultOp(const FOperatorSettings& InSettings, const FLiteral& InVertexDefault)
+		{
+			return GetDefault(InSettings, InVertexDefault);
+		}
+
+		static void Calculate(TMathOpSubtract& InInstanceData, const FTimeReadRef& InPrimaryOperand, const TArray<FTimeReadRef>& InAdditionalOperands, FTimeWriteRef& OutResult)
+		{
+			if (!ensure(!InAdditionalOperands.IsEmpty()))
+			{
+				return;
+			}
+
+			*OutResult = *InPrimaryOperand;
+			for (int32 i = 0; i < InAdditionalOperands.Num(); ++i)
+			{
+				*OutResult -= *InAdditionalOperands[i];
+			}
+		}
+	};
+
+	// int64 has to be explicitly implemented as they are PODs that are not implemented as a literal base type.
+	template <>
+	class TMathOpSubtract<int64>
+	{
+	public:
+		static const FVertexInterface& GetVertexInterface()
+		{
+			static const FVertexInterface Interface(
+				FInputVertexInterface(
+					TInputDataVertexModel<int64>(MathOpNames::PrimaryOperandName, LOCTEXT("MathOpMinuendTooltip", "Minuend."), 0),
+					TInputDataVertexModel<int64>(MathOpNames::AdditionalOperandsName, LOCTEXT("MathOpSubtrahendsTooltip", "Subtrahend(s)."), 0)
+				),
+				FOutputVertexInterface(
+					TOutputDataVertexModel<int64>(TEXT("Out"), LOCTEXT("MathOpOutTooltip", "Subtraction result"))
+				)
+			);
+
+			return Interface;
+		}
+
+		static int32 GetDefault(const FOperatorSettings& InSettings, const FLiteral& InVertexDefault)
+		{
+			return static_cast<int64>(InVertexDefault.Value.Get<int32>());
+		}
+
+		static int32 GetDefaultOp(const FOperatorSettings& InSettings, const FLiteral& InVertexDefault)
+		{
+			return static_cast<int64>(InVertexDefault.Value.Get<int32>());
+		}
+
+		static void Calculate(TMathOpSubtract& InInstanceData, const FInt64ReadRef& InPrimaryOperand, const TArray<FInt64ReadRef>& InAdditionalOperands, FInt64WriteRef& OutResult)
+		{
+			*OutResult = *InPrimaryOperand;
+
+			if (!ensure(!InAdditionalOperands.IsEmpty()))
+			{
+				return;
+			}
+
+			for (int32 i = 0; i < InAdditionalOperands.Num(); ++i)
+			{
+				const int64 OperandValue = *InAdditionalOperands[i];
+				*OutResult -= OperandValue;
+			}
+		}
+	};
+
 	template <>
 	class TMathOpMultiply<FAudioBuffer>
 	{
 	public:
-		static FAudioBuffer GetDefault(const FOperatorSettings* InSettings)
+		static const FVertexInterface& GetVertexInterface()
 		{
-			return FAudioBuffer(InSettings ? InSettings->GetNumFramesPerBlock() : 0);
+			static const FVertexInterface Interface(
+				FInputVertexInterface(
+					TInputDataVertexModel<FAudioBuffer>(MathOpNames::PrimaryOperandName, LOCTEXT("MathOpAudioInitMultiplicandTooltip", "Initial audio to multiply.")),
+					TInputDataVertexModel<FAudioBuffer>(MathOpNames::AdditionalOperandsName, LOCTEXT("MathOpMultiplyAudioSubtrahendsTooltip", "Additional audio to multiply sample-by-sample."))
+				),
+				FOutputVertexInterface(
+					TOutputDataVertexModel<FAudioBuffer>(TEXT("Out"), LOCTEXT("MathOpMultiplyAudioOutTooltip", "Resulting buffer"))
+				)
+			);
+
+			return Interface;
 		}
 
-		static void Calculate(const TArray<FAudioBufferReadRef>& InOperands, FAudioBufferWriteRef& OutResult)
+		static FAudioBuffer GetDefault(const FOperatorSettings& InSettings, const FLiteral& InVertexDefault)
 		{
-			if (!ensure(!InOperands.IsEmpty()))
+			return FAudioBuffer(InSettings.GetNumFramesPerBlock());
+		}
+
+		static FAudioBuffer GetDefaultOp(const FOperatorSettings& InSettings, const FLiteral& InVertexDefault)
+		{
+			return GetDefault(InSettings, InVertexDefault);
+		}
+
+		static void Calculate(TMathOpMultiply& InInstanceData, const FAudioBufferReadRef& InPrimaryOperand, const TArray<FAudioBufferReadRef>& InAdditionalOperands, FAudioBufferWriteRef& OutResult)
+		{
+			if (!ensure(!InAdditionalOperands.IsEmpty()))
 			{
 				OutResult->Zero();
 				return;
 			}
 
-			const int32 NumSamples = InOperands[0]->Num();
+			const int32 NumSamples = InPrimaryOperand->Num();
 			if (!ensure(NumSamples == OutResult->Num()))
 			{
 				// TODO: Error
@@ -434,11 +872,11 @@ namespace Metasound
 				return;
 			}
 
-			FMemory::Memcpy(OutResult->GetData(), InOperands[0]->GetData(), sizeof(float) * NumSamples);
+			FMemory::Memcpy(OutResult->GetData(), InPrimaryOperand->GetData(), sizeof(float) * NumSamples);
 
-			for (int32 i = 1; i < InOperands.Num(); ++i)
+			for (int32 i = 0; i < InAdditionalOperands.Num(); ++i)
 			{
-				const FAudioBuffer& Operand = *InOperands[i];
+				const FAudioBuffer& Operand = *InAdditionalOperands[i];
 				if (!ensure(NumSamples == Operand.Num()))
 				{
 					// TODO: Error
@@ -446,7 +884,7 @@ namespace Metasound
 					return;
 				}
 
-				const float* OperandData = InOperands[i]->GetData();
+				const float* OperandData = InAdditionalOperands[i]->GetData();
 				float* OutData = OutResult->GetData();
 
 				if (NumSamples % AUDIO_SIMD_FLOAT_ALIGNMENT)
@@ -465,170 +903,433 @@ namespace Metasound
 	};
 
 	template <>
-	class TMathOpMultiply<FFloatTime>
+	class TMathOpMultiply<FAudioBuffer, float>
 	{
+		bool bInit = false;
+		float LastGain = 0.0f;
+
 	public:
-		static FFloatTime GetDefault(const FOperatorSettings* InSettings)
+		static const FVertexInterface& GetVertexInterface()
 		{
-			return FFloatTime(1.0f);
+			static const FVertexInterface Interface(
+				FInputVertexInterface(
+					TInputDataVertexModel<FAudioBuffer>(MathOpNames::PrimaryOperandName, LOCTEXT("MathOpAudioMultiplyFloatTooltip", "Audio multiplicand.")),
+					TInputDataVertexModel<float>(MathOpNames::AdditionalOperandsName, LOCTEXT("MathOpAddendAdditionalTooltip", "Float multiplicand to apply sample-by-sample to audio. Interpolates over buffer size on value change."), 1.0f)
+				),
+				FOutputVertexInterface(
+					TOutputDataVertexModel<FAudioBuffer>(TEXT("Out"), LOCTEXT("MathOpAudioFloatMultiplyOutTooltip", "Resulting buffer"))
+				)
+			);
+
+			return Interface;
 		}
 
-		static void Calculate(const TArray<FFloatTimeReadRef>& InOperands, FFloatTimeWriteRef& OutResult)
+		static FAudioBuffer GetDefault(const FOperatorSettings& InSettings, const FLiteral& InVertexDefault)
 		{
-			if (!ensure(!InOperands.IsEmpty()))
+			return FAudioBuffer(InSettings.GetNumFramesPerBlock());
+		}
+
+		static float GetDefaultOp(const FOperatorSettings& InSettings, const FLiteral& InVertexDefault)
+		{
+			return InVertexDefault.Value.Get<float>();
+		}
+
+		static void Calculate(TMathOpMultiply& InInstanceData, const FAudioBufferReadRef& InPrimaryOperand, const TArray<FFloatReadRef>& InAdditionalOperands, FAudioBufferWriteRef& OutResult)
+		{
+			if (!InInstanceData.bInit)
+			{
+				InInstanceData.bInit = true;
+				InInstanceData.LastGain = 1.0f;
+
+				for (int32 i = 0; i < InAdditionalOperands.Num(); ++i)
+				{
+					InInstanceData.LastGain *= *InAdditionalOperands[i];
+				}
+			}
+
+			FMemory::Memcpy(OutResult->GetData(), InPrimaryOperand->GetData(), sizeof(float) * InPrimaryOperand->Num());
+
+			float NewGain = 1.0f;
+			for (int32 i = 0; i < InAdditionalOperands.Num(); ++i)
+			{
+				const int32 SIMDRemainder = OutResult->Num() % AUDIO_SIMD_FLOAT_ALIGNMENT;
+				const int32 SIMDCount = OutResult->Num() - SIMDRemainder;
+
+				Audio::FadeBufferFast(OutResult->GetData(), SIMDCount, InInstanceData.LastGain, *InAdditionalOperands[i]);
+
+				for (int32 j = SIMDCount; j < OutResult->Num(); ++j)
+				{
+					OutResult->GetData()[j] *= *InAdditionalOperands[j];
+				}
+
+				NewGain *= *InAdditionalOperands[i];
+			}
+
+			InInstanceData.LastGain = NewGain;
+		}
+	};
+
+	template <>
+	class TMathOpMultiply<FTime, float>
+	{
+	public:
+		static const FVertexInterface& GetVertexInterface()
+		{
+			static const FVertexInterface Interface(
+				FInputVertexInterface(
+					TInputDataVertexModel<FTime>(MathOpNames::PrimaryOperandName, LOCTEXT("MathOpTimeMultiplyFloatTooltip", "Time multiplicand."), 1.0f),
+					TInputDataVertexModel<float>(MathOpNames::AdditionalOperandsName, LOCTEXT("MathOpAddendAdditionalTooltip", "Float multiplicand(s)."), 1.0f)
+				),
+				FOutputVertexInterface(
+					TOutputDataVertexModel<FTime>(TEXT("Out"), LOCTEXT("MathOpTimeMultiplyOutTooltip", "Resulting buffer"))
+				)
+			);
+
+			return Interface;
+		}
+
+		static FTime GetDefault(const FOperatorSettings& InSettings, const FLiteral& InVertexDefault)
+		{
+			return FTime(InVertexDefault.Value.Get<float>());
+		}
+
+		static float GetDefaultOp(const FOperatorSettings& InSettings, const FLiteral& InVertexDefault)
+		{
+			return 1.0f;
+		}
+
+		static void Calculate(TMathOpMultiply& InInstanceData, const FTimeReadRef& InPrimaryOperand, const TArray<FFloatReadRef>& InAdditionalOperands, FTimeWriteRef& OutResult)
+		{
+			*OutResult = *InPrimaryOperand;
+
+			if (!ensure(!InAdditionalOperands.IsEmpty()))
 			{
 				return;
 			}
 
-			*OutResult = *InOperands[0];
-			for (int32 i = 1; i < InOperands.Num(); ++i)
+			for (int32 i = 0; i < InAdditionalOperands.Num(); ++i)
 			{
-				OutResult->SetSeconds(OutResult->GetSeconds() * InOperands[i]->GetSeconds());
+				const float OperandValue = *InAdditionalOperands[i];
+				*OutResult *= OperandValue;
 			}
 		}
 	};
 
+	// int64 has to be explicitly implemented as they are PODs that are not implemented as a literal base type.
 	template <>
-	class TMathOpDivide<FFloatTime>
+	class TMathOpMultiply<int64>
 	{
 	public:
-		static FFloatTime GetDefault(const FOperatorSettings* InSettings)
+		static const FVertexInterface& GetVertexInterface()
 		{
-			return FFloatTime(1.0f);
+			static const FVertexInterface Interface(
+				FInputVertexInterface(
+					TInputDataVertexModel<int64>(MathOpNames::PrimaryOperandName, LOCTEXT("MathOpInitMultiplicandTooltip", "Initial multiplicand."), 1),
+					TInputDataVertexModel<int64>(MathOpNames::AdditionalOperandsName, LOCTEXT("MathOpMultiplicandsTooltip", "Additional multiplicand(s)."), 1)
+				),
+				FOutputVertexInterface(
+					TOutputDataVertexModel<int64>(TEXT("Out"), LOCTEXT("MultiplicationResultTooltip", "Multiplication result"))
+				)
+			);
+
+			return Interface;
 		}
 
-		static void Calculate(const TArray<FFloatTimeReadRef>& InOperands, FFloatTimeWriteRef& OutResult)
+		static int32 GetDefault(const FOperatorSettings& InSettings, const FLiteral& InVertexDefault)
 		{
-			if (!ensure(!InOperands.IsEmpty()))
+			return static_cast<int64>(InVertexDefault.Value.Get<int32>());
+		}
+
+		static int32 GetDefaultOp(const FOperatorSettings& InSettings, const FLiteral& InVertexDefault)
+		{
+			return static_cast<int64>(InVertexDefault.Value.Get<int32>());
+		}
+
+		static void Calculate(TMathOpMultiply& InInstanceData, const FInt64ReadRef& InPrimaryOperand, const TArray<FInt64ReadRef>& InAdditionalOperands, FInt64WriteRef& OutResult)
+		{
+			*OutResult = *InPrimaryOperand;
+
+			if (!ensure(!InAdditionalOperands.IsEmpty()))
 			{
 				return;
 			}
 
-			*OutResult = *InOperands[0];
-			for (int32 i = 1; i < InOperands.Num(); ++i)
+			for (int32 i = 0; i < InAdditionalOperands.Num(); ++i)
 			{
-				const FFloatTime& OperandValue = *InOperands[i];
-				if (OperandValue.GetSeconds() == 0.0f)
+				const int64 OperandValue = *InAdditionalOperands[i];
+				*OutResult *= OperandValue;
+			}
+		}
+	};
+
+	template <>
+	class TMathOpDivide<FTime, float>
+	{
+	public:
+		static const FVertexInterface& GetVertexInterface()
+		{
+			static const FVertexInterface Interface(
+				FInputVertexInterface(
+					TInputDataVertexModel<FTime>(MathOpNames::PrimaryOperandName, LOCTEXT("MathOpTimeMultiplyFloatTooltip", "Time multiplicand."), 1.0f),
+					TInputDataVertexModel<float>(MathOpNames::AdditionalOperandsName, LOCTEXT("MathOpAddendAdditionalTooltip", "Float multiplicand(s)."), 1.0f)
+				),
+				FOutputVertexInterface(
+					TOutputDataVertexModel<FTime>(TEXT("Out"), LOCTEXT("MathOpTimeMultiplyOutTooltip", "Resulting buffer"))
+				)
+			);
+
+			return Interface;
+		}
+
+		static FTime GetDefault(const FOperatorSettings& InSettings, const FLiteral& InVertexDefault)
+		{
+			return FTime(InVertexDefault.Value.Get<float>());
+		}
+
+		static float GetDefaultOp(const FOperatorSettings& InSettings, const FLiteral& InVertexDefault)
+		{
+			return 1.0f;
+		}
+
+		static void Calculate(TMathOpDivide& InInstanceData, const FTimeReadRef& InPrimaryOperand, const TArray<FFloatReadRef>& InAdditionalOperands, FTimeWriteRef& OutResult)
+		{
+			*OutResult = *InPrimaryOperand;
+
+			if (!ensure(!InAdditionalOperands.IsEmpty()))
+			{
+				return;
+			}
+
+			for (int32 i = 0; i < InAdditionalOperands.Num(); ++i)
+			{
+				const float OperandValue = *InAdditionalOperands[i];
+				if (OperandValue == 0.0f)
 				{
 					// TODO: Error here
-					return;
+					continue;
+				}
+
+				*OutResult /= OperandValue;
+			}
+		}
+	};
+
+	// int64 has to be explicitly implemented as they are PODs that are not implemented as a literal base type.
+	template <>
+	class TMathOpDivide<int64>
+	{
+	public:
+		static const FVertexInterface& GetVertexInterface()
+		{
+			static const FVertexInterface Interface(
+				FInputVertexInterface(
+					TInputDataVertexModel<int64>(MathOpNames::PrimaryOperandName, LOCTEXT("MathOpRandomMinTooltip", "Random value minimum."), 1),
+					TInputDataVertexModel<int64>(MathOpNames::AdditionalOperandsName, LOCTEXT("MathOpRandomMaxTooltip", "Random value maximum."), 1)
+				),
+				FOutputVertexInterface(
+					TOutputDataVertexModel<int64>(TEXT("Out"), LOCTEXT("MathOpRandomOutTooltip", "Resulting random value"))
+				)
+			);
+
+			return Interface;
+		}
+
+		static int64 GetDefault(const FOperatorSettings& InSettings, const FLiteral& InVertexDefault)
+		{
+			return static_cast<int64>(InVertexDefault.Value.Get<int32>());
+		}
+
+		static int64 GetDefaultOp(const FOperatorSettings& InSettings, const FLiteral& InVertexDefault)
+		{
+			return static_cast<int64>(InVertexDefault.Value.Get<int32>());
+		}
+
+		static void Calculate(TMathOpDivide& InInstanceData, const FInt64ReadRef& InPrimaryOperand, const TArray<FInt64ReadRef>& InAdditionalOperands, FInt64WriteRef& OutResult)
+		{
+			*OutResult = *InPrimaryOperand;
+
+			if (!ensure(!InAdditionalOperands.IsEmpty()))
+			{
+				return;
+			}
+
+			for (int32 i = 0; i < InAdditionalOperands.Num(); ++i)
+			{
+				const int64 OperandValue = *InAdditionalOperands[i];
+				if (OperandValue == 0)
+				{
+					// TODO: Error here
+					continue;
 				}
 
-				OutResult->SetSeconds(OutResult->GetSeconds() / OperandValue.GetSeconds());
+				*OutResult /= OperandValue;
 			}
 		}
 	};
 
-	template <typename TDataClass>
+	template <typename TDataClass, typename TOperandDataClass = TDataClass>
 	class TMathOpRandRange
 	{
 	public:
 		using TDataClassReadRef = TDataReadReference<TDataClass>;
+		using TOperandDataClassReadRef = TDataReadReference<TOperandDataClass>;
 		using TDataClassWriteRef = TDataWriteReference<TDataClass>;
 
-		static TDataClass GetDefault(const FOperatorSettings* InSettings)
+		static const FVertexInterface& GetVertexInterface()
+		{
+			static const FVertexInterface Interface(
+				FInputVertexInterface(
+					TInputDataVertexModel<TDataClass>(MathOpNames::PrimaryOperandName, LOCTEXT("MathOpRandomMinTooltip", "Random value minimum."), static_cast<TDataClass>(0.0f)),
+					TInputDataVertexModel<TOperandDataClass>(MathOpNames::AdditionalOperandsName, LOCTEXT("MathOpRandomMaxTooltip", "Random value maximum."), static_cast<TOperandDataClass>(1.0f))
+				),
+				FOutputVertexInterface(
+					TOutputDataVertexModel<TDataClass>(TEXT("Out"), LOCTEXT("MathOpRandomOutTooltip", "Resulting random value"))
+				)
+			);
+
+			return Interface;
+		}
+
+		static TDataClass GetDefault(const FOperatorSettings& InSettings, const FLiteral& InVertexDefault)
+		{
+			return InVertexDefault.Value.Get<TDataClass>();
+		}
+
+		static TOperandDataClass GetDefaultOp(const FOperatorSettings& InSettings, const FLiteral& InVertexDefault)
 		{
-			return 0;
+			return InVertexDefault.Value.Get<TOperandDataClass>();
 		}
 
-		static void Calculate(const TArray<TDataClassReadRef>& InOperands, TDataClassWriteRef& OutResult)
+		static void Calculate(TMathOpRandRange& InInstanceData, const TDataClassReadRef& InPrimaryOperand, const TArray<TOperandDataClassReadRef>& InAdditionalOperands, TDataClassWriteRef& OutResult)
 		{
-			if (InOperands.Num() > 1)
+			if (ensure(!InAdditionalOperands.IsEmpty()))
 			{
-				const TDataClass& Min = *InOperands[0];
-				const TDataClass& Max = *InOperands[1];
+				const TDataClass& Min = *InPrimaryOperand;
+				const TOperandDataClass& Max = *InAdditionalOperands[0];
 				*OutResult = FMath::RandRange(Min, Max);
 			}
-			else
-			{
-				// TODO: Error
-			}
 		}
 	};
 
 	template <>
-	class TMathOpRandRange<FFloatTime>
+	class TMathOpRandRange<FTime>
 	{
 	public:
-		static FFloatTime GetDefault(const FOperatorSettings* InSettings)
+		static const FVertexInterface& GetVertexInterface()
 		{
-			return FFloatTime(0.0f);
+			static const FVertexInterface Interface(
+				FInputVertexInterface(
+					TInputDataVertexModel<FTime>(MathOpNames::PrimaryOperandName, LOCTEXT("MathOpRandomMinTooltip", "Random value minimum."), 0.0f),
+					TInputDataVertexModel<FTime>(MathOpNames::AdditionalOperandsName, LOCTEXT("MathOpRandomMaxTooltip", "Random value maximum."), 1.0f)
+				),
+				FOutputVertexInterface(
+					TOutputDataVertexModel<FTime>(TEXT("Out"), LOCTEXT("MathOpRandomOutTooltip", "Resulting random value"))
+				)
+			);
+
+			return Interface;
 		}
 
-		static void Calculate(const TArray<FFloatTimeReadRef>& InOperands, FFloatTimeWriteRef& OutResult)
+		static FTime GetDefault(const FOperatorSettings& InSettings, const FLiteral& InVertexDefault)
 		{
-			if (InOperands.Num() > 1)
-			{
-				const FFloatTime& Min = *InOperands[0];
-				const FFloatTime& Max = *InOperands[1];
-				const float FloatResult = FMath::RandRange(Min.GetSeconds(), Max.GetSeconds());
-				OutResult->SetSeconds(FloatResult);
-			}
-			else
+			return FTime(InVertexDefault.Value.Get<float>());
+		}
+
+		static FTime GetDefaultOp(const FOperatorSettings& InSettings, const FLiteral& InVertexDefault)
+		{
+			return FTime(InVertexDefault.Value.Get<float>());
+		}
+
+		static void Calculate(TMathOpRandRange& InInstanceData, const FTimeReadRef& InPrimaryOperand, const TArray<FTimeReadRef>& InAdditionalOperands, FTimeWriteRef& OutResult)
+		{
+			if (ensure(!InAdditionalOperands.IsEmpty()))
 			{
-				// TODO: Error
+				const FTime& Min = *InPrimaryOperand;
+				const FTime& Max = *InAdditionalOperands[0];
+
+				// Not a true random as precision is lost with randomizing over float range
+				const FTime Delta = Max.GetSeconds() - Min.GetSeconds();
+				const FTime Result = Min.GetSeconds() + (Delta * FMath::FRand());
+				*OutResult = Result;
 			}
 		}
 	};
 
+	// int64 has to be explicitly implemented as they are PODs that are not implemented as a literal base type.
 	template <>
-	class TMathOpRandRange<FFrequency>
+	class TMathOpRandRange<int64>
 	{
 	public:
-		static FFrequency GetDefault(const FOperatorSettings* InSettings)
+		static const FVertexInterface& GetVertexInterface()
 		{
-			return FFrequency(0.0f);
+			static const FVertexInterface Interface(
+				FInputVertexInterface(
+					TInputDataVertexModel<int64>(MathOpNames::PrimaryOperandName, LOCTEXT("MathOpRandomMinTooltip", "Random value minimum."), 0),
+					TInputDataVertexModel<int64>(MathOpNames::AdditionalOperandsName, LOCTEXT("MathOpRandomMaxTooltip", "Random value maximum."), 1)
+				),
+				FOutputVertexInterface(
+					TOutputDataVertexModel<int64>(TEXT("Out"), LOCTEXT("MathOpRandomOutTooltip", "Resulting random value"))
+				)
+			);
+
+			return Interface;
 		}
 
-		static void Calculate(const TArray<FFrequencyReadRef>& InOperands, FFrequencyWriteRef& OutResult)
+		static int64 GetDefault(const FOperatorSettings& InSettings, const FLiteral& InVertexDefault)
 		{
-			if (InOperands.Num() > 1)
-			{
-				const FFrequency& Min = *InOperands[0];
-				const FFrequency& Max = *InOperands[1];
-				const float FloatResult = FMath::RandRange(Min.GetHertz(), Max.GetHertz());
-				OutResult->SetHertz(FloatResult);
-			}
-			else
+			return static_cast<int64>(InVertexDefault.Value.Get<int32>());
+		}
+
+		static int64 GetDefaultOp(const FOperatorSettings& InSettings, const FLiteral& InVertexDefault)
+		{
+			return static_cast<int64>(InVertexDefault.Value.Get<int32>());
+		}
+
+		static void Calculate(TMathOpRandRange& InInstanceData, const FInt64ReadRef& InPrimaryOperand, const TArray<FInt64ReadRef>& InAdditionalOperands, FInt64WriteRef& OutResult)
+		{
+			if (ensure(!InAdditionalOperands.IsEmpty()))
 			{
-				// TODO: Error
+				const int64 Min = *InPrimaryOperand;
+				const int64 Max = *InAdditionalOperands[0];
+				*OutResult = FMath::RandRange(Min, Max);
 			}
 		}
 	};
 
 	// Definitions
-	DEFINE_METASOUND_MATHOP(Add, Float, float, MathOpDisplayNames::GetAddDisplayName<float>(), LOCTEXT("Metasound_MathAddFloatNodeDescription", "Adds floats."))
-	DEFINE_METASOUND_MATHOP(Add, Int32, int32, MathOpDisplayNames::GetAddDisplayName<int32>(), LOCTEXT("Metasound_MathAddInt32NodeDescription", "Adds int32s."))
-	DEFINE_METASOUND_MATHOP(Add, Int64, int64, MathOpDisplayNames::GetAddDisplayName<int64>(), LOCTEXT("Metasound_MathAddInt64NodeDescription", "Adds int64s."))
-	DEFINE_METASOUND_MATHOP(Add, Frequency, FFrequency, MathOpDisplayNames::GetAddDisplayName<FFrequency>(), LOCTEXT("Metasound_MathAddFrequencyNodeDescription", "Adds frequency values."))
-	DEFINE_METASOUND_MATHOP(Add, FloatTime, FFloatTime, MathOpDisplayNames::GetAddDisplayName<FFloatTime>(), LOCTEXT("Metasound_MathAddFloatTimeNodeDescription", "Adds time values."))
-	DEFINE_METASOUND_MATHOP(Add, AudioBuffer, FAudioBuffer, MathOpDisplayNames::GetAddDisplayName<FAudioBuffer>(), LOCTEXT("Metasound_MathAddBufferNodeDescription", "Adds buffers together by sample."))
-
-	DEFINE_METASOUND_MATHOP(Subtract, Float, float, MathOpDisplayNames::GetSubtractDisplayName<float>(), LOCTEXT("Metasound_MathSubractFloatNodeDescription", "Subtracts floats."))
-	DEFINE_METASOUND_MATHOP(Subtract, Int32, int32, MathOpDisplayNames::GetSubtractDisplayName<int32>(), LOCTEXT("Metasound_MathSubractInt32NodeDescription", "Subtracts int32s."))
-	DEFINE_METASOUND_MATHOP(Subtract, Int64, int64, MathOpDisplayNames::GetSubtractDisplayName<int64>(), LOCTEXT("Metasound_MathSubractInt64NodeDescription", "Subtracts int64s."))
-	DEFINE_METASOUND_MATHOP(Subtract, Frequency, FFrequency, MathOpDisplayNames::GetSubtractDisplayName<FFrequency>(), LOCTEXT("Metasound_MathSubractFrequencyNodeDescription", "Subtracts frequency values."))
-	DEFINE_METASOUND_MATHOP(Subtract, FloatTime, FFloatTime, MathOpDisplayNames::GetSubtractDisplayName<FFloatTime>(), LOCTEXT("Metasound_MathSubractFloatTimeNodeDescription", "Subtracts time values."))
-	DEFINE_METASOUND_MATHOP(Subtract, AudioBuffer, FAudioBuffer, MathOpDisplayNames::GetSubtractDisplayName<FAudioBuffer>(), LOCTEXT("Metasound_MathSubtractBufferNodeDescription", "Subtracts buffers by sample."))
-
-	DEFINE_METASOUND_MATHOP(Multiply, Float, float, MathOpDisplayNames::GetMultiplyDisplayName<float>(), LOCTEXT("Metasound_MathMultiplyFloatNodeDescription", "Multiplies floats."))
-	DEFINE_METASOUND_MATHOP(Multiply, Int32, int32, MathOpDisplayNames::GetMultiplyDisplayName<int32>(), LOCTEXT("Metasound_MathMultiplyInt32NodeDescription", "Multiplies int32s."))
-	DEFINE_METASOUND_MATHOP(Multiply, Int64, int64, MathOpDisplayNames::GetMultiplyDisplayName<int64>(), LOCTEXT("Metasound_MathMultiplyInt64NodeDescription", "Multiplies int64s."))
-	DEFINE_METASOUND_MATHOP(Multiply, Frequency, FFrequency, MathOpDisplayNames::GetMultiplyDisplayName<FFrequency>(), LOCTEXT("Metasound_MathMultiplyFrequencyNodeDescription", "Multiplies frequency values."))
-	DEFINE_METASOUND_MATHOP(Multiply, FloatTime, FFloatTime, MathOpDisplayNames::GetMultiplyDisplayName<FFloatTime>(), LOCTEXT("Metasound_MathMultiplyFloatTimeNodeDescription", "Multiplies time values."))
-	DEFINE_METASOUND_MATHOP(Multiply, AudioBuffer, FAudioBuffer, MathOpDisplayNames::GetMultiplyDisplayName<FAudioBuffer>(), LOCTEXT("Metasound_MathMultiplyBufferNodeDescription", "Multiplies buffers together by sample."))
-
-	DEFINE_METASOUND_MATHOP(Divide, Float, float, MathOpDisplayNames::GetDivideDisplayName<float>(), LOCTEXT("Metasound_MathDivideFloatNodeDescription", "Divide float by another float."))
-	DEFINE_METASOUND_MATHOP(Divide, Int32, int32, MathOpDisplayNames::GetDivideDisplayName<int32>(), LOCTEXT("Metasound_MathDivideInt32NodeDescription", "Divide int32 by another int32."))
-	DEFINE_METASOUND_MATHOP(Divide, Int64, int64, MathOpDisplayNames::GetDivideDisplayName<int64>(), LOCTEXT("Metasound_MathDivideInt64NodeDescription", "Divide int64 by another int64."))
-	DEFINE_METASOUND_MATHOP(Divide, Frequency, FFrequency, MathOpDisplayNames::GetDivideDisplayName<FFrequency>(), LOCTEXT("Metasound_MathDivideFrequencyNodeDescription", "Divide frequency by another frequency."))
-	DEFINE_METASOUND_MATHOP(Divide, FloatTime, FFloatTime, MathOpDisplayNames::GetDivideDisplayName<FFloatTime>(), LOCTEXT("Metasound_MathDivideFloatTimeNodeDescription", "Divide time by another time."))
-
-	DEFINE_METASOUND_MATHOP(RandRange, Float, float, MathOpDisplayNames::GetRandRangeDisplayName<float>(), LOCTEXT("Metasound_MathRandRangeFloatNodeDescription", "Computes random float value in the range [Min, Max]."))
-	DEFINE_METASOUND_MATHOP(RandRange, Int32, int32, MathOpDisplayNames::GetRandRangeDisplayName<int32>(), LOCTEXT("Metasound_MathRandRangeInt32NodeDescription", "Computes random int32 in the range [Min, Max]."))
-	DEFINE_METASOUND_MATHOP(RandRange, Int64, int64, MathOpDisplayNames::GetRandRangeDisplayName<int64>(), LOCTEXT("Metasound_MathRandRangeInt64NodeDescription", "Computes random int64 in the range [Min, Max]."))
-	DEFINE_METASOUND_MATHOP(RandRange, Frequency, FFrequency, MathOpDisplayNames::GetRandRangeDisplayName<FFrequency>(), LOCTEXT("Metasound_MathRandRangeFrequencyNodeDescription", "Computes random frequency value in the range [Min, Max]."))
-	DEFINE_METASOUND_MATHOP(RandRange, FloatTime, FFloatTime, MathOpDisplayNames::GetRandRangeDisplayName<FFloatTime>(), LOCTEXT("Metasound_MathRandRangeFloatTimeNodeDescription", "Computes random time value in the range [Min, Max]."))
+	DEFINE_METASOUND_MATHOP(Add, Float, float, LOCTEXT("Metasound_MathAddFloatNodeDescription", "Adds floats."))
+	DEFINE_METASOUND_MATHOP(Add, Int32, int32, LOCTEXT("Metasound_MathAddInt32NodeDescription", "Adds int32s."))
+	DEFINE_METASOUND_MATHOP(Add, Int64, int64, LOCTEXT("Metasound_MathAddInt64NodeDescription", "Adds int64s."))
+	DEFINE_METASOUND_MATHOP(Add, Audio, FAudioBuffer, LOCTEXT("Metasound_MathAddBufferNodeDescription", "Adds buffers together by sample."))
+	DEFINE_METASOUND_MATHOP(Add, Time, FTime, LOCTEXT("Metasound_MathAddTimeNodeDescription", "Adds time values."))
+
+	DEFINE_METASOUND_OPERAND_TYPED_MATHOP(Add, Audio, FAudioBuffer, Float, float, LOCTEXT("Metasound_MathAddAudioFloatNodeDescription", "Add floats to buffer sample-by-sample."))
+
+	DEFINE_METASOUND_MATHOP(Subtract, Float, float, LOCTEXT("Metasound_MathSubractFloatNodeDescription", "Subtracts floats."))
+	DEFINE_METASOUND_MATHOP(Subtract, Int32, int32, LOCTEXT("Metasound_MathSubractInt32NodeDescription", "Subtracts int32s."))
+ 	DEFINE_METASOUND_MATHOP(Subtract, Int64, int64, LOCTEXT("Metasound_MathSubractInt64NodeDescription", "Subtracts int64s."))
+	DEFINE_METASOUND_MATHOP(Subtract, Audio, FAudioBuffer, LOCTEXT("Metasound_MathSubtractBufferNodeDescription", "Subtracts buffers sample-by-sample."))
+	DEFINE_METASOUND_MATHOP(Subtract, Time, FTime, LOCTEXT("Metasound_MathSubractTimeNodeDescription", "Subtracts time values."))
+
+	DEFINE_METASOUND_MATHOP(Multiply, Float, float, LOCTEXT("Metasound_MathMultiplyFloatNodeDescription", "Multiplies floats."))
+	DEFINE_METASOUND_MATHOP(Multiply, Int32, int32, LOCTEXT("Metasound_MathMultiplyInt32NodeDescription", "Multiplies int32s."))
+ 	DEFINE_METASOUND_MATHOP(Multiply, Int64, int64, LOCTEXT("Metasound_MathMultiplyInt64NodeDescription", "Multiplies int64s."))
+	DEFINE_METASOUND_MATHOP(Multiply, Audio, FAudioBuffer, LOCTEXT("Metasound_MathMultiplyBufferNodeDescription", "Multiplies buffers together sample-by-sample."))
+
+	DEFINE_METASOUND_OPERAND_TYPED_MATHOP(Multiply, Audio, FAudioBuffer, Float, float, LOCTEXT("Metasound_MathMultiplyAudioByFloatDescription", "Multiplies buffer by float scalars."))
+	DEFINE_METASOUND_OPERAND_TYPED_MATHOP(Multiply, Time, FTime, Float, float, LOCTEXT("Metasound_MathMultiplyTimeNodeDescription", "Scales time by floats."))
+
+	DEFINE_METASOUND_MATHOP(Divide, Float, float, LOCTEXT("Metasound_MathDivideFloatNodeDescription", "Divide float by another float."))
+	DEFINE_METASOUND_MATHOP(Divide, Int32, int32, LOCTEXT("Metasound_MathDivideInt32NodeDescription", "Divide int32 by another int32."))
+ 	DEFINE_METASOUND_MATHOP(Divide, Int64, int64, LOCTEXT("Metasound_MathDivideInt64NodeDescription", "Divide int64 by another int64."))
+
+	DEFINE_METASOUND_OPERAND_TYPED_MATHOP(Divide, Time, FTime, Float, float, LOCTEXT("Metasound_MathDivideTimeNodeDescription", "Divides time by floats."))
+
+	DEFINE_METASOUND_MATHOP(RandRange, Float, float, LOCTEXT("Metasound_MathRandRangeFloatNodeDescription", "Computes random float value in the range [Min, Max]."))
+	DEFINE_METASOUND_MATHOP(RandRange, Int32, int32, LOCTEXT("Metasound_MathRandRangeInt32NodeDescription", "Computes random int32 in the range [Min, Max]."))
+ 	DEFINE_METASOUND_MATHOP(RandRange, Int64, int64, LOCTEXT("Metasound_MathRandRangeInt64NodeDescription", "Computes random int64 in the range [Min, Max]."))
+	DEFINE_METASOUND_MATHOP(RandRange, Time, FTime, LOCTEXT("Metasound_MathRandRangeTimeNodeDescription", "Computes random time value in the range [Min, Max]."))
+
 }
 
 #undef LOCTEXT_NAMESPACE // MetasoundMathOpNode
diff --git a/Engine/Plugins/Runtime/Metasound/Source/MetasoundStandardNodes/Public/MetasoundMixerNode.h b/Engine/Plugins/Runtime/Metasound/Source/MetasoundStandardNodes/Public/MetasoundMixerNode.h
--- a/Engine/Plugins/Runtime/Metasound/Source/MetasoundStandardNodes/Public/MetasoundMixerNode.h
+++ b/Engine/Plugins/Runtime/Metasound/Source/MetasoundStandardNodes/Public/MetasoundMixerNode.h
@@ -17,11 +17,6 @@ namespace Metasound
 	class METASOUNDSTANDARDNODES_API FMixerNode : public FNodeFacade
 	{
 	public:
-		FMixerNode(const FString& InName, const FGuid& InInstanceID, float InDefaultMixGainCents);
 		FMixerNode(const FNodeInitData& InInitData);
-
-		float GetDefaultMixGainLinear() const;
-	private:
-		float DefaultMixGainLinear = 1.0f;
 	};
 }
diff --git a/Engine/Plugins/Runtime/Metasound/Source/MetasoundStandardNodes/Public/MetasoundOscNode.h b/Engine/Plugins/Runtime/Metasound/Source/MetasoundStandardNodes/Public/MetasoundOscNode.h
--- a/Engine/Plugins/Runtime/Metasound/Source/MetasoundStandardNodes/Public/MetasoundOscNode.h
+++ b/Engine/Plugins/Runtime/Metasound/Source/MetasoundStandardNodes/Public/MetasoundOscNode.h
@@ -9,7 +9,6 @@
 #include "MetasoundDataReferenceCollection.h"
 #include "MetasoundExecutableOperator.h"
 #include "MetasoundFacade.h"
-#include "MetasoundFrequency.h"
 #include "MetasoundNode.h"
 #include "MetasoundOperatorInterface.h"
 
@@ -17,19 +16,7 @@ namespace Metasound
 {
 	class METASOUNDSTANDARDNODES_API FOscNode : public FNodeFacade
 	{
-		public:
-
-		FOscNode(const FString& InInstanceName, const FGuid& InInstanceID, float InDefaultFrequency, bool bInDefaultEnablement);
-
+	public:
 		FOscNode(const FNodeInitData& InInitData);
-
-		
-		float GetDefaultFrequency() const;
-		bool GetDefaultEnablement() const;
-
-		private:
-
-		float DefaultFrequency = 440.f;		
-		bool bDefaultEnablement = true;
 	};
 }
diff --git a/Engine/Plugins/Runtime/Metasound/Source/MetasoundStandardNodes/Public/MetasoundSampleCounter.h b/Engine/Plugins/Runtime/Metasound/Source/MetasoundStandardNodes/Public/MetasoundSampleCounter.h
new file mode 100644
--- /dev/null
+++ b/Engine/Plugins/Runtime/Metasound/Source/MetasoundStandardNodes/Public/MetasoundSampleCounter.h
@@ -0,0 +1,322 @@
+// Copyright Epic Games, Inc. All Rights Reserved.
+#pragma once
+
+#include "MetasoundDataReference.h"
+#include "MetasoundDataTypeRegistrationMacro.h"
+#include "MetasoundOperatorSettings.h"
+#include "MetasoundTime.h"
+
+
+namespace Metasound
+{
+	using FSampleCount = int64;
+	static_assert(TIsIntegral<FSampleCount>::Value, "FSample must be integral type.");
+
+	class FSampleCounter
+	{
+		private:
+			FSampleCount NumSamples = 0;
+
+			// Must be non-zero
+			FSampleRate SampleRate = 1;
+
+		public:
+			FSampleCounter()
+			{
+			}
+
+			FSampleCounter(const FSampleCounter& InSampleCounter)
+				: NumSamples(InSampleCounter.NumSamples)
+				, SampleRate(InSampleCounter.SampleRate)
+			{
+			}
+
+			FSampleCounter(FSampleCount InNumSamples, FSampleRate InSampleRate)
+				: NumSamples(InNumSamples)
+				, SampleRate(InSampleRate)
+			{
+			}
+
+			/** Construct a sample counter with a sample count,
+			 * time resolution, and sample rate.
+			 *
+			 * @param InTime - The initial time value.
+			 * @param InSampleRate - The sample rate associated with this time 
+			 *						 object.
+			 */
+			static FSampleCounter FromTime(const FTime& InTime, FSampleRate InSampleRate)
+			{
+				return { FSampleCount(InTime.GetSeconds() * InSampleRate), InSampleRate };
+			}
+
+			/** Set the sample rate of this object without changing the number
+			 * of samples. This will result in the object representing a
+			 * different amount of time.
+			 */
+			void SetSampleRate(FSampleRate InSampleRate)
+			{
+				if (!ensure(InSampleRate > 0))
+				{
+					InSampleRate = 1;
+				}
+
+				SampleRate = InSampleRate;
+			}
+
+			/* Returns sample rate */
+			FSampleRate GetSampleRate() const
+			{
+				return SampleRate;
+			}
+
+			FTime ToTime() const
+			{
+				return FTime(static_cast<FTime::TimeType>(NumSamples) / SampleRate);
+			}
+
+			/** Return the number of samples. */
+			FSampleCount GetNumSamples() const
+			{ 
+				return NumSamples;
+			}
+
+			/** Return the number of samples which represent this time duration
+			 * using a different sample rate. 
+			 *
+			 * @param InOtherSampleRate - Sample rate to use when calculating
+			 * 							  the number of samples.
+			 */
+			FSampleCount GetNumSamples(FSampleRate InOtherSampleRate) const
+			{
+				if (InOtherSampleRate == SampleRate)
+				{
+					return NumSamples;
+				}
+
+				return NumSamples * InOtherSampleRate / SampleRate;
+			}
+
+			/** Set the number of samples from time value (This is functionally
+			  * discretization and will result in loss of precision)
+			  */
+			void SetNumSamples(const FTime& InTime)
+			{ 
+				NumSamples = InTime.GetSeconds() * SampleRate;
+			}
+
+			/** Set the number of samples. */
+			void SetNumSamples(FSampleCount InNumSamples)
+			{
+				NumSamples = InNumSamples;
+			}
+
+			FSampleCounter& operator=(const FSampleCounter& InRHS)
+			{
+				SampleRate = InRHS.SampleRate;
+				NumSamples = InRHS.NumSamples;
+				return *this;
+			}
+
+			friend FSampleCounter& operator+=(FSampleCounter& InLHS, const FSampleCounter& InRHS)
+			{
+				InLHS.NumSamples += InRHS.GetNumSamples(InLHS.SampleRate);
+				return InLHS;
+			}
+
+			friend FSampleCounter& operator+=(FSampleCounter& InLHS, const FTime& InRHS)
+			{
+				InLHS.NumSamples += InRHS * InLHS.SampleRate;
+				return InLHS;
+			}
+
+			friend FSampleCounter& operator-=(FSampleCounter& InLHS, const FSampleCounter& InRHS)
+			{
+				InLHS.NumSamples -= InRHS.GetNumSamples(InLHS.SampleRate);
+				return InLHS;
+			}
+
+			friend FSampleCounter& operator-=(FSampleCounter& InLHS, const FTime& InRHS)
+			{
+				InLHS.NumSamples -= InRHS * InLHS.SampleRate;
+				return InLHS;
+			}
+
+			friend FSampleCounter operator+(const FSampleCounter& InLHS, const FSampleCounter& InRHS)
+			{
+				return { InLHS.NumSamples + InRHS.NumSamples, InLHS.SampleRate };
+			}
+
+			friend FSampleCounter operator+(const FSampleCounter& InLHS, const FTime& InRHS)
+			{
+				return { FSampleCount(InLHS.NumSamples + (InRHS * InLHS.SampleRate)), InLHS.SampleRate };
+			}
+
+			friend FSampleCounter operator-(const FSampleCounter& InLHS, const FSampleCounter& InRHS)
+			{
+				return { FSampleCount(InLHS.NumSamples - InRHS.NumSamples), InLHS.SampleRate };
+			}
+
+			friend FSampleCounter operator-(const FSampleCounter& InLHS, const FTime& InRHS)
+			{
+				return { FSampleCount(InLHS.NumSamples - (InRHS * InLHS.SampleRate)), InLHS.SampleRate };
+			}
+
+			template<typename ArithmeticType>
+			friend FSampleCounter operator*(const FSampleCounter& InLHS, const ArithmeticType& InRHS)
+			{
+				static_assert(TIsArithmetic<ArithmeticType>::Value, "Multiplicand type must be arithmetic.");
+				return FSampleCounter(InLHS.NumSamples * InRHS, InLHS.SampleRate);
+			}
+
+			template<typename ArithmeticType>
+			friend FSampleCounter operator*(const ArithmeticType& InLHS, const FSampleCounter& InRHS)
+			{
+				static_assert(TIsArithmetic<ArithmeticType>::Value, "Multiplicand type must be arithmetic.");
+				return FSampleCounter(InLHS * InLHS.NumSamples, InLHS.SampleRate);
+			}
+
+			template<typename ArithmeticType>
+			friend FSampleCounter operator/(const FSampleCounter& InLHS, ArithmeticType InRHS)
+			{
+				static_assert(TIsArithmetic<ArithmeticType>::Value, "Divisor type must be arithmetic.");
+				return FSampleCounter(InLHS.NumSamples / InRHS, InLHS.SampleRate);
+			}
+
+			template<typename ArithmeticType>
+			friend ArithmeticType& operator/(ArithmeticType InLHS, const FSampleCounter& InRHS)
+			{
+				static_assert(TIsArithmetic<ArithmeticType>::Value, "Divisor type must be arithmetic.");
+				return InLHS / InRHS.NumSamples;
+			}
+
+			template<typename ArithmeticType>
+			friend FSampleCounter& operator*=(FSampleCounter& InLHS, ArithmeticType InRHS)
+			{
+				static_assert(TIsArithmetic<ArithmeticType>::Value, "Multiplicand type must be arithmetic.");
+				InLHS.NumSamples *= InRHS;
+				return InLHS;
+			}
+
+			template<typename ArithmeticType>
+			friend FSampleCounter& operator/=(FSampleCounter& InLHS, ArithmeticType InRHS)
+			{
+				static_assert(TIsArithmetic<ArithmeticType>::Value, "Divisor type must be arithmetic.");
+				InLHS.NumSamples /= InRHS;
+				return InLHS;
+			}
+
+			template<typename ArithmeticType>
+			friend ArithmeticType& operator/=(ArithmeticType& InLHS, const FSampleCounter& InRHS)
+			{
+				static_assert(TIsArithmetic<ArithmeticType>::Value, "Divisor type must be arithmetic.");
+				InLHS.NumSamples /= InRHS;
+				return InLHS;
+			}
+
+			template<typename ArithmeticType>
+			friend FSampleCounter operator+=(FSampleCounter& InLHS, const ArithmeticType& InRHS)
+			{
+				static_assert(TIsArithmetic<ArithmeticType>::Value, "Addend type must be arithmetic.");
+				InLHS.NumSamples += InRHS;
+				return InLHS;
+			}
+
+			template<typename ArithmeticType>
+			friend FSampleCounter& operator-=(FSampleCounter& InLHS, const ArithmeticType& InRHS)
+			{
+				static_assert(TIsArithmetic<ArithmeticType>::Value, "Subtrahend type must be arithmetic.");
+				InLHS.NumSamples -= InRHS;
+				return InLHS;
+			}
+
+			friend bool operator<(const FSampleCounter& InLHS, const FTime& InRHS)
+			{
+				return InLHS.ToTime() < InRHS;
+			}
+
+			friend bool operator<(const FTime& InLHS, const FSampleCounter& InRHS)
+			{
+				return InLHS < InRHS.ToTime();
+			}
+
+			friend bool operator<(const FSampleCounter& InLHS, const FSampleCounter& InRHS)
+			{
+				return InLHS.NumSamples < InRHS.NumSamples;
+			}
+
+			friend bool operator>(const FSampleCounter& InLHS, const FTime& InRHS)
+			{
+				return InLHS.ToTime() > InRHS;
+			}
+
+			friend bool operator>(const FTime& InLHS, const FSampleCounter& InRHS)
+			{
+				return InLHS > InRHS.ToTime();
+			}
+
+			friend bool operator>(const FSampleCounter& InLHS, const FSampleCounter& InRHS)
+			{
+				return InLHS.NumSamples > InRHS.NumSamples;
+			}
+
+			friend bool operator<=(const FSampleCounter& InLHS, const FTime& InRHS)
+			{
+				return InLHS.ToTime() <= InRHS;
+			}
+
+			friend bool operator<=(const FTime& InLHS, const FSampleCounter& InRHS)
+			{
+				return InLHS <= InRHS.ToTime();
+			}
+
+			friend bool operator<=(const FSampleCounter& InLHS, const FSampleCounter& InRHS)
+			{
+				return InLHS.NumSamples <= InRHS.NumSamples;
+			}
+
+			friend bool operator>=(const FSampleCounter& InLHS, const FTime& InRHS)
+			{
+				return InLHS.ToTime() >= InRHS;
+			}
+
+			friend bool operator>=(const FTime& InLHS, const FSampleCounter& InRHS)
+			{
+				return InLHS >= InRHS.ToTime();
+			}
+
+			friend bool operator>=(const FSampleCounter& InLHS, const FSampleCounter& InRHS)
+			{
+				return InLHS.NumSamples >= InRHS.NumSamples;
+			}
+
+			friend bool operator==(const FSampleCounter& InLHS, const FTime& InRHS)
+			{
+				return InLHS.ToTime() == InRHS;
+			}
+
+			friend bool operator==(const FTime& InLHS, const FSampleCounter& InRHS)
+			{
+				return InLHS == InRHS.ToTime();
+			}
+
+			friend bool operator==(const FSampleCounter& InLHS, const FSampleCounter& InRHS)
+			{
+				return InLHS.NumSamples == InRHS.GetNumSamples(InLHS.SampleRate);
+			}
+
+			friend bool operator!=(const FSampleCounter& InLHS, const FTime& InRHS)
+			{
+				return InLHS.ToTime() != InRHS;
+			}
+
+			friend bool operator!=(const FTime& InLHS, const FSampleCounter& InRHS)
+			{
+				return InLHS != InRHS.ToTime();
+			}
+
+			friend bool operator!=(const FSampleCounter& InLHS, const FSampleCounter& InRHS)
+			{
+				return InLHS.NumSamples != InRHS.NumSamples;
+			}
+	};
+}
diff --git a/Engine/Plugins/Runtime/Metasound/Source/MetasoundStandardNodes/Public/MetasoundSwitchGateNode.h b/Engine/Plugins/Runtime/Metasound/Source/MetasoundStandardNodes/Public/MetasoundSwitchGateNode.h
--- a/Engine/Plugins/Runtime/Metasound/Source/MetasoundStandardNodes/Public/MetasoundSwitchGateNode.h
+++ b/Engine/Plugins/Runtime/Metasound/Source/MetasoundStandardNodes/Public/MetasoundSwitchGateNode.h
@@ -9,7 +9,6 @@
 #include "MetasoundDataReferenceCollection.h"
 #include "MetasoundExecutableOperator.h"
 #include "MetasoundFacade.h"
-#include "MetasoundFrequency.h"
 #include "MetasoundNode.h"
 #include "MetasoundOperatorInterface.h"
 
diff --git a/Engine/Plugins/Runtime/Metasound/Source/MetasoundStandardNodes/Public/MetasoundTime.h b/Engine/Plugins/Runtime/Metasound/Source/MetasoundStandardNodes/Public/MetasoundTime.h
--- a/Engine/Plugins/Runtime/Metasound/Source/MetasoundStandardNodes/Public/MetasoundTime.h
+++ b/Engine/Plugins/Runtime/Metasound/Source/MetasoundStandardNodes/Public/MetasoundTime.h
@@ -1,646 +1,356 @@
 // Copyright Epic Games, Inc. All Rights Reserved.
-
 #pragma once
 
 #include "MetasoundDataReference.h"
 #include "MetasoundDataTypeRegistrationMacro.h"
+#include "MetasoundOperatorSettings.h"
 
-namespace Metasound
-{ 
-	/** Resolution of time value */
-	enum class ETimeResolution : uint8
-	{
-		/** Time value given in seconds. */
-		Seconds,
 
-		/** Time value given in milliseconds. */
-		Milliseconds,
-
-		/** Time value given in microseconds. */
-		Microseconds
-	};
+namespace Metasound
+{
+	class FSampleCounter;
 
-	/** TTime represents a time value. It provides basic mathematical 
-	 * operations, clearly defined getters and setters, and interoperability 
-	 * with TTime objects utilizing different template arguments. 
-	 *
-	 * QuantizationType declares the type of object which stores time 
-	 * information.
-	 *
-	 * TimeType declares the type of object which represent time values 
-	 * returned from TTime's getters and accepted by TTime's setters.
-	 */
-	template<typename QuantizationType, typename TimeType>
-	class TTime
+	// Strongly typed time class
+	class FTime
 	{
-			static_assert(TIsFloatingPoint<TimeType>::Value, "TimeType must be floating point");
+		using TimeType = double;
+		static_assert(TIsFloatingPoint<TimeType>::Value, "TimeType must be floating point.");
+
+		TimeType Time = 0;
 
 		public:
-			using FQuantizationType = QuantizationType;
-			using FTimeType = TimeType;
-
-			/** Default constructor */
-			TTime()
-			:	Value(0)
-			{}
-			
-			/** Initialize with a time value and time resolution.
-			 *
-			 * @param InValue - The initial value of the time object for the 
-			 * 					given resolution.
-			 * @param InResolution - The Resolution of the given InValue.
-			 */
-			TTime(QuantizationType InValue, ETimeResolution InResolution)
-			:	Value(0)
-			{
-				// Different setters called based upon resolution.
-				switch (InResolution)
-				{
-					case ETimeResolution::Microseconds:
-						SetMicroseconds(InValue);
-						break;
-
-					case ETimeResolution::Milliseconds:
-						SetMilliseconds(InValue);
-						break;
-
-					case ETimeResolution::Seconds:
-					default:
-						SetSeconds(InValue);
-				}
-			}
-
-			/** Constructor used by the metasound backend. */
-			TTime(QuantizationType InSeconds)
-				: TTime(InSeconds, ETimeResolution::Seconds)
-			{}
-
-			/** Returns the time as seconds. */
-			TimeType GetSeconds() const 
-			{ 
-				return static_cast<TimeType>(Value);
+			FTime()
+			{
 			}
 
-			/** Returns the time as milliseconds. */
-			TimeType GetMilliseconds() const 
-			{ 
-				return static_cast<TimeType>(Value * QuantizationType(1e3));
+			FTime(TimeType InTime)
+				: Time(InTime)
+			{
 			}
 
-			/** Returns the time as microseconds. */
-			TimeType GetMicroseconds() const 
-			{ 
-				return static_cast<TimeType>(Value * QuantizationType(1e6));
+			FTime(const FTime& InOtherTime)
+				: Time(InOtherTime.Time)
+			{
 			}
 
-			/** Returns the number of samples which represent the current time 
-			 * value.
-			 *
-			 * @param InSampleRate - The sample rate to use when calculating
-			 * 						 the number of samples. 
-			 *
-			 * @return The number of samples. 
-			 */
-			QuantizationType GetNumSamples(TimeType InSampleRate) const
+			static const FTime& Zero()
 			{
-				return static_cast<QuantizationType>(Value * InSampleRate);
+				static const FTime ZeroTime;
+				return ZeroTime;
 			}
 
-			/** Set the time value in seconds. */
-			void SetSeconds(TimeType InSeconds) 
+			/** Return the time as seconds. */
+			TimeType GetSeconds() const
 			{ 
-				Value = static_cast<QuantizationType>(InSeconds);
+				return Time;
 			}
 
-			/** Set the time value in milliseconds. */
-			void SetMilliseconds(TimeType InMilliseconds) 
-			{ 
-				Value = static_cast<QuantizationType>(InMilliseconds * TimeType(1e-3));
+			FTime& operator=(const FTime& InRHS)
+			{
+				Time = InRHS.Time;
+				return *this;
 			}
 
-			/** Set the time value in microseconds. */
-			void SetMicroseconds(TimeType InMicroseconds) 
-			{ 
-				Value = static_cast<QuantizationType>(InMicroseconds * TimeType(1e-6));
+			template<typename ArithmeticType>
+			friend ArithmeticType& operator*=(ArithmeticType& InLHS, const FTime& InRHS)
+			{
+				static_assert(TIsArithmetic<ArithmeticType>::Value, "Must be arithmetic type.");
+				InLHS.Time *= InRHS.Time;
+				return InLHS;
 			}
 
-			/** Set the number of samples to represent. 
-			 *
-			 * @param InNumSamples - The number of samples for the TTime object 
-			 * 						 to represent.
-			 * @param InSampleRate - The sample rate to use when calculating
-			 * 						 the number of samples. 
-			 */
-			void SetNumSamples(QuantizationType InNumSamples, TimeType InSampleRate)
-			{ 
-				SetSeconds(InNumSamples / FMath::Max(InSampleRate, static_cast<TimeType>(SMALL_NUMBER)));
+			template<typename ArithmeticType>
+			friend FTime& operator*=(FTime& InLHS, const ArithmeticType& InRHS)
+			{
+				static_assert(TIsArithmetic<ArithmeticType>::Value, "Must be arithmetic type.");
+				InLHS.Time *= InRHS;
+				return InLHS;
 			}
 
-			/** Assignment operator. */
-			template<typename OtherQuantizationType, typename OtherTimeType>
-			TTime& operator=(const TTime<OtherQuantizationType, OtherTimeType>& InOther)
+			template<typename ArithmeticType>
+			friend ArithmeticType& operator/=(ArithmeticType& InLHS, const FTime& InRHS)
 			{
-				Value = static_cast<QuantizationType>(InOther.GetSeconds());
-				return *this;
+				static_assert(TIsArithmetic<ArithmeticType>::Value, "Must be arithmetic type.");
+				InLHS.Time /= InRHS.Time;
+				return InLHS;
 			}
 
-			/** Addition assignment operator. 
-			 * 
-			 * Adds the Other's time to this time.
-			 */
-			template<typename OtherQuantizationType, typename OtherTimeType>
-			TTime& operator+=(const TTime<OtherQuantizationType, OtherTimeType>& InOther)
+			template<typename ArithmeticType>
+			friend FTime& operator/=(FTime& InLHS, const ArithmeticType& InRHS)
 			{
-				Value += static_cast<QuantizationType>(InOther.GetSeconds());
-				return *this;
+				static_assert(TIsArithmetic<ArithmeticType>::Value, "Must be arithmetic type.");
+				InLHS.Time /= InRHS;
+				return InLHS;
 			}
 
-			/** Subtraction assignment operator. 
-			 * 
-			 * Subtracts the Other's time from this time.
-			 */
-			template<typename OtherQuantizationType, typename OtherTimeType>
-			TTime& operator-=(const TTime<OtherQuantizationType, OtherTimeType>& InOther)
+			FTime& operator-=(const FTime& InRHS)
 			{
-				Value -= static_cast<QuantizationType>(InOther.GetSeconds());
+				Time -= InRHS.Time;
 				return *this;
 			}
 
-			/* Multiplication assignment operator. 
-			 *
-			 * Multiplies this time by the value. 
-			 */
-			template<typename OtherType>
-			TTime& operator*=(OtherType InValue)
+			template<typename ArithmeticType>
+			friend ArithmeticType& operator-=(ArithmeticType& InLHS, const FTime& InRHS)
 			{
-				static_assert(TIsArithmetic<OtherType>::Value, "TTime can only be multiplied with arithmetic types");
-
-				Value *= InValue;
-				return *this;
+				static_assert(TIsArithmetic<ArithmeticType>::Value, "Must be arithmetic type.");
+				InLHS -= InRHS.Time;
+				return InLHS;
 			}
 
-			/* Division assignment operator. 
-			 *
-			 * Divides this time by the value. 
-			 */
-			template<typename OtherType>
-			TTime& operator/=(OtherType InValue)
+			template<typename ArithmeticType>
+			friend FTime& operator-=(FTime& InLHS, const ArithmeticType& InRHS)
 			{
-				static_assert(TIsArithmetic<OtherType>::Value, "TTime can only be divided with arithmetic types");
+				static_assert(TIsArithmetic<ArithmeticType>::Value, "Must be arithmetic type.");
+				InLHS.Time -= InRHS;
+				return InLHS;
+			}
 
-				Value /= InValue;
+			FTime& operator+=(const FTime& InRHS)
+			{
+				Time += InRHS.Time;
 				return *this;
 			}
 
-			/* Comparison operator global friend functions */
-
-			template<typename LHSQuantizationType, typename LHSTimeType, typename RHSQuantizationType, typename RHSTimeType>
-			friend bool operator<(const TTime<LHSQuantizationType, LHSTimeType>& LHS, const TTime<RHSQuantizationType, RHSTimeType>& RHS);
-
-			template<typename LHSQuantizationType, typename LHSTimeType, typename RHSQuantizationType, typename RHSTimeType>
-			friend bool operator>(const TTime<LHSQuantizationType, LHSTimeType>& LHS, const TTime<RHSQuantizationType, RHSTimeType>& RHS);
-
-			template<typename LHSQuantizationType, typename LHSTimeType, typename RHSQuantizationType, typename RHSTimeType>
-			friend bool operator<=(const TTime<LHSQuantizationType, LHSTimeType>& LHS, const TTime<RHSQuantizationType, RHSTimeType>& RHS);
-
-			template<typename LHSQuantizationType, typename LHSTimeType, typename RHSQuantizationType, typename RHSTimeType>
-			friend bool operator>=(const TTime<LHSQuantizationType, LHSTimeType>& LHS, const TTime<RHSQuantizationType, RHSTimeType>& RHS);
-
-			template<typename LHSQuantizationType, typename LHSTimeType, typename RHSQuantizationType, typename RHSTimeType>
-			friend bool operator==(const TTime<LHSQuantizationType, LHSTimeType>& LHS, const TTime<RHSQuantizationType, RHSTimeType>& RHS);
-
-			template<typename LHSQuantizationType, typename LHSTimeType, typename RHSQuantizationType, typename RHSTimeType>
-			friend bool operator!=(const TTime<LHSQuantizationType, LHSTimeType>& LHS, const TTime<RHSQuantizationType, RHSTimeType>& RHS);
-
-			/* Arithmetic operator global friend functions */
-
-			template<typename LHSQuantizationType, typename LHSTimeType, typename RHSQuantizationType, typename RHSTimeType>
-			friend TTime<LHSQuantizationType, LHSTimeType> operator+(TTime<LHSQuantizationType, LHSTimeType> LHS, const TTime<RHSQuantizationType, RHSTimeType>& RHS);
-
-			template<typename LHSQuantizationType, typename LHSTimeType, typename RHSQuantizationType, typename RHSTimeType>
-			friend TTime<LHSQuantizationType, LHSTimeType> operator-(TTime<LHSQuantizationType, LHSTimeType> LHS, const TTime<RHSQuantizationType, RHSTimeType>& RHS);
-
-			template<typename LHSQuantizationType, typename LHSTimeType, typename OtherType>
-			friend TTime<LHSQuantizationType, LHSTimeType> operator*(TTime<LHSQuantizationType, LHSTimeType> LHS, OtherType RHS);
-
-			template<typename RHSQuantizationType, typename RHSTimeType, typename OtherType>
-			friend TTime<RHSQuantizationType, RHSTimeType> operator*(OtherType LHS, const TTime<RHSQuantizationType, RHSTimeType>& RHS);
-
-			template<typename LHSQuantizationType, typename LHSTimeType, typename OtherType>
-			friend TTime<LHSQuantizationType, LHSTimeType> operator/(TTime<LHSQuantizationType, LHSTimeType> LHS, OtherType RHS);
-
-
-		private:
-
-			QuantizationType Value;
-	};
-
-
-	// holds on to a count of samples.
-	struct FTimeSampleCounter
-	{
-		// defines the type for counting samples.
-		typedef int32 FCountType;
-		FCountType Num;
-	};
-
-	/** TTime sample counting template specialization holds time values which
-	 * are quantized to sample boundaries. It provides basic mathematical 
-	 * operations, clearly defined getters and setters, and interoperability 
-	 * with generic TTime objects utilizing different template arguments. 
-	 *
-	 * This specialization differs distinctly by storing the number of samples
-	 * rather than the number of seconds in the quantization type. Additionally,
-	 * instances have an associated sample rate member which is used to convert
-	 * between TimeType and the sample count.
-	 *
-	 * TimeType declares the type of object which represent time values 
-	 * returned from TTime's getters.
-	 */
-	template<typename TimeType>
-	class TTime<FTimeSampleCounter, TimeType>
-	{
-		public:
-			static_assert(TIsFloatingPoint<TimeType>::Value, "TimeType must be floating point");
-
-			// using types from namespace.
-			using FCountType = FTimeSampleCounter::FCountType;
-			
-			/** Construct a sample count based TTime with a given time value, 
-			 * time resolution, and sample rate.
-			 *
-			 * @param InValue - The initial time value.
-			 * @param InResolution - The time resolution of InValue.
-			 * @param InSampleRate - The sample rate associated with this time 
-			 * 						 object.
-			 */
-			TTime(TimeType InValue, ETimeResolution InResolution, TimeType InSampleRate)
+			template<typename ArithmeticType>
+			friend ArithmeticType& operator+=(ArithmeticType& InLHS, const FTime& InRHS)
 			{
-				// Need to set sample rate before setting time or else the 
-				// resulting sample count will be incorrect.
-				SetSampleRate(InSampleRate);
-
-				switch (InResolution)
-				{
-					case ETimeResolution::Microseconds:
-						SetMicroseconds(InValue);
-						break;
-
-					case ETimeResolution::Milliseconds:
-						SetMilliseconds(InValue);
-						break;
-
-					case ETimeResolution::Seconds:
-					default:
-						SetSeconds(InValue);
-				}
+				static_assert(TIsArithmetic<ArithmeticType>::Value, "Must be arithmetic type.");
+				InLHS += InRHS.Time;
+				return InLHS;
 			}
 
-			/**
-			 * Constructor used by the Metasound Frontend.
-			 *
-			 * @param InSettings the operator settings to use.
-			 * @param InSeconds The number in time to initialize to.
-			 */
-			TTime(const FOperatorSettings& InSettings, float InSeconds)
-				: TTime(InSeconds, ETimeResolution::Seconds, InSettings.GetSampleRate())
-			{}
-
-			TTime()
-				: TTime(0.0f, ETimeResolution::Seconds, SMALL_NUMBER)
-			{}
-
-			/** Construct a sample count based TTime using a number of samples
-			 * and a sample rate.
-			 *
-			 * @param InNumSamples - The initial number of samples.
-			 * @param InSampleRate - The sample rate of this instance.
-			 */
-			TTime(FCountType InNumSamples, TimeType InSampleRate)
+			template<typename ArithmeticType>
+			friend FTime& operator+=(FTime& InLHS, const ArithmeticType& InRHS)
 			{
-				SetSampleRate(InSampleRate);
-
-				SetNumSamples(InNumSamples);
+				static_assert(TIsArithmetic<ArithmeticType>::Value, "Must be arithmetic type.");
+				InLHS.Time += InRHS;
+				return InLHS;
 			}
 
-			/** Set the sample rate of this object without changing the number
-			 * of samples. This will result in the object representing a 
-			 * different amount of time. 
-			 */
-			void SetSampleRate(TimeType InSampleRate)
+			FTime operator-(const FTime& InRHS) const
 			{
-				if (!ensure(InSampleRate > 0))
-				{
-					InSampleRate = SMALL_NUMBER;
-				}
-
-				SampleRate = InSampleRate;
+				return Time - InRHS.Time;
 			}
 
-			/** Set the sample rate of this object and change the number
-			 * of samples. This will result in the object representing 
-			 * approximately the same amount of time as it did before the 
-			 * sample rate was changed.
-			 */
-			void SetSampleRateAndUpdateSampleCount(TimeType InSampleRate)
+			template<typename ArithmeticType>
+			friend ArithmeticType operator-(const ArithmeticType& InLHS, const FTime& InRHS)
 			{
-				if (!ensure(InSampleRate > 0))
-				{
-					InSampleRate = SMALL_NUMBER;
-				}
-
-				Value.Num = GetNumSamples(InSampleRate);
-				
-				SampleRate = InSampleRate;
+				static_assert(TIsArithmetic<ArithmeticType>::Value, "Must be arithmetic type.");
+				return InLHS - InRHS.Time;
 			}
 
-			/** Return the sample rate. */
-			TimeType GetSampleRate() const
+			template<typename ArithmeticType>
+			friend FTime operator-(const FTime& InLHS, const ArithmeticType& InRHS)
 			{
-				return SampleRate;
+				static_assert(TIsArithmetic<ArithmeticType>::Value, "Must be arithmetic type.");
+				return InLHS.Time - InRHS;
 			}
 
-			/** Return the time as seconds. */
-			TimeType GetSeconds() const 
-			{ 
-				return static_cast<TimeType>(Value.Num / SampleRate);
+			FTime operator+(const FTime& InRHS) const
+			{
+				return Time + InRHS.Time;
 			}
 
-			/** Return the time as milliseconds. */
-			TimeType GetMilliseconds() const 
-			{ 
-				return static_cast<TimeType>(Value.Num * 1e3 / SampleRate);
+			template<typename ArithmeticType>
+			friend ArithmeticType operator+(const ArithmeticType& InLHS, const FTime& InRHS)
+			{
+				static_assert(TIsArithmetic<ArithmeticType>::Value, "Must be arithmetic type.");
+				return InLHS + InRHS.Time;
 			}
 
-			/** Return the time as microseconds. */
-			TimeType GetMicroseconds() const 
-			{ 
-				return static_cast<TimeType>(Value.Num * 1e6 / SampleRate);
+			template<typename ArithmeticType>
+			friend FTime operator+(const FTime& InLHS, const ArithmeticType& InRHS)
+			{
+				static_assert(TIsArithmetic<ArithmeticType>::Value, "Must be arithmetic type.");
+				return InLHS.Time + InRHS;
 			}
 
-			/** Return the number of samples. */
-			FCountType GetNumSamples() const
-			{ 
-				return Value.Num;
+			template<typename ArithmeticType>
+			ArithmeticType operator*(const FTime& InRHS) const
+			{
+				static_assert(TIsArithmetic<ArithmeticType>::Value, "Must be arithmetic type.");
+				return static_cast<ArithmeticType>(Time * InRHS.Time);
 			}
 
-			/** Return the number of samples which represent this time duration
-			 * using a different sample rate. 
-			 *
-			 * @param InOtherSampleRate - Sample rate to use when calculating 
-			 * 							  the number of samples.  
-			 */
-			FCountType GetNumSamples(TimeType InOtherSampleRate) const
-			{ 
-				if (InOtherSampleRate == SampleRate)
-				{
-					return Value.Num;
-				}
-				else
-				{
-					return FMath::RoundToInt(GetSeconds() * InOtherSampleRate);
-				}
-			}
-
-			/** Set the time value in seconds. */
-			void SetSeconds(TimeType InSeconds) 
-			{ 
-				Value.Num = FMath::RoundToInt(InSeconds * SampleRate);
+			template<typename ArithmeticType>
+			friend ArithmeticType operator*(const ArithmeticType& InLHS, const FTime& InRHS)
+			{
+				static_assert(TIsArithmetic<ArithmeticType>::Value, "Must be arithmetic type.");
+				return InLHS * InRHS.Time;
 			}
 
-			/** Set the time value in milliseconds. */
-			void SetMilliseconds(TimeType InMilliseconds) 
-			{ 
-				Value.Num = FMath::RoundToInt(InMilliseconds * SampleRate * TimeType(1e-3));
+			template<typename ArithmeticType>
+			friend FTime operator*(const FTime& InLHS, const ArithmeticType& InRHS)
+			{
+				static_assert(TIsArithmetic<ArithmeticType>::Value, "Must be arithmetic type.");
+				return InLHS.Time * InRHS;
 			}
 
-			/** Set the time value in microseconds. */
-			void SetMicroseconds(TimeType InMicroseconds) 
-			{ 
-				Value.Num = FMath::RoundToInt(InMicroseconds * SampleRate * TimeType(1e-6));
+			template<typename ArithmeticType>
+			ArithmeticType operator/(const FTime& InRHS) const
+			{
+				static_assert(TIsArithmetic<ArithmeticType>::Value, "Must be arithmetic type.");
+				return static_cast<ArithmeticType>(Time / InRHS.Time);
 			}
 
-			/** Set the number of samples. */
-			void SetNumSamples(FCountType InNumSamples)
-			{ 
-				Value.Num = InNumSamples;
+			template<typename ArithmeticType>
+			friend ArithmeticType& operator/(const ArithmeticType& InLHS, const FTime& InRHS)
+			{
+				static_assert(TIsArithmetic<ArithmeticType>::Value, "Must be arithmetic type.");
+				return InLHS / InRHS.Time;
 			}
 
-			/** Assignment operator 
-			 *
-			 * This assignment preserves this objects sample rate when being 
-			 * assigned to TTime objects of different quantization types.
-			 * */
-			template<typename OtherQuantizationType, typename OtherTimeType>
-			TTime& operator=(const TTime<OtherQuantizationType, OtherTimeType>& Other)
+			template<typename ArithmeticType>
+			friend FTime operator/(const FTime& InLHS, const ArithmeticType& InRHS)
 			{
-				Value.Num = Other.GetNumSamples(SampleRate);
-
-				return *this;
+				static_assert(TIsArithmetic<ArithmeticType>::Value, "Must be arithmetic type.");
+				return InLHS.Time / InRHS;
 			}
 
-			/** Assignment operator 
-			 *
-			 * This assignment copies the sample rate and sample count of TTime
-			 * objects of equal quantization types.
-			 */
-			template<typename OtherTimeType>
-			TTime& operator=(const TTime<FTimeSampleCounter, OtherTimeType>& Other)
+			bool operator<(const FTime& InRHS) const
 			{
-				Value.Num = Other.Value.Num;
-				SampleRate = Other.SampleRate;
-
-				return *this;
+				return Time < InRHS.Time;
 			}
 
-			/** Addition assignment operator. */
-			template<typename OtherQuantizationType, typename OtherTimeType>
-			TTime& operator+=(const TTime<OtherQuantizationType, OtherTimeType>& Other)
+			template<typename ArithmeticType>
+			friend bool operator<(const ArithmeticType& InLHS, const FTime& InRHS)
 			{
-				Value.Num += Other.GetNumSamples(SampleRate);
+				static_assert(TIsArithmetic<ArithmeticType>::Value, "Must be arithmetic type.");
+				return InLHS < InRHS.Time;
+			}
 
-				return *this;
+			template<typename ArithmeticType>
+			friend bool operator<(const FTime& InLHS, const ArithmeticType& InRHS)
+			{
+				static_assert(TIsArithmetic<ArithmeticType>::Value, "Must be arithmetic type.");
+				return InLHS.Time < InRHS;
 			}
 
-			/** Subtraction assignment operator. */
-			template<typename OtherQuantizationType, typename OtherTimeType>
-			TTime& operator-=(const TTime<OtherQuantizationType, OtherTimeType>& Other)
+			bool operator<=(const FTime& InRHS) const
 			{
-				Value.Num -= Other.GetNumSamples(SampleRate);
+				return Time <= InRHS.Time;
+			}
 
-				return *this;
+			template<typename ArithmeticType>
+			friend bool operator<=(const ArithmeticType& InLHS, const FTime& InRHS)
+			{
+				static_assert(TIsArithmetic<ArithmeticType>::Value, "Must be arithmetic type.");
+				return InLHS <= InRHS.Time;
 			}
 
-			/** Multiplication assignment operator. */
-			template<typename OtherType>
-			TTime& operator*=(OtherType Other)
+			template<typename ArithmeticType>
+			friend bool operator<=(const FTime& InLHS, const ArithmeticType& InRHS)
 			{
-				static_assert(TIsArithmetic<OtherType>::Value, "TTime can only be multiplied with arithmetic types");
+				static_assert(TIsArithmetic<ArithmeticType>::Value, "Must be arithmetic type.");
+				return InLHS.Time <= InRHS;
+			}
 
-				Value.Num *= Other;
-				return *this;
+			bool operator>(const FTime& InRHS) const
+			{
+				return Time > InRHS.Time;
 			}
 
-			/** Division assignment operator. */
-			template<typename OtherType>
-			TTime& operator/=(OtherType Other)
+			template<typename ArithmeticType>
+			friend bool operator>(const FTime& InLHS, const ArithmeticType& InRHS)
 			{
-				static_assert(TIsArithmetic<OtherType>::Value, "TTime can only be divided with arithmetic types");
+				static_assert(TIsArithmetic<ArithmeticType>::Value, "Must be arithmetic type.");
+				return InLHS.Time > InRHS;
+			}
 
-				Value.Num /= Other;
-				return *this;
+			template<typename ArithmeticType>
+			friend bool operator>(const ArithmeticType& InLHS, const FTime& InRHS)
+			{
+				static_assert(TIsArithmetic<ArithmeticType>::Value, "Must be arithmetic type.");
+				return InLHS > InRHS.Time;
 			}
 
-			/* Comparison operator global friend functions */
-			
-			template<typename LHSQuantizationType, typename LHSTimeType, typename RHSQuantizationType, typename RHSTimeType>
-			friend bool operator<(const TTime<LHSQuantizationType, LHSTimeType>& LHS, const TTime<RHSQuantizationType, RHSTimeType>& RHS);
+			bool operator>=(const FTime& InRHS) const
+			{
+				return Time >= InRHS.Time;
+			}
 
-			template<typename LHSQuantizationType, typename LHSTimeType, typename RHSQuantizationType, typename RHSTimeType>
-			friend bool operator>(const TTime<LHSQuantizationType, LHSTimeType>& LHS, const TTime<RHSQuantizationType, RHSTimeType>& RHS);
+			template<typename ArithmeticType>
+			friend bool operator>=(const ArithmeticType& InLHS, const FTime& InRHS)
+			{
+				static_assert(TIsArithmetic<ArithmeticType>::Value, "Must be arithmetic type.");
+				return InLHS >= InRHS.Time;
+			}
 
-			template<typename LHSQuantizationType, typename LHSTimeType, typename RHSQuantizationType, typename RHSTimeType>
-			friend bool operator<=(const TTime<LHSQuantizationType, LHSTimeType>& LHS, const TTime<RHSQuantizationType, RHSTimeType>& RHS);
+			template<typename ArithmeticType>
+			friend bool operator>=(const FTime& InLHS, const ArithmeticType& InRHS)
+			{
+				static_assert(TIsArithmetic<ArithmeticType>::Value, "Must be arithmetic type.");
+				return InLHS.Time >= InRHS;
+			}
 
-			template<typename LHSQuantizationType, typename LHSTimeType, typename RHSQuantizationType, typename RHSTimeType>
-			friend bool operator>=(const TTime<LHSQuantizationType, LHSTimeType>& LHS, const TTime<RHSQuantizationType, RHSTimeType>& RHS);
+			template<typename ArithmeticType>
+			friend bool operator!=(const ArithmeticType& InLHS, const FTime& InRHS)
+			{
+				static_assert(TIsArithmetic<ArithmeticType>::Value, "Must be arithmetic type.");
+				return InLHS != InRHS.Time;
+			}
 
-			template<typename LHSQuantizationType, typename LHSTimeType, typename RHSQuantizationType, typename RHSTimeType>
-			friend bool operator==(const TTime<LHSQuantizationType, LHSTimeType>& LHS, const TTime<RHSQuantizationType, RHSTimeType>& RHS);
+			bool operator!=(const FTime& InRHS) const
+			{
+				return Time != InRHS.Time;
+			}
 
-			template<typename LHSQuantizationType, typename LHSTimeType, typename RHSQuantizationType, typename RHSTimeType>
-			friend bool operator!=(const TTime<LHSQuantizationType, LHSTimeType>& LHS, const TTime<RHSQuantizationType, RHSTimeType>& RHS);
+			template<typename ArithmeticType>
+			friend bool operator!=(const FTime& InLHS, const ArithmeticType& InRHS)
+			{
+				static_assert(TIsArithmetic<ArithmeticType>::Value, "Must be arithmetic type.");
+				return InLHS.Time != InRHS;
+			}
 
-			/* Arithmetic operator global friend functions */
+			friend bool operator==(const FTime& InLHS, const FTime& InRHS)
+			{
+				return InLHS.Time == InRHS.Time;
+			}
 
-			template<typename LHSQuantizationType, typename LHSTimeType, typename RHSQuantizationType, typename RHSTimeType>
-			friend TTime<LHSQuantizationType, LHSTimeType> operator+(TTime<LHSQuantizationType, LHSTimeType> LHS, const TTime<RHSQuantizationType, RHSTimeType>& RHS);
+			template<typename ArithmeticType>
+			friend bool operator==(const ArithmeticType& InLHS, const FTime& InRHS)
+			{
+				static_assert(TIsArithmetic<ArithmeticType>::Value, "Must be arithmetic type.");
+				return InLHS == InRHS.Time;
+			}
 
-			template<typename LHSQuantizationType, typename LHSTimeType, typename RHSQuantizationType, typename RHSTimeType>
-			friend TTime<LHSQuantizationType, LHSTimeType> operator-(TTime<LHSQuantizationType, LHSTimeType> LHS, const TTime<RHSQuantizationType, RHSTimeType>& RHS);
+			template<typename ArithmeticType>
+			friend bool operator==(const FTime& InLHS, const ArithmeticType& InRHS)
+			{
+				static_assert(TIsArithmetic<ArithmeticType>::Value, "Must be arithmetic type.");
+				return InLHS.Time == InRHS;
+			}
 
-			template<typename LHSQuantizationType, typename LHSTimeType, typename OtherType>
-			friend TTime<LHSQuantizationType, LHSTimeType> operator*(TTime<LHSQuantizationType, LHSTimeType> LHS, OtherType RHS);
+			template<typename ArithmeticType = TimeType>
+			static ArithmeticType ToMilliseconds(const FTime& InTime)
+			{
+				static_assert(TIsArithmetic<ArithmeticType>::Value, "Must be arithmetic type.");
+				return static_cast<ArithmeticType>(InTime.GetSeconds() * 1e3);
+			}
 
-			template<typename RHSQuantizationType, typename RHSTimeType, typename OtherType>
-			friend TTime<RHSQuantizationType, RHSTimeType> operator*(OtherType LHS, const TTime<RHSQuantizationType, RHSTimeType>& RHS);
+			template<typename ArithmeticType = TimeType>
+			static ArithmeticType ToMicroseconds(const FTime& InTime)
+			{
+				static_assert(TIsArithmetic<ArithmeticType>::Value, "Must be arithmetic type.");
+				return static_cast<ArithmeticType>(InTime.GetSeconds() * 1e6);
+			}
 
-			template<typename LHSQuantizationType, typename LHSTimeType, typename OtherType>
-			friend TTime<LHSQuantizationType, LHSTimeType> operator/(TTime<LHSQuantizationType, LHSTimeType> LHS, OtherType RHS);
+			template<typename ArithmeticType>
+			static FTime FromMilliseconds(const ArithmeticType& InMilliseconds)
+			{
+				static_assert(TIsArithmetic<ArithmeticType>::Value, "Must be arithmetic type.");
+				return FTime(static_cast<TimeType>(InMilliseconds * 1e-3));
+			}
 
-		private:
+			template<typename ArithmeticType>
+			static FTime FromMicroseconds(const ArithmeticType& InMicroseconds)
+			{
+				static_assert(TIsArithmetic<ArithmeticType>::Value, "Must be arithmetic type.");
+				return FTime(static_cast<TimeType>(InMicroseconds * 1e-6));
+			}
 
-			FTimeSampleCounter Value = {0};
-			TimeType SampleRate = static_cast<TimeType>(SMALL_NUMBER);
+			friend FSampleCounter;
 	};
 
-
-	/** Less-than comparison of TTime */
-	template<typename LHSQuantizationType, typename LHSTimeType, typename RHSQuantizationType, typename RHSTimeType>
-	bool operator<(const TTime<LHSQuantizationType, LHSTimeType>& LHS, const TTime<RHSQuantizationType, RHSTimeType>& RHS)
-	{ 
-		return LHS.GetSeconds() < RHS.GetSeconds();
-	}
-
-	/** Greater-than comparison of TTime */
-	template<typename LHSQuantizationType, typename LHSTimeType, typename RHSQuantizationType, typename RHSTimeType>
-	bool operator>(const TTime<LHSQuantizationType, LHSTimeType>& LHS, const TTime<RHSQuantizationType, RHSTimeType>& RHS)
-	{ 
-		return RHS < LHS; 
-	}
-
-	/** Less-than-or-equal-to comparison of TTime */
-	template<typename LHSQuantizationType, typename LHSTimeType, typename RHSQuantizationType, typename RHSTimeType>
-	bool operator<=(const TTime<LHSQuantizationType, LHSTimeType>& LHS, const TTime<RHSQuantizationType, RHSTimeType>& RHS)
-	{ 
-		return !(LHS > RHS); 
-	}
-
-	/** Greater-than-or-equal-to comparison of TTime */
-	template<typename LHSQuantizationType, typename LHSTimeType, typename RHSQuantizationType, typename RHSTimeType>
-	bool operator>=(const TTime<LHSQuantizationType, LHSTimeType>& LHS, const TTime<RHSQuantizationType, RHSTimeType>& RHS)
-	{ 
-		return !(LHS < RHS); 
-	}
-
-	/** Equal comparison of TTime */
-	template<typename LHSQuantizationType, typename LHSTimeType, typename RHSQuantizationType, typename RHSTimeType>
-	bool operator==(const TTime<LHSQuantizationType, LHSTimeType>& LHS, const TTime<RHSQuantizationType, RHSTimeType>& RHS)
-	{
-		return LHS.GetSeconds() == RHS.GetSeconds();
-	}
-
-	/** Not-equal comparison of TTime */
-	template<typename LHSQuantizationType, typename LHSTimeType, typename RHSQuantizationType, typename RHSTimeType>
-	bool operator!=(const TTime<LHSQuantizationType, LHSTimeType>& LHS, const TTime<RHSQuantizationType, RHSTimeType>& RHS)
-	{ 
-		return !(LHS == RHS); 
-	}
-
-	/** Addition operator of TTime */
-	template<typename LHSQuantizationType, typename LHSTimeType, typename RHSQuantizationType, typename RHSTimeType>
-	TTime<LHSQuantizationType, LHSTimeType> operator+(TTime<LHSQuantizationType, LHSTimeType> LHS, const TTime<RHSQuantizationType, RHSTimeType>& RHS)
-	{
-		LHS += RHS;
-		return LHS;
-	}
-
-	/** Subtraction operator of TTime */
-	template<typename LHSQuantizationType, typename LHSTimeType, typename RHSQuantizationType, typename RHSTimeType>
-	TTime<LHSQuantizationType, LHSTimeType> operator-(TTime<LHSQuantizationType, LHSTimeType> LHS, const TTime<RHSQuantizationType, RHSTimeType>& RHS)
-	{
-		LHS -= RHS;
-		return LHS;
-	}
-
-	/** Multiplication operator of TTime */
-	template<typename LHSQuantizationType, typename LHSTimeType, typename OtherType>
-	TTime<LHSQuantizationType, LHSTimeType> operator*(TTime<LHSQuantizationType, LHSTimeType> LHS, OtherType RHS)
-	{
-		static_assert(TIsArithmetic<OtherType>::Value, "TTime can only be multiplied with arithmetic types");
-
-		LHS *= RHS;
-		return LHS;
-	}
-
-	/** Division operator of TTime */
-	template<typename LHSQuantizationType, typename LHSTimeType, typename OtherType>
-	TTime<LHSQuantizationType, LHSTimeType> operator/(TTime<LHSQuantizationType, LHSTimeType> LHS, OtherType RHS)
-	{
-		static_assert(TIsArithmetic<OtherType>::Value, "TTime can only be divided with arithmetic types");
-
-		LHS /= RHS;
-		return LHS;
-	}
-
-	/** Multiplication operator of TTime */
-	template<typename RHSQuantizationType, typename RHSTimeType, typename OtherType>
-	TTime<RHSQuantizationType, RHSTimeType> operator*(OtherType LHS, const TTime<RHSQuantizationType, RHSTimeType>& RHS)
-	{
-		static_assert(TIsArithmetic<OtherType>::Value, "TTime can only be multiplied with arithmetic types");
-
-		return RHS * LHS;
-	}
-	
-
-	/** FFloatTime stores and represents all time values as 32bit floating 
-	 * point values. 
-	 */
-	typedef TTime<float, float> FFloatTime;
-
-	DECLARE_METASOUND_DATA_REFERENCE_TYPES(FFloatTime, METASOUNDSTANDARDNODES_API, FFloatTimeTypeInfo, FFloatTimeReadRef, FFloatTimeWriteRef);
-
-	/** FDoubleTime stores and represents all time values as 64bit floating
-	 * point values.
-	 */
-	typedef TTime<double, double> FDoubleTime;
-
-	DECLARE_METASOUND_DATA_REFERENCE_TYPES(FDoubleTime, METASOUNDSTANDARDNODES_API, FDoubleTimeTypeInfo, FDoubleTimeReadRef, FDoubleTimeWriteRef);
-
-	/** FSampleTime stores time quantized to sample boundaries. It returns time
-	 * as 32bit floating point values.
-	 */
-	typedef TTime<FTimeSampleCounter, float> FSampleTime;
-
-	DECLARE_METASOUND_DATA_REFERENCE_TYPES(FSampleTime, METASOUNDSTANDARDNODES_API, FSampleTimeTypeInfo, FSampleTimeReadRef, FSampleTimeWriteRef);
+	DECLARE_METASOUND_DATA_REFERENCE_TYPES(FTime, METASOUNDSTANDARDNODES_API, FTimeTypeInfo, FTimeReadRef, FTimeWriteRef);
 }
diff --git a/Engine/Plugins/Runtime/Metasound/Source/MetasoundStandardNodes/Public/MetasoundTrigger.h b/Engine/Plugins/Runtime/Metasound/Source/MetasoundStandardNodes/Public/MetasoundTrigger.h
--- a/Engine/Plugins/Runtime/Metasound/Source/MetasoundStandardNodes/Public/MetasoundTrigger.h
+++ b/Engine/Plugins/Runtime/Metasound/Source/MetasoundStandardNodes/Public/MetasoundTrigger.h
@@ -5,6 +5,7 @@
 #include "MetasoundDataReference.h"
 #include "MetasoundDataTypeRegistrationMacro.h"
 #include "MetasoundOperatorSettings.h"
+#include "MetasoundTime.h"
 
 namespace Metasound
 {
@@ -165,7 +166,7 @@ namespace Metasound
 
 			bool bHasTrigger = false;
 			int32 NumFramesPerBlock = 0;
-			float SampleRate;
+			FSampleRate SampleRate = 0;
 			int32 LastTriggerIndexInBlock = 0;
 	};
 
