Remove local change not meant to submit in last CL

diff --git a/Engine/Plugins/Runtime/Metasound/Source/MetasoundStandardNodes/MetasoundStandardNodes.Build.cs b/Engine/Plugins/Runtime/Metasound/Source/MetasoundStandardNodes/MetasoundStandardNodes.Build.cs
--- a/Engine/Plugins/Runtime/Metasound/Source/MetasoundStandardNodes/MetasoundStandardNodes.Build.cs
+++ b/Engine/Plugins/Runtime/Metasound/Source/MetasoundStandardNodes/MetasoundStandardNodes.Build.cs
@@ -9,9 +9,6 @@ namespace UnrealBuildTool.Rules
 	{
 		public MetasoundStandardNodes(ReadOnlyTargetRules Target) : base(Target)
 		{
-			OptimizeCode = CodeOptimization.Never;
-			PCHUsage = ModuleRules.PCHUsageMode.NoSharedPCHs;
-
 			PublicDependencyModuleNames.AddRange
 			(
 				new string[]
