Fix potential period being set to 0 causing infinite loop in PeriodicTrigger

diff --git a/Engine/Plugins/Runtime/Metasound/Source/MetasoundStandardNodes/Private/MetasoundPeriodicTriggerNode.cpp b/Engine/Plugins/Runtime/Metasound/Source/MetasoundStandardNodes/Private/MetasoundPeriodicTriggerNode.cpp
--- a/Engine/Plugins/Runtime/Metasound/Source/MetasoundStandardNodes/Private/MetasoundPeriodicTriggerNode.cpp
+++ b/Engine/Plugins/Runtime/Metasound/Source/MetasoundStandardNodes/Private/MetasoundPeriodicTriggerNode.cpp
@@ -100,7 +100,11 @@ namespace Metasound
 
 		if (bEnabled)
 		{
-			const FSampleCount PeriodInSamples = FSampleCounter::FromTime(*Period, SampleCounter.GetSampleRate()).GetNumSamples();
+			FSampleCount PeriodInSamples = FSampleCounter::FromTime(*Period, SampleCounter.GetSampleRate()).GetNumSamples();
+
+			// Time must march on, can't stay in the now forever.
+			PeriodInSamples = FMath::Max(static_cast<FSampleCount>(1), PeriodInSamples);
+
 			while ((SampleCounter - BlockSize) <= 0)
 			{
 				const int32 StartOffset = static_cast<int32>(SampleCounter.GetNumSamples());
