MIDIToFreq metasound node

diff --git a/Engine/Plugins/Runtime/Metasound/Source/MetasoundStandardNodes/Private/MetasoundMIDIToFreqNode.cpp b/Engine/Plugins/Runtime/Metasound/Source/MetasoundStandardNodes/Private/MetasoundMIDIToFreqNode.cpp
new file mode 100644
--- /dev/null
+++ b/Engine/Plugins/Runtime/Metasound/Source/MetasoundStandardNodes/Private/MetasoundMIDIToFreqNode.cpp
@@ -0,0 +1,133 @@
+// Copyright Epic Games, Inc. All Rights Reserved.
+
+#include "MetasoundMIDIToFreqNode.h"
+
+#include "Internationalization/Text.h"
+#include "MetasoundExecutableOperator.h"
+#include "MetasoundNodeRegistrationMacro.h"
+#include "MetasoundPrimitives.h"
+#include "MetasoundStandardNodesNames.h"
+#include "DSP/Dsp.h"
+
+#define LOCTEXT_NAMESPACE "MetasoundStandardNodes"
+
+namespace Metasound
+{
+	// Param list
+	static const TCHAR* InParamNameMIDI = TEXT("In MIDI");
+	static const TCHAR* OutParamNameFreq = TEXT("Out Freq");
+
+	class FMIDIToFreqOperator : public TExecutableOperator<FMIDIToFreqOperator>
+	{
+	public:
+
+		static const FNodeClassMetadata& GetNodeInfo();
+		static const FVertexInterface& GetVertexInterface();
+		static TUniquePtr<IOperator> CreateOperator(const FCreateOperatorParams& InParams, FBuildErrorArray& OutErrors);
+
+		FMIDIToFreqOperator(const FOperatorSettings& InSettings, const FInt32ReadRef& InMIDINote);
+
+		virtual FDataReferenceCollection GetInputs() const override;
+		virtual FDataReferenceCollection GetOutputs() const override;
+		void Execute();
+
+	private:
+		// The input midi value
+		FInt32ReadRef MIDINote;
+
+		// The output frequency
+		FFloatWriteRef FreqOutput;
+
+		// Cached midi note value. Used to catch if the value changes to recompute freq output.
+		int32 PrevMidiNote;
+	};
+
+	FMIDIToFreqOperator::FMIDIToFreqOperator(const FOperatorSettings& InSettings, const FInt32ReadRef& InMIDINote)
+		: MIDINote(InMIDINote)
+		, FreqOutput(FFloatWriteRef::CreateNew(Audio::GetFrequencyFromMidi(*InMIDINote)))
+		, PrevMidiNote(*InMIDINote)
+	{
+	}
+
+	FDataReferenceCollection FMIDIToFreqOperator::GetInputs() const
+	{
+		FDataReferenceCollection InputDataReferences;
+		InputDataReferences.AddDataReadReference(InParamNameMIDI, FInt32ReadRef(MIDINote));
+
+		return InputDataReferences;
+	}
+
+	FDataReferenceCollection FMIDIToFreqOperator::GetOutputs() const
+	{
+		FDataReferenceCollection OutputDataReferences;
+		OutputDataReferences.AddDataReadReference(OutParamNameFreq, FFloatReadRef(FreqOutput));
+		return OutputDataReferences;
+	}
+
+	void FMIDIToFreqOperator::Execute()
+	{
+		// Only do anything if the midi note changes
+		if (*MIDINote != PrevMidiNote)
+		{
+			PrevMidiNote = *MIDINote;
+
+			*FreqOutput = Audio::GetFrequencyFromMidi(FMath::Clamp(*MIDINote, 0, 127));
+		}
+	}
+
+	const FVertexInterface& FMIDIToFreqOperator::GetVertexInterface()
+	{
+		static const FVertexInterface Interface(
+			FInputVertexInterface(
+				TInputDataVertexModel<int32>(InParamNameMIDI, LOCTEXT("MIDIInToolTip", "Input MIDI note value."), 60)
+			),
+			FOutputVertexInterface(
+				TOutputDataVertexModel<float>(OutParamNameFreq, LOCTEXT("OutFreqTooltip", "The output frequency (hz) value corresponding to the MIDI note."))
+			)
+		);
+
+		return Interface;
+	}
+
+	const FNodeClassMetadata& FMIDIToFreqOperator::GetNodeInfo()
+	{
+		auto InitNodeInfo = []() -> FNodeClassMetadata
+		{
+			FNodeClassMetadata Info;
+			Info.ClassName = { Metasound::StandardNodes::Namespace, TEXT("MIDIToFreq"), Metasound::StandardNodes::AudioVariant };
+			Info.MajorVersion = 1;
+			Info.MinorVersion = 0;
+			Info.DisplayName = LOCTEXT("Metasound_MidiToFreqDisplayName", "MIDIToFreq");
+			Info.Description = LOCTEXT("Metasound_MidiToFreqNodeDescription", "Converts a MIDI note value to a frequency (hz) value.");
+			Info.Author = PluginAuthor;
+			Info.PromptIfMissing = PluginNodeMissingPrompt;
+			Info.DefaultInterface = GetVertexInterface();
+
+			return Info;
+		};
+
+		static const FNodeClassMetadata Info = InitNodeInfo();
+
+		return Info;
+	}
+
+	TUniquePtr<IOperator> FMIDIToFreqOperator::CreateOperator(const FCreateOperatorParams& InParams, FBuildErrorArray& OutErrors)
+	{
+		const FMIDIToFreqNode& MIDIToFreqNode = static_cast<const FMIDIToFreqNode&>(InParams.Node);
+		const FDataReferenceCollection& InputCollection = InParams.InputDataReferences;
+		const FInputVertexInterface& InputInterface = GetVertexInterface().GetInputInterface();
+
+		FInt32ReadRef InMIDINote = InputCollection.GetDataReadReferenceOrConstructWithVertexDefault<int32>(InputInterface, InParamNameMIDI);
+
+		return MakeUnique<FMIDIToFreqOperator>(InParams.OperatorSettings, InMIDINote);
+	}
+
+	FMIDIToFreqNode::FMIDIToFreqNode(const FNodeInitData& InitData)
+		: FNodeFacade(InitData.InstanceName, InitData.InstanceID, TFacadeOperatorClass<FMIDIToFreqOperator>())
+	{
+	}
+
+	METASOUND_REGISTER_NODE(FMIDIToFreqNode)
+}
+
+#undef LOCTEXT_NAMESPACE
diff --git a/Engine/Plugins/Runtime/Metasound/Source/MetasoundStandardNodes/Public/MetasoundMIDIToFreqNode.h b/Engine/Plugins/Runtime/Metasound/Source/MetasoundStandardNodes/Public/MetasoundMIDIToFreqNode.h
new file mode 100644
--- /dev/null
+++ b/Engine/Plugins/Runtime/Metasound/Source/MetasoundStandardNodes/Public/MetasoundMIDIToFreqNode.h
@@ -0,0 +1,23 @@
+// Copyright Epic Games, Inc. All Rights Reserved.
+
+#pragma once
+
+#include "CoreMinimal.h"
+#include "MetasoundFacade.h"
+
+namespace Metasound
+{
+	/** FMIDIToFreqNode
+	 *
+	 *  Outputs the frequency value for a given MIDI note.
+	 */
+	class METASOUNDSTANDARDNODES_API FMIDIToFreqNode : public FNodeFacade
+	{
+	public:
+		/**
+		 * Constructor used by the Metasound Frontend.
+		 */
+		FMIDIToFreqNode(const FNodeInitData& InitData);
+	};
+} // namespace Metasound
+
