Supporting automatic registration of Metasound TArray types when a data type is registered.

diff --git a/Engine/Plugins/Runtime/Metasound/Source/MetasoundEngine/Public/MetasoundWave.h b/Engine/Plugins/Runtime/Metasound/Source/MetasoundEngine/Public/MetasoundWave.h
--- a/Engine/Plugins/Runtime/Metasound/Source/MetasoundEngine/Public/MetasoundWave.h
+++ b/Engine/Plugins/Runtime/Metasound/Source/MetasoundEngine/Public/MetasoundWave.h
@@ -21,9 +21,17 @@ namespace Metasound
 	class METASOUNDENGINE_API FWaveAsset
 	{
 	public:
+		// TODO: Implement using TSharedPtr<> to enable cheaper assignment when 
+		// choosing between multiple instances of FWaveAsset.
 		TUniquePtr<FSoundWaveProxy> SoundWaveProxy;
 
 		FWaveAsset() = default;
+		FWaveAsset(const FWaveAsset&) = delete;
+
+		// Define move constructor and move assignment since
+		// object cannot be copied to due TUniquePtr member.
+		FWaveAsset(FWaveAsset&&) = default;
+		FWaveAsset& operator=(FWaveAsset&&) = default;
 
 		FWaveAsset& operator=(const FWaveAsset& Other)
 		{
@@ -116,4 +124,4 @@ namespace Audio
 	}; // class FSimpleDecoderWrapper
 
 
-} // namespace Audio
\ No newline at end of file
+} // namespace Audio
diff --git a/Engine/Plugins/Runtime/Metasound/Source/MetasoundFrontend/Private/MetasoundFrontend.cpp b/Engine/Plugins/Runtime/Metasound/Source/MetasoundFrontend/Private/MetasoundFrontend.cpp
--- a/Engine/Plugins/Runtime/Metasound/Source/MetasoundFrontend/Private/MetasoundFrontend.cpp
+++ b/Engine/Plugins/Runtime/Metasound/Source/MetasoundFrontend/Private/MetasoundFrontend.cpp
@@ -193,8 +193,6 @@ REGISTER_METASOUND_DATATYPE(float, "Float", ::Metasound::ELiteralType::Float)
 REGISTER_METASOUND_DATATYPE(double, "Double", ::Metasound::ELiteralType::Float)
 REGISTER_METASOUND_DATATYPE(FString, "String", ::Metasound::ELiteralType::String)
 
-REGISTER_METASOUND_DATATYPE(TArray<float>, "FloatArray", ::Metasound::ELiteralType::FloatArray)
-
 REGISTER_METASOUND_DATATYPE(Metasound::FAudioBuffer, "Audio")
 REGISTER_METASOUND_DATATYPE(Metasound::FSendAddress, "Transmission:Address", ::Metasound::ELiteralType::String)
 
diff --git a/Engine/Plugins/Runtime/Metasound/Source/MetasoundFrontend/Private/MetasoundFrontendRegistries.cpp b/Engine/Plugins/Runtime/Metasound/Source/MetasoundFrontend/Private/MetasoundFrontendRegistries.cpp
--- a/Engine/Plugins/Runtime/Metasound/Source/MetasoundFrontend/Private/MetasoundFrontendRegistries.cpp
+++ b/Engine/Plugins/Runtime/Metasound/Source/MetasoundFrontend/Private/MetasoundFrontendRegistries.cpp
@@ -11,20 +11,121 @@
 #define WITH_METASOUND_FRONTEND 0
 #endif
 
-namespace MetasoundFrontendRegistryPrivate
+namespace Metasound
 {
-	// All registry keys should be created through this function to ensure consistency.
-	Metasound::Frontend::FNodeRegistryKey GetRegistryKey(const FName& InClassName, int32 InMajorVersion)
+	namespace MetasoundFrontendRegistryPrivate
 	{
-		Metasound::Frontend::FNodeRegistryKey Key;
+		// All registry keys should be created through this function to ensure consistency.
+		Frontend::FNodeRegistryKey GetRegistryKey(const FName& InClassName, int32 InMajorVersion)
+		{
+			Frontend::FNodeRegistryKey Key;
+
+			Key.NodeClassFullName = InClassName;
 
-		Key.NodeClassFullName = InClassName;
+			// NodeHash is hash of node name and major version.
+			Key.NodeHash = FCrc::StrCrc32(*Key.NodeClassFullName.ToString());
+			Key.NodeHash = HashCombine(Key.NodeHash, FCrc::TypeCrc32(InMajorVersion));
 
-		// NodeHash is hash of node name and major version.
-		Key.NodeHash = FCrc::StrCrc32(*Key.NodeClassFullName.ToString());
-		Key.NodeHash = HashCombine(Key.NodeHash, FCrc::TypeCrc32(InMajorVersion));
+			return Key;
+		}
 
-		return Key;
+		// Return the compatible literal with the most descriptive type.
+		// TODO: Currently TIsParsable<> allows for implicit conversion of
+		// constructor arguments of integral types which can cause some confusion
+		// here when trying to match a literal type to a constructor. For example:
+		//
+		// struct FBoolConstructibleType
+		// {
+		// 	FBoolConstructibleType(bool InValue);
+		// };
+		//
+		// static_assert(TIsParsable<FBoolConstructible, double>::Value); 
+		//
+		// Implicit conversions are currently allowed in TIsParsable because this
+		// is perfectly legal syntax.
+		//
+		// double Value = 10.0;
+		// FBoolConstructibleType BoolConstructible = Value;
+		//
+		// There are some tricks to possibly disable implicit conversions when
+		// checking for specific constructors, but they are yet to be implemented 
+		// and are untested. Here's the basic idea.
+		//
+		// template<DataType, DesiredIntegralArgType>
+		// struct TOnlyConvertIfIsSame
+		// {
+		// 		// Implicit conversion only defined if types match.
+		// 		template<typename SuppliedIntegralArgType, std::enable_if<std::is_same<std::decay<SuppliedIntegralArgType>::type, DesiredIntegralArgType>::value, int> = 0>
+		// 		operator DesiredIntegralArgType()
+		// 		{
+		// 			return DesiredIntegralArgType{};
+		// 		}
+		// };
+		//
+		// static_assert(false == std::is_constructible<FBoolConstructibleType, TOnlyConvertIfSame<double>>::value);
+		// static_assert(true == std::is_constructible<FBoolConstructibleType, TOnlyConvertIfSame<bool>>::value);
+		ELiteralType GetMostDescriptiveLiteralForDataType(const FDataTypeRegistryInfo& InDataTypeInfo)
+		{
+			if (InDataTypeInfo.bIsProxyArrayParsable)
+			{
+				return ELiteralType::UObjectProxyArray;
+			}
+			else if (InDataTypeInfo.bIsProxyParsable)
+			{
+				return ELiteralType::UObjectProxy;
+			}
+			else if (InDataTypeInfo.bIsEnum && InDataTypeInfo.bIsIntParsable)
+			{
+				return ELiteralType::Integer;
+			}
+			else if (InDataTypeInfo.bIsStringArrayParsable)
+			{
+				return ELiteralType::StringArray;
+			}
+			else if (InDataTypeInfo.bIsFloatArrayParsable)
+			{
+				return ELiteralType::FloatArray;
+			}
+			else if (InDataTypeInfo.bIsIntArrayParsable)
+			{
+				return ELiteralType::IntegerArray;
+			}
+			else if (InDataTypeInfo.bIsBoolArrayParsable)
+			{
+				return ELiteralType::BooleanArray;
+			}
+			else if (InDataTypeInfo.bIsStringParsable)
+			{
+				return ELiteralType::String;
+			}
+			else if (InDataTypeInfo.bIsFloatParsable)
+			{
+				return ELiteralType::Float;
+			}
+			else if (InDataTypeInfo.bIsIntParsable)
+			{
+				return ELiteralType::Integer;
+			}
+			else if (InDataTypeInfo.bIsBoolParsable)
+			{
+				return ELiteralType::Boolean;
+			}
+			else if (InDataTypeInfo.bIsDefaultArrayParsable)
+			{
+				return ELiteralType::NoneArray; 
+			}
+			else if (InDataTypeInfo.bIsDefaultParsable)
+			{
+				return ELiteralType::None;
+			}
+			else
+			{
+				// if we ever hit this, something has gone wrong with the REGISTER_METASOUND_DATATYPE macro.
+				// we should have failed to compile if any of these are false.
+				checkNoEntry();
+				return ELiteralType::Invalid;
+			}
+		}
 	}
 }
 
@@ -200,6 +301,8 @@ TArray<::Metasound::Frontend::FConverterNodeInfo> FMetasoundFrontendRegistryCont
 	}
 }
 
+
+
 Metasound::ELiteralType FMetasoundFrontendRegistryContainer::GetDesiredLiteralTypeForDataType(FName InDataType) const
 {
 	if (!DataTypeRegistry.Contains(InDataType))
@@ -208,46 +311,15 @@ Metasound::ELiteralType FMetasoundFrontendRegistryContainer::GetDesiredLiteralTy
 	}
 
 	const FDataTypeRegistryElement& DataTypeInfo = DataTypeRegistry[InDataType];
-	
-	if (DataTypeInfo.Info.bIsEnum)
-	{
-		return Metasound::ELiteralType::Integer;
-	}
 
 	// If there's a designated preferred literal type for this datatype, use that.
-	if (DataTypeInfo.Info.PreferredLiteralType != Metasound::ELiteralType::None)
+	if (DataTypeInfo.Info.PreferredLiteralType != Metasound::ELiteralType::Invalid)
 	{
 		return DataTypeInfo.Info.PreferredLiteralType;
 	}
 
 	// Otherwise, we opt for the highest precision construction option available.
-	if (DataTypeInfo.Info.bIsStringParsable)
-	{
-		return Metasound::ELiteralType::String;
-	}
-	else if (DataTypeInfo.Info.bIsFloatParsable)
-	{
-		return Metasound::ELiteralType::Float;
-	}
-	else if (DataTypeInfo.Info.bIsIntParsable)
-	{
-		return Metasound::ELiteralType::Integer;
-	}
-	else if (DataTypeInfo.Info.bIsBoolParsable)
-	{
-		return Metasound::ELiteralType::Boolean;
-	}
-	else if (DataTypeInfo.Info.bIsDefaultParsable)
-	{
-		return Metasound::ELiteralType::None;
-	}
-	else
-	{
-		// if we ever hit this, something has gone terribly wrong with the REGISTER_METASOUND_DATATYPE macro.
-		// we should have failed to compile if any of these are false.
-		checkNoEntry();
-		return Metasound::ELiteralType::Invalid;
-	}
+	return Metasound::MetasoundFrontendRegistryPrivate::GetMostDescriptiveLiteralForDataType(DataTypeInfo.Info);
 }
 
 UClass* FMetasoundFrontendRegistryContainer::GetLiteralUClassForDataType(FName InDataType) const
@@ -383,12 +455,12 @@ bool FMetasoundFrontendRegistryContainer::GetRegistryKey(const Metasound::Fronte
 
 Metasound::Frontend::FNodeRegistryKey FMetasoundFrontendRegistryContainer::GetRegistryKey(const FNodeClassMetadata& InNodeMetadata)
 {
-	return MetasoundFrontendRegistryPrivate::GetRegistryKey(InNodeMetadata.ClassName.GetFullName(), InNodeMetadata.MajorVersion);
+	return Metasound::MetasoundFrontendRegistryPrivate::GetRegistryKey(InNodeMetadata.ClassName.GetFullName(), InNodeMetadata.MajorVersion);
 }
 
 Metasound::Frontend::FNodeRegistryKey FMetasoundFrontendRegistryContainer::GetRegistryKey(const FMetasoundFrontendClassMetadata& InNodeMetadata)
 {
-	return MetasoundFrontendRegistryPrivate::GetRegistryKey(InNodeMetadata.ClassName.GetFullName(), InNodeMetadata.Version.Major);
+	return Metasound::MetasoundFrontendRegistryPrivate::GetRegistryKey(InNodeMetadata.ClassName.GetFullName(), InNodeMetadata.Version.Major);
 }
 
 bool FMetasoundFrontendRegistryContainer::GetFrontendClassFromRegistered(const FMetasoundFrontendClassMetadata& InMetadata, FMetasoundFrontendClass& OutClass)
diff --git a/Engine/Plugins/Runtime/Metasound/Source/MetasoundFrontend/Public/MetasoundDataTypeRegistrationMacro.h b/Engine/Plugins/Runtime/Metasound/Source/MetasoundFrontend/Public/MetasoundDataTypeRegistrationMacro.h
--- a/Engine/Plugins/Runtime/Metasound/Source/MetasoundFrontend/Public/MetasoundDataTypeRegistrationMacro.h
+++ b/Engine/Plugins/Runtime/Metasound/Source/MetasoundFrontend/Public/MetasoundDataTypeRegistrationMacro.h
@@ -9,8 +9,10 @@
 #include "MetasoundConverterNodeRegistrationMacro.h"
 #include "MetasoundDataFactory.h"
 #include "MetasoundDataReference.h"
+#include "MetasoundEnum.h"
 #include "MetasoundFrontendRegistries.h"
 #include "MetasoundInputNode.h"
+#include "MetasoundLiteral.h"
 #include "MetasoundNodeRegistrationMacro.h"
 #include "MetasoundOperatorInterface.h"
 #include "MetasoundOutputNode.h"
@@ -18,276 +20,382 @@
 #include "MetasoundReceiveNode.h"
 #include "MetasoundRouter.h"
 #include "MetasoundSendNode.h"
-#include "MetasoundEnum.h"
 
 #include <type_traits>
 
 namespace Metasound
 {
-	// SFINAE used to optionally invoke subclasses of IAudioProxyDataFactory when we can.
-	template<typename UClassToUse, typename TEnableIf<TIsDerivedFrom<UClassToUse, IAudioProxyDataFactory>::Value, bool>::Type = true>
-	IAudioProxyDataFactory* CastToAudioProxyDataFactory(UObject* InObject)
+	namespace MetasoundDataTypeRegistrationPrivate
 	{
-		UClassToUse* DowncastObject = Cast<UClassToUse>(InObject);
-		if (ensureAlways(DowncastObject))
+		// Returns the Array version of a literal type if it exists.
+		template<ELiteralType LiteralType>
+		struct TLiteralArrayEnum 
+		{
+			// Default to TArray default constructor by using
+			// ELiteralType::None
+			static constexpr ELiteralType Value = ELiteralType::None;
+		};
+
+		// Specialization for None->NoneArray
+		template<>
+		struct TLiteralArrayEnum<ELiteralType::None>
+		{
+			static constexpr ELiteralType Value = ELiteralType::NoneArray;
+		};
+		 
+		// Specialization for Boolean->BooleanArray
+		template<>
+		struct TLiteralArrayEnum<ELiteralType::Boolean>
+		{
+			static constexpr ELiteralType Value = ELiteralType::BooleanArray;
+		};
+		
+		// Specialization for Integer->IntegerArray
+		template<>
+		struct TLiteralArrayEnum<ELiteralType::Integer>
+		{
+			static constexpr ELiteralType Value = ELiteralType::IntegerArray;
+		};
+		
+		// Specialization for Float->FloatArray
+		template<>
+		struct TLiteralArrayEnum<ELiteralType::Float>
+		{
+			static constexpr ELiteralType Value = ELiteralType::FloatArray;
+		};
+		
+		// Specialization for String->StringArray
+		template<>
+		struct TLiteralArrayEnum<ELiteralType::String>
+		{
+			static constexpr ELiteralType Value = ELiteralType::StringArray;
+		};
+		
+		// Specialization for UObjectProxy->UObjectProxyArray
+		template<>
+		struct TLiteralArrayEnum<ELiteralType::UObjectProxy>
 		{
-			return static_cast<IAudioProxyDataFactory*>(DowncastObject);
+			static constexpr ELiteralType Value = ELiteralType::UObjectProxyArray;
+		};
+
+		// SFINAE used to optionally invoke subclasses of IAudioProxyDataFactory when we can.
+		template<typename UClassToUse, typename TEnableIf<TIsDerivedFrom<UClassToUse, IAudioProxyDataFactory>::Value, bool>::Type = true>
+		IAudioProxyDataFactory* CastToAudioProxyDataFactory(UObject* InObject)
+		{
+			UClassToUse* DowncastObject = Cast<UClassToUse>(InObject);
+			if (ensureAlways(DowncastObject))
+			{
+				return static_cast<IAudioProxyDataFactory*>(DowncastObject);
+			}
+			else
+			{
+				return nullptr;
+			}
 		}
-		else
+
+		template<typename UClassToUse, typename TEnableIf<!TIsDerivedFrom<UClassToUse, IAudioProxyDataFactory>::Value, bool>::Type = true>
+		IAudioProxyDataFactory* CastToAudioProxyDataFactory(UObject* InObject)
 		{
 			return nullptr;
 		}
-	}
-
-	template<typename UClassToUse, typename TEnableIf<!TIsDerivedFrom<UClassToUse, IAudioProxyDataFactory>::Value, bool>::Type = true>
-	IAudioProxyDataFactory* CastToAudioProxyDataFactory(UObject* InObject)
-	{
-		return nullptr;
-	}
-	
-	// Helper utility to test if we can transmit a datatype between a send and a receive node.
-	template <typename TDataType>
-	struct TIsTransmittable
-	{
-	private:
-		static constexpr bool bIsCopyConstructible = std::is_copy_constructible<TDataType>::value;
-		static constexpr bool bIsCopyAssignable = std::is_copy_assignable<TDataType>::value;
-
-		// TODO: audio types were intended to be send/receive nodes but they require 
-		// template specialization.  TIsTransmittable should ask the Send and Receive nodes
-		// if they are transmittable rather than attempting to do the logic for all
-		// types here. 
-		//static constexpr bool bIsAudioDataType = TIsDerivedFrom<TDataType, IAudioDataType>::Value;
+		
+		// Helper utility to test if we can transmit a datatype between a send and a receive node.
+		template <typename TDataType>
+		struct TIsTransmittable
+		{
+		private:
+			static constexpr bool bIsCopyConstructible = std::is_copy_constructible<TDataType>::value;
+			static constexpr bool bIsCopyAssignable = std::is_copy_assignable<TDataType>::value;
 
-		static constexpr bool bCanBeTransmitted = bIsCopyConstructible && bIsCopyAssignable;
+			// TODO: audio types were intended to be send/receive nodes but they require 
+			// template specialization.  TIsTransmittable should ask the Send and Receive nodes
+			// if they are transmittable rather than attempting to do the logic for all
+			// types here. 
+			//static constexpr bool bIsAudioDataType = TIsDerivedFrom<TDataType, IAudioDataType>::Value;
 
-	public:
+			static constexpr bool bCanBeTransmitted = bIsCopyConstructible && bIsCopyAssignable;
 
-		static constexpr bool Value = bCanBeTransmitted;
-	};
+		public:
 
-	// This utility function can be used to optionally check to see if we can transmit a data type, and autogenerate send and receive nodes for that datatype.
-	template<typename TDataType, typename TEnableIf<TIsTransmittable<TDataType>::Value, bool>::Type = true>
-	void AttemptToRegisterSendAndReceiveNodes()
-	{
-		ensureAlways(RegisterNodeWithFrontend<Metasound::TSendNode<TDataType>>());
-		ensureAlways(RegisterNodeWithFrontend<Metasound::TReceiveNode<TDataType>>());
-	}
+			static constexpr bool Value = bCanBeTransmitted;
+		};
 
-	template<typename TDataType, typename TEnableIf<!TIsTransmittable<TDataType>::Value, bool>::Type = true>
-	void AttemptToRegisterSendAndReceiveNodes()
-	{
-		// This implementation intentionally noops, because Metasound::TIsTransmittable is false for this datatype.
-		// This is either because the datatype is not trivially copyable, and thus can't be buffered between threads,
-		// or it's not an audio buffer type, which we use Audio::FPatchMixerSplitter instances for.
-	}
+		// Specialization of TIsTransmittable<> for TArray to handle lax TArray copy constructor 
+		// definition. This can be removed once updates to TArray are merged into the same codebase.
+		// TODO: delete me eventually.
+		template<typename TElementType>
+		struct TIsTransmittable<TArray<TElementType>>
+		{
+		private:
+			// Depend on copy constructor of elements to determine whether TArray is copy constructible
+			// and copy assignable. Definitions for TArray copy construction and copy assignment are
+			// defined by default whether or not the TArray element types support them. This logic
+			// skips checking the TArray object itselt and instead checks the TArray element type
+			// directly.
+			static constexpr bool bIsElementCopyConstructible = std::is_copy_constructible<TElementType>::value;
 
+			static constexpr bool bCanBeTransmitted = bIsElementCopyConstructible;
 
+		public:
 
-	// This utility function can be used to check to see if we can static cast between two types, and autogenerate a node for that static cast.
-	template<typename TFromDataType, typename TToDataType, typename std::enable_if<std::is_convertible<TFromDataType, TToDataType>::value, int>::type = 0>
-	void AttemptToRegisterConverter()
-	{
-		using FConverterNode = Metasound::TAutoConverterNode<TFromDataType, TToDataType>;
+			static constexpr bool Value = bCanBeTransmitted;
+		};
 
-		const FNodeClassMetadata& Metadata = FConverterNode::GetAutoConverterNodeMetadata();
-		const Metasound::Frontend::FNodeRegistryKey Key = FMetasoundFrontendRegistryContainer::GetRegistryKey(Metadata);
+		// This utility function can be used to optionally check to see if we can transmit a data type, and autogenerate send and receive nodes for that datatype.
+		template<typename TDataType, typename TEnableIf<TIsTransmittable<TDataType>::Value, bool>::Type = true>
+		void AttemptToRegisterSendAndReceiveNodes()
+		{
+			ensureAlways(RegisterNodeWithFrontend<Metasound::TSendNode<TDataType>>());
+			ensureAlways(RegisterNodeWithFrontend<Metasound::TReceiveNode<TDataType>>());
+		}
 
-		if (!std::is_same<TFromDataType, TToDataType>::value && !FMetasoundFrontendRegistryContainer::Get()->IsNodeRegistered(Key))
+		template<typename TDataType, typename TEnableIf<!TIsTransmittable<TDataType>::Value, bool>::Type = true>
+		void AttemptToRegisterSendAndReceiveNodes()
 		{
-			ensureAlways(RegisterNodeWithFrontend<FConverterNode>(Metadata));
-			
-			bool bSucessfullyRegisteredConversionNode = RegisterConversionNode<FConverterNode, TFromDataType, TToDataType>(FConverterNode::GetInputName(), FConverterNode::GetOutputName(), Metadata);
-			ensureAlways(bSucessfullyRegisteredConversionNode);
+			// This implementation intentionally noops, because Metasound::TIsTransmittable is false for this datatype.
+			// This is either because the datatype is not trivially copyable, and thus can't be buffered between threads,
+			// or it's not an audio buffer type, which we use Audio::FPatchMixerSplitter instances for.
 		}
-	}
 
-	template<typename TFromDataType, typename TToDataType, typename std::enable_if<!std::is_convertible<TFromDataType, TToDataType>::value, int>::type = 0>
-	void AttemptToRegisterConverter()
-	{
-		// This implementation intentionally noops, because static_cast<TFromDataType>(TToDataType&) is invalid.
-	}
+		// This utility function can be used to check to see if we can static cast between two types, and autogenerate a node for that static cast.
+		template<typename TFromDataType, typename TToDataType, typename std::enable_if<std::is_convertible<TFromDataType, TToDataType>::value, int>::type = 0>
+		void AttemptToRegisterConverter()
+		{
+			using FConverterNode = Metasound::TAutoConverterNode<TFromDataType, TToDataType>;
 
-	// Here we attempt to infer and autogenerate conversions for basic datatypes.
-	template<typename TDataType>
-	void RegisterConverterNodes()
-	{
-		// Conversions to this data type:
-		AttemptToRegisterConverter<bool, TDataType>();
-		AttemptToRegisterConverter<int32, TDataType>();
-		AttemptToRegisterConverter<float, TDataType>();
-		AttemptToRegisterConverter<FString, TDataType>();
-
-		// Conversions from this data type:
-		AttemptToRegisterConverter<TDataType, bool>();
-		AttemptToRegisterConverter<TDataType, int32>();
-		AttemptToRegisterConverter<TDataType, float>();
-		AttemptToRegisterConverter<TDataType, FString>();
-	}
-	
-	// SFINAE for enum types.
-	template<typename TDataType, typename std::enable_if<TEnumTraits<TDataType>::bIsEnum, bool>::type = true>
-	bool RegisterEnumDataTypeWithFrontend()
-	{
-		using InnerType = typename TEnumTraits<TDataType>::InnerType;
-		using FStringHelper = TEnumStringHelper<InnerType>;
+			const FNodeClassMetadata& Metadata = FConverterNode::GetAutoConverterNodeMetadata();
+			const Metasound::Frontend::FNodeRegistryKey Key = FMetasoundFrontendRegistryContainer::GetRegistryKey(Metadata);
 
-		struct FEnumHandler : Metasound::Frontend::IEnumDataTypeInterface
-		{
-			FName GetNamespace() const override
+			if (!std::is_same<TFromDataType, TToDataType>::value && !FMetasoundFrontendRegistryContainer::Get()->IsNodeRegistered(Key))
 			{
-				return FStringHelper::GetNamespace();
+				ensureAlways(RegisterNodeWithFrontend<FConverterNode>(Metadata));
+				
+				bool bSucessfullyRegisteredConversionNode = RegisterConversionNode<FConverterNode, TFromDataType, TToDataType>(FConverterNode::GetInputName(), FConverterNode::GetOutputName(), Metadata);
+				ensureAlways(bSucessfullyRegisteredConversionNode);
 			}
-			int32 GetDefaultValue() const override
-			{
-				return static_cast<int32>(TEnumTraits<TDataType>::DefaultValue);
-			}
-			const TArray<FGenericInt32Entry>& GetAllEntries() const override
+		}
+
+		template<typename TFromDataType, typename TToDataType, typename std::enable_if<!std::is_convertible<TFromDataType, TToDataType>::value, int>::type = 0>
+		void AttemptToRegisterConverter()
+		{
+			// This implementation intentionally noops, because static_cast<TFromDataType>(TToDataType&) is invalid.
+		}
+
+		// Here we attempt to infer and autogenerate conversions for basic datatypes.
+		template<typename TDataType>
+		void RegisterConverterNodes()
+		{
+			// Conversions to this data type:
+			AttemptToRegisterConverter<bool, TDataType>();
+			AttemptToRegisterConverter<int32, TDataType>();
+			AttemptToRegisterConverter<float, TDataType>();
+			AttemptToRegisterConverter<FString, TDataType>();
+
+			// Conversions from this data type:
+			AttemptToRegisterConverter<TDataType, bool>();
+			AttemptToRegisterConverter<TDataType, int32>();
+			AttemptToRegisterConverter<TDataType, float>();
+			AttemptToRegisterConverter<TDataType, FString>();
+		}
+
+		// SFINAE for enum types.
+		template<typename TDataType, typename std::enable_if<TEnumTraits<TDataType>::bIsEnum, bool>::type = true>
+		bool RegisterEnumDataTypeWithFrontend()
+		{
+			using InnerType = typename TEnumTraits<TDataType>::InnerType;
+			using FStringHelper = TEnumStringHelper<InnerType>;
+
+			struct FEnumHandler : Metasound::Frontend::IEnumDataTypeInterface
 			{
-				auto BuildIntEntries = []()
+				FName GetNamespace() const override
+				{
+					return FStringHelper::GetNamespace();
+				}
+				int32 GetDefaultValue() const override
 				{
-					// Convert to int32 representation 
-					TArray<FGenericInt32Entry> IntEntries;
-					IntEntries.Reserve(FStringHelper::GetAllEntries().Num());
-					for (const TEnumEntry<InnerType>& i : FStringHelper::GetAllEntries())
+					return static_cast<int32>(TEnumTraits<TDataType>::DefaultValue);
+				}
+				const TArray<FGenericInt32Entry>& GetAllEntries() const override
+				{
+					auto BuildIntEntries = []()
 					{
-						IntEntries.Emplace(i);
-					}
+						// Convert to int32 representation 
+						TArray<FGenericInt32Entry> IntEntries;
+						IntEntries.Reserve(FStringHelper::GetAllEntries().Num());
+						for (const TEnumEntry<InnerType>& i : FStringHelper::GetAllEntries())
+						{
+							IntEntries.Emplace(i);
+						}
+						return IntEntries;
+					};
+					static const TArray<FGenericInt32Entry> IntEntries = BuildIntEntries();
 					return IntEntries;
-				};
-				static const TArray<FGenericInt32Entry> IntEntries = BuildIntEntries();
-				return IntEntries;
-			}
-		};
+				}
+			};
 
-		return FMetasoundFrontendRegistryContainer::Get()->RegisterEnumDataInterface(
-			GetMetasoundDataTypeName<TDataType>(), MakeShared<FEnumHandler>());
-	}
-	
-	// SFINAE stub for non-enum types. 
-	template<typename TDataType, typename std::enable_if<!TEnumTraits<TDataType>::bIsEnum, bool>::type = true>
-	bool RegisterEnumDataTypeWithFrontend() { return false; }
-	
-	template<typename TDataType, ELiteralType PreferredArgType = ELiteralType::None, typename UClassToUse = UObject>
-	bool RegisterDataTypeWithFrontend()
-	{
-		// if we reenter this code (because DECLARE_METASOUND_DATA_REFERENCE_TYPES was called twice with the same type),
-		// we catch it here.
-		static bool bAlreadyRegisteredThisDataType = false;
-		if (bAlreadyRegisteredThisDataType)
-		{
-			UE_LOG(LogTemp, Display, TEXT("Tried to call REGISTER_METASOUND_DATATYPE twice with the same class %s. ignoring the second call. Likely because REGISTER_METASOUND_DATATYPE is in a header that's used in multiple modules. Consider moving it to a private header or cpp file."), TDataReferenceTypeInfo<TDataType>::TypeName)
-			return false;
+			return FMetasoundFrontendRegistryContainer::Get()->RegisterEnumDataInterface(
+				GetMetasoundDataTypeName<TDataType>(), MakeShared<FEnumHandler>());
 		}
+		
+		// SFINAE stub for non-enum types. 
+		template<typename TDataType, typename std::enable_if<!TEnumTraits<TDataType>::bIsEnum, bool>::type = true>
+		bool RegisterEnumDataTypeWithFrontend() { return false; }
 
-		bAlreadyRegisteredThisDataType = true;
 
-		// Lambdas that generate our template-instantiated input and output nodes:
-		FCreateInputNodeFunction InputNodeConstructor = [](FInputNodeConstructorParams&& InParams) -> TUniquePtr<INode>
-		{
-			return TUniquePtr<INode>(new TInputNode<TDataType>(InParams.InNodeName, InParams.InInstanceID, InParams.InVertexName, MoveTemp(InParams.InitParam)));
-		};
 
-		FCreateMetasoundFrontendClassFunction CreateFrontendInputClass = []() -> FMetasoundFrontendClass
+		template<typename TDataType, ELiteralType PreferredArgType = ELiteralType::None, typename UClassToUse = UObject>
+		bool RegisterDataTypeWithFrontendInternal()
 		{
-			// Create class info using prototype node
-			// TODO: register input nodes with static class info.
-  			static TInputNode<TDataType> Prototype(TEXT(""), FGuid(), TEXT(""), FLiteral());
-			return Metasound::Frontend::GenerateClassDescription(Prototype.GetMetadata(), EMetasoundFrontendClassType::Input);
-		};
+			// if we reenter this code (because DECLARE_METASOUND_DATA_REFERENCE_TYPES was called twice with the same type),
+			// we catch it here.
+			static bool bAlreadyRegisteredThisDataType = false;
+			if (bAlreadyRegisteredThisDataType)
+			{
+				UE_LOG(LogTemp, Display, TEXT("Tried to call REGISTER_METASOUND_DATATYPE twice with the same class %s. ignoring the second call. Likely because REGISTER_METASOUND_DATATYPE is in a header that's used in multiple modules. Consider moving it to a private header or cpp file."), TDataReferenceTypeInfo<TDataType>::TypeName)
+				return false;
+			}
 
-		FCreateOutputNodeFunction OutputNodeConstructor = [](const FOutputNodeConstructorParams& InParams) -> TUniquePtr<INode>
-		{
-			return TUniquePtr<INode>(new TOutputNode<TDataType>(InParams.InNodeName, InParams.InInstanceID, InParams.InVertexName));
-		};
+			bAlreadyRegisteredThisDataType = true;
 
-		FCreateMetasoundFrontendClassFunction CreateFrontendOutputClass = []() -> FMetasoundFrontendClass
-		{
-			// Create class info using prototype node
-			// TODO: register input nodes with static class info.
-			static TOutputNode<TDataType> Prototype(TEXT(""), FGuid(), TEXT(""));
-			return Metasound::Frontend::GenerateClassDescription(Prototype.GetMetadata(), EMetasoundFrontendClassType::Output);
-		};
-		
-		// By default, this function should not be used, unless the preferred arg type is UObjectProxy or UObjectProxyArray, and UClassToUse should be specified.
-		FCreateAudioProxyFunction ProxyGenerator = [](UObject* InObject) -> Audio::IProxyDataPtr
-		{
-			checkNoEntry();
-			return Audio::IProxyDataPtr(nullptr);
-		};
+			// Lambdas that generate our template-instantiated input and output nodes:
+			FCreateInputNodeFunction InputNodeConstructor = [](FInputNodeConstructorParams&& InParams) -> TUniquePtr<INode>
+			{
+				return TUniquePtr<INode>(new TInputNode<TDataType>(InParams.InNodeName, InParams.InInstanceID, InParams.InVertexName, MoveTemp(InParams.InitParam)));
+			};
 
-		// If this datatype uses a UObject or UObject array literal, we generate a lambda to build a proxy here:
-		constexpr bool bSpecifiedUClassForProxy = !TIsSame<UClassToUse, UObject>::Value;
-		if (bSpecifiedUClassForProxy)
-		{
-			static_assert(!bSpecifiedUClassForProxy || std::is_base_of<IAudioProxyDataFactory, UClassToUse>::value, "If a Metasound Datatype uses a UObject as a literal, the UClass of that object needs to also derive from Audio::IProxyDataFactory. See USoundWave as an example.");
-			ProxyGenerator = [](UObject* InObject) -> Audio::IProxyDataPtr
+			FCreateMetasoundFrontendClassFunction CreateFrontendInputClass = []() -> FMetasoundFrontendClass
 			{
-				IAudioProxyDataFactory* ObjectAsFactory = CastToAudioProxyDataFactory<UClassToUse>(InObject);
+				// Create class info using prototype node
+				// TODO: register input nodes with static class info.
+				static TInputNode<TDataType> Prototype(TEXT(""), FGuid(), TEXT(""), FLiteral());
+				return Metasound::Frontend::GenerateClassDescription(Prototype.GetMetadata(), EMetasoundFrontendClassType::Input);
+			};
 
-				if (ensureAlways(ObjectAsFactory))
-				{
-					static FName ProxySubsystemName = TEXT("Metasound");
-					
-					Audio::FProxyDataInitParams ProxyInitParams;
-					ProxyInitParams.NameOfFeatureRequestingProxy = ProxySubsystemName;
+			FCreateOutputNodeFunction OutputNodeConstructor = [](const FOutputNodeConstructorParams& InParams) -> TUniquePtr<INode>
+			{
+				return TUniquePtr<INode>(new TOutputNode<TDataType>(InParams.InNodeName, InParams.InInstanceID, InParams.InVertexName));
+			};
 
-					return ObjectAsFactory->CreateNewProxyData(ProxyInitParams);
-				}
-				else
+			FCreateMetasoundFrontendClassFunction CreateFrontendOutputClass = []() -> FMetasoundFrontendClass
+			{
+				// Create class info using prototype node
+				// TODO: register input nodes with static class info.
+				static TOutputNode<TDataType> Prototype(TEXT(""), FGuid(), TEXT(""));
+				return Metasound::Frontend::GenerateClassDescription(Prototype.GetMetadata(), EMetasoundFrontendClassType::Output);
+			};
+			
+			// By default, this function should not be used, unless the preferred arg type is UObjectProxy or UObjectProxyArray, and UClassToUse should be specified.
+			FCreateAudioProxyFunction ProxyGenerator = [](UObject* InObject) -> Audio::IProxyDataPtr
+			{
+				checkNoEntry();
+				return Audio::IProxyDataPtr(nullptr);
+			};
+
+			// If this datatype uses a UObject or UObject array literal, we generate a lambda to build a proxy here:
+			constexpr bool bSpecifiedUClassForProxy = !TIsSame<UClassToUse, UObject>::Value;
+			if (bSpecifiedUClassForProxy)
+			{
+				static_assert(!bSpecifiedUClassForProxy || std::is_base_of<IAudioProxyDataFactory, UClassToUse>::value, "If a Metasound Datatype uses a UObject as a literal, the UClass of that object needs to also derive from Audio::IProxyDataFactory. See USoundWave as an example.");
+				ProxyGenerator = [](UObject* InObject) -> Audio::IProxyDataPtr
 				{
-					return Audio::IProxyDataPtr(nullptr);
-				}
+					IAudioProxyDataFactory* ObjectAsFactory = CastToAudioProxyDataFactory<UClassToUse>(InObject);
+
+					if (ensureAlways(ObjectAsFactory))
+					{
+						static FName ProxySubsystemName = TEXT("Metasound");
+						
+						Audio::FProxyDataInitParams ProxyInitParams;
+						ProxyInitParams.NameOfFeatureRequestingProxy = ProxySubsystemName;
+
+						return ObjectAsFactory->CreateNewProxyData(ProxyInitParams);
+					}
+					else
+					{
+						return Audio::IProxyDataPtr(nullptr);
+					}
+				};
+			}
+
+			static const FName DataTypeName = GetMetasoundDataTypeName<TDataType>();
+
+			// Pack all of our various constructor lambdas to a single struct.
+			FDataTypeConstructorCallbacks Callbacks = 
+			{ 
+				MoveTemp(InputNodeConstructor), 
+				MoveTemp(CreateFrontendInputClass),
+				MoveTemp(OutputNodeConstructor), 
+				MoveTemp(CreateFrontendOutputClass),
+				MoveTemp(ProxyGenerator) 
 			};
-		}
 
-		static const FName DataTypeName = GetMetasoundDataTypeName<TDataType>();
 
-		// Pack all of our various constructor lambdas to a single struct.
-		FDataTypeConstructorCallbacks Callbacks = 
-		{ 
-			MoveTemp(InputNodeConstructor), 
-			MoveTemp(CreateFrontendInputClass),
-			MoveTemp(OutputNodeConstructor), 
-			MoveTemp(CreateFrontendOutputClass),
-			MoveTemp(ProxyGenerator) 
-		};
+			FDataTypeRegistryInfo RegistryInfo;
 
+			RegistryInfo.DataTypeName = DataTypeName;
+			RegistryInfo.PreferredLiteralType = PreferredArgType;
 
-		FDataTypeRegistryInfo RegistryInfo;
+			RegistryInfo.bIsDefaultParsable = TIsParsable<TDataType, FLiteral::FNone>::Value;		
+			RegistryInfo.bIsBoolParsable = TIsParsable<TDataType, bool>::Value;
+			RegistryInfo.bIsIntParsable = TIsParsable<TDataType, int32>::Value;
+			RegistryInfo.bIsFloatParsable = TIsParsable<TDataType, float>::Value;
+			RegistryInfo.bIsStringParsable = TIsParsable<TDataType, FString>::Value;
+			RegistryInfo.bIsProxyParsable = TIsParsable<TDataType, const Audio::IProxyDataPtr&>::Value;
 
-		RegistryInfo.DataTypeName = DataTypeName;
-		RegistryInfo.PreferredLiteralType = PreferredArgType;
+			RegistryInfo.bIsDefaultArrayParsable = TIsParsable<TDataType, TArray<FLiteral::FNone>>::Value;		
+			RegistryInfo.bIsBoolArrayParsable = TIsParsable<TDataType, TArray<bool>>::Value;
+			RegistryInfo.bIsIntArrayParsable = TIsParsable<TDataType, TArray<int32>>::Value;
+			RegistryInfo.bIsFloatArrayParsable = TIsParsable<TDataType, TArray<float>>::Value;
+			RegistryInfo.bIsStringArrayParsable = TIsParsable<TDataType, TArray<FString>>::Value;
 
-		RegistryInfo.bIsBoolParsable = TIsParsable<TDataType, bool>::Value;
-		RegistryInfo.bIsIntParsable = TIsParsable<TDataType, int32>::Value;
-		RegistryInfo.bIsFloatParsable = TIsParsable<TDataType, float>::Value;
-		RegistryInfo.bIsStringParsable = TIsParsable<TDataType, FString>::Value;
-		RegistryInfo.bIsProxyParsable = TIsParsable<TDataType, const Audio::IProxyDataPtr&>::Value;
-		RegistryInfo.bIsProxyArrayParsable = TIsParsable<TDataType, const TArray<Audio::IProxyDataPtr>& >::Value;
-		RegistryInfo.bIsDefaultParsable = TIsParsable<TDataType>::Value;		
-		RegistryInfo.bIsEnum = TEnumTraits<TDataType>::bIsEnum;
+			RegistryInfo.bIsProxyArrayParsable = TIsParsable<TDataType, const TArray<Audio::IProxyDataPtr>& >::Value;
+			RegistryInfo.bIsEnum = TEnumTraits<TDataType>::bIsEnum;
 
-		RegistryInfo.bIsTransmittable = TIsTransmittable<TDataType>::Value;
-		
-		RegistryInfo.ProxyGeneratorClass = UClassToUse::StaticClass();
+			RegistryInfo.bIsTransmittable = TIsTransmittable<TDataType>::Value;
+			
+			RegistryInfo.ProxyGeneratorClass = UClassToUse::StaticClass();
 
-		bool bSucceeded = FMetasoundFrontendRegistryContainer::Get()->RegisterDataType(RegistryInfo, MoveTemp(Callbacks));
-		ensureAlwaysMsgf(bSucceeded, TEXT("Failed to register data type %s in the node registry!"), *GetMetasoundDataTypeString<TDataType>());
+			bool bSucceeded = FMetasoundFrontendRegistryContainer::Get()->RegisterDataType(RegistryInfo, MoveTemp(Callbacks));
+			ensureAlwaysMsgf(bSucceeded, TEXT("Failed to register data type %s in the node registry!"), *GetMetasoundDataTypeString<TDataType>());
 
-		// If its an enum, register its data interface AFTER we've registered the data type.
-		if (TEnumTraits<TDataType>::bIsEnum)
-		{
-			if (!ensure(RegisterEnumDataTypeWithFrontend<TDataType>()))
+			// If its an enum, register its data interface AFTER we've registered the data type.
+			if (TEnumTraits<TDataType>::bIsEnum)
 			{
-				return false;
+				if (!ensure(RegisterEnumDataTypeWithFrontend<TDataType>()))
+				{
+					return false;
+				}
 			}
+
+			RegisterConverterNodes<TDataType>();
+			AttemptToRegisterSendAndReceiveNodes<TDataType>();
+			
+			return bSucceeded;
 		}
 
-		RegisterConverterNodes<TDataType>();
-		AttemptToRegisterSendAndReceiveNodes<TDataType>();
-		
-		return bSucceeded;
+		template<typename TDataType, ELiteralType PreferredArgType, typename UClassToUse>
+		bool RegisterDataTypeArrayWithFrontend()
+		{
+			using namespace MetasoundDataTypeRegistrationPrivate;
+
+			return RegisterDataTypeWithFrontendInternal<TArray<TDataType>, TLiteralArrayEnum<PreferredArgType>::Value, UClassToUse>();
+		}
+	}
+	
+	template<typename TDataType, ELiteralType PreferredArgType = ELiteralType::None, typename UClassToUse = UObject>
+	bool RegisterDataTypeWithFrontend()
+	{
+		using namespace MetasoundDataTypeRegistrationPrivate;
+
+
+		// Register TDataType as a metasound data type.
+		bool bSuccess = RegisterDataTypeWithFrontendInternal<TDataType, PreferredArgType, UClassToUse>();
+		ensure(bSuccess);
+
+		// Register TArray<TDataType> as a metasound data type.
+		bSuccess = bSuccess && RegisterDataTypeArrayWithFrontend<TDataType, PreferredArgType, UClassToUse>();
+		ensure(bSuccess);
+
+		return bSuccess;
 	}
 
 
@@ -314,6 +422,13 @@ namespace Metasound
 
 #define CANNOT_REGISTER_METASOUND_DATA_TYPE_ASSERT_STRING(DataType) \
 "To register " #DataType " to be used as a Metasounds input or output type, it needs a default constructor or one of the following constructors must be implemented:  " \
+#DataType "(), " \
+#DataType "(bool InValue), " \
+#DataType "(int32 InValue), " \
+#DataType "(float InValue), " \
+#DataType "(const FString& InString)" \
+#DataType "(const Audio::IProxyDataPtr& InData),  or " \
+#DataType "(const TArray<Audio::IProxyDataPtr>& InProxyArray)."\
 #DataType "(const ::Metasound::FOperatorSettings& InSettings), " \
 #DataType "(const ::Metasound::FOperatorSettings& InSettings, bool InValue), " \
 #DataType "(const ::Metasound::FOperatorSettings& InSettings, int32 InValue), " \
diff --git a/Engine/Plugins/Runtime/Metasound/Source/MetasoundFrontend/Public/MetasoundFrontendRegistries.h b/Engine/Plugins/Runtime/Metasound/Source/MetasoundFrontend/Public/MetasoundFrontendRegistries.h
--- a/Engine/Plugins/Runtime/Metasound/Source/MetasoundFrontend/Public/MetasoundFrontendRegistries.h
+++ b/Engine/Plugins/Runtime/Metasound/Source/MetasoundFrontend/Public/MetasoundFrontendRegistries.h
@@ -13,6 +13,49 @@
 
 namespace Metasound
 {
+	// This struct is passed to FInputNodeConstructorCallback and FOutputNodeConstructorCallback.
+	struct FInputNodeConstructorParams
+	{
+		// the instance name and name of the specific connection that we should use.
+		const FString& InNodeName;
+		const FGuid& InInstanceID;
+		const FString& InVertexName;
+
+		FLiteral InitParam;
+	};
+
+	struct FOutputNodeConstructorParams
+	{
+		// the instance name and name of the specific connection that we should use.
+		const FString& InNodeName;
+		const FGuid& InInstanceID;
+		const FString& InVertexName;
+	};
+
+	/*
+	class IMetasoundDataTypeRegistryEntry
+	{
+	public:
+		virtual TUniquePtr<Metasound::INode> CreateInputNode(FInputNodeConstructorParams&&) const = 0;
+		virtual TUniquePtr<Metasound::INode> CreateOutputNode(const FOutputNodeConstructorParams&) const = 0;
+	};
+
+	class IMetasoundProxyRegistryEntry
+	{
+	public:
+		// This function is used to create a proxy from a datatype's base uclass.
+		virtual Audio::IProxyDataPtr CreateAudioProxy(UObject*) const = 0;
+	};
+
+	class IMetasoundNodeRegistryEntry
+	{
+	public:
+		virtual TUniquePtr<Metasound::INode> CreateNode(const Metasound::FNodeInitData&) const = 0;
+		virtual FMetasoundFrontendClass CreateMetasoundFrontendClass() const = 0;
+	};
+	*/
+
+
 	typedef TFunction<TUniquePtr<Metasound::INode>(::Metasound::FInputNodeConstructorParams&&)> FCreateInputNodeFunction;
 	typedef TFunction<TUniquePtr<Metasound::INode>(const ::Metasound::FOutputNodeConstructorParams&)> FCreateOutputNodeFunction;
 
@@ -34,19 +77,21 @@ namespace Metasound
 		// This indicates the type can only be constructed with FOperatorSettings or the default constructor.
 		bool bIsDefaultParsable = false;
 
-		// These bools signify what basic
-		// UProperty primitives we can use to describe this data type as a literal in a document.
+		// These bools signify what basic literal types can be use to describe this data type.
 		bool bIsBoolParsable = false;
 		bool bIsIntParsable = false;
 		bool bIsFloatParsable = false;
 		bool bIsStringParsable = false;
+		bool bIsProxyParsable = false;
+		bool bIsDefaultArrayParsable = false;
+		bool bIsBoolArrayParsable = false;
+		bool bIsIntArrayParsable = false;
+		bool bIsFloatArrayParsable = false;
+		bool bIsStringArrayParsable = false;
+		bool bIsProxyArrayParsable = false;
 
 		// Is an TEnum wrapped enum
 		bool bIsEnum = false;
-		
-		// these are used for using UObjects, or arrays of UObjects.
-		bool bIsProxyParsable = false;
-		bool bIsProxyArrayParsable = false;
 
 		// Determines whether the type can be used with send/receive transmitters
 		bool bIsTransmittable = false;
diff --git a/Engine/Plugins/Runtime/Metasound/Source/MetasoundFrontend/Public/MetasoundPrimitives.h b/Engine/Plugins/Runtime/Metasound/Source/MetasoundFrontend/Public/MetasoundPrimitives.h
--- a/Engine/Plugins/Runtime/Metasound/Source/MetasoundFrontend/Public/MetasoundPrimitives.h
+++ b/Engine/Plugins/Runtime/Metasound/Source/MetasoundFrontend/Public/MetasoundPrimitives.h
@@ -18,7 +18,5 @@ namespace Metasound
 	DECLARE_METASOUND_DATA_REFERENCE_TYPES(double, METASOUNDFRONTEND_API, FDoubleTypeInfo, FDoubleReadRef, FDoubleWriteRef);
 
 	DECLARE_METASOUND_DATA_REFERENCE_TYPES(FString, METASOUNDFRONTEND_API, FStringTypeInfo, FStringReadRef, FStringWriteRef);
-
-	DECLARE_METASOUND_DATA_REFERENCE_TYPES(TArray<float>, METASOUNDFRONTEND_API, FFloatArrayTypeInfo, FFloatArrayReadRef, FFloatArrayWriteRef);
 }
 
diff --git a/Engine/Plugins/Runtime/Metasound/Source/MetasoundGraphCore/Public/MetasoundDataFactory.h b/Engine/Plugins/Runtime/Metasound/Source/MetasoundGraphCore/Public/MetasoundDataFactory.h
--- a/Engine/Plugins/Runtime/Metasound/Source/MetasoundGraphCore/Public/MetasoundDataFactory.h
+++ b/Engine/Plugins/Runtime/Metasound/Source/MetasoundGraphCore/Public/MetasoundDataFactory.h
@@ -409,7 +409,46 @@ namespace Metasound
 			template<
 				typename VariantParseType,
 				typename VariantType,
-				typename std::enable_if< VariantParseType::bCreateArrayWithArg, int>::type = 0
+				typename std::enable_if< VariantParseType::bCreateArrayElementsWithArg, int>::type = 0
+			>
+			static auto CreateNewFromVariant(const FOperatorSettings& InSettings, const VariantType& InVariant)
+			{
+				using FConstructorForwarding = TExplicitArgsConstructorForwarding<DataType>;
+
+				using FArgType = typename VariantParseType::ArgType;
+				using FArgElementType = typename VariantParseType::ArgElementType;
+				using FElementType = typename VariantParseType::ElementType;
+				using FElementConstructorForwarding = TExplicitArgsConstructorForwarding<FElementType, FArgElementType>;
+
+				using FElementFactory = TDataFactoryInternal<FElementType, DataFactoryPrivate::TDataTypeCreator<FElementType>>;
+
+				const FArgType& ConstructorValues = InVariant.template Get<FArgType>();
+
+				// Array construct objects
+				DataType OutputContainer;
+				for (const FArgElementType& ConstructorValue : ConstructorValues)
+				{
+					OutputContainer.Emplace(FElementFactory::template CreateNew<FElementConstructorForwarding>(InSettings, ConstructorValue));
+				}
+
+				// Move construct array.
+				return FInternalFactory::template CreateNew<FConstructorForwarding>(InSettings, MoveTemp(OutputContainer));
+			}
+
+			/** Create a new object using FOperatorSettings and TVariant.
+			 *
+			 * Note: In this implementation, the InVariant object is an array of 
+			 * FLiteral::FNone where each array element simply denotes that a corresponding
+			 * element should be constructed with FOperatorSettings.
+			 *
+			 * @tparam VariantParseType - A type which describes how to parse the InVariant
+			 *                            argument.
+			 * @tparam VariantType - The expected underlying type stored in InVariant.
+			 */
+			template<
+				typename VariantParseType,
+				typename VariantType,
+				typename std::enable_if< VariantParseType::bCreateArrayElementsWithSettings, int>::type = 0
 			>
 			static auto CreateNewFromVariant(const FOperatorSettings& InSettings, const VariantType& InVariant)
 			{
@@ -428,7 +467,7 @@ namespace Metasound
 				DataType OutputContainer;
 				for (const FArgElementType& ConstructorValue : ConstructorValues)
 				{
-					OutputContainer.Add(FElementFactory::template CreateNew<FElementConstructorForwarding>(InSettings, ConstructorValue));
+					OutputContainer.Add(FElementFactory::template CreateNew<FElementConstructorForwarding>(InSettings));
 				}
 
 				// Move construct array.
@@ -448,7 +487,7 @@ namespace Metasound
 			template<
 				typename VariantParseType,
 				typename VariantType,
-				typename std::enable_if< VariantParseType::bCreateArrayWithoutArg, int>::type = 0
+				typename std::enable_if< VariantParseType::bCreateArrayElementsWithDefaultConstructor, int>::type = 0
 			>
 			static auto CreateNewFromVariant(const FOperatorSettings& InSettings, const VariantType& InVariant)
 			{
@@ -649,8 +688,9 @@ namespace Metasound
 			using FFallbackHelper = TDataTypeVariantFallbackHelper<DataType, FirstVariantType, AdditionalVariantTypes...>;
 
 			// Determine which elementwise construction methods are supported
-			static constexpr bool bCanCreateArrayWithArg = FDesiredConstructorTraits::bIsArrayConstructibleWithArgs || FDesiredConstructorTraits::bIsArrayConstructibleWithSettingsAndArgs;
-			static constexpr bool bCanCreateArrayWithoutArg = FDesiredConstructorTraits::bIsArrayConstructibleWithSettings || FDesiredConstructorTraits::bIsArrayDefaultConstructible;
+			static constexpr bool bCanCreateArrayElementsWithArg = FDesiredConstructorTraits::bIsArrayConstructibleWithArgs || FDesiredConstructorTraits::bIsArrayConstructibleWithSettingsAndArgs;
+			static constexpr bool bCanCreateArrayElementsWithSettings = FDesiredConstructorTraits::bIsArrayConstructibleWithSettings; 
+			static constexpr bool bCanCreateArrayElementsWithDefaultConstructor = FDesiredConstructorTraits::bIsArrayDefaultConstructible;
 
 			// Determine which construction methods are supported
 			static constexpr bool bCanCreateWithArg = FDesiredConstructorTraits::bIsConstructibleWithArgs || FDesiredConstructorTraits::bIsConstructibleWithSettingsAndArgs;
@@ -659,15 +699,16 @@ namespace Metasound
 
 		public:
 
-			static constexpr bool bCannotForwardToConstructor = !(bCanCreateWithArg || bCanCreateWithoutArg || bCanCreateWithFallbackArg || bCanCreateArrayWithArg || bCanCreateArrayWithoutArg);
+			static constexpr bool bCannotForwardToConstructor = !(bCanCreateWithArg || bCanCreateWithoutArg || bCanCreateWithFallbackArg || bCanCreateArrayElementsWithArg || bCanCreateArrayElementsWithDefaultConstructor);
 
 			// Determine which construction method to use.
-			static constexpr bool bCreateArrayWithArg = bCanCreateArrayWithArg;
-			static constexpr bool bCreateArrayWithoutArg = !(bCanCreateArrayWithArg) && bCanCreateArrayWithoutArg;
+			static constexpr bool bCreateArrayElementsWithArg = bCanCreateArrayElementsWithArg;
+			static constexpr bool bCreateArrayElementsWithSettings = !(bCanCreateArrayElementsWithArg) && bCanCreateArrayElementsWithSettings;
+			static constexpr bool bCreateArrayElementsWithDefaultConstructor = !(bCanCreateArrayElementsWithArg || bCreateArrayElementsWithSettings) && bCanCreateArrayElementsWithDefaultConstructor;
 
-			static constexpr bool bCreateWithArg = !(bCreateArrayWithArg || bCreateArrayWithoutArg) && bCanCreateWithArg;
-			static constexpr bool bCreateWithoutArg = !(bCreateArrayWithArg || bCreateArrayWithoutArg || bCreateWithArg) && bCanCreateWithoutArg;
-			static constexpr bool bCreateWithFallbackArg = !(bCreateArrayWithArg || bCreateArrayWithoutArg || bCreateWithArg || bCreateWithoutArg) && bCanCreateWithFallbackArg;
+			static constexpr bool bCreateWithArg = !(bCreateArrayElementsWithArg || bCreateArrayElementsWithSettings || bCreateArrayElementsWithDefaultConstructor) && bCanCreateWithArg;
+			static constexpr bool bCreateWithoutArg = !(bCreateArrayElementsWithArg || bCreateArrayElementsWithSettings || bCreateArrayElementsWithDefaultConstructor || bCreateWithArg) && bCanCreateWithoutArg;
+			static constexpr bool bCreateWithFallbackArg = !(bCreateArrayElementsWithArg || bCreateArrayElementsWithSettings || bCreateArrayElementsWithDefaultConstructor || bCreateWithArg || bCreateWithoutArg) && bCanCreateWithFallbackArg;
 
 			// Determine types
 			using ElementType = typename TElementType<DataType>::Type;
diff --git a/Engine/Plugins/Runtime/Metasound/Source/MetasoundGraphCore/Public/MetasoundDataReference.h b/Engine/Plugins/Runtime/Metasound/Source/MetasoundGraphCore/Public/MetasoundDataReference.h
--- a/Engine/Plugins/Runtime/Metasound/Source/MetasoundGraphCore/Public/MetasoundDataReference.h
+++ b/Engine/Plugins/Runtime/Metasound/Source/MetasoundGraphCore/Public/MetasoundDataReference.h
@@ -21,17 +21,31 @@
 		\
 		static const DataType* const TypePtr; \
 	}; \
-	\
+	/*TODO remove typedefs */ \
 	typedef ::Metasound::TDataReferenceTypeInfo<DataType> DataTypeInfoTypeName; \
 	\
 	typedef ::Metasound::TDataReadReference<DataType> DataReadReferenceTypeName; \
-	typedef ::Metasound::TDataWriteReference<DataType> DataWriteReferenceTypeName;
+	typedef ::Metasound::TDataWriteReference<DataType> DataWriteReferenceTypeName;\
+	template<> \
+	struct ::Metasound::TDataReferenceTypeInfo<TArray<DataType>> \
+	{ \
+		static ModuleApi const TCHAR* TypeName; \
+		static ModuleApi const FMetasoundDataTypeId TypeId; \
+		\
+		private: \
+		\
+		static const TArray<DataType>* const TypePtr; \
+	};
 
 // This only needs to be called if you don't plan on calling REGISTER_METASOUND_DATATYPE.
 #define DEFINE_METASOUND_DATA_TYPE(DataType, DataTypeName) \
 	const TCHAR*  				::Metasound::TDataReferenceTypeInfo<DataType>::TypeName = TEXT(DataTypeName); \
 	const DataType* const ::Metasound::TDataReferenceTypeInfo<DataType>::TypePtr = nullptr; \
-	const void* const ::Metasound::TDataReferenceTypeInfo<DataType>::TypeId = static_cast<const FMetasoundDataTypeId>(&::Metasound::TDataReferenceTypeInfo<DataType>::TypePtr);
+	const void* const ::Metasound::TDataReferenceTypeInfo<DataType>::TypeId = static_cast<const FMetasoundDataTypeId>(&::Metasound::TDataReferenceTypeInfo<DataType>::TypePtr); \
+	/* Array definition */\
+	const TCHAR*  				::Metasound::TDataReferenceTypeInfo<TArray<DataType>>::TypeName = TEXT(DataTypeName ":Array"); \
+	const TArray<DataType>* const ::Metasound::TDataReferenceTypeInfo<TArray<DataType>>::TypePtr = nullptr; \
+	const void* const ::Metasound::TDataReferenceTypeInfo<TArray<DataType>>::TypeId = static_cast<const FMetasoundDataTypeId>(&::Metasound::TDataReferenceTypeInfo<TArray<DataType>>::TypePtr);
 
 namespace Audio
 {
@@ -96,14 +110,6 @@ namespace Metasound
 	{
 		static METASOUNDGRAPHCORE_API const TCHAR* TypeName;
 		static METASOUNDGRAPHCORE_API const void* const TypeId;
-		static constexpr bool bIsStringParsable = false;
-		static constexpr bool bIsBoolParsable = false;
-		static constexpr bool bIsIntParsable = false;
-		static constexpr bool bIsFloatParsable = false;
-		static constexpr bool bIsProxyParsable = false;
-		static constexpr bool bIsProxyArrayParsable = false;
-		static constexpr bool bIsConstructableWithSettings = false;
-		static constexpr bool bCanUseDefaultConstructor = false;
 
 		private:
 
diff --git a/Engine/Plugins/Runtime/Metasound/Source/MetasoundGraphCore/Public/MetasoundLiteral.h b/Engine/Plugins/Runtime/Metasound/Source/MetasoundGraphCore/Public/MetasoundLiteral.h
--- a/Engine/Plugins/Runtime/Metasound/Source/MetasoundGraphCore/Public/MetasoundLiteral.h
+++ b/Engine/Plugins/Runtime/Metasound/Source/MetasoundGraphCore/Public/MetasoundLiteral.h
@@ -331,22 +331,4 @@ namespace Metasound
 	};
 
 
-	// This struct is passed to FInputNodeConstructorCallback and FOutputNodeConstructorCallback.
-	struct FInputNodeConstructorParams
-	{
-		// the instance name and name of the specific connection that we should use.
-		const FString& InNodeName;
-		const FGuid& InInstanceID;
-		const FString& InVertexName;
-
-		FLiteral InitParam;
-	};
-
-	struct FOutputNodeConstructorParams
-	{
-		// the instance name and name of the specific connection that we should use.
-		const FString& InNodeName;
-		const FGuid& InInstanceID;
-		const FString& InVertexName;
-	};
 }
