Random node support
Fix for audio buffer add

diff --git a/Engine/Plugins/Runtime/Metasound/Source/MetasoundStandardNodes/Private/MetasoundMathNodes.cpp b/Engine/Plugins/Runtime/Metasound/Source/MetasoundStandardNodes/Private/MetasoundMathNodes.cpp
--- a/Engine/Plugins/Runtime/Metasound/Source/MetasoundStandardNodes/Private/MetasoundMathNodes.cpp
+++ b/Engine/Plugins/Runtime/Metasound/Source/MetasoundStandardNodes/Private/MetasoundMathNodes.cpp
@@ -29,11 +29,6 @@ namespace Metasound
 	METASOUND_REGISTER_NODE(FDivideInt32Node)
 	METASOUND_REGISTER_NODE(FDivideInt64Node)
 	METASOUND_REGISTER_NODE(FDivideTimeFloatNode)
-
-	METASOUND_REGISTER_NODE(FRandRangeFloatNode)
-	METASOUND_REGISTER_NODE(FRandRangeInt32Node)
-	METASOUND_REGISTER_NODE(FRandRangeInt64Node)
-	METASOUND_REGISTER_NODE(FRandRangeTimeNode)
 }
 
 #undef LOCTEXT_NAMESPACE // MetasoundMathOpNode
diff --git a/Engine/Plugins/Runtime/Metasound/Source/MetasoundStandardNodes/Private/MetasoundRandomNode.cpp b/Engine/Plugins/Runtime/Metasound/Source/MetasoundStandardNodes/Private/MetasoundRandomNode.cpp
new file mode 100644
--- /dev/null
+++ b/Engine/Plugins/Runtime/Metasound/Source/MetasoundStandardNodes/Private/MetasoundRandomNode.cpp
@@ -0,0 +1,406 @@
+// Copyright Epic Games, Inc. All Rights Reserved.
+
+#include "MetasoundRandomNode.h"
+
+#include "Internationalization/Text.h"
+#include "MetasoundExecutableOperator.h"
+#include "MetasoundNodeRegistrationMacro.h"
+#include "MetasoundPrimitives.h"
+#include "MetasoundStandardNodesNames.h"
+#include "MetasoundTrigger.h"
+#include "Math/RandomStream.h"
+#include "HAL/PlatformTime.h"
+
+#define LOCTEXT_NAMESPACE "MetasoundStandardNodes"
+
+namespace Metasound
+{
+	static const TCHAR* InParamNameSeed = TEXT("Seed");
+	static const TCHAR* InParamNameMin = TEXT("Min");
+	static const TCHAR* InParamNameMax = TEXT("Max");
+	static const TCHAR* InParamNameNextTrigger = TEXT("Next");
+	static const TCHAR* InParamNameResetTrigger = TEXT("Reset");
+	static const TCHAR* OutParamNameValue = TEXT("Value");
+
+	//////////////////////////////////////////////////////////////////////////
+	// Random Int
+
+	class FRandomIntOperator : public TExecutableOperator<FRandomIntOperator>
+	{
+	public:
+
+		static const FNodeClassMetadata& GetNodeInfo();
+		static const FVertexInterface& GetVertexInterface();
+		static TUniquePtr<IOperator> CreateOperator(const FCreateOperatorParams& InParams, FBuildErrorArray& OutErrors);
+
+		FRandomIntOperator(const FOperatorSettings& InSettings,
+			const FInt32ReadRef& InSeed, 
+			const FInt32ReadRef& InMin, 
+			const FInt32ReadRef& InMax, 
+			const FTriggerReadRef& InTriggerNext, 
+			const FTriggerReadRef& InTriggerReset);
+
+		virtual FDataReferenceCollection GetInputs() const override;
+		virtual FDataReferenceCollection GetOutputs() const override;
+		void Execute();
+
+	private:
+		void EvaluateSeedChanges();
+
+		FInt32ReadRef Seed;
+		FInt32ReadRef Min;
+		FInt32ReadRef Max;
+		FTriggerReadRef TriggerNext;
+		FTriggerReadRef TriggerReset;
+
+		FInt32WriteRef ValueOutput;
+
+		FRandomStream RandomStream;
+
+		// Used to track if the random node starts non-zero and then becomes 0 (i.e. we want a new random seed)
+		bool bIsZeroSeeded = false;
+		bool bIsRandomStreamInitialized = false;
+	};
+
+	FRandomIntOperator::FRandomIntOperator(const FOperatorSettings& InSettings,
+		const FInt32ReadRef& InSeed,
+		const FInt32ReadRef& InMin,
+		const FInt32ReadRef& InMax,
+		const FTriggerReadRef& InTriggerNext,
+		const FTriggerReadRef& InTriggerReset)
+		: Seed(InSeed)
+		, Min(InMin)
+		, Max(InMax)
+		, TriggerNext(InTriggerNext)
+		, TriggerReset(InTriggerReset)
+		, ValueOutput(FInt32WriteRef::CreateNew(0))
+		, bIsZeroSeeded(*Seed == 0)
+		, bIsRandomStreamInitialized(false)
+	{
+		*ValueOutput = *Min;
+	}
+
+	FDataReferenceCollection FRandomIntOperator::GetInputs() const
+	{
+		FDataReferenceCollection InputDataReferences;
+		InputDataReferences.AddDataReadReference(InParamNameSeed, FInt32ReadRef(Seed));
+		InputDataReferences.AddDataReadReference(InParamNameMin, FInt32ReadRef(Min));
+		InputDataReferences.AddDataReadReference(InParamNameMax, FInt32ReadRef(Max));
+		InputDataReferences.AddDataReadReference(InParamNameNextTrigger, FTriggerReadRef(TriggerNext));
+		InputDataReferences.AddDataReadReference(InParamNameResetTrigger, FTriggerReadRef(TriggerReset));
+
+		return InputDataReferences;
+	}
+
+	FDataReferenceCollection FRandomIntOperator::GetOutputs() const
+	{
+		FDataReferenceCollection OutputDataReferences;
+		OutputDataReferences.AddDataReadReference(OutParamNameValue, FInt32ReadRef(ValueOutput));
+		return OutputDataReferences;
+	}
+
+	void FRandomIntOperator::EvaluateSeedChanges()
+	{
+		// if we have a non-zero seed
+		if (*Seed != 0)
+		{
+			// If we were previously zero-seeded OR our seed has changed
+			if (bIsZeroSeeded || !bIsRandomStreamInitialized || *Seed != RandomStream.GetInitialSeed())
+			{
+				bIsRandomStreamInitialized = true;
+				bIsZeroSeeded = false;
+				RandomStream.Initialize(*Seed);
+			}
+		}
+		// If we are zero-seeded now BUT were previously not, we need to randomize our seed
+		else if (!bIsZeroSeeded || !bIsRandomStreamInitialized)
+		{
+			bIsRandomStreamInitialized = true;
+			bIsZeroSeeded = true;
+			RandomStream.Initialize(FPlatformTime::Cycles());
+		}
+	}
+
+	void FRandomIntOperator::Execute()
+	{
+		if (TriggerReset->IsTriggeredInBlock())
+		{
+			EvaluateSeedChanges();
+			RandomStream.Reset();
+		}
+
+		if (TriggerNext->IsTriggeredInBlock())
+		{
+			EvaluateSeedChanges();
+			*ValueOutput = RandomStream.RandRange(*Min, *Max);
+		}
+
+// 		TriggerReset->ExecuteBlock(
+// 			// OnPreTrigger
+// 			[&](int32 StartFrame, int32 EndFrame)
+// 			{
+// 			},
+// 			// OnTrigger
+// 				[&](int32 StartFrame, int32 EndFrame)
+// 			{
+// 				EvaluateSeedChanges();
+// 				RandomStream.Reset();
+// 
+// 				UE_LOG(LogTemp, Log, TEXT("RESET"));
+// 			});
+// 
+// 		TriggerNext->ExecuteBlock(
+// 			// OnPreTrigger
+// 			[&](int32 StartFrame, int32 EndFrame)
+// 			{
+// 			},
+// 			// OnTrigger
+// 				[&](int32 StartFrame, int32 EndFrame)
+// 			{
+// 				EvaluateSeedChanges();
+// 				*ValueOutput = RandomStream.RandRange(*Min, *Max);
+// 
+// 				UE_LOG(LogTemp, Log, TEXT("%d"), *ValueOutput);
+// 			})
+	}
+
+	const FVertexInterface& FRandomIntOperator::GetVertexInterface()
+	{
+		static const FVertexInterface Interface(
+			FInputVertexInterface(
+				TInputDataVertexModel<int32>(InParamNameSeed, LOCTEXT("SeedTooltip", "The seed value to use for the random node. Set to 0 to use a random seed."), 0),
+				TInputDataVertexModel<int32>(InParamNameMin, LOCTEXT("MinTooltip", "Min random integer."), 0),
+				TInputDataVertexModel<int32>(InParamNameMax, LOCTEXT("MaxTooltip", "Max random integer."), 100),
+				TInputDataVertexModel<FTrigger>(InParamNameNextTrigger, LOCTEXT("NextTooltip", "Trigger to generate the next random integer.")),
+				TInputDataVertexModel<FTrigger>(InParamNameResetTrigger, LOCTEXT("ResetTooltip", "Trigger to reset the random sequence with the supplied seed. Useful to get randomized repetition."))
+			),
+			FOutputVertexInterface(
+				TOutputDataVertexModel<int32>(OutParamNameValue, LOCTEXT("ValueTooltip", "The current randomly generated random integer value."))
+			)
+		);
+
+		return Interface;
+	}
+
+	const FNodeClassMetadata& FRandomIntOperator::GetNodeInfo()
+	{
+		auto InitNodeInfo = []() -> FNodeClassMetadata
+		{
+			FNodeClassMetadata Info;
+			Info.ClassName = { Metasound::StandardNodes::Namespace, TEXT("Random"), TEXT("Int32") };
+			Info.MajorVersion = 1;
+			Info.MinorVersion = 0;
+			Info.DisplayName = LOCTEXT("Metasound_RandomIntDisplayName", "Random Int");
+			Info.Description = LOCTEXT("Metasound_RandomIntNodeDescription", "Generates a seedable random integer in the given value range.");
+			Info.Author = PluginAuthor;
+			Info.PromptIfMissing = PluginNodeMissingPrompt;
+			Info.DefaultInterface = GetVertexInterface();
+
+			return Info;
+		};
+
+		static const FNodeClassMetadata Info = InitNodeInfo();
+
+		return Info;
+	}
+
+	TUniquePtr<IOperator> FRandomIntOperator::CreateOperator(const FCreateOperatorParams& InParams, FBuildErrorArray& OutErrors)
+	{
+		const FRandomIntNode& RandomIntNode = static_cast<const FRandomIntNode&>(InParams.Node);
+		const FDataReferenceCollection& InputCollection = InParams.InputDataReferences;
+		const FInputVertexInterface& InputInterface = GetVertexInterface().GetInputInterface();
+
+		FInt32ReadRef SeedValue = InputCollection.GetDataReadReferenceOrConstructWithVertexDefault<int32>(InputInterface, InParamNameSeed);
+		FInt32ReadRef MinValue = InputCollection.GetDataReadReferenceOrConstructWithVertexDefault<int32>(InputInterface, InParamNameMin);
+		FInt32ReadRef MaxValue = InputCollection.GetDataReadReferenceOrConstructWithVertexDefault<int32>(InputInterface, InParamNameMax);
+		FTriggerReadRef TriggerNext = InputCollection.GetDataReadReferenceOrConstruct<FTrigger>(InParamNameNextTrigger, InParams.OperatorSettings);
+		FTriggerReadRef TriggerReset = InputCollection.GetDataReadReferenceOrConstruct<FTrigger>(InParamNameResetTrigger, InParams.OperatorSettings);
+
+		return MakeUnique<FRandomIntOperator>(InParams.OperatorSettings, SeedValue, MinValue, MaxValue, TriggerNext, TriggerReset);
+	}
+
+	FRandomIntNode::FRandomIntNode(const FNodeInitData& InitData)
+		: FNodeFacade(InitData.InstanceName, InitData.InstanceID, TFacadeOperatorClass<FRandomIntOperator>())
+	{
+	}
+
+	METASOUND_REGISTER_NODE(FRandomIntNode)
+
+	//////////////////////////////////////////////////////////////////////////
+	// Random Float
+
+	class FRandomFloatOperator : public TExecutableOperator<FRandomFloatOperator>
+	{
+	public:
+
+		static const FNodeClassMetadata& GetNodeInfo();
+		static const FVertexInterface& GetVertexInterface();
+		static TUniquePtr<IOperator> CreateOperator(const FCreateOperatorParams& InParams, FBuildErrorArray& OutErrors);
+
+		FRandomFloatOperator(const FOperatorSettings& InSettings,
+			const FInt32ReadRef& InSeed,
+			const FFloatReadRef& InMin,
+			const FFloatReadRef& InMax,
+			const FTriggerReadRef& InTriggerNext,
+			const FTriggerReadRef& InTriggerReset);
+
+		virtual FDataReferenceCollection GetInputs() const override;
+		virtual FDataReferenceCollection GetOutputs() const override;
+		void Execute();
+
+	private:
+		void EvaluateSeedChanges();
+
+		FInt32ReadRef Seed;
+		FFloatReadRef Min;
+		FFloatReadRef Max;
+		FTriggerReadRef TriggerNext;
+		FTriggerReadRef TriggerReset;
+
+		FFloatWriteRef ValueOutput;
+
+		FRandomStream RandomStream;
+
+		// Used to track if the random node starts non-zero and then becomes 0 (i.e. we want a new random seed)
+		bool bIsZeroSeeded = false;
+		bool bIsRandomStreamInitialized = false;
+	};
+
+	FRandomFloatOperator::FRandomFloatOperator(const FOperatorSettings& InSettings,
+		const FInt32ReadRef& InSeed,
+		const FFloatReadRef& InMin,
+		const FFloatReadRef& InMax,
+		const FTriggerReadRef& InTriggerNext,
+		const FTriggerReadRef& InTriggerReset)
+		: Seed(InSeed)
+		, Min(InMin)
+		, Max(InMax)
+		, TriggerNext(InTriggerNext)
+		, TriggerReset(InTriggerReset)
+		, ValueOutput(FFloatWriteRef::CreateNew(0.0f))
+		, bIsZeroSeeded(*Seed == 0)
+		, bIsRandomStreamInitialized(false)
+	{
+		*ValueOutput = *Min;
+	}
+
+	FDataReferenceCollection FRandomFloatOperator::GetInputs() const
+	{
+		FDataReferenceCollection InputDataReferences;
+		InputDataReferences.AddDataReadReference(InParamNameSeed, FInt32ReadRef(Seed));
+		InputDataReferences.AddDataReadReference(InParamNameMin, FFloatReadRef(Min));
+		InputDataReferences.AddDataReadReference(InParamNameMax, FFloatReadRef(Max));
+		InputDataReferences.AddDataReadReference(InParamNameNextTrigger, FTriggerReadRef(TriggerNext));
+		InputDataReferences.AddDataReadReference(InParamNameResetTrigger, FTriggerReadRef(TriggerReset));
+
+		return InputDataReferences;
+	}
+
+	FDataReferenceCollection FRandomFloatOperator::GetOutputs() const
+	{
+		FDataReferenceCollection OutputDataReferences;
+		OutputDataReferences.AddDataReadReference(OutParamNameValue, FFloatReadRef(ValueOutput));
+		return OutputDataReferences;
+	}
+
+	void FRandomFloatOperator::EvaluateSeedChanges()
+	{
+		// if we have a non-zero seed
+		if (*Seed != 0)
+		{
+			// If we were previously zero-seeded OR our seed has changed
+			if (bIsZeroSeeded || !bIsRandomStreamInitialized || *Seed != RandomStream.GetInitialSeed())
+			{
+				bIsRandomStreamInitialized = true;
+				bIsZeroSeeded = false;
+				RandomStream.Initialize(*Seed);
+			}
+		}
+		// If we are zero-seeded now BUT were previously not, we need to randomize our seed
+		else if (!bIsZeroSeeded || !bIsRandomStreamInitialized)
+		{
+			bIsRandomStreamInitialized = true;
+			bIsZeroSeeded = true;
+			RandomStream.Initialize(FPlatformTime::Cycles());
+		}
+	}
+
+	void FRandomFloatOperator::Execute()
+	{
+		if (TriggerReset->IsTriggeredInBlock())
+		{
+			EvaluateSeedChanges();
+			RandomStream.Reset();
+		}
+
+		if (TriggerNext->IsTriggeredInBlock())
+		{
+			EvaluateSeedChanges();
+			*ValueOutput = RandomStream.FRandRange(*Min, *Max);
+		}
+	}
+
+	const FVertexInterface& FRandomFloatOperator::GetVertexInterface()
+	{
+		static const FVertexInterface Interface(
+			FInputVertexInterface(
+				TInputDataVertexModel<int32>(InParamNameSeed, LOCTEXT("SeedTooltip", "The seed value to use for the random node. Set to 0 to use a random seed."), 0),
+				TInputDataVertexModel<float>(InParamNameMin, LOCTEXT("MinTooltip", "Min random float."), 0.0f),
+				TInputDataVertexModel<float>(InParamNameMax, LOCTEXT("MaxTooltip", "Max random float."), 1.0f),
+				TInputDataVertexModel<FTrigger>(InParamNameNextTrigger, LOCTEXT("NextTooltip", "Trigger to generate the next random float.")),
+				TInputDataVertexModel<FTrigger>(InParamNameResetTrigger, LOCTEXT("ResetTooltip", "Trigger to reset the random sequence with the supplied seed. Useful to get randomized repetition."))
+			),
+			FOutputVertexInterface(
+				TOutputDataVertexModel<float>(OutParamNameValue, LOCTEXT("ValueTooltip", "The current randomly generated float value."))
+			)
+		);
+
+		return Interface;
+	}
+
+	const FNodeClassMetadata& FRandomFloatOperator::GetNodeInfo()
+	{
+		auto InitNodeInfo = []() -> FNodeClassMetadata
+		{
+			FNodeClassMetadata Info;
+			Info.ClassName = { Metasound::StandardNodes::Namespace, TEXT("Random"), TEXT("Float") };
+			Info.MajorVersion = 1;
+			Info.MinorVersion = 0;
+			Info.DisplayName = LOCTEXT("Metasound_RandomFloatDisplayName", "Random Float");
+			Info.Description = LOCTEXT("Metasound_RandomFloatNodeDescription", "Generates a seedable random float in the given value range.");
+			Info.Author = PluginAuthor;
+			Info.PromptIfMissing = PluginNodeMissingPrompt;
+			Info.DefaultInterface = GetVertexInterface();
+
+			return Info;
+		};
+
+		static const FNodeClassMetadata Info = InitNodeInfo();
+
+		return Info;
+	}
+
+	TUniquePtr<IOperator> FRandomFloatOperator::CreateOperator(const FCreateOperatorParams& InParams, FBuildErrorArray& OutErrors)
+	{
+		const FRandomFloatNode& RandomIntNode = static_cast<const FRandomFloatNode&>(InParams.Node);
+		const FDataReferenceCollection& InputCollection = InParams.InputDataReferences;
+		const FInputVertexInterface& InputInterface = GetVertexInterface().GetInputInterface();
+
+		FInt32ReadRef SeedValue = InputCollection.GetDataReadReferenceOrConstructWithVertexDefault<int32>(InputInterface, InParamNameSeed);
+		FFloatReadRef MinValue = InputCollection.GetDataReadReferenceOrConstructWithVertexDefault<float>(InputInterface, InParamNameMin);
+		FFloatReadRef MaxValue = InputCollection.GetDataReadReferenceOrConstructWithVertexDefault<float>(InputInterface, InParamNameMax);
+		FTriggerReadRef TriggerNext = InputCollection.GetDataReadReferenceOrConstruct<FTrigger>(InParamNameNextTrigger, InParams.OperatorSettings);
+		FTriggerReadRef TriggerReset = InputCollection.GetDataReadReferenceOrConstruct<FTrigger>(InParamNameResetTrigger, InParams.OperatorSettings);
+
+		return MakeUnique<FRandomFloatOperator>(InParams.OperatorSettings, SeedValue, MinValue, MaxValue, TriggerNext, TriggerReset);
+	}
+
+	FRandomFloatNode::FRandomFloatNode(const FNodeInitData& InitData)
+		: FNodeFacade(InitData.InstanceName, InitData.InstanceID, TFacadeOperatorClass<FRandomFloatOperator>())
+	{
+	}
+
+	METASOUND_REGISTER_NODE(FRandomFloatNode)
+}
+
+#undef LOCTEXT_NAMESPACE
diff --git a/Engine/Plugins/Runtime/Metasound/Source/MetasoundStandardNodes/Public/MetasoundMathNodes.h b/Engine/Plugins/Runtime/Metasound/Source/MetasoundStandardNodes/Public/MetasoundMathNodes.h
--- a/Engine/Plugins/Runtime/Metasound/Source/MetasoundStandardNodes/Public/MetasoundMathNodes.h
+++ b/Engine/Plugins/Runtime/Metasound/Source/MetasoundStandardNodes/Public/MetasoundMathNodes.h
@@ -120,13 +120,6 @@ namespace Metasound
 				FText::FromString(GetMetasoundDataTypeString<OperandDataType>()));
 			return DisplayName;
 		}
-
-		template<typename DataType>
-		const FText GetRandRangeDisplayName()
-		{
-			static const FText DisplayName = FText::Format(LOCTEXT("RandRangeNodeDisplayNamePattern", "Random ({0})"), FText::FromString(GetMetasoundDataTypeString<DataType>()));
-			return DisplayName;
-		}
 	}
 
 	template <typename TDerivedClass, typename TMathOpClass, typename TDataClass, typename TOperandDataClass>
@@ -490,7 +483,7 @@ namespace Metasound
 				return;
 			}
 
-			FMemory::Memcpy(OutResult->GetData(), InAdditionalOperands[0]->GetData(), sizeof(float) * NumSamples);
+			FMemory::Memcpy(OutResult->GetData(), InPrimaryOperand->GetData(), sizeof(float) * NumSamples);
 
 			for (int32 i = 0; i < InAdditionalOperands.Num(); ++i)
 			{
@@ -1167,134 +1160,6 @@ namespace Metasound
 		}
 	};
 
-	template <typename TDataClass, typename TOperandDataClass = TDataClass>
-	class TMathOpRandRange
-	{
-	public:
-		using TDataClassReadRef = TDataReadReference<TDataClass>;
-		using TOperandDataClassReadRef = TDataReadReference<TOperandDataClass>;
-		using TDataClassWriteRef = TDataWriteReference<TDataClass>;
-
-		static const FVertexInterface& GetVertexInterface()
-		{
-			static const FVertexInterface Interface(
-				FInputVertexInterface(
-					TInputDataVertexModel<TDataClass>(MathOpNames::PrimaryOperandName, LOCTEXT("MathOpRandomMinTooltip", "Random value minimum."), static_cast<TDataClass>(0.0f)),
-					TInputDataVertexModel<TOperandDataClass>(MathOpNames::AdditionalOperandsName, LOCTEXT("MathOpRandomMaxTooltip", "Random value maximum."), static_cast<TOperandDataClass>(1.0f))
-				),
-				FOutputVertexInterface(
-					TOutputDataVertexModel<TDataClass>(TEXT("Out"), LOCTEXT("MathOpRandomOutTooltip", "Resulting random value"))
-				)
-			);
-
-			return Interface;
-		}
-
-		static TDataClass GetDefault(const FOperatorSettings& InSettings, const FLiteral& InVertexDefault)
-		{
-			return InVertexDefault.Value.Get<TDataClass>();
-		}
-
-		static TOperandDataClass GetDefaultOp(const FOperatorSettings& InSettings, const FLiteral& InVertexDefault)
-		{
-			return InVertexDefault.Value.Get<TOperandDataClass>();
-		}
-
-		static void Calculate(TMathOpRandRange& InInstanceData, const TDataClassReadRef& InPrimaryOperand, const TArray<TOperandDataClassReadRef>& InAdditionalOperands, TDataClassWriteRef& OutResult)
-		{
-			if (ensure(!InAdditionalOperands.IsEmpty()))
-			{
-				const TDataClass& Min = *InPrimaryOperand;
-				const TOperandDataClass& Max = *InAdditionalOperands[0];
-				*OutResult = FMath::RandRange(Min, Max);
-			}
-		}
-	};
-
-	template <>
-	class TMathOpRandRange<FTime>
-	{
-	public:
-		static const FVertexInterface& GetVertexInterface()
-		{
-			static const FVertexInterface Interface(
-				FInputVertexInterface(
-					TInputDataVertexModel<FTime>(MathOpNames::PrimaryOperandName, LOCTEXT("MathOpRandomMinTooltip", "Random value minimum."), 0.0f),
-					TInputDataVertexModel<FTime>(MathOpNames::AdditionalOperandsName, LOCTEXT("MathOpRandomMaxTooltip", "Random value maximum."), 1.0f)
-				),
-				FOutputVertexInterface(
-					TOutputDataVertexModel<FTime>(TEXT("Out"), LOCTEXT("MathOpRandomOutTooltip", "Resulting random value"))
-				)
-			);
-
-			return Interface;
-		}
-
-		static FTime GetDefault(const FOperatorSettings& InSettings, const FLiteral& InVertexDefault)
-		{
-			return FTime(InVertexDefault.Value.Get<float>());
-		}
-
-		static FTime GetDefaultOp(const FOperatorSettings& InSettings, const FLiteral& InVertexDefault)
-		{
-			return FTime(InVertexDefault.Value.Get<float>());
-		}
-
-		static void Calculate(TMathOpRandRange& InInstanceData, const FTimeReadRef& InPrimaryOperand, const TArray<FTimeReadRef>& InAdditionalOperands, FTimeWriteRef& OutResult)
-		{
-			if (ensure(!InAdditionalOperands.IsEmpty()))
-			{
-				const FTime& Min = *InPrimaryOperand;
-				const FTime& Max = *InAdditionalOperands[0];
-
-				// Not a true random as precision is lost with randomizing over float range
-				const FTime Delta = Max.GetSeconds() - Min.GetSeconds();
-				const FTime Result = Min.GetSeconds() + (Delta * FMath::FRand());
-				*OutResult = Result;
-			}
-		}
-	};
-
-	// int64 has to be explicitly implemented as they are PODs that are not implemented as a literal base type.
-	template <>
-	class TMathOpRandRange<int64>
-	{
-	public:
-		static const FVertexInterface& GetVertexInterface()
-		{
-			static const FVertexInterface Interface(
-				FInputVertexInterface(
-					TInputDataVertexModel<int64>(MathOpNames::PrimaryOperandName, LOCTEXT("MathOpRandomMinTooltip", "Random value minimum."), 0),
-					TInputDataVertexModel<int64>(MathOpNames::AdditionalOperandsName, LOCTEXT("MathOpRandomMaxTooltip", "Random value maximum."), 1)
-				),
-				FOutputVertexInterface(
-					TOutputDataVertexModel<int64>(TEXT("Out"), LOCTEXT("MathOpRandomOutTooltip", "Resulting random value"))
-				)
-			);
-
-			return Interface;
-		}
-
-		static int64 GetDefault(const FOperatorSettings& InSettings, const FLiteral& InVertexDefault)
-		{
-			return static_cast<int64>(InVertexDefault.Value.Get<int32>());
-		}
-
-		static int64 GetDefaultOp(const FOperatorSettings& InSettings, const FLiteral& InVertexDefault)
-		{
-			return static_cast<int64>(InVertexDefault.Value.Get<int32>());
-		}
-
-		static void Calculate(TMathOpRandRange& InInstanceData, const FInt64ReadRef& InPrimaryOperand, const TArray<FInt64ReadRef>& InAdditionalOperands, FInt64WriteRef& OutResult)
-		{
-			if (ensure(!InAdditionalOperands.IsEmpty()))
-			{
-				const int64 Min = *InPrimaryOperand;
-				const int64 Max = *InAdditionalOperands[0];
-				*OutResult = FMath::RandRange(Min, Max);
-			}
-		}
-	};
 
 	// Definitions
 	DEFINE_METASOUND_MATHOP(Add, Float, float, LOCTEXT("Metasound_MathAddFloatNodeDescription", "Adds floats."))
@@ -1324,12 +1189,6 @@ namespace Metasound
  	DEFINE_METASOUND_MATHOP(Divide, Int64, int64, LOCTEXT("Metasound_MathDivideInt64NodeDescription", "Divide int64 by another int64."))
 
 	DEFINE_METASOUND_OPERAND_TYPED_MATHOP(Divide, Time, FTime, Float, float, LOCTEXT("Metasound_MathDivideTimeNodeDescription", "Divides time by floats."))
-
-	DEFINE_METASOUND_MATHOP(RandRange, Float, float, LOCTEXT("Metasound_MathRandRangeFloatNodeDescription", "Computes random float value in the range [Min, Max]."))
-	DEFINE_METASOUND_MATHOP(RandRange, Int32, int32, LOCTEXT("Metasound_MathRandRangeInt32NodeDescription", "Computes random int32 in the range [Min, Max]."))
- 	DEFINE_METASOUND_MATHOP(RandRange, Int64, int64, LOCTEXT("Metasound_MathRandRangeInt64NodeDescription", "Computes random int64 in the range [Min, Max]."))
-	DEFINE_METASOUND_MATHOP(RandRange, Time, FTime, LOCTEXT("Metasound_MathRandRangeTimeNodeDescription", "Computes random time value in the range [Min, Max]."))
-
 }
 
 #undef LOCTEXT_NAMESPACE // MetasoundMathOpNode
diff --git a/Engine/Plugins/Runtime/Metasound/Source/MetasoundStandardNodes/Public/MetasoundRandomNode.h b/Engine/Plugins/Runtime/Metasound/Source/MetasoundStandardNodes/Public/MetasoundRandomNode.h
new file mode 100644
--- /dev/null
+++ b/Engine/Plugins/Runtime/Metasound/Source/MetasoundStandardNodes/Public/MetasoundRandomNode.h
@@ -0,0 +1,38 @@
+// Copyright Epic Games, Inc. All Rights Reserved.
+
+#pragma once
+
+#include "CoreMinimal.h"
+#include "MetasoundFacade.h"
+
+namespace Metasound
+{
+	/** FRandomIntNode
+	 *
+	 *  Generates a random integer value when triggered.
+	 */
+	class METASOUNDSTANDARDNODES_API FRandomIntNode : public FNodeFacade
+	{
+	public:
+		/**
+		 * Constructor used by the Metasound Frontend.
+		 */
+		FRandomIntNode(const FNodeInitData& InInitData);
+	};
+
+	/** FRandomFloatNode
+	 *
+	 *  Generates a random float value when triggered.
+	 */
+	class METASOUNDSTANDARDNODES_API FRandomFloatNode : public FNodeFacade
+	{
+	public:
+		/**
+		 * Constructor used by the Metasound Frontend.
+		 */
+		FRandomFloatNode(const FNodeInitData& InInitData);
+	};
+
+
+} // namespace Metasound
+
