Fix MetasoundSource transmission bug by adding data type name to address.

diff --git a/Engine/Plugins/Runtime/Metasound/Source/MetasoundEngine/Private/MetasoundSource.cpp b/Engine/Plugins/Runtime/Metasound/Source/MetasoundEngine/Private/MetasoundSource.cpp
--- a/Engine/Plugins/Runtime/Metasound/Source/MetasoundEngine/Private/MetasoundSource.cpp
+++ b/Engine/Plugins/Runtime/Metasound/Source/MetasoundEngine/Private/MetasoundSource.cpp
@@ -372,14 +372,14 @@ Metasound::FOperatorSettings UMetasoundSource::GetOperatorSettings(Metasound::FS
 	return Metasound::FOperatorSettings(InSampleRate, BlockRate);
 }
 
-Metasound::FSendAddress UMetasoundSource::CreateSendAddress(uint64 InInstanceID, const FString& InVertexName) const
+Metasound::FSendAddress UMetasoundSource::CreateSendAddress(uint64 InInstanceID, const FString& InVertexName, const FName& InDataTypeName) const
 {
 	using namespace Metasound;
 
 	FSendAddress Address;
 
 	Address.Subsystem = GetSubsystemNameForSendScope(ETransmissionScope::Global);
-	Address.ChannelName = FName(FString::Printf(TEXT("%d:%s"), InInstanceID, *InVertexName));
+	Address.ChannelName = FName(FString::Printf(TEXT("%d:%s:%s"), InInstanceID, *InVertexName, *InDataTypeName.ToString()));
 
 	return Address;
 }
@@ -397,13 +397,17 @@ TArray<UMetasoundSource::FSendInfoAndVertexName> UMetasoundSource::GetSendInfos(
 	for (const FString& VertexName : SendVertices)
 	{
 		FConstNodeHandle InputNode = RootGraph->GetInputNodeWithName(VertexName);
+		// TODO: re-expose InputNode Inputs. Currently they are hidden since they cannot be 
+		// connected. But instead, they should be accessible through this API, but fail
+		// to connect. As a quick fix, the output handles are used, but this is not
+		// sustainable for situations where the input node is specialized. 
 		for (FConstOutputHandle InputHandle : InputNode->GetConstOutputs())
 		{
 			FSendInfoAndVertexName Info;
 
 			// TODO: incorporate PointID into address. But need to ensure that PointID
 			// will be maintained after injecting Receive nodes. 
-			Info.SendInfo.Address = CreateSendAddress(InInstanceID, InputHandle->GetName());
+			Info.SendInfo.Address = CreateSendAddress(InInstanceID, InputHandle->GetName(), InputHandle->GetDataType());
 			Info.SendInfo.ParameterName = FName(InputHandle->GetDisplayName().ToString()); // TODO: display name hack. Need to have nameing consistent in editor for inputs
 			Info.SendInfo.TypeName = InputHandle->GetDataType();
 			Info.VertexName = VertexName;
diff --git a/Engine/Plugins/Runtime/Metasound/Source/MetasoundEngine/Public/MetasoundSource.h b/Engine/Plugins/Runtime/Metasound/Source/MetasoundEngine/Public/MetasoundSource.h
--- a/Engine/Plugins/Runtime/Metasound/Source/MetasoundEngine/Public/MetasoundSource.h
+++ b/Engine/Plugins/Runtime/Metasound/Source/MetasoundEngine/Public/MetasoundSource.h
@@ -146,7 +146,7 @@ private:
 	bool CopyDocumentAndInjectReceiveNodes(uint64 InInstanceID, const FMetasoundFrontendDocument& InSourceDoc, FMetasoundFrontendDocument& OutDestDoc) const;	
 	TArray<FString> GetTransmittableInputVertexNames() const;
 	Metasound::FOperatorSettings GetOperatorSettings(Metasound::FSampleRate InSampleRate) const;
-	Metasound::FSendAddress CreateSendAddress(uint64 InInstanceID, const FString& InVertexName) const;
+	Metasound::FSendAddress CreateSendAddress(uint64 InInstanceID, const FString& InVertexName, const FName& InDataTypeName) const;
 	
 
 
