Bug fix for incorrect literal type returned in GetDefaultParamForDataType

diff --git a/Engine/Plugins/Runtime/Metasound/Source/MetasoundFrontend/Private/MetasoundFrontendLiteral.cpp b/Engine/Plugins/Runtime/Metasound/Source/MetasoundFrontend/Private/MetasoundFrontendLiteral.cpp
--- a/Engine/Plugins/Runtime/Metasound/Source/MetasoundFrontend/Private/MetasoundFrontendLiteral.cpp
+++ b/Engine/Plugins/Runtime/Metasound/Source/MetasoundFrontend/Private/MetasoundFrontendLiteral.cpp
@@ -524,7 +524,7 @@ namespace Metasound
 			{
 				ELiteralType DesiredArgType = Registry->GetDesiredLiteralTypeForDataType(InDataType);
 
-				FLiteral::GetDefaultForType(DesiredArgType);
+				return FLiteral::GetDefaultForType(DesiredArgType);
 			}
 			return FLiteral::CreateInvalid();
 		}
