Add missing case coverage for literal compatibility check

diff --git a/Engine/Plugins/Runtime/Metasound/Source/MetasoundFrontend/Private/MetasoundFrontendRegistries.cpp b/Engine/Plugins/Runtime/Metasound/Source/MetasoundFrontend/Private/MetasoundFrontendRegistries.cpp
--- a/Engine/Plugins/Runtime/Metasound/Source/MetasoundFrontend/Private/MetasoundFrontendRegistries.cpp
+++ b/Engine/Plugins/Runtime/Metasound/Source/MetasoundFrontend/Private/MetasoundFrontendRegistries.cpp
@@ -351,18 +351,38 @@ bool FMetasoundFrontendRegistryContainer::DoesDataTypeSupportLiteralType(FName I
 		{
 			return DataTypeInfo.Info.bIsBoolParsable;
 		}
+		case Metasound::ELiteralType::BooleanArray:
+		{
+			return DataTypeInfo.Info.bIsBoolArrayParsable;
+		}
+
 		case Metasound::ELiteralType::Integer:
 		{
 			return DataTypeInfo.Info.bIsIntParsable;
 		}
+		case Metasound::ELiteralType::IntegerArray:
+		{
+			return DataTypeInfo.Info.bIsIntArrayParsable;
+		}
+
 		case Metasound::ELiteralType::Float:
 		{
 			return DataTypeInfo.Info.bIsFloatParsable;
 		}
+		case Metasound::ELiteralType::FloatArray:
+		{
+			return DataTypeInfo.Info.bIsFloatArrayParsable;
+		}
+
 		case Metasound::ELiteralType::String:
 		{
 			return DataTypeInfo.Info.bIsStringParsable;
 		}
+		case Metasound::ELiteralType::StringArray:
+		{
+			return DataTypeInfo.Info.bIsStringArrayParsable;
+		}
+
 		case Metasound::ELiteralType::UObjectProxy:
 		{
 			return DataTypeInfo.Info.bIsProxyParsable;
@@ -371,13 +391,20 @@ bool FMetasoundFrontendRegistryContainer::DoesDataTypeSupportLiteralType(FName I
 		{
 			return DataTypeInfo.Info.bIsProxyArrayParsable;
 		}
+
 		case Metasound::ELiteralType::None:
 		{
 			return DataTypeInfo.Info.bIsDefaultParsable;
 		}
+		case Metasound::ELiteralType::NoneArray:
+		{
+			return DataTypeInfo.Info.bIsDefaultArrayParsable;
+		}
+
 		case Metasound::ELiteralType::Invalid:
 		default:
 		{
+			static_assert(static_cast<int32>(Metasound::ELiteralType::COUNT) == 13, "Possible missing case coverage for ELiteralType");
 			return false;
 		}
 	}
diff --git a/Engine/Plugins/Runtime/Metasound/Source/MetasoundGraphCore/Public/MetasoundLiteral.h b/Engine/Plugins/Runtime/Metasound/Source/MetasoundGraphCore/Public/MetasoundLiteral.h
--- a/Engine/Plugins/Runtime/Metasound/Source/MetasoundGraphCore/Public/MetasoundLiteral.h
+++ b/Engine/Plugins/Runtime/Metasound/Source/MetasoundGraphCore/Public/MetasoundLiteral.h
@@ -32,7 +32,9 @@ namespace Metasound
 		FloatArray,
 		StringArray,
 		UObjectProxyArray,
-		Invalid  //< The literal is in an invalid state and cannot be used to construct an object.
+		Invalid,  //< The literal is in an invalid state and cannot be used to construct an object.
+
+		COUNT
 	};
 
 
