Metasound Array Nodes
- Added ArrayGet, ArraySet, ArraySlice, and ArrayConcat nodes
- Array nodes are automatically registered with new data types.
- Moved FTrigger and FTime to MetasoundFrontend module because they are needed for array node registration.
- Added array node unit test

diff --git a/Engine/Plugins/Runtime/Metasound/Source/MetasoundFrontend/Private/MetasoundArrayNodes.cpp b/Engine/Plugins/Runtime/Metasound/Source/MetasoundFrontend/Private/MetasoundArrayNodes.cpp
new file mode 100644
--- /dev/null
+++ b/Engine/Plugins/Runtime/Metasound/Source/MetasoundFrontend/Private/MetasoundArrayNodes.cpp
@@ -0,0 +1,104 @@
+// Copyright Epic Games, Inc. All Rights Reserved.
+
+#include "MetasoundArrayNodes.h"
+
+#define LOCTEXT_NAMESPACE "MetasoundFrontend"
+
+namespace Metasound
+{
+	namespace MetasoundArrayNodesPrivate
+	{
+		FNodeClassMetadata CreateArrayNodeClassMetadata(const FName& InDataTypeName, const FName& InOperatorName, const FText& InDisplayName, const FText& InDescription, const FVertexInterface& InDefaultInterface)
+		{
+			FNodeClassMetadata Metadata
+			{
+				FNodeClassName{FName("Array"), InOperatorName, InDataTypeName},
+				1, // Major Version
+				0, // Minor Version
+				InDisplayName, 
+				InDescription,
+				PluginAuthor,
+				PluginNodeMissingPrompt,
+				InDefaultInterface,
+				{LOCTEXT("ArrayCategory", "Array")},
+				{TEXT("Array")},
+				FNodeDisplayStyle{}
+			};
+
+			return Metadata;
+		}
+	}
+
+	namespace ArrayNodeVertexNames
+	{
+		/* Input Vertx Names */
+		const FString& GetInputArrayName()
+		{
+			static const FString Name = TEXT("Array");
+			return Name;
+		}
+
+		const FString& GetInputLeftArrayName()
+		{
+			static const FString Name = TEXT("Left Array");
+			return Name;
+		}
+
+		const FString& GetInputRightArrayName()
+		{
+			static const FString Name = TEXT("Right Array");
+			return Name;
+		}
+
+		const FString& GetInputTriggerName()
+		{
+			static const FString Name = TEXT("Trigger");
+			return Name;
+		}
+
+		const FString& GetInputIndexName()
+		{
+			static const FString Name = TEXT("Index");
+			return Name;
+		}
+
+		const FString& GetInputStartIndexName()
+		{
+			static const FString Name = TEXT("Start Index");
+			return Name;
+		}
+
+		const FString& GetInputEndIndexName()
+		{
+			static const FString Name = TEXT("End Index");
+			return Name;
+		}
+
+		const FString& GetInputValueName()
+		{
+			static const FString Name = TEXT("Value");
+			return Name;
+		}
+
+		/* Output Vertex Names */
+		const FString& GetOutputNumName()
+		{
+			static const FString Name = TEXT("Num");
+			return Name;
+		}
+
+		const FString& GetOutputValueName()
+		{
+			static const FString Name = TEXT("Element");
+			return Name;
+		}
+
+		const FString& GetOutputArrayName()
+		{
+			static const FString Name = TEXT("Array");
+			return Name;
+		}
+	}
+}
+
+#undef LOCTEXT_NAMESPACE
diff --git a/Engine/Plugins/Runtime/Metasound/Source/MetasoundFrontend/Private/MetasoundFrontend.cpp b/Engine/Plugins/Runtime/Metasound/Source/MetasoundFrontend/Private/MetasoundFrontend.cpp
--- a/Engine/Plugins/Runtime/Metasound/Source/MetasoundFrontend/Private/MetasoundFrontend.cpp
+++ b/Engine/Plugins/Runtime/Metasound/Source/MetasoundFrontend/Private/MetasoundFrontend.cpp
@@ -193,6 +193,8 @@ REGISTER_METASOUND_DATATYPE(float, "Float", ::Metasound::ELiteralType::Float)
 REGISTER_METASOUND_DATATYPE(double, "Double", ::Metasound::ELiteralType::Float)
 REGISTER_METASOUND_DATATYPE(FString, "String", ::Metasound::ELiteralType::String)
 
+REGISTER_METASOUND_DATATYPE(Metasound::FTrigger, "Trigger", ::Metasound::ELiteralType::Boolean)
+REGISTER_METASOUND_DATATYPE(Metasound::FTime, "Time", ::Metasound::ELiteralType::Float)
 REGISTER_METASOUND_DATATYPE(Metasound::FAudioBuffer, "Audio")
 REGISTER_METASOUND_DATATYPE(Metasound::FSendAddress, "Transmission:Address", ::Metasound::ELiteralType::String)
 
diff --git a/Engine/Plugins/Runtime/Metasound/Source/MetasoundFrontend/Private/MetasoundTrigger.cpp b/Engine/Plugins/Runtime/Metasound/Source/MetasoundFrontend/Private/MetasoundTrigger.cpp
new file mode 100644
--- /dev/null
+++ b/Engine/Plugins/Runtime/Metasound/Source/MetasoundFrontend/Private/MetasoundTrigger.cpp
@@ -0,0 +1,121 @@
+// Copyright Epic Games, Inc. All Rights Reserved.
+#include "MetasoundTrigger.h"
+
+#include "MetasoundDataReference.h"
+#include "MetasoundDataTypeRegistrationMacro.h"
+#include "MetasoundOperatorSettings.h"
+
+namespace Metasound
+{
+	FTrigger::FTrigger(const FOperatorSettings& InSettings, bool bShouldTrigger)
+	:	NumFramesPerBlock(InSettings.GetNumFramesPerBlock())
+	,	SampleRate(InSettings.GetSampleRate())
+	{
+		if (bShouldTrigger)
+		{
+			TriggerFrame(0);
+		}
+	}
+
+	FTrigger::FTrigger(const FOperatorSettings& InSettings, int32 InFrameToTrigger)
+	:	NumFramesPerBlock(InSettings.GetNumFramesPerBlock())
+	,	SampleRate(InSettings.GetSampleRate())
+	{
+		TriggerFrame(InFrameToTrigger);
+	}
+
+	FTrigger::FTrigger(const FOperatorSettings& InSettings)
+		: FTrigger(InSettings, false)
+	{}
+
+	void FTrigger::TriggerFrame(int32 InFrameToTrigger)
+	{
+		// Insert trigger frame index into sorted triggered frames
+		TriggeredFrames.Insert(InFrameToTrigger, Algo::LowerBound(TriggeredFrames, InFrameToTrigger));
+
+		if (InFrameToTrigger < NumFramesPerBlock)
+		{
+			// Update last trigger pos. 
+			UpdateLastTriggerIndexInBlock();
+		}
+
+		bHasTrigger = true;
+	}
+
+
+	void FTrigger::AdvanceBlock()
+	{
+		Advance(NumFramesPerBlock);
+	}
+
+	void FTrigger::Advance(int32 InNumFrames)
+	{
+		if (bHasTrigger)
+		{
+			const int32 NumToRemove = Algo::LowerBound(TriggeredFrames, InNumFrames);
+
+			if (NumToRemove)
+			{
+				TriggeredFrames.RemoveAt(0, NumToRemove, false /* bAllowShrinking */);
+			}
+
+			const int32 Num = TriggeredFrames.Num();
+
+			if (Num > 0)
+			{
+				for (int32 i = 0; i < Num; i++)
+				{
+					TriggeredFrames[i] -= InNumFrames;
+				}
+			}
+			else
+			{
+				bHasTrigger = false;
+			}
+
+			UpdateLastTriggerIndexInBlock();
+		}
+	}
+
+	int32 FTrigger::Num() const
+	{
+		return TriggeredFrames.Num();
+	}
+
+	int32 FTrigger::NumTriggeredInBlock() const
+	{
+		return LastTriggerIndexInBlock;
+	}
+
+	int32 FTrigger::operator[](int32 InTriggerIndex) const
+	{
+		return TriggeredFrames[InTriggerIndex];
+	}
+
+	bool FTrigger::IsTriggered() const
+	{
+		return bHasTrigger;
+	}
+
+	bool FTrigger::IsTriggeredInBlock() const
+	{
+		return LastTriggerIndexInBlock > 0;
+	}
+
+	FTrigger::operator bool() const 
+	{
+		return IsTriggeredInBlock();
+	}
+
+	void FTrigger::Reset()
+	{
+		TriggeredFrames.Reset();
+		LastTriggerIndexInBlock = 0;
+		bHasTrigger = false;
+	}
+
+	void FTrigger::UpdateLastTriggerIndexInBlock()
+	{
+		LastTriggerIndexInBlock = Algo::LowerBound(TriggeredFrames, NumFramesPerBlock);
+	}
+}
diff --git a/Engine/Plugins/Runtime/Metasound/Source/MetasoundFrontend/Public/MetasoundArrayNodes.h b/Engine/Plugins/Runtime/Metasound/Source/MetasoundFrontend/Public/MetasoundArrayNodes.h
new file mode 100644
--- /dev/null
+++ b/Engine/Plugins/Runtime/Metasound/Source/MetasoundFrontend/Public/MetasoundArrayNodes.h
@@ -0,0 +1,747 @@
+// Copyright Epic Games, Inc. All Rights Reserved.
+
+#pragma once
+
+#include "CoreMinimal.h"
+#include "Internationalization/Text.h"
+#include "MetasoundBuilderInterface.h"
+#include "MetasoundDataFactory.h"
+#include "MetasoundExecutableOperator.h"
+#include "MetasoundFacade.h"
+#include "MetasoundLog.h"
+#include "MetasoundNodeInterface.h"
+#include "MetasoundNodeRegistrationMacro.h"
+#include "MetasoundOperatorInterface.h"
+#include "MetasoundPrimitives.h"
+#include "MetasoundTrigger.h"
+
+#include <type_traits>
+
+#define LOCTEXT_NAMESPACE "MetasoundFrontend"
+
+
+namespace Metasound
+{
+	namespace MetasoundArrayNodesPrivate
+	{
+		// Convenience function for make FNodeClassMetadata of array nodes.
+		METASOUNDFRONTEND_API FNodeClassMetadata CreateArrayNodeClassMetadata(const FName& InDataTypeName, const FName& InOperatorName, const FText& InDisplayName, const FText& InDescription, const FVertexInterface& InDefaultInterface);
+
+		// Retrieve the ElementType from an ArrayType
+		template<typename ArrayType>
+		struct TArrayElementType
+		{
+			// Default implementation has Type. 
+		};
+
+		// ElementType specialization for TArray types.
+		template<typename ElementType>
+		struct TArrayElementType<TArray<ElementType>>
+		{
+			using Type = ElementType;
+		};
+		
+		// TODO: Remove me and swap after fixing FDataReferenceCollection::GetDataReadReferenceOrConstructWithVertexDefault
+		//
+		// This existing implementation of `GetDataReadReferenceOrConstructWithVertexDefault` does not automatically
+		// utilize the set type of the FLiteral or the TDataReadReferenceLiteralFactory 
+		// for that type. The implementation in FDataReferenceCollection should be replaced
+		// with this implementation. It will require reworking of all nodes which
+		// call `GetDataReadReferenceOrConstructWithVertexDefault`.
+		template<typename DataType>
+		TDataReadReference<DataType> GetDataReadReferenceOrConstructWithVertexDefault(const FDataReferenceCollection& InCollection, const FInputVertexInterface& InputVertices, const FString& InName, const FOperatorSettings& InSettings)
+		{
+			using FDataFactory = TDataReadReferenceLiteralFactory<DataType>;
+
+			if (InCollection.ContainsDataReadReference<DataType>(InName))
+			{
+				return InCollection.GetDataReadReference<DataType>(InName);
+			}
+			else
+			{
+				// TODO: replace mechanism in data vertex to "CreateDefaultValue"
+				// instead of get in order to avoid the need to have a `Clone` on 
+				// FLiteral. Then move FLiteral when needed.
+				const FLiteral& Literal = InputVertices[InName].GetDefaultValue();
+				return FDataFactory::CreateExplicitArgs(InSettings, Literal);
+			}
+		}
+	}
+
+	namespace ArrayNodeVertexNames
+	{
+		/* Input Vertex Names */
+		METASOUNDFRONTEND_API const FString& GetInputArrayName();
+		METASOUNDFRONTEND_API const FString& GetInputLeftArrayName();
+		METASOUNDFRONTEND_API const FString& GetInputRightArrayName();
+		METASOUNDFRONTEND_API const FString& GetInputTriggerName();
+		METASOUNDFRONTEND_API const FString& GetInputStartIndexName();
+		METASOUNDFRONTEND_API const FString& GetInputEndIndexName();
+		METASOUNDFRONTEND_API const FString& GetInputIndexName();
+		METASOUNDFRONTEND_API const FString& GetInputValueName();
+
+		/* Output Vertex Names */
+		METASOUNDFRONTEND_API const FString& GetOutputNumName();
+		METASOUNDFRONTEND_API const FString& GetOutputValueName();
+		METASOUNDFRONTEND_API const FString& GetOutputArrayName();
+	};
+
+	/** TArrayNumOperator gets the number of elements in an Array. The operator
+	 * uses the FNodeFacade and defines the vertex, metadata and vertex interface
+	 * statically on the operator class. */
+	template<typename ArrayType>
+	class TArrayNumOperator : public TExecutableOperator<TArrayNumOperator<ArrayType>>
+	{
+	public:
+		using FArrayDataReadReference = TDataReadReference<ArrayType>;
+
+		// Declare the vertex interface
+		static const FVertexInterface& GetDefaultInterface()
+		{
+			using namespace ArrayNodeVertexNames;
+
+			static const FVertexInterface DefaultInterface(
+				FInputVertexInterface(
+					TInputDataVertexModel<ArrayType>(GetInputArrayName(), LOCTEXT("ArrayOpArrayNumInput", "Array to inspect."))
+				),
+				FOutputVertexInterface(
+					TOutputDataVertexModel<int32>(GetOutputNumName(), LOCTEXT("ArrayOpArrayNumOutput", "Number of elements in the array."))
+				)
+			);
+
+			return DefaultInterface;
+		}
+
+		static const FNodeClassMetadata& GetNodeInfo()
+		{
+			auto CreateNodeClassMetadata = []() -> FNodeClassMetadata
+			{
+				FName DataTypeName = GetMetasoundDataTypeName<ArrayType>();
+				FName OperatorName = TEXT("Num"); 
+				FText NodeDisplayName = FText::Format(LOCTEXT("ArrayOpArrayNumDisplayNamePattern", "Num ({0})"), FText::FromString(GetMetasoundDataTypeString<ArrayType>()));
+				FText NodeDescription = LOCTEXT("ArrayOpArrayNumDescription", "Number of elements in the array");
+				FVertexInterface NodeInterface = GetDefaultInterface();
+			
+				return MetasoundArrayNodesPrivate::CreateArrayNodeClassMetadata(DataTypeName, OperatorName, NodeDisplayName, NodeDescription, NodeInterface);
+			};
+			
+			static const FNodeClassMetadata Metadata = CreateNodeClassMetadata();
+
+			return Metadata;
+		}
+
+		static TUniquePtr<IOperator> CreateOperator(const FCreateOperatorParams& InParams, TArray<TUniquePtr<IOperatorBuildError>>& OutErrors)
+		{
+			using namespace ArrayNodeVertexNames;
+			using namespace MetasoundArrayNodesPrivate;
+
+			const FInputVertexInterface& Inputs = InParams.Node.GetVertexInterface().GetInputInterface();
+
+			// Get the input array or construct an empty one. 
+			FArrayDataReadReference Array = GetDataReadReferenceOrConstructWithVertexDefault<ArrayType>(InParams.InputDataReferences, Inputs, GetInputArrayName(), InParams.OperatorSettings);
+
+			return MakeUnique<TArrayNumOperator>(Array);
+		}
+
+		TArrayNumOperator(FArrayDataReadReference InArray)
+		: Array(InArray)
+		, Num(TDataWriteReference<int32>::CreateNew())
+		{
+			// Initialize value for downstream nodes.
+			*Num = Array->Num();
+		}
+
+		virtual ~TArrayNumOperator() = default;
+
+		virtual FDataReferenceCollection GetInputs() const override
+		{
+			using namespace ArrayNodeVertexNames;
+
+			FDataReferenceCollection Inputs;
+
+			Inputs.AddDataReadReference(GetInputArrayName(), Array);
+
+			return Inputs;
+		}
+
+		virtual FDataReferenceCollection GetOutputs() const override
+		{
+			using namespace ArrayNodeVertexNames;
+
+			FDataReferenceCollection Outputs;
+
+			Outputs.AddDataReadReference(GetOutputNumName(), Num);
+
+			return Outputs;
+		}
+
+		void Execute()
+		{
+			*Num = Array->Num();
+		}
+
+	private:
+
+		FArrayDataReadReference Array;
+		TDataWriteReference<int32> Num;
+	};
+
+	template<typename ArrayType>
+	class TArrayNumNode : public FNodeFacade
+	{
+	public:
+		TArrayNumNode(const FNodeInitData& InInitData)
+		: FNodeFacade(InInitData.InstanceName, InInitData.InstanceID, TFacadeOperatorClass<TArrayNumOperator<ArrayType>>())
+		{
+		}
+
+		virtual ~TArrayNumNode() = default;
+	};
+
+	/** TArrayGetOperator copies a value from the array to the output when
+	 * a trigger occurs. Initially, the output value is default constructed and
+	 * will remain that way until until a trigger is encountered.
+	 */
+	template<typename ArrayType>
+	class TArrayGetOperator : public TExecutableOperator<TArrayGetOperator<ArrayType>>
+	{
+	public:
+		using FArrayDataReadReference = TDataReadReference<ArrayType>;
+		using ElementType = typename MetasoundArrayNodesPrivate::TArrayElementType<ArrayType>::Type;
+
+		static const FVertexInterface& GetDefaultInterface()
+		{
+			using namespace ArrayNodeVertexNames;
+			static const FVertexInterface DefaultInterface(
+				FInputVertexInterface(
+					TInputDataVertexModel<FTrigger>(GetInputTriggerName(), LOCTEXT("ArrayOpArrayGetTrigger", "Trigger to get value.")),
+					TInputDataVertexModel<ArrayType>(GetInputArrayName(), LOCTEXT("ArrayOpArrayGetInput", "Input Array.")),
+					TInputDataVertexModel<int32>(GetInputIndexName(), LOCTEXT("ArrayOpArrayGetIndex", "Index in Array."))
+				),
+				FOutputVertexInterface(
+					TOutputDataVertexModel<ElementType>(GetOutputValueName(), LOCTEXT("ArrayOpArrayGetOutput", "Value of element at array index."))
+				)
+			);
+
+			return DefaultInterface;
+		}
+
+		static const FNodeClassMetadata& GetNodeInfo()
+		{
+			auto CreateNodeClassMetadata = []() -> FNodeClassMetadata
+			{
+				FName DataTypeName = GetMetasoundDataTypeName<ArrayType>();
+				FName OperatorName = TEXT("Get"); 
+				FText NodeDisplayName = FText::Format(LOCTEXT("ArrayOpArrayGetDisplayNamePattern", "Get ({0})"), FText::FromString(GetMetasoundDataTypeString<ArrayType>()));
+				FText NodeDescription = LOCTEXT("ArrayOpArrayGetDescription", "Get element at index in array.");
+				FVertexInterface NodeInterface = GetDefaultInterface();
+			
+				return MetasoundArrayNodesPrivate::CreateArrayNodeClassMetadata(DataTypeName, OperatorName, NodeDisplayName, NodeDescription, NodeInterface);
+			};
+			
+			static const FNodeClassMetadata Metadata = CreateNodeClassMetadata();
+
+			return Metadata;
+		}
+
+		static TUniquePtr<IOperator> CreateOperator(const FCreateOperatorParams& InParams, TArray<TUniquePtr<IOperatorBuildError>>& OutErrors)
+		{
+			using namespace ArrayNodeVertexNames;
+			using namespace MetasoundArrayNodesPrivate;
+
+			const FInputVertexInterface& Inputs = InParams.Node.GetVertexInterface().GetInputInterface();
+
+			// Input Trigger
+			TDataReadReference<FTrigger> Trigger = GetDataReadReferenceOrConstructWithVertexDefault<FTrigger>(InParams.InputDataReferences, Inputs, GetInputTriggerName(), InParams.OperatorSettings);
+			
+			// Input Array
+			FArrayDataReadReference Array = GetDataReadReferenceOrConstructWithVertexDefault<ArrayType>(InParams.InputDataReferences, Inputs, GetInputArrayName(), InParams.OperatorSettings);
+
+			// Input Index
+			TDataReadReference<int32> Index = GetDataReadReferenceOrConstructWithVertexDefault<int32>(InParams.InputDataReferences, Inputs, GetInputIndexName(), InParams.OperatorSettings);
+
+			return MakeUnique<TArrayGetOperator>(InParams.OperatorSettings, Trigger, Array, Index);
+		}
+
+
+		TArrayGetOperator(const FOperatorSettings& InSettings, TDataReadReference<FTrigger> InTrigger, FArrayDataReadReference InArray, TDataReadReference<int32> InIndex)
+		: Trigger(InTrigger)
+		, Array(InArray)
+		, Index(InIndex)
+		, Value(TDataWriteReferenceFactory<ElementType>::CreateAny(InSettings))
+		{
+		}
+
+		virtual ~TArrayGetOperator() = default;
+
+		virtual FDataReferenceCollection GetInputs() const override
+		{
+			using namespace ArrayNodeVertexNames;
+
+			FDataReferenceCollection Inputs;
+
+			Inputs.AddDataReadReference(GetInputTriggerName(), Trigger);
+			Inputs.AddDataReadReference(GetInputArrayName(), Array);
+			Inputs.AddDataReadReference(GetInputIndexName(), Index);
+
+			return Inputs;
+		}
+
+		virtual FDataReferenceCollection GetOutputs() const override
+		{
+			using namespace ArrayNodeVertexNames;
+
+			FDataReferenceCollection Outputs;
+
+			Outputs.AddDataReadReference(GetOutputValueName(), Value);
+
+			return Outputs;
+		}
+
+		void Execute()
+		{
+			// Only perform get on trigger.
+			if (*Trigger)
+			{
+				const int32 IndexValue = *Index;
+				const ArrayType& ArrayRef = *Array;
+
+				if (IndexValue < ArrayRef.Num())
+				{
+					*Value = ArrayRef[IndexValue];
+				}
+				else
+				{
+					UE_LOG(LogMetasound, Error, TEXT("Attempt to get value at invalid index [ArraySize:%d, Index:%d]"), ArrayRef.Num(), IndexValue);
+				}
+			}
+		}
+
+	private:
+
+		TDataReadReference<FTrigger> Trigger;
+		FArrayDataReadReference Array;
+		TDataReadReference<int32> Index;
+		TDataWriteReference<ElementType> Value;
+	};
+
+	template<typename ArrayType>
+	class TArrayGetNode : public FNodeFacade
+	{
+	public:
+		TArrayGetNode(const FNodeInitData& InInitData)
+		: FNodeFacade(InInitData.InstanceName, InInitData.InstanceID, TFacadeOperatorClass<TArrayGetOperator<ArrayType>>())
+		{
+		}
+
+		virtual ~TArrayGetNode() = default;
+	};
+
+	/** TArraySetOperator sets an element in an array to a specific value. */
+	template<typename ArrayType>
+	class TArraySetOperator : public TExecutableOperator<TArraySetOperator<ArrayType>>
+	{
+	public:
+		using FArrayDataReadReference = TDataReadReference<ArrayType>;
+		using FArrayDataWriteReference = TDataWriteReference<ArrayType>;
+		using ElementType = typename MetasoundArrayNodesPrivate::TArrayElementType<ArrayType>::Type;
+
+		static const FVertexInterface& GetDefaultInterface()
+		{
+			using namespace ArrayNodeVertexNames;
+			static const FVertexInterface DefaultInterface(
+				FInputVertexInterface(
+					TInputDataVertexModel<FTrigger>(GetInputTriggerName(), LOCTEXT("ArrayOpArraySetTrigger", "Trigger to set value.")),
+					TInputDataVertexModel<ArrayType>(GetInputArrayName(), LOCTEXT("ArrayOpArraySetInput", "Input Array.")),
+					TInputDataVertexModel<int32>(GetInputIndexName(), LOCTEXT("ArrayOpArraySetIndex", "Index in Array.")),
+					TInputDataVertexModel<ElementType>(GetInputValueName(), LOCTEXT("ArrayOpArraySetElement", "Value to set"))
+				),
+				FOutputVertexInterface(
+					TOutputDataVertexModel<ArrayType>(GetOutputArrayName(), LOCTEXT("ArrayOpArraySetOutput", "Array after setting."))
+				)
+			);
+
+			return DefaultInterface;
+		}
+
+		static const FNodeClassMetadata& GetNodeInfo()
+		{
+			auto CreateNodeClassMetadata = []() -> FNodeClassMetadata
+			{
+				FName DataTypeName = GetMetasoundDataTypeName<ArrayType>();
+				FName OperatorName = TEXT("Set"); 
+				FText NodeDisplayName = FText::Format(LOCTEXT("ArrayOpArraySetDisplayNamePattern", "Set ({0})"), FText::FromString(GetMetasoundDataTypeString<ArrayType>()));
+				FText NodeDescription = LOCTEXT("ArrayOpArraySetDescription", "Set element at index in array.");
+				FVertexInterface NodeInterface = GetDefaultInterface();
+			
+				return MetasoundArrayNodesPrivate::CreateArrayNodeClassMetadata(DataTypeName, OperatorName, NodeDisplayName, NodeDescription, NodeInterface);
+			};
+			
+			static const FNodeClassMetadata Metadata = CreateNodeClassMetadata();
+
+			return Metadata;
+		}
+
+		static TUniquePtr<IOperator> CreateOperator(const FCreateOperatorParams& InParams, TArray<TUniquePtr<IOperatorBuildError>>& OutErrors)
+		{
+			using namespace ArrayNodeVertexNames;
+			using namespace MetasoundArrayNodesPrivate;
+
+			const FInputVertexInterface& Inputs = InParams.Node.GetVertexInterface().GetInputInterface();
+			
+			TDataReadReference<FTrigger> Trigger = GetDataReadReferenceOrConstructWithVertexDefault<FTrigger>(InParams.InputDataReferences, Inputs, GetInputTriggerName(), InParams.OperatorSettings);
+
+			FArrayDataReadReference InitArray = GetDataReadReferenceOrConstructWithVertexDefault<ArrayType>(InParams.InputDataReferences, Inputs, GetInputArrayName(), InParams.OperatorSettings);
+			FArrayDataWriteReference Array = TDataWriteReferenceFactory<ArrayType>::CreateExplicitArgs(InParams.OperatorSettings, *InitArray);
+
+			TDataReadReference<int32> Index = GetDataReadReferenceOrConstructWithVertexDefault<int32>(InParams.InputDataReferences, Inputs, GetInputIndexName(), InParams.OperatorSettings);
+
+			TDataReadReference<ElementType> Value = GetDataReadReferenceOrConstructWithVertexDefault<ElementType>(InParams.InputDataReferences, Inputs, GetInputValueName(), InParams.OperatorSettings);
+
+			return MakeUnique<TArraySetOperator>(InParams.OperatorSettings, Trigger, InitArray, Array, Index, Value);
+		}
+
+
+		TArraySetOperator(const FOperatorSettings& InSettings, TDataReadReference<FTrigger> InTrigger, FArrayDataReadReference InInitArray, FArrayDataWriteReference InArray, TDataReadReference<int32> InIndex, TDataReadReference<ElementType> InValue)
+		: OperatorSettings(InSettings)
+		, Trigger(InTrigger)
+		, InitArray(InInitArray)
+		, Array(InArray)
+		, Index(InIndex)
+		, Value(InValue)
+		{
+		}
+
+		virtual ~TArraySetOperator() = default;
+
+		virtual FDataReferenceCollection GetInputs() const override
+		{
+			using namespace ArrayNodeVertexNames;
+			FDataReferenceCollection Inputs;
+
+			Inputs.AddDataReadReference(GetInputTriggerName(), Trigger);
+			Inputs.AddDataReadReference(GetInputArrayName(), InitArray);
+			Inputs.AddDataReadReference(GetInputIndexName(), Index);
+			Inputs.AddDataReadReference(GetInputValueName(), Value);
+
+			return Inputs;
+		}
+
+		virtual FDataReferenceCollection GetOutputs() const override
+		{
+			using namespace ArrayNodeVertexNames;
+			FDataReferenceCollection Outputs;
+
+			Outputs.AddDataReadReference(GetOutputArrayName(), Array);
+
+			return Outputs;
+		}
+
+		void Execute()
+		{
+			if (*Trigger)
+			{
+				const int32 IndexValue = *Index;
+				ArrayType& ArrayRef = *Array;
+
+				if ((IndexValue >= 0) && (IndexValue < ArrayRef.Num()))
+				{
+					ArrayRef[IndexValue] = *Value;
+				}
+				else
+				{
+					UE_LOG(LogMetasound, Error, TEXT("Attempt to set value at invalid index [ArraySize:%d, Index:%d]"), ArrayRef.Num(), IndexValue);
+				}
+			}
+		}
+
+	private:
+		FOperatorSettings OperatorSettings;
+
+		TDataReadReference<FTrigger> Trigger;
+		FArrayDataReadReference InitArray;
+		FArrayDataWriteReference Array;
+		TDataReadReference<int32> Index;
+		TDataReadReference<ElementType> Value;
+	};
+
+	template<typename ArrayType>
+	class TArraySetNode : public FNodeFacade
+	{
+	public:
+		TArraySetNode(const FNodeInitData& InInitData)
+		: FNodeFacade(InInitData.InstanceName, InInitData.InstanceID, TFacadeOperatorClass<TArraySetOperator<ArrayType>>())
+		{
+		}
+
+		virtual ~TArraySetNode() = default;
+	};
+
+	/** TArrayConcatOperator concatenates two arrays on trigger. */
+	template<typename ArrayType>
+	class TArrayConcatOperator : public TExecutableOperator<TArrayConcatOperator<ArrayType>>
+	{
+	public:
+		using FArrayDataReadReference = TDataReadReference<ArrayType>;
+		using FArrayDataWriteReference = TDataWriteReference<ArrayType>;
+		using ElementType = typename MetasoundArrayNodesPrivate::TArrayElementType<ArrayType>::Type;
+
+		static const FVertexInterface& GetDefaultInterface()
+		{
+			using namespace ArrayNodeVertexNames;
+
+			static const FVertexInterface DefaultInterface(
+				FInputVertexInterface(
+					TInputDataVertexModel<FTrigger>(GetInputTriggerName(), LOCTEXT("ArrayOpArrayConcatTrigger", "Trigger to set value.")),
+					TInputDataVertexModel<ArrayType>(GetInputLeftArrayName(), LOCTEXT("ArrayOpArrayConcatInputLeft", "Input Left Array.")),
+					TInputDataVertexModel<ArrayType>(GetInputRightArrayName(), LOCTEXT("ArrayOpArrayConcatInputRight", "Input Left Array."))
+				),
+				FOutputVertexInterface(
+					TOutputDataVertexModel<ArrayType>(GetOutputArrayName(), LOCTEXT("ArrayOpArrayConcatOutput", "Array after setting."))
+				)
+			);
+
+			return DefaultInterface;
+		}
+
+		static const FNodeClassMetadata& GetNodeInfo()
+		{
+			auto CreateNodeClassMetadata = []() -> FNodeClassMetadata
+			{
+				FName DataTypeName = GetMetasoundDataTypeName<ArrayType>();
+				FName OperatorName = TEXT("Concat"); 
+				FText NodeDisplayName = FText::Format(LOCTEXT("ArrayOpArrayConcatDisplayNamePattern", "Concatenate ({0})"), FText::FromString(GetMetasoundDataTypeString<ArrayType>()));
+				FText NodeDescription = LOCTEXT("ArrayOpArrayConcatDescription", "Concatenates two arrays on trigger.");
+				FVertexInterface NodeInterface = GetDefaultInterface();
+			
+				return MetasoundArrayNodesPrivate::CreateArrayNodeClassMetadata(DataTypeName, OperatorName, NodeDisplayName, NodeDescription, NodeInterface);
+			};
+			
+			static const FNodeClassMetadata Metadata = CreateNodeClassMetadata();
+
+			return Metadata;
+		}
+
+		static TUniquePtr<IOperator> CreateOperator(const FCreateOperatorParams& InParams, TArray<TUniquePtr<IOperatorBuildError>>& OutErrors)
+		{
+			using namespace ArrayNodeVertexNames;
+			using namespace MetasoundArrayNodesPrivate;
+
+			const FInputVertexInterface& Inputs = InParams.Node.GetVertexInterface().GetInputInterface();
+			
+			TDataReadReference<FTrigger> Trigger = GetDataReadReferenceOrConstructWithVertexDefault<FTrigger>(InParams.InputDataReferences, Inputs, GetInputTriggerName(), InParams.OperatorSettings);
+
+			FArrayDataReadReference LeftArray = GetDataReadReferenceOrConstructWithVertexDefault<ArrayType>(InParams.InputDataReferences, Inputs, GetInputLeftArrayName(), InParams.OperatorSettings);
+			FArrayDataReadReference RightArray = GetDataReadReferenceOrConstructWithVertexDefault<ArrayType>(InParams.InputDataReferences, Inputs, GetInputRightArrayName(), InParams.OperatorSettings);
+
+			FArrayDataWriteReference OutArray = TDataWriteReferenceFactory<ArrayType>::CreateExplicitArgs(InParams.OperatorSettings);
+
+			return MakeUnique<TArrayConcatOperator>(Trigger, LeftArray, RightArray, OutArray);
+		}
+
+
+		TArrayConcatOperator(TDataReadReference<FTrigger> InTrigger, FArrayDataReadReference InLeftArray, FArrayDataReadReference InRightArray, FArrayDataWriteReference InOutArray)
+		: Trigger(InTrigger)
+		, LeftArray(InLeftArray)
+		, RightArray(InRightArray)
+		, OutArray(InOutArray)
+		{
+		}
+
+		virtual ~TArrayConcatOperator() = default;
+
+		virtual FDataReferenceCollection GetInputs() const override
+		{
+			using namespace ArrayNodeVertexNames;
+			FDataReferenceCollection Inputs;
+
+			Inputs.AddDataReadReference(GetInputTriggerName(), Trigger);
+			Inputs.AddDataReadReference(GetInputLeftArrayName(), LeftArray);
+			Inputs.AddDataReadReference(GetInputRightArrayName(), RightArray);
+
+			return Inputs;
+		}
+
+		virtual FDataReferenceCollection GetOutputs() const override
+		{
+			using namespace ArrayNodeVertexNames;
+			FDataReferenceCollection Outputs;
+
+			Outputs.AddDataReadReference(GetOutputArrayName(), OutArray);
+
+			return Outputs;
+		}
+
+		void Execute()
+		{
+			if (*Trigger)
+			{
+				*OutArray = *LeftArray;
+				OutArray->Append(*RightArray);
+			}
+		}
+
+	private:
+
+		TDataReadReference<FTrigger> Trigger;
+		FArrayDataReadReference LeftArray;
+		FArrayDataReadReference RightArray;
+		FArrayDataWriteReference OutArray;
+	};
+
+	template<typename ArrayType>
+	class TArrayConcatNode : public FNodeFacade
+	{
+	public:
+		TArrayConcatNode(const FNodeInitData& InInitData)
+		: FNodeFacade(InInitData.InstanceName, InInitData.InstanceID, TFacadeOperatorClass<TArrayConcatOperator<ArrayType>>())
+		{
+		}
+
+		virtual ~TArrayConcatNode() = default;
+	};
+
+	/** TArraySliceOperator slices an array on trigger. */
+	template<typename ArrayType>
+	class TArraySliceOperator : public TExecutableOperator<TArraySliceOperator<ArrayType>>
+	{
+	public:
+		using FArrayDataReadReference = TDataReadReference<ArrayType>;
+		using FArrayDataWriteReference = TDataWriteReference<ArrayType>;
+		using ElementType = typename MetasoundArrayNodesPrivate::TArrayElementType<ArrayType>::Type;
+
+		static const FVertexInterface& GetDefaultInterface()
+		{
+			using namespace ArrayNodeVertexNames;
+
+			static const FVertexInterface DefaultInterface(
+				FInputVertexInterface(
+					TInputDataVertexModel<FTrigger>(GetInputTriggerName(), LOCTEXT("ArrayOpArraySliceTrigger", "Trigger to set value.")),
+					TInputDataVertexModel<ArrayType>(GetInputArrayName(), LOCTEXT("ArrayOpArraySliceInputLeft", "Input Array.")),
+					TInputDataVertexModel<int32>(GetInputStartIndexName(), LOCTEXT("ArrayOpArraySliceStartIndex", "First index to include.")),
+					TInputDataVertexModel<int32>(GetInputEndIndexName(), LOCTEXT("ArrayOpArraySliceEndIndex", "Last index to include."))
+
+				),
+				FOutputVertexInterface(
+					TOutputDataVertexModel<ArrayType>(GetOutputArrayName(), LOCTEXT("ArrayOpArraySliceOutput", "Array after slicing."))
+				)
+			);
+
+			return DefaultInterface;
+		}
+
+		static const FNodeClassMetadata& GetNodeInfo()
+		{
+			auto CreateNodeClassMetadata = []() -> FNodeClassMetadata
+			{
+				FName DataTypeName = GetMetasoundDataTypeName<ArrayType>();
+				FName OperatorName = TEXT("Slice"); 
+				FText NodeDisplayName = FText::Format(LOCTEXT("ArrayOpArraySliceDisplayNamePattern", "Slice ({0})"), FText::FromString(GetMetasoundDataTypeString<ArrayType>()));
+				FText NodeDescription = LOCTEXT("ArrayOpArraySliceDescription", "Slice array on trigger.");
+				FVertexInterface NodeInterface = GetDefaultInterface();
+			
+				return MetasoundArrayNodesPrivate::CreateArrayNodeClassMetadata(DataTypeName, OperatorName, NodeDisplayName, NodeDescription, NodeInterface);
+			};
+			
+			static const FNodeClassMetadata Metadata = CreateNodeClassMetadata();
+
+			return Metadata;
+		}
+
+		static TUniquePtr<IOperator> CreateOperator(const FCreateOperatorParams& InParams, TArray<TUniquePtr<IOperatorBuildError>>& OutErrors)
+		{
+			using namespace ArrayNodeVertexNames;
+			using namespace MetasoundArrayNodesPrivate;
+
+			const FInputVertexInterface& Inputs = InParams.Node.GetVertexInterface().GetInputInterface();
+			
+			TDataReadReference<FTrigger> Trigger = GetDataReadReferenceOrConstructWithVertexDefault<FTrigger>(InParams.InputDataReferences, Inputs, GetInputTriggerName(), InParams.OperatorSettings);
+
+			FArrayDataReadReference InArray = GetDataReadReferenceOrConstructWithVertexDefault<ArrayType>(InParams.InputDataReferences, Inputs, GetInputArrayName(), InParams.OperatorSettings);
+
+			TDataReadReference<int32> StartIndex = GetDataReadReferenceOrConstructWithVertexDefault<int32>(InParams.InputDataReferences, Inputs, GetInputStartIndexName(), InParams.OperatorSettings);
+			TDataReadReference<int32> EndIndex = GetDataReadReferenceOrConstructWithVertexDefault<int32>(InParams.InputDataReferences, Inputs, GetInputEndIndexName(), InParams.OperatorSettings);
+
+			FArrayDataWriteReference OutArray = TDataWriteReferenceFactory<ArrayType>::CreateExplicitArgs(InParams.OperatorSettings);
+
+			return MakeUnique<TArraySliceOperator>(Trigger, InArray, StartIndex, EndIndex, OutArray);
+		}
+
+
+		TArraySliceOperator(TDataReadReference<FTrigger> InTrigger, FArrayDataReadReference InInputArray, TDataReadReference<int32> InStartIndex, TDataReadReference<int32> InEndIndex, FArrayDataWriteReference InOutputArray)
+		: Trigger(InTrigger)
+		, InputArray(InInputArray)
+		, StartIndex(InStartIndex)
+		, EndIndex(InEndIndex)
+		, OutputArray(InOutputArray)
+		{
+		}
+
+		virtual ~TArraySliceOperator() = default;
+
+		virtual FDataReferenceCollection GetInputs() const override
+		{
+			using namespace ArrayNodeVertexNames;
+
+			FDataReferenceCollection Inputs;
+
+			Inputs.AddDataReadReference(GetInputTriggerName(), Trigger);
+			Inputs.AddDataReadReference(GetInputArrayName(), InputArray);
+			Inputs.AddDataReadReference(GetInputStartIndexName(), StartIndex);
+			Inputs.AddDataReadReference(GetInputEndIndexName(), EndIndex);
+
+			return Inputs;
+		}
+
+		virtual FDataReferenceCollection GetOutputs() const override
+		{
+			using namespace ArrayNodeVertexNames;
+
+			FDataReferenceCollection Outputs;
+
+			Outputs.AddDataReadReference(GetOutputArrayName(), OutputArray);
+
+			return Outputs;
+		}
+
+		void Execute()
+		{
+			if (*Trigger)
+			{
+				OutputArray->Reset();
+
+				const ArrayType& InputArrayRef = *InputArray;
+				const int32 StartIndexValue = FMath::Max(0, *StartIndex);
+				const int32 EndIndexValue = FMath::Min(InputArrayRef.Num(), *EndIndex);
+
+				if (StartIndexValue < EndIndexValue)
+				{
+					const int32 Num = EndIndexValue - StartIndexValue;
+					OutputArray->Append(&InputArrayRef[StartIndexValue], Num);
+				}
+			}
+		}
+
+	private:
+
+		TDataReadReference<FTrigger> Trigger;
+		FArrayDataReadReference InputArray;
+		TDataReadReference<int32> StartIndex;
+		TDataReadReference<int32> EndIndex;
+		FArrayDataWriteReference OutputArray;
+	};
+
+	template<typename ArrayType>
+	class TArraySliceNode : public FNodeFacade
+	{
+	public:
+		TArraySliceNode(const FNodeInitData& InInitData)
+		: FNodeFacade(InInitData.InstanceName, InInitData.InstanceID, TFacadeOperatorClass<TArraySliceOperator<ArrayType>>())
+		{
+		}
+
+		virtual ~TArraySliceNode() = default;
+	};
+}
+
+#undef LOCTEXT_NAMESPACE
diff --git a/Engine/Plugins/Runtime/Metasound/Source/MetasoundFrontend/Public/MetasoundArrayNodesRegistration.h b/Engine/Plugins/Runtime/Metasound/Source/MetasoundFrontend/Public/MetasoundArrayNodesRegistration.h
new file mode 100644
--- /dev/null
+++ b/Engine/Plugins/Runtime/Metasound/Source/MetasoundFrontend/Public/MetasoundArrayNodesRegistration.h
@@ -0,0 +1,134 @@
+// Copyright Epic Games, Inc. All Rights Reserved.
+
+#pragma once
+
+#include "MetasoundArrayNodes.h"
+#include "MetasoundNodeRegistrationMacro.h"
+
+#include <type_traits>
+
+namespace Metasound
+{
+
+	namespace MetasoundArrayNodesPrivate
+	{
+		// TArrayNodeSupport acts as a configuration sturct to determine whether
+		// a particular TArrayNode can be instantiated for a specific ArrayType.
+		//
+		// Some ArrayNodes require that the array elements have certain properties
+		// such as default element constructors, element copy constructors, etc.
+		template<typename ArrayType>
+		struct TArrayNodeSupport
+		{
+			using ElementType = typename TArrayElementType<ArrayType>::Type;
+			
+			// Array num is supported for all array types.
+			static constexpr bool bIsArrayNumSupported = true;
+
+			// Element must be default parsable to create get operator because a
+			// value must be returned even if the index is invalid. Also values are
+			// assigned by copy.
+			static constexpr bool bIsArrayGetSupported = TIsParsable<ElementType>::Value && std::is_copy_assignable<ElementType>::value;
+
+			// Element must be copy assignable to set the value.
+			static constexpr bool bIsArraySetSupported = std::is_copy_assignable<ElementType>::value && std::is_copy_constructible<ElementType>::value;
+
+			// Elements must be copy constructible
+			static constexpr bool bIsArrayConcatSupported = std::is_copy_constructible<ElementType>::value;
+
+			// Elements must be copy constructible
+			static constexpr bool bIsArraySliceSupported = std::is_copy_constructible<ElementType>::value;
+		};
+
+		template<typename ArrayType, typename std::enable_if<TArrayNodeSupport<ArrayType>::bIsArrayGetSupported, bool>::type = true>
+		bool RegisterArrayGetNode()
+		{
+			using FNodeType = typename Metasound::TArrayGetNode<ArrayType>;
+			return RegisterNodeWithFrontend<FNodeType>();
+		}
+
+		template<typename ArrayType, typename std::enable_if<!TArrayNodeSupport<ArrayType>::bIsArrayGetSupported, bool>::type = true>
+		bool RegisterArrayGetNode()
+		{
+			// No op if not supported
+			return true;
+		}
+
+		template<typename ArrayType, typename std::enable_if<TArrayNodeSupport<ArrayType>::bIsArraySetSupported, bool>::type = true>
+		bool RegisterArraySetNode()
+		{
+			using FNodeType = typename Metasound::TArraySetNode<ArrayType>;
+
+			static_assert(TArrayNodeSupport<ArrayType>::bIsArraySetSupported, "TArraySetNode<> is not supported by array type");
+
+			return RegisterNodeWithFrontend<FNodeType>();
+		}
+
+		template<typename ArrayType, typename std::enable_if<!TArrayNodeSupport<ArrayType>::bIsArraySetSupported, bool>::type = true>
+		bool RegisterArraySetNode()
+		{
+			// No op if not supported
+			return true;
+		}
+
+		template<typename ArrayType, typename std::enable_if<TArrayNodeSupport<ArrayType>::bIsArraySliceSupported, bool>::type = true>
+		bool RegisterArraySliceNode()
+		{
+			using FNodeType = typename Metasound::TArraySliceNode<ArrayType>;
+
+			static_assert(TArrayNodeSupport<ArrayType>::bIsArraySliceSupported, "TArraySliceNode<> is not supported by array type");
+
+			return RegisterNodeWithFrontend<FNodeType>();
+		}
+
+		template<typename ArrayType, typename std::enable_if<!TArrayNodeSupport<ArrayType>::bIsArraySliceSupported, bool>::type = true>
+		bool RegisterArraySliceNode()
+		{
+			// No op if not supported
+			return true;
+		}
+
+		template<typename ArrayType, typename std::enable_if<TArrayNodeSupport<ArrayType>::bIsArrayConcatSupported, bool>::type = true>
+		bool RegisterArrayConcatNode()
+		{
+			using FNodeType = typename Metasound::TArrayConcatNode<ArrayType>;
+
+			static_assert(TArrayNodeSupport<ArrayType>::bIsArrayConcatSupported, "TArrayConcatNode<> is not supported by array type");
+
+			return RegisterNodeWithFrontend<FNodeType>();
+		}
+
+		template<typename ArrayType, typename std::enable_if<!TArrayNodeSupport<ArrayType>::bIsArrayConcatSupported, bool>::type = true>
+		bool RegisterArrayConcatNode()
+		{
+			// No op if not supported
+			return true;
+		}
+
+		template<typename ArrayType>
+		bool RegisterArrayNumNode()
+		{
+			static_assert(TArrayNodeSupport<ArrayType>::bIsArrayNumSupported, "TArrayNumNode<> is not supported by array type");
+
+			return ensureAlways(RegisterNodeWithFrontend<Metasound::TArrayNumNode<ArrayType>>());
+		}
+	}
+
+	/** Registers all available array nodes which can be instantiated for the given
+	 * ArrayType. Some nodes cannot be instantiated due to limitations of the 
+	 * array elements.
+	 */
+	template<typename ArrayType>
+	bool RegisterArrayNodes()
+	{
+		using namespace MetasoundArrayNodesPrivate;
+
+		bool bSuccess = RegisterArrayNumNode<ArrayType>();
+		bSuccess = bSuccess && RegisterArrayGetNode<ArrayType>();
+		bSuccess = bSuccess && RegisterArraySetNode<ArrayType>();
+		bSuccess = bSuccess && RegisterArraySliceNode<ArrayType>();
+		bSuccess = bSuccess && RegisterArrayConcatNode<ArrayType>();
+
+		return bSuccess;
+	}
+}
diff --git a/Engine/Plugins/Runtime/Metasound/Source/MetasoundFrontend/Public/MetasoundDataTypeRegistrationMacro.h b/Engine/Plugins/Runtime/Metasound/Source/MetasoundFrontend/Public/MetasoundDataTypeRegistrationMacro.h
--- a/Engine/Plugins/Runtime/Metasound/Source/MetasoundFrontend/Public/MetasoundDataTypeRegistrationMacro.h
+++ b/Engine/Plugins/Runtime/Metasound/Source/MetasoundFrontend/Public/MetasoundDataTypeRegistrationMacro.h
@@ -5,6 +5,7 @@
 #include "CoreMinimal.h"
 #include "Templates/Casts.h"
 
+#include "MetasoundArrayNodesRegistration.h"
 #include "MetasoundAutoConverterNode.h"
 #include "MetasoundConverterNodeRegistrationMacro.h"
 #include "MetasoundDataFactory.h"
@@ -243,6 +244,8 @@ namespace Metasound
 
 
 
+
+
 		template<typename TDataType, ELiteralType PreferredArgType = ELiteralType::None, typename UClassToUse = UObject>
 		bool RegisterDataTypeWithFrontendInternal()
 		{
@@ -376,8 +379,12 @@ namespace Metasound
 		bool RegisterDataTypeArrayWithFrontend()
 		{
 			using namespace MetasoundDataTypeRegistrationPrivate;
+			using TArrayType = TArray<TDataType>;
+
+			bool bSuccess = RegisterDataTypeWithFrontendInternal<TArrayType, TLiteralArrayEnum<PreferredArgType>::Value, UClassToUse>();
+			bSuccess = bSuccess && RegisterArrayNodes<TArrayType>();
 
-			return RegisterDataTypeWithFrontendInternal<TArray<TDataType>, TLiteralArrayEnum<PreferredArgType>::Value, UClassToUse>();
+			return bSuccess;
 		}
 	}
 	
diff --git a/Engine/Plugins/Runtime/Metasound/Source/MetasoundFrontend/Public/MetasoundTime.h b/Engine/Plugins/Runtime/Metasound/Source/MetasoundFrontend/Public/MetasoundTime.h
new file mode 100644
--- /dev/null
+++ b/Engine/Plugins/Runtime/Metasound/Source/MetasoundFrontend/Public/MetasoundTime.h
@@ -0,0 +1,354 @@
+// Copyright Epic Games, Inc. All Rights Reserved.
+#pragma once
+
+#include "MetasoundDataReference.h"
+#include "MetasoundOperatorSettings.h"
+
+namespace Metasound
+{
+	class FSampleCounter;
+
+	// Strongly typed time class
+	class FTime
+	{
+		using TimeType = double;
+		static_assert(TIsFloatingPoint<TimeType>::Value, "TimeType must be floating point.");
+
+		TimeType Time = 0;
+
+		public:
+			FTime()
+			{
+			}
+
+			FTime(TimeType InTime)
+				: Time(InTime)
+			{
+			}
+
+			FTime(const FTime& InOtherTime)
+				: Time(InOtherTime.Time)
+			{
+			}
+
+			static const FTime& Zero()
+			{
+				static const FTime ZeroTime;
+				return ZeroTime;
+			}
+
+			/** Return the time as seconds. */
+			TimeType GetSeconds() const
+			{ 
+				return Time;
+			}
+
+			FTime& operator=(const FTime& InRHS)
+			{
+				Time = InRHS.Time;
+				return *this;
+			}
+
+			template<typename ArithmeticType>
+			friend ArithmeticType& operator*=(ArithmeticType& InLHS, const FTime& InRHS)
+			{
+				static_assert(TIsArithmetic<ArithmeticType>::Value, "Must be arithmetic type.");
+				InLHS.Time *= InRHS.Time;
+				return InLHS;
+			}
+
+			template<typename ArithmeticType>
+			friend FTime& operator*=(FTime& InLHS, const ArithmeticType& InRHS)
+			{
+				static_assert(TIsArithmetic<ArithmeticType>::Value, "Must be arithmetic type.");
+				InLHS.Time *= InRHS;
+				return InLHS;
+			}
+
+			template<typename ArithmeticType>
+			friend ArithmeticType& operator/=(ArithmeticType& InLHS, const FTime& InRHS)
+			{
+				static_assert(TIsArithmetic<ArithmeticType>::Value, "Must be arithmetic type.");
+				InLHS.Time /= InRHS.Time;
+				return InLHS;
+			}
+
+			template<typename ArithmeticType>
+			friend FTime& operator/=(FTime& InLHS, const ArithmeticType& InRHS)
+			{
+				static_assert(TIsArithmetic<ArithmeticType>::Value, "Must be arithmetic type.");
+				InLHS.Time /= InRHS;
+				return InLHS;
+			}
+
+			FTime& operator-=(const FTime& InRHS)
+			{
+				Time -= InRHS.Time;
+				return *this;
+			}
+
+			template<typename ArithmeticType>
+			friend ArithmeticType& operator-=(ArithmeticType& InLHS, const FTime& InRHS)
+			{
+				static_assert(TIsArithmetic<ArithmeticType>::Value, "Must be arithmetic type.");
+				InLHS -= InRHS.Time;
+				return InLHS;
+			}
+
+			template<typename ArithmeticType>
+			friend FTime& operator-=(FTime& InLHS, const ArithmeticType& InRHS)
+			{
+				static_assert(TIsArithmetic<ArithmeticType>::Value, "Must be arithmetic type.");
+				InLHS.Time -= InRHS;
+				return InLHS;
+			}
+
+			FTime& operator+=(const FTime& InRHS)
+			{
+				Time += InRHS.Time;
+				return *this;
+			}
+
+			template<typename ArithmeticType>
+			friend ArithmeticType& operator+=(ArithmeticType& InLHS, const FTime& InRHS)
+			{
+				static_assert(TIsArithmetic<ArithmeticType>::Value, "Must be arithmetic type.");
+				InLHS += InRHS.Time;
+				return InLHS;
+			}
+
+			template<typename ArithmeticType>
+			friend FTime& operator+=(FTime& InLHS, const ArithmeticType& InRHS)
+			{
+				static_assert(TIsArithmetic<ArithmeticType>::Value, "Must be arithmetic type.");
+				InLHS.Time += InRHS;
+				return InLHS;
+			}
+
+			FTime operator-(const FTime& InRHS) const
+			{
+				return Time - InRHS.Time;
+			}
+
+			template<typename ArithmeticType>
+			friend ArithmeticType operator-(const ArithmeticType& InLHS, const FTime& InRHS)
+			{
+				static_assert(TIsArithmetic<ArithmeticType>::Value, "Must be arithmetic type.");
+				return InLHS - InRHS.Time;
+			}
+
+			template<typename ArithmeticType>
+			friend FTime operator-(const FTime& InLHS, const ArithmeticType& InRHS)
+			{
+				static_assert(TIsArithmetic<ArithmeticType>::Value, "Must be arithmetic type.");
+				return InLHS.Time - InRHS;
+			}
+
+			FTime operator+(const FTime& InRHS) const
+			{
+				return Time + InRHS.Time;
+			}
+
+			template<typename ArithmeticType>
+			friend ArithmeticType operator+(const ArithmeticType& InLHS, const FTime& InRHS)
+			{
+				static_assert(TIsArithmetic<ArithmeticType>::Value, "Must be arithmetic type.");
+				return InLHS + InRHS.Time;
+			}
+
+			template<typename ArithmeticType>
+			friend FTime operator+(const FTime& InLHS, const ArithmeticType& InRHS)
+			{
+				static_assert(TIsArithmetic<ArithmeticType>::Value, "Must be arithmetic type.");
+				return InLHS.Time + InRHS;
+			}
+
+			template<typename ArithmeticType>
+			ArithmeticType operator*(const FTime& InRHS) const
+			{
+				static_assert(TIsArithmetic<ArithmeticType>::Value, "Must be arithmetic type.");
+				return static_cast<ArithmeticType>(Time * InRHS.Time);
+			}
+
+			template<typename ArithmeticType>
+			friend ArithmeticType operator*(const ArithmeticType& InLHS, const FTime& InRHS)
+			{
+				static_assert(TIsArithmetic<ArithmeticType>::Value, "Must be arithmetic type.");
+				return InLHS * InRHS.Time;
+			}
+
+			template<typename ArithmeticType>
+			friend FTime operator*(const FTime& InLHS, const ArithmeticType& InRHS)
+			{
+				static_assert(TIsArithmetic<ArithmeticType>::Value, "Must be arithmetic type.");
+				return InLHS.Time * InRHS;
+			}
+
+			template<typename ArithmeticType>
+			ArithmeticType operator/(const FTime& InRHS) const
+			{
+				static_assert(TIsArithmetic<ArithmeticType>::Value, "Must be arithmetic type.");
+				return static_cast<ArithmeticType>(Time / InRHS.Time);
+			}
+
+			template<typename ArithmeticType>
+			friend ArithmeticType& operator/(const ArithmeticType& InLHS, const FTime& InRHS)
+			{
+				static_assert(TIsArithmetic<ArithmeticType>::Value, "Must be arithmetic type.");
+				return InLHS / InRHS.Time;
+			}
+
+			template<typename ArithmeticType>
+			friend FTime operator/(const FTime& InLHS, const ArithmeticType& InRHS)
+			{
+				static_assert(TIsArithmetic<ArithmeticType>::Value, "Must be arithmetic type.");
+				return InLHS.Time / InRHS;
+			}
+
+			bool operator<(const FTime& InRHS) const
+			{
+				return Time < InRHS.Time;
+			}
+
+			template<typename ArithmeticType>
+			friend bool operator<(const ArithmeticType& InLHS, const FTime& InRHS)
+			{
+				static_assert(TIsArithmetic<ArithmeticType>::Value, "Must be arithmetic type.");
+				return InLHS < InRHS.Time;
+			}
+
+			template<typename ArithmeticType>
+			friend bool operator<(const FTime& InLHS, const ArithmeticType& InRHS)
+			{
+				static_assert(TIsArithmetic<ArithmeticType>::Value, "Must be arithmetic type.");
+				return InLHS.Time < InRHS;
+			}
+
+			bool operator<=(const FTime& InRHS) const
+			{
+				return Time <= InRHS.Time;
+			}
+
+			template<typename ArithmeticType>
+			friend bool operator<=(const ArithmeticType& InLHS, const FTime& InRHS)
+			{
+				static_assert(TIsArithmetic<ArithmeticType>::Value, "Must be arithmetic type.");
+				return InLHS <= InRHS.Time;
+			}
+
+			template<typename ArithmeticType>
+			friend bool operator<=(const FTime& InLHS, const ArithmeticType& InRHS)
+			{
+				static_assert(TIsArithmetic<ArithmeticType>::Value, "Must be arithmetic type.");
+				return InLHS.Time <= InRHS;
+			}
+
+			bool operator>(const FTime& InRHS) const
+			{
+				return Time > InRHS.Time;
+			}
+
+			template<typename ArithmeticType>
+			friend bool operator>(const FTime& InLHS, const ArithmeticType& InRHS)
+			{
+				static_assert(TIsArithmetic<ArithmeticType>::Value, "Must be arithmetic type.");
+				return InLHS.Time > InRHS;
+			}
+
+			template<typename ArithmeticType>
+			friend bool operator>(const ArithmeticType& InLHS, const FTime& InRHS)
+			{
+				static_assert(TIsArithmetic<ArithmeticType>::Value, "Must be arithmetic type.");
+				return InLHS > InRHS.Time;
+			}
+
+			bool operator>=(const FTime& InRHS) const
+			{
+				return Time >= InRHS.Time;
+			}
+
+			template<typename ArithmeticType>
+			friend bool operator>=(const ArithmeticType& InLHS, const FTime& InRHS)
+			{
+				static_assert(TIsArithmetic<ArithmeticType>::Value, "Must be arithmetic type.");
+				return InLHS >= InRHS.Time;
+			}
+
+			template<typename ArithmeticType>
+			friend bool operator>=(const FTime& InLHS, const ArithmeticType& InRHS)
+			{
+				static_assert(TIsArithmetic<ArithmeticType>::Value, "Must be arithmetic type.");
+				return InLHS.Time >= InRHS;
+			}
+
+			template<typename ArithmeticType>
+			friend bool operator!=(const ArithmeticType& InLHS, const FTime& InRHS)
+			{
+				static_assert(TIsArithmetic<ArithmeticType>::Value, "Must be arithmetic type.");
+				return InLHS != InRHS.Time;
+			}
+
+			bool operator!=(const FTime& InRHS) const
+			{
+				return Time != InRHS.Time;
+			}
+
+			template<typename ArithmeticType>
+			friend bool operator!=(const FTime& InLHS, const ArithmeticType& InRHS)
+			{
+				static_assert(TIsArithmetic<ArithmeticType>::Value, "Must be arithmetic type.");
+				return InLHS.Time != InRHS;
+			}
+
+			friend bool operator==(const FTime& InLHS, const FTime& InRHS)
+			{
+				return InLHS.Time == InRHS.Time;
+			}
+
+			template<typename ArithmeticType>
+			friend bool operator==(const ArithmeticType& InLHS, const FTime& InRHS)
+			{
+				static_assert(TIsArithmetic<ArithmeticType>::Value, "Must be arithmetic type.");
+				return InLHS == InRHS.Time;
+			}
+
+			template<typename ArithmeticType>
+			friend bool operator==(const FTime& InLHS, const ArithmeticType& InRHS)
+			{
+				static_assert(TIsArithmetic<ArithmeticType>::Value, "Must be arithmetic type.");
+				return InLHS.Time == InRHS;
+			}
+
+			template<typename ArithmeticType = TimeType>
+			static ArithmeticType ToMilliseconds(const FTime& InTime)
+			{
+				static_assert(TIsArithmetic<ArithmeticType>::Value, "Must be arithmetic type.");
+				return static_cast<ArithmeticType>(InTime.GetSeconds() * 1e3);
+			}
+
+			template<typename ArithmeticType = TimeType>
+			static ArithmeticType ToMicroseconds(const FTime& InTime)
+			{
+				static_assert(TIsArithmetic<ArithmeticType>::Value, "Must be arithmetic type.");
+				return static_cast<ArithmeticType>(InTime.GetSeconds() * 1e6);
+			}
+
+			template<typename ArithmeticType>
+			static FTime FromMilliseconds(const ArithmeticType& InMilliseconds)
+			{
+				static_assert(TIsArithmetic<ArithmeticType>::Value, "Must be arithmetic type.");
+				return FTime(static_cast<TimeType>(InMilliseconds * 1e-3));
+			}
+
+			template<typename ArithmeticType>
+			static FTime FromMicroseconds(const ArithmeticType& InMicroseconds)
+			{
+				static_assert(TIsArithmetic<ArithmeticType>::Value, "Must be arithmetic type.");
+				return FTime(static_cast<TimeType>(InMicroseconds * 1e-6));
+			}
+
+			friend FSampleCounter;
+	};
+
+	DECLARE_METASOUND_DATA_REFERENCE_TYPES(FTime, METASOUNDFRONTEND_API, FTimeTypeInfo, FTimeReadRef, FTimeWriteRef);
+}
diff --git a/Engine/Plugins/Runtime/Metasound/Source/MetasoundFrontend/Public/MetasoundTrigger.h b/Engine/Plugins/Runtime/Metasound/Source/MetasoundFrontend/Public/MetasoundTrigger.h
new file mode 100644
--- /dev/null
+++ b/Engine/Plugins/Runtime/Metasound/Source/MetasoundFrontend/Public/MetasoundTrigger.h
@@ -0,0 +1,173 @@
+// Copyright Epic Games, Inc. All Rights Reserved.
+
+#pragma once
+
+#include "MetasoundDataReference.h"
+#include "MetasoundOperatorSettings.h"
+#include "MetasoundTime.h"
+
+namespace Metasound
+{
+	/** FTrigger supports sample accurate triggering, sample accurate internal tracking,
+	 * and a convenient interface for running trigger-aligned audio signal processing
+	 * routines on buffers.. 
+	 *
+	 * FTriggers are triggered using FTrigger::TriggerTime or FTrigger::TriggerFrame.
+	 * FTriggers track time internally by calling FTrigger::Advance.
+	 * Executing audio signal processing on buffers can be performed by calling
+	 * FTrigger::ExecuteBlock or FTrigger::LookAhead.
+	 */
+	class METASOUNDFRONTEND_API FTrigger
+	{
+		public:
+			/** FTrigger constructor. 
+			 *
+			 * @param InSettings - Operator settings.
+			 * @param bShouldTrigger - If true, triggers first sample.
+			 */
+			explicit FTrigger(const FOperatorSettings& InSettings, bool bShouldTrigger);
+
+			/** FTrigger constructor. 
+			 *
+			 * @param InFrameToTrigger - Set specific frame to trigger.
+			 * @param InSettings - Operator settings.
+			 */
+			explicit FTrigger(const FOperatorSettings& InSettings, int32 InFrameToTrigger);
+
+			/** FTrigger constructor. By default it is not triggered.
+			 *
+			 * @param InSettings - Operator settings.
+			 */
+			FTrigger(const FOperatorSettings& InSettings);
+
+			/** Trigger a specific frame in the future.
+			 *
+			 * @param InFrameToTrigger - Index of frame to trigger.
+			 */
+			void TriggerFrame(int32 InFrameToTrigger);
+
+			/** Trigger a specific time in the future.
+			 *
+			 * @param InTime - Time to set trigger.
+			 */
+			template<typename TimeType>
+			void TriggerTime(const TimeType& InTime)
+			{
+				TriggerFrame(InTime.GetNumSamples(SampleRate));
+			}
+
+			/** Advance internal frame counters by block size. */
+			void AdvanceBlock();
+
+			/** Advance internal frame counters by specific frame count. */
+			void Advance(int32 InNumFrames);
+
+			/** Number of triggered frames. */
+			int32 Num() const;
+
+			/** Returns true if there is a trigger in the current block of audio */
+			int32 NumTriggeredInBlock() const;
+
+			/** Returns frame index for a given trigger index. 
+			 *
+			 * @param InIndex - Index of trigger. Must be a value between 0 and Num().
+			 *
+			 * @return The frame of the trigger.
+			 */
+			int32 operator[](int32 InTriggerIndex) const;
+
+			/** Returns true if there are any triggered frames. */
+			bool IsTriggered() const;
+
+			/** Returns true there is a trigger in the current block of audio. */
+			bool IsTriggeredInBlock() const;
+
+			/** Implicit conversion of FTrigger into bool by calling IsTriggeredInBlock()
+			 *
+			 * @return Returns true if frame is triggered in current block.
+			 */
+			operator bool() const;
+
+			/** Removes all triggered frames. */
+			void Reset();
+
+			/** Executes one block of frames and calls underlying InPreTrigger and InOnTrigger
+			 * functions with frame indices.
+			 *
+			 * @param InPreTrigger - A function which handles frames before the first 
+			 *                   trigger in the current block. The function must accept 
+			 *                   the arguments (int32 StartFrame, int32 EndFrame).
+			 * @param InOnTrigger - A function which handles frames starting with the 
+			 *                  triggers index and ending the next trigger index or the
+			 *                  number of frames in a block.. The function must
+			 *                  accept the arguments (int32 StartFrame, int32 EndFrame).
+			 */
+			template<typename PreTriggerType, typename OnTriggerType>
+			void ExecuteBlock(PreTriggerType InPreTrigger, OnTriggerType InOnTrigger) const
+			{
+				ExecuteFrames(NumFramesPerBlock, LastTriggerIndexInBlock, InPreTrigger, InOnTrigger);
+			}
+
+			/** Executes a desired number of frames and calls underlying InPreTrigger 
+			 * and InOnTrigger functions with frame indices.
+			 *
+			 * @param InNumFrames - The number of frames to look ahead.
+			 * @param InPreTrigger - A function which handles frames before the first 
+			 *                   trigger in the range of frames. The function must accept 
+			 *                   the arguments (int32 StartFrame, int32 EndFrame).
+			 * @param InOnTrigger - A function which handles frames starting with the 
+			 *                  triggers index and ending the next trigger index or InNumFrames
+			 *             		The function must accept the arguments 
+			 *             		(int32 StartFrame, int32 EndFrame).
+			 */
+			template<typename PreTriggerType, typename OnTriggerType>
+			void LookAhead(int32 InNumFrames, PreTriggerType InPreTrigger, OnTriggerType InOnTrigger) const
+			{
+				if (InNumFrames > 0)
+				{
+					int32 LastTriggerIndexInLookAhead= Algo::LowerBound(TriggeredFrames, InNumFrames);
+					
+					ExecuteFrames(InNumFrames, LastTriggerIndexInLookAhead, InPreTrigger, InOnTrigger);
+				}
+			}
+
+		private:
+			void UpdateLastTriggerIndexInBlock();
+
+			template<typename PreTriggerType, typename OnTriggerType>
+			void ExecuteFrames(int32 InNumFrames, int32 InLastTriggerIndex, PreTriggerType InPreTrigger, OnTriggerType InOnTrigger) const
+			{
+				if (InLastTriggerIndex <= 0)
+				{
+					InPreTrigger(0, InNumFrames);
+					return;
+				}
+
+				const int32* TriggeredFramesData = TriggeredFrames.GetData();
+
+				if (TriggeredFramesData[0] > 0)
+				{
+					InPreTrigger(0, TriggeredFramesData[0]);
+				}
+
+				const int32 NextToLastTriggerIndex = InLastTriggerIndex - 1;
+
+				for (int32 i = 0; i < NextToLastTriggerIndex; i++)
+				{
+					InOnTrigger(TriggeredFramesData[i], TriggeredFramesData[i + 1]);
+				}
+
+				InOnTrigger(TriggeredFramesData[NextToLastTriggerIndex], InNumFrames);
+			}
+
+
+			TArray<int32> TriggeredFrames;
+
+			bool bHasTrigger = false;
+			int32 NumFramesPerBlock = 0;
+			FSampleRate SampleRate = 0;
+			int32 LastTriggerIndexInBlock = 0;
+	};
+
+	DECLARE_METASOUND_DATA_REFERENCE_TYPES(FTrigger, METASOUNDFRONTEND_API, FTriggerTypeInfo, FTriggerReadRef, FTriggerWriteRef);
+}
diff --git a/Engine/Plugins/Runtime/Metasound/Source/MetasoundGraphCore/Public/MetasoundDataReferenceCollection.h b/Engine/Plugins/Runtime/Metasound/Source/MetasoundGraphCore/Public/MetasoundDataReferenceCollection.h
--- a/Engine/Plugins/Runtime/Metasound/Source/MetasoundGraphCore/Public/MetasoundDataReferenceCollection.h
+++ b/Engine/Plugins/Runtime/Metasound/Source/MetasoundGraphCore/Public/MetasoundDataReferenceCollection.h
@@ -226,10 +226,15 @@ namespace Metasound
 			 *
 			 * @return A readable data reference.
 			 */
+			// TODO: FIXME
+			// 1. Remove extra template parameters "DefaultDataType", "ConstructorArgTypes"
+			// 2. Use data factory instead of read reference class to construct
+			// 3. Pipe in FOperatorSettings to the factory call.
+			// 4. Build from literal, not from DefaultDataType.
 			template<typename DataType, typename DefaultDataType = DataType, typename... ConstructorArgTypes>
 			TDataReadReference<DataType> GetDataReadReferenceOrConstructWithVertexDefault(const FInputVertexInterface& InputVertices, const FString& InName) const
 			{
-				using FDataRefType = TDataReadReference<DataType>;
+				using FDataRef = TDataReadReference<DataType>;
 
 				if (ContainsDataReadReference<DataType>(InName))
 				{
@@ -238,7 +243,7 @@ namespace Metasound
 				else
 				{
 					const DefaultDataType DefaultValue = InputVertices[InName].GetDefaultValue().Value.Get<DefaultDataType>();
-					return FDataRefType::CreateNew(DefaultValue);
+					return FDataRef::CreateNew(DefaultValue);
 				}
 			}
 
diff --git a/Engine/Plugins/Runtime/Metasound/Source/MetasoundStandardNodes/Private/MetasoundStandardNodesModule.cpp b/Engine/Plugins/Runtime/Metasound/Source/MetasoundStandardNodes/Private/MetasoundStandardNodesModule.cpp
--- a/Engine/Plugins/Runtime/Metasound/Source/MetasoundStandardNodes/Private/MetasoundStandardNodesModule.cpp
+++ b/Engine/Plugins/Runtime/Metasound/Source/MetasoundStandardNodes/Private/MetasoundStandardNodesModule.cpp
@@ -9,8 +9,6 @@
 #include "Modules/ModuleInterface.h"
 #include "Modules/ModuleManager.h"
 
-REGISTER_METASOUND_DATATYPE(Metasound::FTrigger, "Trigger", ::Metasound::ELiteralType::Boolean)
-REGISTER_METASOUND_DATATYPE(Metasound::FTime, "Time", ::Metasound::ELiteralType::Float)
 
 namespace Metasound 
 {
diff --git a/Engine/Plugins/Runtime/Metasound/Source/MetasoundStandardNodes/Private/MetasoundTrigger.cpp b/Engine/Plugins/Runtime/Metasound/Source/MetasoundStandardNodes/Private/MetasoundTrigger.cpp
deleted file mode 100644
diff --git a/Engine/Plugins/Runtime/Metasound/Source/MetasoundStandardNodes/Public/MetasoundTime.h b/Engine/Plugins/Runtime/Metasound/Source/MetasoundStandardNodes/Public/MetasoundTime.h
deleted file mode 100644
diff --git a/Engine/Plugins/Runtime/Metasound/Source/MetasoundStandardNodes/Public/MetasoundTrigger.h b/Engine/Plugins/Runtime/Metasound/Source/MetasoundStandardNodes/Public/MetasoundTrigger.h
deleted file mode 100644
