Ability to change Metasound BlockRate via cvar

diff --git a/Engine/Plugins/Runtime/Metasound/Source/MetasoundEngine/Private/MetasoundSource.cpp b/Engine/Plugins/Runtime/Metasound/Source/MetasoundEngine/Private/MetasoundSource.cpp
--- a/Engine/Plugins/Runtime/Metasound/Source/MetasoundEngine/Private/MetasoundSource.cpp
+++ b/Engine/Plugins/Runtime/Metasound/Source/MetasoundEngine/Private/MetasoundSource.cpp
@@ -23,6 +23,16 @@
 #include "EdGraph/EdGraph.h"
 #endif // WITH_EDITORONLY_DATA
 
+static float MetaSoundBlockRateCVar = 0;
+FAutoConsoleVariableRef CVarMetaSoundBlockRate(
+	TEXT("au.MetaSound.BlockRate"),
+	MetaSoundBlockRateCVar,
+	TEXT("Sets block rate (blocks per second) of MetaSounds.\n")
+	TEXT("Default: 100.0f"),
+	ECVF_Default);
+
+
+
 #define LOCTEXT_NAMESPACE "MetasoundSource"
 
 static FName MetasoundSourceArchetypeName = FName(TEXT("Metasound Source"));
@@ -368,7 +378,8 @@ TArray<FString> UMetasoundSource::GetTransmittableInputVertexNames() const
 
 Metasound::FOperatorSettings UMetasoundSource::GetOperatorSettings(Metasound::FSampleRate InSampleRate) const
 {
-	constexpr float BlockRate = 100.f; // Metasound graph gets evaluated 100 times per second.
+	// Metasound graph gets evaluated 100 times per second by default.
+	float BlockRate = FMath::Clamp(MetaSoundBlockRateCVar, 1.0f, 200.0f); 
 	return Metasound::FOperatorSettings(InSampleRate, BlockRate);
 }
 
