Metasound InterpTo node

diff --git a/Engine/Plugins/Runtime/Metasound/Source/MetasoundStandardNodes/Private/MetasoundInterpToNode.cpp b/Engine/Plugins/Runtime/Metasound/Source/MetasoundStandardNodes/Private/MetasoundInterpToNode.cpp
new file mode 100644
--- /dev/null
+++ b/Engine/Plugins/Runtime/Metasound/Source/MetasoundStandardNodes/Private/MetasoundInterpToNode.cpp
@@ -0,0 +1,167 @@
+// Copyright Epic Games, Inc. All Rights Reserved.
+
+#include "MetasoundInterpToNode.h"
+
+#include "Internationalization/Text.h"
+#include "MetasoundExecutableOperator.h"
+#include "MetasoundNodeRegistrationMacro.h"
+#include "MetasoundPrimitives.h"
+#include "MetasoundStandardNodesNames.h"
+#include "MetasoundTrigger.h"
+#include "MetasoundTime.h"
+#include "DSP/VolumeFader.h"
+
+#define LOCTEXT_NAMESPACE "MetasoundStandardNodes"
+
+namespace Metasound
+{
+	namespace InterpTo
+	{
+		static const TCHAR* InParamNameTarget = TEXT("Target");
+		static const TCHAR* InParamNameInterpTime = TEXT("Interp Time");
+		static const TCHAR* OutParamNameValue = TEXT("Value");
+	}
+
+	class FInterpToOperator : public TExecutableOperator<FInterpToOperator>
+	{
+	public:
+
+		static const FNodeClassMetadata& GetNodeInfo();
+		static const FVertexInterface& GetVertexInterface();
+		static TUniquePtr<IOperator> CreateOperator(const FCreateOperatorParams& InParams, FBuildErrorArray& OutErrors);
+
+		FInterpToOperator(const FOperatorSettings& InSettings, const FFloatReadRef& InTargetValue, const FTimeReadRef& InInterpTime);
+
+		virtual FDataReferenceCollection GetInputs() const override;
+		virtual FDataReferenceCollection GetOutputs() const override;
+		void Execute();
+
+	private:
+		// The target value of the lerp. The output will lerp from it's current value to the output value.
+		FFloatReadRef TargetValue;
+
+		// The amount of time to do lerp
+		FTimeReadRef InterpTime;
+
+		// The current output value.
+		FFloatWriteRef ValueOutput;
+
+		// Volume fader object which performs the interpolating
+		Audio::FVolumeFader VolumeFader;
+
+		// The time-delta per block
+		float BlockTimeDelta = 0.0f;
+
+		// The previous target value
+		float PreviousTargetValue = 0.0f;
+	};
+
+	FInterpToOperator::FInterpToOperator(const FOperatorSettings& InSettings, const FFloatReadRef& InTargetValue, const FTimeReadRef& InInterpTime)
+		: TargetValue(InTargetValue)
+		, InterpTime(InInterpTime)
+		, ValueOutput(FFloatWriteRef::CreateNew(*TargetValue))
+	{
+		// Set the fade to start at the value specified in the current value
+		VolumeFader.SetVolume(*TargetValue);
+
+		float BlockRate = InSettings.GetActualBlockRate();
+		BlockTimeDelta = 1.0f / BlockRate;
+
+		PreviousTargetValue = *TargetValue;
+	}
+
+	FDataReferenceCollection FInterpToOperator::GetInputs() const
+	{
+		FDataReferenceCollection InputDataReferences;
+		InputDataReferences.AddDataReadReference(InterpTo::InParamNameTarget, FFloatReadRef(TargetValue));
+		InputDataReferences.AddDataReadReference(InterpTo::InParamNameInterpTime, FTimeReadRef(InterpTime));
+
+		return InputDataReferences;
+	}
+
+	FDataReferenceCollection FInterpToOperator::GetOutputs() const
+	{
+		FDataReferenceCollection OutputDataReferences;
+		OutputDataReferences.AddDataReadReference(InterpTo::OutParamNameValue, FFloatReadRef(ValueOutput));
+		return OutputDataReferences;
+	}
+
+	void FInterpToOperator::Execute()
+	{
+		// Update the value output with the current value in case it was changed
+		if (!FMath::IsNearlyEqual(PreviousTargetValue, *TargetValue))
+		{
+			PreviousTargetValue = *TargetValue;
+			// Start the volume fader on the interp trigger
+			float FadeSeconds = (float)InterpTime->GetSeconds();
+			VolumeFader.StartFade(*TargetValue, FadeSeconds, Audio::EFaderCurve::Linear);
+		}
+
+		// Perform the fading
+		if (VolumeFader.IsFading())
+		{
+			VolumeFader.Update(BlockTimeDelta);
+		}
+
+		// Update the fader w/ the current volume
+		*ValueOutput = VolumeFader.GetVolume();
+	}
+
+	const FVertexInterface& FInterpToOperator::GetVertexInterface()
+	{
+		static const FVertexInterface Interface(
+			FInputVertexInterface(
+				TInputDataVertexModel<FTime>(InterpTo::InParamNameInterpTime, LOCTEXT("InterpTimeTooltip", "The time to interpolate from the current value to the target value."), 0.1f),
+				TInputDataVertexModel<float>(InterpTo::InParamNameTarget, LOCTEXT("TargetValueTooltip", "Target value."), 1.0f)
+			),
+			FOutputVertexInterface(
+				TOutputDataVertexModel<float>(InterpTo::OutParamNameValue, LOCTEXT("ValueTooltip", "The current value."))
+			)
+		);
+
+		return Interface;
+	}
+
+	const FNodeClassMetadata& FInterpToOperator::GetNodeInfo()
+	{
+		auto InitNodeInfo = []() -> FNodeClassMetadata
+		{
+			FNodeClassMetadata Info;
+			Info.ClassName = { Metasound::StandardNodes::Namespace, TEXT("InterpTo"), Metasound::StandardNodes::AudioVariant };
+			Info.MajorVersion = 1;
+			Info.MinorVersion = 0;
+			Info.DisplayName = LOCTEXT("Metasound_InterpDisplayName", "InterpTo");
+			Info.Description = LOCTEXT("Metasound_InterpNodeDescription", "Interpolates between the current value and a target value over the specified time.");
+			Info.Author = PluginAuthor;
+			Info.PromptIfMissing = PluginNodeMissingPrompt;
+			Info.DefaultInterface = GetVertexInterface();
+
+			return Info;
+		};
+
+		static const FNodeClassMetadata Info = InitNodeInfo();
+
+		return Info;
+	}
+
+	TUniquePtr<IOperator> FInterpToOperator::CreateOperator(const FCreateOperatorParams& InParams, FBuildErrorArray& OutErrors)
+	{
+		const FInterpToNode& InterpToNode = static_cast<const FInterpToNode&>(InParams.Node);
+		const FDataReferenceCollection& InputCollection = InParams.InputDataReferences;
+		const FInputVertexInterface& InputInterface = GetVertexInterface().GetInputInterface();
+
+		FFloatReadRef TargetValue = InputCollection.GetDataReadReferenceOrConstructWithVertexDefault<float>(InputInterface, InterpTo::InParamNameTarget);
+		FTimeReadRef InterpTime = InputCollection.GetDataReadReferenceOrConstructWithVertexDefault<FTime, float>(InputInterface, InterpTo::InParamNameInterpTime);
+
+		return MakeUnique<FInterpToOperator>(InParams.OperatorSettings, TargetValue, InterpTime);
+	}
+
+	FInterpToNode::FInterpToNode(const FNodeInitData& InitData)
+		: FNodeFacade(InitData.InstanceName, InitData.InstanceID, TFacadeOperatorClass<FInterpToOperator>())
+	{
+	}
+
+	METASOUND_REGISTER_NODE(FInterpToNode)
+}
+
+#undef LOCTEXT_NAMESPACE
diff --git a/Engine/Plugins/Runtime/Metasound/Source/MetasoundStandardNodes/Public/MetasoundInterpToNode.h b/Engine/Plugins/Runtime/Metasound/Source/MetasoundStandardNodes/Public/MetasoundInterpToNode.h
new file mode 100644
--- /dev/null
+++ b/Engine/Plugins/Runtime/Metasound/Source/MetasoundStandardNodes/Public/MetasoundInterpToNode.h
@@ -0,0 +1,23 @@
+// Copyright Epic Games, Inc. All Rights Reserved.
+
+#pragma once
+
+#include "CoreMinimal.h"
+#include "MetasoundFacade.h"
+
+namespace Metasound
+{
+	/** FInterpToNode
+	 *
+	 *  Interpolates to a target value over a given time.
+	 */
+	class METASOUNDSTANDARDNODES_API FInterpToNode : public FNodeFacade
+	{
+	public:
+		/**
+		 * Constructor used by the Metasound Frontend.
+		 */
+		FInterpToNode(const FNodeInitData& InitData);
+	};
+} // namespace Metasound
+
