Nuke Unformatted audio type

diff --git a/Engine/Plugins/Runtime/Metasound/Source/MetasoundStandardNodes/Private/MetasoundAudioFormats.cpp b/Engine/Plugins/Runtime/Metasound/Source/MetasoundStandardNodes/Private/MetasoundAudioFormats.cpp
--- a/Engine/Plugins/Runtime/Metasound/Source/MetasoundStandardNodes/Private/MetasoundAudioFormats.cpp
+++ b/Engine/Plugins/Runtime/Metasound/Source/MetasoundStandardNodes/Private/MetasoundAudioFormats.cpp
@@ -8,51 +8,12 @@
 #include "MetasoundLiteral.h"
 #include "MetasoundDataTypeRegistrationMacro.h"
 
-REGISTER_METASOUND_DATATYPE(Metasound::FUnformattedAudio, "Audio:Unformatted");
 REGISTER_METASOUND_DATATYPE(Metasound::FMonoAudioFormat, "Audio:Mono");
 REGISTER_METASOUND_DATATYPE(Metasound::FStereoAudioFormat, "Audio:Stereo");
 REGISTER_METASOUND_DATATYPE(Metasound::FMultichannelAudioFormat, "Audio:Multichannel", ELiteralType::Integer);
 
 namespace Metasound
 {
-
-	/* FUnformattedAudio */
-	FUnformattedAudio::FUnformattedAudio(int32 InNumFrames, int32 InNumChannels, int32 InMaxNumChannels)
-	:	NumFrames(InNumFrames)
-	,	NumChannels(0)
-	,	MaxNumChannels(InMaxNumChannels)
-	{
-		NumFrames = FMath::Max(NumFrames, 0);
-		MaxNumChannels = FMath::Max(MaxNumChannels, 0);
-
-		for (int32 i = 0; i < MaxNumChannels; i++)
-		{
-			FAudioBufferWriteRef Audio = FAudioBufferWriteRef::CreateNew(NumFrames);
-			Audio->Zero();
-
-			WritableBufferStorage.Add(Audio);
-			ReadableBufferStorage.Add(Audio);
-		}
-
-		SetNumChannels(InNumChannels);
-	}
-
-	FUnformattedAudio::FUnformattedAudio(const FOperatorSettings& InSettings, int32 InInitialNumChannels)
-		: FUnformattedAudio(InSettings.GetNumFramesPerBlock(), InInitialNumChannels, 8)
-	{
-	}
-
-	int32 FUnformattedAudio::SetNumChannels(int32 InNumChannels)
-	{
-		NumChannels = FMath::Max(0, FMath::Min(InNumChannels, MaxNumChannels));
-
-		ReadableBuffers = TArrayView<const FAudioBufferReadRef>(ReadableBufferStorage.GetData(), NumChannels);
-		WritableBuffers = TArrayView<const FAudioBufferWriteRef>(WritableBufferStorage.GetData(), NumChannels);
-
-		return NumChannels;
-	}
-
-
 	/* FMultichannelAudioFormat */
 
 	FMultichannelAudioFormat::FMultichannelAudioFormat()
@@ -83,7 +44,6 @@ namespace Metasound
 		: FMultichannelAudioFormat(InSettings.GetNumFramesPerBlock(), InNumChannels)
 	{}
 
-
 	FMultichannelAudioFormat::FMultichannelAudioFormat(TArrayView<const FAudioBufferWriteRef> InWriteRefs)
 	:	NumChannels(InWriteRefs.Num())
 	{
@@ -123,5 +83,4 @@ namespace Metasound
 			ReadableBuffers = ReadableBufferStorage;
 		}
 	}
-
-}
+} // namespace Metasound
diff --git a/Engine/Plugins/Runtime/Metasound/Source/MetasoundStandardNodes/Public/MetasoundAudioFormats.h b/Engine/Plugins/Runtime/Metasound/Source/MetasoundStandardNodes/Public/MetasoundAudioFormats.h
--- a/Engine/Plugins/Runtime/Metasound/Source/MetasoundStandardNodes/Public/MetasoundAudioFormats.h
+++ b/Engine/Plugins/Runtime/Metasound/Source/MetasoundStandardNodes/Public/MetasoundAudioFormats.h
@@ -9,95 +9,6 @@
 
 namespace Metasound
 {
-	/** FUnformattedAudio
-	 *
-	 * FUnformattedAudio represents deinterleaved multichannel audio which can dynamically change
-	 * the number of channels of audio. In possible, using FMultichannelAudioFormat is preferred over
-	 * using FUnformattedAudio when the number of audio channels is known at construction.
-	 *
-	 * The audio buffers in FUnformattedAudio are shared data references which can be accessed outside
-	 * of the FUnformattedAudio. All audio buffers within a FUnformattedAudio object must contain the 
-	 * same number of audio frames.
-	 */
-	class METASOUNDSTANDARDNODES_API FUnformattedAudio : public IAudioDataType
-	{
-		public:
-			/** FUnformattedAudio Constructor.
-			 * 
-			 * @param InNumFrames - Number of frames to hold in each buffer.
-			 * @param InNumChannels - Initial number of audio channels.
-			 * @param InMaxNumChannels - Maximum number of audio channels.
-			 */
-			FUnformattedAudio(int32 InNumFrames, int32 InNumChannels, int32 InMaxNumChannels);
-
-			FUnformattedAudio()
-				: FUnformattedAudio(0, 0, 0)
-			{}
-
-			/**
-			 * FUnformattedAudio Constructor used by the metasound frontend.
-			 *
-			 * @param InSettings - Operator Settings passed in on construction.
-			 * @param InNumChannels - initial number of audio channels.
-			 */
-			explicit FUnformattedAudio(const FOperatorSettings& InSettings, int32 InNumChannels);
-
-
-			/** Sets the number of audio channels.
-			 * 
-			 * The requested number of channels is internally clamped to the supported range 
-			 * of [0, GetMaxNumChannels()].
-			 *
-			 * @parma InNumChannels - The desired number of channels.
-			 *
-			 * @return The actual number of channels. This may differ from the requested number 
-			 * 		   of channels if the requested number exceeded the maximum number of channels.
-			 */
-			int32 SetNumChannels(int32 InNumChannels);
-
-			/** Return the number of channels. */
-			int32 GetNumChannels() const { return NumChannels; }
-
-			/** Return the maximum number of channels. */
-			int32 GetMaxNumChannels() const { return MaxNumChannels; }
-
-			/** Return an array view of the readable buffer references.
-			 *
-			 * This array will have GetNumChannels() elements.
-			 */
-			const TArrayView<const FAudioBufferReadRef> GetBuffers() const { return ReadableBuffers; }
-
-			/** Return an array view of the writable buffer references.
-			 *
-			 * This array will have GetNumChannels() elements.
-			 */
-			const TArrayView<const FAudioBufferWriteRef> GetBuffers() { return WritableBuffers; }
-
-			/** Return an array of the readable buffer reference storage.
-			 *
-			 * This array will have GetMaxNumChannels() elements.
-			 */
-			const TArray<FAudioBufferReadRef>& GetStorage() const { return ReadableBufferStorage; }
-
-			/** Return an array of the writable buffer reference storage.
-			 *
-			 * This array will have GetMaxNumChannels() elements.
-			 */
-			const TArray<FAudioBufferWriteRef>& GetStorage() { return WritableBufferStorage; }
-			
-		private:
-			int32 NumFrames;
-			int32 NumChannels;
-			int32 MaxNumChannels;
-
-			TArrayView<const FAudioBufferReadRef> ReadableBuffers;
-			TArrayView<const FAudioBufferWriteRef> WritableBuffers;
-
-			TArray<FAudioBufferReadRef> ReadableBufferStorage;
-			TArray<FAudioBufferWriteRef> WritableBufferStorage;
-	};
-
-
 	/** FMultichannelAudioFormat
 	 *
 	 * FMultichannelAudioFormat represents deinterleaved multichannel audio which supports a constant
@@ -414,8 +325,6 @@ namespace Metasound
 	};
 
 
-	DECLARE_METASOUND_DATA_REFERENCE_TYPES(FUnformattedAudio, METASOUNDSTANDARDNODES_API, FUnformattedAudioTypeInfo, FUnformattedAudioReadRef, FUnformattedAudioWriteRef);
-	
 	DECLARE_METASOUND_DATA_REFERENCE_TYPES(FMultichannelAudioFormat, METASOUNDSTANDARDNODES_API, FMultichannelAudioFormatTypeInfo, FMultichannelAudioFormatReadRef, FMultichannelAudioFormatWriteRef);
 
 	DECLARE_METASOUND_DATA_REFERENCE_TYPES(FMonoAudioFormat, METASOUNDSTANDARDNODES_API, FMonoAudioFormatTypeInfo, FMonoAudioFormatReadRef, FMonoAudioFormatWriteRef);
