Metasound mono and stereo delay nodes

diff --git a/Engine/Plugins/Runtime/Metasound/Source/MetasoundStandardNodes/Private/MetasoundDelayNode.cpp b/Engine/Plugins/Runtime/Metasound/Source/MetasoundStandardNodes/Private/MetasoundDelayNode.cpp
new file mode 100644
--- /dev/null
+++ b/Engine/Plugins/Runtime/Metasound/Source/MetasoundStandardNodes/Private/MetasoundDelayNode.cpp
@@ -0,0 +1,232 @@
+// Copyright Epic Games, Inc. All Rights Reserved.
+
+#include "MetasoundDelayNode.h"
+
+#include "Internationalization/Text.h"
+#include "MetasoundExecutableOperator.h"
+#include "MetasoundNodeRegistrationMacro.h"
+#include "MetasoundPrimitives.h"
+#include "MetasoundStandardNodesNames.h"
+#include "MetasoundTrigger.h"
+#include "MetasoundTime.h"
+#include "MetasoundAudioBuffer.h"
+#include "DSP/Delay.h"
+
+#define LOCTEXT_NAMESPACE "MetasoundStandardNodes"
+
+namespace Metasound
+{
+	namespace Delay
+	{
+		static const TCHAR* InParamNameAudioInput = TEXT("In");
+		static const TCHAR* InParamNameDelayTime = TEXT("Delay Time");
+		static const TCHAR* InParamNameDryLevel = TEXT("Dry Level");
+		static const TCHAR* InParamNameWetLevel = TEXT("Wet Level");
+		static const TCHAR* InParamNameFeedbackAmount = TEXT("Feedback");
+		static const TCHAR* OutParamNameAudio = TEXT("Out");
+
+		// TODO: make this a static vertex
+		static float MaxDelaySeconds = 5.0f;
+	}
+
+	class FDelayOperator : public TExecutableOperator<FDelayOperator>
+	{
+	public:
+
+		static const FNodeClassMetadata& GetNodeInfo();
+		static const FVertexInterface& GetVertexInterface();
+		static TUniquePtr<IOperator> CreateOperator(const FCreateOperatorParams& InParams, FBuildErrorArray& OutErrors);
+
+		FDelayOperator(const FOperatorSettings& InSettings, 
+			const FAudioBufferReadRef& InAudioInput, 
+			const FTimeReadRef& InDelayTime, 
+			const FFloatReadRef& InDryLevel, 
+			const FFloatReadRef& InWetLevel, 
+			const FFloatReadRef& InFeedback);
+
+		virtual FDataReferenceCollection GetInputs() const override;
+		virtual FDataReferenceCollection GetOutputs() const override;
+		void Execute();
+
+	private:
+		float GetInputDelayTimeMsec() const;
+
+		// The input audio buffer
+		FAudioBufferReadRef AudioInput;
+
+		// The amount of delay time
+		FTimeReadRef DelayTime;
+
+		// The the dry level
+		FFloatReadRef DryLevel;
+
+		// The the wet level
+		FFloatReadRef WetLevel;
+
+		// The feedback amount
+		FFloatReadRef Feedback;
+
+		// The audio output
+		FAudioBufferWriteRef AudioOutput;
+
+		// The internal delay buffer
+		Audio::FDelay DelayBuffer;
+
+		// The current delay time
+		float PrevDelayTimeMsec;
+
+		// Feedback sample
+		float FeedbackSample;
+	};
+
+	FDelayOperator::FDelayOperator(const FOperatorSettings& InSettings, 
+		const FAudioBufferReadRef& InAudioInput, 
+		const FTimeReadRef& InDelayTime,
+		const FFloatReadRef& InDryLevel, 
+		const FFloatReadRef& InWetLevel, 
+		const FFloatReadRef& InFeedback)
+
+		: AudioInput(InAudioInput)
+		, DelayTime(InDelayTime)
+		, DryLevel(InDryLevel)
+		, WetLevel(InWetLevel)
+		, Feedback(InFeedback)
+		, AudioOutput(FAudioBufferWriteRef::CreateNew(InSettings))
+		, PrevDelayTimeMsec(GetInputDelayTimeMsec())
+		, FeedbackSample(0.0f)
+	{
+		DelayBuffer.Init(InSettings.GetSampleRate(), Delay::MaxDelaySeconds);
+		DelayBuffer.SetDelayMsec(PrevDelayTimeMsec);
+	}
+
+	FDataReferenceCollection FDelayOperator::GetInputs() const
+	{
+		FDataReferenceCollection InputDataReferences;
+		InputDataReferences.AddDataReadReference(Delay::InParamNameAudioInput, FAudioBufferReadRef(AudioInput));
+		InputDataReferences.AddDataReadReference(Delay::InParamNameDelayTime, FTimeReadRef(DelayTime));
+		InputDataReferences.AddDataReadReference(Delay::InParamNameDryLevel, FFloatReadRef(DryLevel));
+		InputDataReferences.AddDataReadReference(Delay::InParamNameWetLevel, FFloatReadRef(WetLevel));
+		InputDataReferences.AddDataReadReference(Delay::InParamNameFeedbackAmount, FFloatReadRef(Feedback));
+
+		return InputDataReferences;
+	}
+
+	FDataReferenceCollection FDelayOperator::GetOutputs() const
+	{
+		FDataReferenceCollection OutputDataReferences;
+		OutputDataReferences.AddDataReadReference(Delay::OutParamNameAudio, FAudioBufferReadRef(AudioOutput));
+		return OutputDataReferences;
+	}
+
+	float FDelayOperator::GetInputDelayTimeMsec() const
+	{
+		// Clamp the delay time to the max delay allowed
+		return 1000.0f * FMath::Clamp((float)DelayTime->GetSeconds(), 0.0f, Delay::MaxDelaySeconds);
+	}
+
+	void FDelayOperator::Execute()
+	{
+		// Get clamped delay time
+		float CurrentInputDelayTime = GetInputDelayTimeMsec();
+
+		// Check to see if our delay amount has changed
+		if (!FMath::IsNearlyEqual(PrevDelayTimeMsec, CurrentInputDelayTime))
+		{
+			PrevDelayTimeMsec = CurrentInputDelayTime;
+			DelayBuffer.SetEasedDelayMsec(PrevDelayTimeMsec);
+		}
+
+		const float* InputAudio = AudioInput->GetData();
+
+		float* OutputAudio = AudioOutput->GetData();
+		int32 NumFrames = AudioInput->Num();
+
+		// Clamp the feedback amount to make sure it's bounded. Clamp to a number slightly less than 1.0
+		float FeedbackAmount = FMath::Clamp(*Feedback, 0.0f, 1.0f - SMALL_NUMBER);
+		float CurrentDryLevel = FMath::Clamp(*DryLevel, 0.0f, 1.0f);
+		float CurrentWetLevel = FMath::Clamp(*WetLevel, 0.0f, 1.0f);
+
+		if (FMath::IsNearlyZero(FeedbackAmount))
+		{
+			FeedbackSample = 0.0f;
+
+			for (int32 FrameIndex = 0; FrameIndex < NumFrames; ++FrameIndex)
+			{
+				OutputAudio[FrameIndex] = CurrentWetLevel * DelayBuffer.ProcessAudioSample(InputAudio[FrameIndex]) + CurrentDryLevel * OutputAudio[FrameIndex];
+			}
+		}
+		else
+		{
+			// There is some amount of feedback so we do the feedback mixing
+			for (int32 FrameIndex = 0; FrameIndex < NumFrames; ++FrameIndex)
+			{
+				OutputAudio[FrameIndex] = CurrentWetLevel * DelayBuffer.ProcessAudioSample(InputAudio[FrameIndex] + FeedbackSample * FeedbackAmount) + CurrentDryLevel * OutputAudio[FrameIndex];
+				FeedbackSample = OutputAudio[FrameIndex];
+			}
+		}
+	}
+
+	const FVertexInterface& FDelayOperator::GetVertexInterface()
+	{
+		static const FVertexInterface Interface(
+			FInputVertexInterface(
+				TInputDataVertexModel<FAudioBuffer>(Delay::InParamNameAudioInput, LOCTEXT("AudioInputTooltip", "Audio input.")),
+				TInputDataVertexModel<FTime>(Delay::InParamNameDelayTime, LOCTEXT("DelayTimeTooltip", "The amount of time to delay the audio."), 1.0f),
+				TInputDataVertexModel<float>(Delay::InParamNameDryLevel, LOCTEXT("DryLevelTooltip", "The dry level of the delay."), 0.0f),
+				TInputDataVertexModel<float>(Delay::InParamNameWetLevel, LOCTEXT("FeedbackTooltip", "The wet level of the delay."), 1.0f),
+				TInputDataVertexModel<float>(Delay::InParamNameFeedbackAmount, LOCTEXT("FeedbackTooltip", "Feedback amount."), 0.0f)
+			),
+			FOutputVertexInterface(
+				TOutputDataVertexModel<FAudioBuffer>(Delay::OutParamNameAudio, LOCTEXT("DelayOutputTooltip", "Audio output."))
+			)
+		);
+
+		return Interface;
+	}
+
+	const FNodeClassMetadata& FDelayOperator::GetNodeInfo()
+	{
+		auto InitNodeInfo = []() -> FNodeClassMetadata
+		{
+			FNodeClassMetadata Info;
+			Info.ClassName = { Metasound::StandardNodes::Namespace, TEXT("Delay"), Metasound::StandardNodes::AudioVariant };
+			Info.MajorVersion = 1;
+			Info.MinorVersion = 0;
+			Info.DisplayName = LOCTEXT("Metasound_DelayDisplayName", "Delay");
+			Info.Description = LOCTEXT("Metasound_DelayNodeDescription", "Delays an audio buffer by the specified amount.");
+			Info.Author = PluginAuthor;
+			Info.PromptIfMissing = PluginNodeMissingPrompt;
+			Info.DefaultInterface = GetVertexInterface();
+
+			return Info;
+		};
+
+		static const FNodeClassMetadata Info = InitNodeInfo();
+
+		return Info;
+	}
+
+	TUniquePtr<IOperator> FDelayOperator::CreateOperator(const FCreateOperatorParams& InParams, FBuildErrorArray& OutErrors)
+	{
+		const FDelayNode& DelayNode = static_cast<const FDelayNode&>(InParams.Node);
+		const FDataReferenceCollection& InputCollection = InParams.InputDataReferences;
+		const FInputVertexInterface& InputInterface = GetVertexInterface().GetInputInterface();
+
+		FAudioBufferReadRef AudioIn = InputCollection.GetDataReadReferenceOrConstruct<FAudioBuffer>(Delay::InParamNameAudioInput, InParams.OperatorSettings);
+		FTimeReadRef DelayTime = InputCollection.GetDataReadReferenceOrConstructWithVertexDefault<FTime, float>(InputInterface, Delay::InParamNameDelayTime);
+		FFloatReadRef DryLevel = InputCollection.GetDataReadReferenceOrConstructWithVertexDefault<float>(InputInterface, Delay::InParamNameDryLevel);
+		FFloatReadRef WetLevel = InputCollection.GetDataReadReferenceOrConstructWithVertexDefault<float>(InputInterface, Delay::InParamNameWetLevel);
+		FFloatReadRef Feedback = InputCollection.GetDataReadReferenceOrConstructWithVertexDefault<float>(InputInterface, Delay::InParamNameFeedbackAmount);
+
+		return MakeUnique<FDelayOperator>(InParams.OperatorSettings, AudioIn, DelayTime, DryLevel, WetLevel, Feedback);
+	}
+
+	FDelayNode::FDelayNode(const FNodeInitData& InitData)
+		: FNodeFacade(InitData.InstanceName, InitData.InstanceID, TFacadeOperatorClass<FDelayOperator>())
+	{
+	}
+
+	METASOUND_REGISTER_NODE(FDelayNode)
+}
+
+#undef LOCTEXT_NAMESPACE
diff --git a/Engine/Plugins/Runtime/Metasound/Source/MetasoundStandardNodes/Private/MetasoundStereoDelayNode.cpp b/Engine/Plugins/Runtime/Metasound/Source/MetasoundStandardNodes/Private/MetasoundStereoDelayNode.cpp
new file mode 100644
--- /dev/null
+++ b/Engine/Plugins/Runtime/Metasound/Source/MetasoundStandardNodes/Private/MetasoundStereoDelayNode.cpp
@@ -0,0 +1,262 @@
+// Copyright Epic Games, Inc. All Rights Reserved.
+
+#include "MetasoundStereoDelayNode.h"
+
+#include "Internationalization/Text.h"
+#include "MetasoundExecutableOperator.h"
+#include "MetasoundNodeRegistrationMacro.h"
+#include "MetasoundPrimitives.h"
+#include "MetasoundStandardNodesNames.h"
+#include "MetasoundTrigger.h"
+#include "MetasoundTime.h"
+#include "MetasoundAudioBuffer.h"
+#include "DSP/Delay.h"
+
+#define LOCTEXT_NAMESPACE "MetasoundStandardNodes"
+
+namespace Metasound
+{
+	namespace StereoDelay
+	{
+		static const TCHAR* InParamNameAudioInputLeft = TEXT("In Left");
+		static const TCHAR* InParamNameAudioInputRight = TEXT("In Right");
+		static const TCHAR* InParamNameDelayTime = TEXT("Delay Time");
+		static const TCHAR* InParamNameDryLevel = TEXT("Dry Level");
+		static const TCHAR* InParamNameWetLevel = TEXT("Wet Level");
+		static const TCHAR* InParamNameFeedbackAmount = TEXT("Feedback");
+		static const TCHAR* OutParamNameAudioLeft = TEXT("Out Left");
+		static const TCHAR* OutParamNameAudioRight = TEXT("Out Right");
+
+		static float MaxDelaySeconds = 5.0f;
+	}
+
+	class FStereoDelayOperator : public TExecutableOperator<FStereoDelayOperator>
+	{
+	public:
+
+		static const FNodeClassMetadata& GetNodeInfo();
+		static const FVertexInterface& GetVertexInterface();
+		static TUniquePtr<IOperator> CreateOperator(const FCreateOperatorParams& InParams, FBuildErrorArray& OutErrors);
+
+		FStereoDelayOperator(const FOperatorSettings& InSettings, 
+			const FAudioBufferReadRef& InLeftAudioInput, 
+			const FAudioBufferReadRef& InRightAudioInput,
+			const FTimeReadRef& InDelayTime,
+			const FFloatReadRef& InDryLevel, 
+			const FFloatReadRef& InWetLevel, 
+			const FFloatReadRef& InFeedback);
+
+		virtual FDataReferenceCollection GetInputs() const override;
+		virtual FDataReferenceCollection GetOutputs() const override;
+		void Execute();
+
+	private:
+		float GetInputDelayTimeMsec() const;
+
+		// The input audio buffer
+		FAudioBufferReadRef LeftAudioInput;
+		FAudioBufferReadRef RightAudioInput;
+
+		// The amount of delay time
+		FTimeReadRef DelayTime;
+
+		// The the dry level
+		FFloatReadRef DryLevel;
+
+		// The the wet level
+		FFloatReadRef WetLevel;
+
+		// The feedback amount
+		FFloatReadRef Feedback;
+
+		// The audio output
+		FAudioBufferWriteRef LeftAudioOutput;
+		FAudioBufferWriteRef RightAudioOutput;
+
+		// The internal delay buffer
+		Audio::FDelay LeftDelayBuffer;
+		Audio::FDelay RightDelayBuffer;
+
+		// The current delay time
+		float PrevDelayTimeMsec;
+	};
+
+	FStereoDelayOperator::FStereoDelayOperator(const FOperatorSettings& InSettings, 
+		const FAudioBufferReadRef& InLeftAudioInput, 
+		const FAudioBufferReadRef& InRightAudioInput,
+		const FTimeReadRef& InDelayTime,
+		const FFloatReadRef& InDryLevel, 
+		const FFloatReadRef& InWetLevel, 
+		const FFloatReadRef& InFeedback)
+
+		: LeftAudioInput(InLeftAudioInput)
+		, RightAudioInput(InRightAudioInput)
+		, DelayTime(InDelayTime)
+		, DryLevel(InDryLevel)
+		, WetLevel(InWetLevel)
+		, Feedback(InFeedback)
+		, LeftAudioOutput(FAudioBufferWriteRef::CreateNew(InSettings))
+		, RightAudioOutput(FAudioBufferWriteRef::CreateNew(InSettings))
+		, PrevDelayTimeMsec(GetInputDelayTimeMsec())
+	{
+		LeftDelayBuffer.Init(InSettings.GetSampleRate(), StereoDelay::MaxDelaySeconds);
+		LeftDelayBuffer.SetDelayMsec(PrevDelayTimeMsec);
+		RightDelayBuffer.Init(InSettings.GetSampleRate(), StereoDelay::MaxDelaySeconds);
+		RightDelayBuffer.SetDelayMsec(PrevDelayTimeMsec);
+	}
+
+	FDataReferenceCollection FStereoDelayOperator::GetInputs() const
+	{
+		FDataReferenceCollection InputDataReferences;
+		InputDataReferences.AddDataReadReference(StereoDelay::InParamNameAudioInputLeft, FAudioBufferReadRef(LeftAudioInput));
+		InputDataReferences.AddDataReadReference(StereoDelay::InParamNameAudioInputRight, FAudioBufferReadRef(RightAudioInput));
+		InputDataReferences.AddDataReadReference(StereoDelay::InParamNameDelayTime, FTimeReadRef(DelayTime));
+		InputDataReferences.AddDataReadReference(StereoDelay::InParamNameDryLevel, FFloatReadRef(DryLevel));
+		InputDataReferences.AddDataReadReference(StereoDelay::InParamNameWetLevel, FFloatReadRef(WetLevel));
+		InputDataReferences.AddDataReadReference(StereoDelay::InParamNameFeedbackAmount, FFloatReadRef(Feedback));
+
+		return InputDataReferences;
+	}
+
+	FDataReferenceCollection FStereoDelayOperator::GetOutputs() const
+	{
+		FDataReferenceCollection OutputDataReferences;
+		OutputDataReferences.AddDataReadReference(StereoDelay::OutParamNameAudioLeft, FAudioBufferReadRef(LeftAudioOutput));
+		OutputDataReferences.AddDataReadReference(StereoDelay::OutParamNameAudioRight, FAudioBufferReadRef(RightAudioOutput));
+		return OutputDataReferences;
+	}
+
+	float FStereoDelayOperator::GetInputDelayTimeMsec() const
+	{
+		// Clamp the delay time to the max delay allowed
+		return 1000.0f * FMath::Clamp((float)DelayTime->GetSeconds(), 0.0f, StereoDelay::MaxDelaySeconds);
+	}
+
+	void FStereoDelayOperator::Execute()
+	{
+		// Get clamped delay time
+		float CurrentInputDelayTime = GetInputDelayTimeMsec();
+
+		// Check to see if our delay amount has changed
+		if (!FMath::IsNearlyEqual(PrevDelayTimeMsec, CurrentInputDelayTime))
+		{
+			PrevDelayTimeMsec = CurrentInputDelayTime;
+			LeftDelayBuffer.SetEasedDelayMsec(PrevDelayTimeMsec);
+			RightDelayBuffer.SetEasedDelayMsec(PrevDelayTimeMsec);
+		}
+
+		const float* LeftInput = LeftAudioInput->GetData();
+		const float* RightInput = RightAudioInput->GetData();
+
+		float* LeftOutput = LeftAudioOutput->GetData();
+		float* RightOutput = RightAudioOutput->GetData();
+
+		int32 NumFrames = LeftAudioInput->Num();
+
+		// Clamp the feedback amount to make sure it's bounded. Clamp to a number slightly less than 1.0
+		float FeedbackAmount = FMath::Clamp(*Feedback, 0.0f, 1.0f - SMALL_NUMBER);
+		float CurrentDryLevel = FMath::Clamp(*DryLevel, 0.0f, 1.0f);
+		float CurrentWetLevel = FMath::Clamp(*WetLevel, 0.0f, 1.0f);
+
+		if (FMath::IsNearlyZero(FeedbackAmount))
+		{
+			// if pingpong
+			{
+				for (int32 FrameIndex = 0; FrameIndex < NumFrames; ++FrameIndex)
+				{
+					// Ping pong feeds right to left and left to right
+					LeftOutput[FrameIndex] = CurrentWetLevel * LeftDelayBuffer.ProcessAudioSample(RightInput[FrameIndex]) + CurrentDryLevel * LeftInput[FrameIndex];
+					RightOutput[FrameIndex] = CurrentWetLevel * RightDelayBuffer.ProcessAudioSample(LeftInput[FrameIndex]) + CurrentDryLevel * RightInput[FrameIndex];
+				}
+			}
+		}
+		else
+		{
+			// TODO: support different delay cross-modes via enum, currently default to pingpong
+		
+			// if pingpong
+			{
+				// There is some amount of feedback so we do the feedback mixing
+				for (int32 FrameIndex = 0; FrameIndex < NumFrames; ++FrameIndex)
+				{
+					float LeftDelayOut = LeftDelayBuffer.Read();
+					float RightDelayOut = RightDelayBuffer.Read();
+
+					// Pingpong algorithm feeds right output into left and left output to right input
+					float LeftDelayIn = RightInput[FrameIndex] + RightDelayOut  * FeedbackAmount;
+					float RightDelayIn = LeftInput[FrameIndex] + LeftDelayOut * FeedbackAmount;
+
+					LeftOutput[FrameIndex] = CurrentWetLevel * LeftDelayBuffer.ProcessAudioSample(LeftDelayIn) + CurrentDryLevel * LeftInput[FrameIndex];
+					RightOutput[FrameIndex] = CurrentWetLevel * RightDelayBuffer.ProcessAudioSample(RightDelayIn) + CurrentDryLevel * RightInput[FrameIndex];
+				}
+			}
+		}
+	}
+
+	const FVertexInterface& FStereoDelayOperator::GetVertexInterface()
+	{
+		static const FVertexInterface Interface(
+			FInputVertexInterface(
+				TInputDataVertexModel<FAudioBuffer>(StereoDelay::InParamNameAudioInputLeft, LOCTEXT("LeftAudioInputTooltip", "Left channel audio input.")),
+				TInputDataVertexModel<FAudioBuffer>(StereoDelay::InParamNameAudioInputRight, LOCTEXT("RightAudioInputTooltip", "Right channel audio input.")),
+				TInputDataVertexModel<FTime>(StereoDelay::InParamNameDelayTime, LOCTEXT("DelayTimeTooltip", "The amount of time to delay the audio."), 1.0f),
+				TInputDataVertexModel<float>(StereoDelay::InParamNameDryLevel, LOCTEXT("DryLevelTooltip", "The dry level of the delay."), 0.0f),
+				TInputDataVertexModel<float>(StereoDelay::InParamNameWetLevel, LOCTEXT("FeedbackTooltip", "The wet level of the delay."), 1.0f),
+				TInputDataVertexModel<float>(StereoDelay::InParamNameFeedbackAmount, LOCTEXT("FeedbackTooltip", "Feedback amount."), 0.0f)
+			),
+			FOutputVertexInterface(
+				TOutputDataVertexModel<FAudioBuffer>(StereoDelay::OutParamNameAudioLeft, LOCTEXT("LeftDelayOutputTooltip", "Left channel audio output.")),
+				TOutputDataVertexModel<FAudioBuffer>(StereoDelay::OutParamNameAudioRight, LOCTEXT("RightDelayOutputTooltip", "Right channel audio output."))
+			)
+		);
+
+		return Interface;
+	}
+
+	const FNodeClassMetadata& FStereoDelayOperator::GetNodeInfo()
+	{
+		auto InitNodeInfo = []() -> FNodeClassMetadata
+		{
+			FNodeClassMetadata Info;
+			Info.ClassName = { Metasound::StandardNodes::Namespace, TEXT("Stereo Delay"), Metasound::StandardNodes::AudioVariant };
+			Info.MajorVersion = 1;
+			Info.MinorVersion = 0;
+			Info.DisplayName = LOCTEXT("Metasound_StereoDelayDisplayName", "Stereo Delay");
+			Info.Description = LOCTEXT("Metasound_StereoDelayNodeDescription", "Delays a stereo audio buffer by the specified amount.");
+			Info.Author = PluginAuthor;
+			Info.PromptIfMissing = PluginNodeMissingPrompt;
+			Info.DefaultInterface = GetVertexInterface();
+
+			return Info;
+		};
+
+		static const FNodeClassMetadata Info = InitNodeInfo();
+
+		return Info;
+	}
+
+	TUniquePtr<IOperator> FStereoDelayOperator::CreateOperator(const FCreateOperatorParams& InParams, FBuildErrorArray& OutErrors)
+	{
+		const FStereoDelayNode& StereoDelayNode = static_cast<const FStereoDelayNode&>(InParams.Node);
+		const FDataReferenceCollection& InputCollection = InParams.InputDataReferences;
+		const FInputVertexInterface& InputInterface = GetVertexInterface().GetInputInterface();
+
+		FAudioBufferReadRef LeftAudioIn = InputCollection.GetDataReadReferenceOrConstruct<FAudioBuffer>(StereoDelay::InParamNameAudioInputLeft, InParams.OperatorSettings);
+		FAudioBufferReadRef RightAudioIn = InputCollection.GetDataReadReferenceOrConstruct<FAudioBuffer>(StereoDelay::InParamNameAudioInputRight, InParams.OperatorSettings);
+		FTimeReadRef DelayTime = InputCollection.GetDataReadReferenceOrConstructWithVertexDefault<FTime, float>(InputInterface, StereoDelay::InParamNameDelayTime);
+		FFloatReadRef DryLevel = InputCollection.GetDataReadReferenceOrConstructWithVertexDefault<float>(InputInterface, StereoDelay::InParamNameDryLevel);
+		FFloatReadRef WetLevel = InputCollection.GetDataReadReferenceOrConstructWithVertexDefault<float>(InputInterface, StereoDelay::InParamNameWetLevel);
+		FFloatReadRef Feedback = InputCollection.GetDataReadReferenceOrConstructWithVertexDefault<float>(InputInterface, StereoDelay::InParamNameFeedbackAmount);
+
+		return MakeUnique<FStereoDelayOperator>(InParams.OperatorSettings, LeftAudioIn, RightAudioIn, DelayTime, DryLevel, WetLevel, Feedback);
+	}
+
+	FStereoDelayNode::FStereoDelayNode(const FNodeInitData& InitData)
+		: FNodeFacade(InitData.InstanceName, InitData.InstanceID, TFacadeOperatorClass<FStereoDelayOperator>())
+	{
+	}
+
+	METASOUND_REGISTER_NODE(FStereoDelayNode)
+}
+
+#undef LOCTEXT_NAMESPACE
diff --git a/Engine/Plugins/Runtime/Metasound/Source/MetasoundStandardNodes/Public/MetasoundDelayNode.h b/Engine/Plugins/Runtime/Metasound/Source/MetasoundStandardNodes/Public/MetasoundDelayNode.h
new file mode 100644
--- /dev/null
+++ b/Engine/Plugins/Runtime/Metasound/Source/MetasoundStandardNodes/Public/MetasoundDelayNode.h
@@ -0,0 +1,23 @@
+// Copyright Epic Games, Inc. All Rights Reserved.
+
+#pragma once
+
+#include "CoreMinimal.h"
+#include "MetasoundFacade.h"
+
+namespace Metasound
+{
+	/** FDelayNode
+	 *
+	 *  Delays an audio buffer by a specified amount.
+	 */
+	class METASOUNDSTANDARDNODES_API FDelayNode : public FNodeFacade
+	{
+	public:
+		/**
+		 * Constructor used by the Metasound Frontend.
+		 */
+		FDelayNode(const FNodeInitData& InitData);
+	};
+} // namespace Metasound
+
diff --git a/Engine/Plugins/Runtime/Metasound/Source/MetasoundStandardNodes/Public/MetasoundStereoDelayNode.h b/Engine/Plugins/Runtime/Metasound/Source/MetasoundStandardNodes/Public/MetasoundStereoDelayNode.h
new file mode 100644
--- /dev/null
+++ b/Engine/Plugins/Runtime/Metasound/Source/MetasoundStandardNodes/Public/MetasoundStereoDelayNode.h
@@ -0,0 +1,23 @@
+// Copyright Epic Games, Inc. All Rights Reserved.
+
+#pragma once
+
+#include "CoreMinimal.h"
+#include "MetasoundFacade.h"
+
+namespace Metasound
+{
+	/** FStereoDelayNode
+	 *
+	 *  Delays an audio buffer by a specified amount.
+	 */
+	class METASOUNDSTANDARDNODES_API FStereoDelayNode : public FNodeFacade
+	{
+	public:
+		/**
+		 * Constructor used by the Metasound Frontend.
+		 */
+		FStereoDelayNode(const FNodeInitData& InitData);
+	};
+} // namespace Metasound
+
