Checkpoint: Metasound Node Inspector Panel & Array Support

diff --git a/Engine/Plugins/Runtime/Metasound/Source/MetasoundEditor/Private/MetasoundDetailCustomization.cpp b/Engine/Plugins/Runtime/Metasound/Source/MetasoundEditor/Private/MetasoundDetailCustomization.cpp
--- a/Engine/Plugins/Runtime/Metasound/Source/MetasoundEditor/Private/MetasoundDetailCustomization.cpp
+++ b/Engine/Plugins/Runtime/Metasound/Source/MetasoundEditor/Private/MetasoundDetailCustomization.cpp
@@ -39,6 +39,14 @@ namespace Metasound
 {
 	namespace Editor
 	{
+		/* Enums to use when grouping the blueprint members in the list panel. The order here will determine the order in the list */
+		enum class ENodeSection : uint8
+		{
+			NONE = 0,
+			INPUTS,
+			OUTPUTS
+		};
+
 		FName BuildChildPath(const FString& InBasePath, FName InPropertyName)
 		{
 			return FName(InBasePath + TEXT(".") + InPropertyName.ToString());
@@ -78,7 +86,7 @@ namespace Metasound
 		}
 
 		template <typename T>
-		void BuildIOFixedArray(IDetailLayoutBuilder& InDetailLayout, FName InCategoryName, FName InPropertyName, const TSet<FString>& InRequiredValues)
+		void BuildIOFixedArray(IDetailLayoutBuilder& InDetailLayout, FName InCategoryName, FName InPropertyName)
 		{
 			const bool bIsInput = InCategoryName == "Inputs";
 
@@ -111,7 +119,6 @@ namespace Metasound
 
 				FString Name;
 				const bool bNameFound = NameProperty->GetValue(Name) == FPropertyAccess::Success;
-				const bool bIsRequired = bNameFound && InRequiredValues.Contains(Name);
 
 				// Hide literal members
 				if (LiteralInputs.Contains(Name))
@@ -119,74 +126,16 @@ namespace Metasound
 					continue;
 				}
 
-				FText DisplayName;
-				DisplayNameProperty->GetValue(DisplayName);
-				DisplayNameProperty->SetInstanceMetaData("LastValidName", DisplayName.ToString());
-
-				FSimpleDelegate NameChangeDelegate = FSimpleDelegate::CreateLambda([ChangedIndex = i, DisplayNameProperty, ArrayHandle, DisplayNamePropertyName]
-				{
-					FText MatchDisplayName;
-					if (DisplayNameProperty->GetValue(MatchDisplayName) != FPropertyAccess::Success)
-					{
-						return;
-					}
-
-					if (const FString* LastValidName = DisplayNameProperty->GetInstanceMetaData("LastValidName"))
-					{
-						uint32 NumElements = 0;
-						ArrayHandle->GetNumElements(NumElements);
-						for (int32 Elem = 0; Elem < static_cast<int32>(NumElements); ++Elem)
-						{
-							TSharedRef<IPropertyHandle> ElemHandle = ArrayHandle->GetElement(Elem);
-							TSharedPtr<IPropertyHandle> ElemNameProperty = ElemHandle->GetChildHandle(DisplayNamePropertyName, true /* bRecurse */);
-							FText ElemDisplayName;
-							if (ElemNameProperty->GetValue(ElemDisplayName) == FPropertyAccess::Success)
-							{
-								if (Elem != ChangedIndex && !ElemDisplayName.CompareTo(MatchDisplayName))
-								{
-									DisplayNameProperty->SetValue(FText::FromString(*LastValidName));
-									FNotificationInfo Info(LOCTEXT("MetasoundEditor_InvalidRename", "Rename failed: Input/output with name already exists."));
-									Info.bFireAndForget = true;
-									Info.ExpireDuration = 2.0f;
-									Info.bUseThrobber = true;
-									FSlateNotificationManager::Get().AddNotification(Info);
-									return;
-								}
-							}
-						}
-					}
-
-					DisplayNameProperty->SetInstanceMetaData("LastValidName", MatchDisplayName.ToString());
-				});
-
-				DisplayNameProperty->SetOnPropertyValueChanged(NameChangeDelegate);
-
 				CategoryBuilder.AddCustomRow(ParentProperty->GetPropertyDisplayName())
-				.EditCondition(!bIsRequired, nullptr)
 				.NameContent()
 				[
 					SNew(STextBlock)
 					.Font(IDetailLayoutBuilder::GetDetailFontBold())
-					.Text(TAttribute<FText>::Create([i, bIsRequired, DisplayNameProperty, TypeProperty]()
+					.Text(TAttribute<FText>::Create([i, DisplayNameProperty]()
 					{
-						FName TypeName;
-						TypeProperty->GetValue(TypeName);
-						FString TypeNameString = TypeName.ToString();
-
-						// Remove namespace info to keep concise
-						TypeNameString.RightChopInline(TypeNameString.Find(TEXT(":"), ESearchCase::IgnoreCase, ESearchDir::FromEnd) + 1);
-
 						FText DisplayName;
 						DisplayNameProperty->GetValue(DisplayName);
-
-						if (bIsRequired)
-						{
-							return FText::Format(LOCTEXT("MetasoundEditor_FixedIOArrayRequiredEntry_Format", "{0}. {1} ({2}, Required)"), FText::AsNumber(i + 1), DisplayName, FText::FromString(TypeNameString));
-						}
-						else
-						{
-							return FText::Format(LOCTEXT("MetasoundEditor_FixedIOArray_Format", "{0}. {1} ({2})"), FText::AsNumber(i + 1), DisplayName, FText::FromString(TypeNameString));
-						}
+						return DisplayName;
 					}))
 					.ToolTipText(TAttribute<FText>::Create([ToolTipProperty]()
 					{
@@ -195,28 +144,6 @@ namespace Metasound
 						return ToolTip;
 					}))
 				];
-
-				if (!bIsRequired)
-				{
-					CategoryBuilder.AddProperty(DisplayNameProperty);
-					CategoryBuilder.AddProperty(ToolTipProperty);
-				}
-
-				if (bIsInput)
-				{
-					const FName DefaultsPropertyName = GET_MEMBER_NAME_CHECKED(FMetasoundFrontendClassInput, Defaults);
-					TSharedPtr<IPropertyHandle> DefaultsProperty = ArrayItemHandle->GetChildHandle(DefaultsPropertyName);
-
-					TSharedPtr<IPropertyHandleArray> DefaultsArrayHandle = DefaultsProperty->AsArray();
-
-					uint32 NumDefaults = 0;
-					DefaultsArrayHandle->GetNumElements(NumDefaults);
-					for (uint32 InputLiteralIndex = 0; InputLiteralIndex < NumDefaults; ++InputLiteralIndex)
-					{
-						TSharedPtr<IPropertyHandle> LiteralHandle = DefaultsArrayHandle->GetElement(InputLiteralIndex)->GetChildHandle(GET_MEMBER_NAME_CHECKED(FMetasoundFrontendVertexLiteral, Value));
-						CategoryBuilder.AddProperty(LiteralHandle);
-					}
-				}
 			}
 
 			FSimpleDelegate RefreshDelegate = FSimpleDelegate::CreateLambda([DetailLayout = &InDetailLayout]()
@@ -224,7 +151,7 @@ namespace Metasound
 				DetailLayout->ForceRefreshDetails();
 			});
 			ArrayHandle->SetOnNumElementsChanged(RefreshDelegate);
-		};
+		}
 
 		FMetasoundDetailCustomization::FMetasoundDetailCustomization(FName InDocumentPropertyName)
 			: IDetailCustomization()
@@ -250,23 +177,17 @@ namespace Metasound
 			// General Category
 			IDetailCategoryBuilder& GeneralCategoryBuilder = DetailLayout.EditCategory("General");
 
-			TArray<TWeakObjectPtr<UObject>> ObjectsCustomized;
-			DetailLayout.GetObjectsBeingCustomized(ObjectsCustomized);
-
 			const FName AuthorPropertyPath = BuildChildPath(GetMetadataPropertyPath(), GET_MEMBER_NAME_CHECKED(FMetasoundFrontendClassMetadata, Author));
 			const FName DescPropertyPath = BuildChildPath(GetMetadataPropertyPath(), GET_MEMBER_NAME_CHECKED(FMetasoundFrontendClassMetadata, Description));
-			const FName NodeTypePropertyPath = BuildChildPath(GetMetadataPropertyPath(), GET_MEMBER_NAME_CHECKED(FMetasoundFrontendClassMetadata, Type));
 			const FName VersionPropertyPath = BuildChildPath(GetMetadataPropertyPath(), GET_MEMBER_NAME_CHECKED(FMetasoundFrontendClassMetadata, Version));
 			const FName MajorVersionPropertyPath = BuildChildPath(VersionPropertyPath, GET_MEMBER_NAME_CHECKED(FMetasoundFrontendVersionNumber, Major));
 			const FName MinorVersionPropertyPath = BuildChildPath(VersionPropertyPath, GET_MEMBER_NAME_CHECKED(FMetasoundFrontendVersionNumber, Minor));
 
 			TSharedPtr<IPropertyHandle> AuthorHandle = DetailLayout.GetProperty(AuthorPropertyPath);
 			TSharedPtr<IPropertyHandle> DescHandle = DetailLayout.GetProperty(DescPropertyPath);
-			TSharedPtr<IPropertyHandle> NodeTypeHandle = DetailLayout.GetProperty(NodeTypePropertyPath);
 			TSharedPtr<IPropertyHandle> MajorVersionHandle = DetailLayout.GetProperty(MajorVersionPropertyPath);
 			TSharedPtr<IPropertyHandle> MinorVersionHandle = DetailLayout.GetProperty(MinorVersionPropertyPath);
 
-			GeneralCategoryBuilder.AddProperty(NodeTypeHandle);
 			GeneralCategoryBuilder.AddProperty(AuthorHandle);
 			GeneralCategoryBuilder.AddProperty(DescHandle);
 			GeneralCategoryBuilder.AddProperty(MajorVersionHandle);
@@ -277,31 +198,33 @@ namespace Metasound
 			// If editing multiple metasound objects, all should be the same type, so safe to just check first in array for
 			// required inputs/outputs
 			TArray<TWeakObjectPtr<UObject>> Objects;
-			TSet<FString> RequiredInputs;
-			TSet<FString> RequiredOutputs;
 			DetailLayout.GetObjectsBeingCustomized(Objects);
-			if (Objects.Num() > 0)
-			{
-				FMetasoundAssetBase* MetasoundAsset = IMetasoundUObjectRegistry::Get().GetObjectAsAssetBase(Objects[0].Get());
-				check(MetasoundAsset);
-
-				Frontend::FDocumentHandle DocumentHandle = MetasoundAsset->GetDocumentHandle();
-				for (const FMetasoundFrontendClassVertex& Desc : DocumentHandle->GetRequiredInputs())
-				{
-					RequiredInputs.Add(Desc.Name);
-				}
-				for (const FMetasoundFrontendClassVertex& Desc : DocumentHandle->GetRequiredOutputs())
-				{
-					RequiredOutputs.Add(Desc.Name);
-				}
-			}
 
 			const FName InterfacePropertyPath = BuildChildPath(GetMetadataRootClassPath(), GET_MEMBER_NAME_CHECKED(FMetasoundFrontendClass, Interface));
 			const FName InputsPropertyPath = BuildChildPath(InterfacePropertyPath, GET_MEMBER_NAME_CHECKED(FMetasoundFrontendClassInterface, Inputs));
 			const FName OutputsPropertyPath = BuildChildPath(InterfacePropertyPath, GET_MEMBER_NAME_CHECKED(FMetasoundFrontendClassInterface, Outputs));
 
-			BuildIOFixedArray<FMetasoundFrontendClassInput>(DetailLayout, "Inputs", InputsPropertyPath, RequiredInputs);
-			BuildIOFixedArray<FMetasoundFrontendClassOutput>(DetailLayout, "Outputs", OutputsPropertyPath, RequiredOutputs);
+			// TODO: Move Fixed Arrays to GraphActionMenu for selection, ordering, and
+			// eventually ability to create input nodes independent of registered inputs.
+// 			SAssignNew(GraphActionMenu, SGraphActionMenu, false)
+// 				.OnGetFilterText(this, &FMetasoundDetailCustomization::GetFilterText)
+// 				.OnCreateWidgetForAction(this, &FMetasoundDetailCustomization::OnCreateWidgetForAction)
+// 				.OnCollectAllActions(this, &FMetasoundDetailCustomization::CollectAllActions)
+// 				.OnCollectStaticSections(this, &FMetasoundDetailCustomization::CollectStaticSections)
+// 				.OnActionDragged(this, &FMetasoundDetailCustomization::OnActionDragged)
+// 				.OnActionSelected(this, &FMetasoundDetailCustomization::OnGlobalActionSelected)
+// 				.OnActionDoubleClicked(this, &FMetasoundDetailCustomization::OnActionDoubleClicked)
+// 				.OnContextMenuOpening(this, &FMetasoundDetailCustomization::OnContextMenuOpening)
+// 				.OnCategoryTextCommitted(this, &FMetasoundDetailCustomization::OnCategoryNameCommitted)
+// 				.OnCanRenameSelectedAction(this, &FMetasoundDetailCustomization::CanRequestRenameOnActionNode)
+// 				.OnGetSectionTitle(this, &FMetasoundDetailCustomization::OnGetSectionTitle)
+// 				.OnGetSectionWidget(this, &FMetasoundDetailCustomization::OnGetSectionWidget)
+// 				.OnActionMatchesName(this, &FMetasoundDetailCustomization::HandleActionMatchesName)
+// 				.AlphaSortItems(false)
+// 				.UseSectionStyling(true);
+			
+			BuildIOFixedArray<FMetasoundFrontendClassInput>(DetailLayout, "Inputs", InputsPropertyPath);
+			BuildIOFixedArray<FMetasoundFrontendClassOutput>(DetailLayout, "Outputs", OutputsPropertyPath);
 
 			// Hack to hide parent structs for nested metadata properties
 			DetailLayout.HideCategory("CustomView");
@@ -324,6 +247,34 @@ namespace Metasound
 			DetailLayout.HideCategory("Subtitles");
 			DetailLayout.HideCategory("Voice Management");
 		}
+
+		FText FMetasoundDetailCustomization::OnGetSectionTitle(int32 InSectionID)
+		{
+			FText SeperatorTitle;
+			/* Setup an appropriate name for the section for this node */
+			switch (static_cast<ENodeSection>(InSectionID))
+			{
+				case ENodeSection::INPUTS:
+				{
+					SeperatorTitle = LOCTEXT("Inputs_Title", "Inputs");
+				}
+				break;
+
+				case ENodeSection::OUTPUTS:
+				{
+					SeperatorTitle = LOCTEXT("Outputs_Title", "Outputs");
+				}
+				break;
+
+				default:
+				{
+					SeperatorTitle = LOCTEXT("Missing_Title", "UNSET");
+				}
+				break;
+			}
+
+			return SeperatorTitle;
+		}
 	} // namespace Editor
 } // namespace Metasound
 #undef LOCTEXT_NAMESPACE
diff --git a/Engine/Plugins/Runtime/Metasound/Source/MetasoundEditor/Private/MetasoundDetailCustomization.h b/Engine/Plugins/Runtime/Metasound/Source/MetasoundEditor/Private/MetasoundDetailCustomization.h
--- a/Engine/Plugins/Runtime/Metasound/Source/MetasoundEditor/Private/MetasoundDetailCustomization.h
+++ b/Engine/Plugins/Runtime/Metasound/Source/MetasoundEditor/Private/MetasoundDetailCustomization.h
@@ -2,6 +2,7 @@
 #pragma once
 
 #include "IDetailCustomization.h"
+#include "SGraphActionMenu.h"
 #include "UObject/NameTypes.h"
 
 // Forward Declarations
@@ -25,7 +26,11 @@ namespace Metasound
 			FName GetMetadataRootClassPath() const;
 			FName GetMetadataPropertyPath() const;
 
+			FText OnGetSectionTitle(int32 InSectionID);
+
 			FName DocumentPropertyName;
+
+// 			TSharedPtr<SGraphActionMenu> GraphActionMenu;
 		};
 	} // namespace Editor
 } // namespace Metasound
\ No newline at end of file
diff --git a/Engine/Plugins/Runtime/Metasound/Source/MetasoundEditor/Private/MetasoundEditor.cpp b/Engine/Plugins/Runtime/Metasound/Source/MetasoundEditor/Private/MetasoundEditor.cpp
--- a/Engine/Plugins/Runtime/Metasound/Source/MetasoundEditor/Private/MetasoundEditor.cpp
+++ b/Engine/Plugins/Runtime/Metasound/Source/MetasoundEditor/Private/MetasoundEditor.cpp
@@ -56,14 +56,23 @@ namespace Metasound
 			.SetGroup(WorkspaceMenuCategoryRef)
 			.SetIcon(FSlateIcon(FEditorStyle::GetStyleSetName(), "GraphEditor.EventGraph_16x"));
 
-			InTabManager->RegisterTabSpawner(TabFactory::Names::Properties, FOnSpawnTab::CreateLambda([InMetasoundProperties = MetasoundProperties](const FSpawnTabArgs& Args)
+			InTabManager->RegisterTabSpawner(TabFactory::Names::Inspector, FOnSpawnTab::CreateLambda([InMetasoundDetails = MetasoundDetails](const FSpawnTabArgs& Args)
 			{
-				return TabFactory::CreatePropertiesTab(InMetasoundProperties, Args);
+				return TabFactory::CreateInspectorTab(InMetasoundDetails, Args);
 			}))
-			.SetDisplayName(LOCTEXT("DetailsTab", "Details"))
+			.SetDisplayName(LOCTEXT("InspectorTab", "Inspector"))
 			.SetGroup(WorkspaceMenuCategoryRef)
 			.SetIcon(FSlateIcon(FEditorStyle::GetStyleSetName(), "LevelEditor.Tabs.Details"));
 
+			InTabManager->RegisterTabSpawner(TabFactory::Names::Metasound, FOnSpawnTab::CreateLambda([InMetasoundGeneral = MetasoundGeneral](const FSpawnTabArgs& Args)
+			{
+				return TabFactory::CreateMetasoundTab(InMetasoundGeneral, Args);
+			}))
+			.SetDisplayName(LOCTEXT("MetasoundTab", "Metasound"))
+			.SetGroup(WorkspaceMenuCategoryRef)
+			.SetIcon(FSlateIcon(FEditorStyle::GetStyleSetName(), "LevelEditor.Tabs.Settings"));
+
+
 			InTabManager->RegisterTabSpawner(TabFactory::Names::Analyzers, FOnSpawnTab::CreateLambda([InMetasoundMeter = MetasoundMeter](const FSpawnTabArgs& Args)
 			{
 				return TabFactory::CreateAnalyzersTab(InMetasoundMeter, Args);
@@ -81,7 +90,8 @@ namespace Metasound
 
 			InTabManager->UnregisterTabSpawner(TabFactory::Names::Analyzers);
 			InTabManager->UnregisterTabSpawner(TabFactory::Names::GraphCanvas);
-			InTabManager->UnregisterTabSpawner(TabFactory::Names::Properties);
+			InTabManager->UnregisterTabSpawner(TabFactory::Names::Inspector);
+			InTabManager->UnregisterTabSpawner(TabFactory::Names::Metasound);
 		}
 
 		FEditor::~FEditor()
@@ -141,7 +151,7 @@ namespace Metasound
 
 			CreateAnalyzers();
 
-			const TSharedRef<FTabManager::FLayout> StandaloneDefaultLayout = FTabManager::NewLayout("Standalone_MetasoundEditor_Layout_v4")
+			const TSharedRef<FTabManager::FLayout> StandaloneDefaultLayout = FTabManager::NewLayout("Standalone_MetasoundEditor_Layout_v6")
 				->AddArea
 				(
 					FTabManager::NewPrimaryArea()
@@ -151,10 +161,22 @@ namespace Metasound
 						->SetSizeCoefficient(0.9f)
 						->Split
 						(
-							FTabManager::NewStack()
-							->SetSizeCoefficient(0.225f)
-							->SetHideTabWell(true)
-							->AddTab(TabFactory::Names::Properties, ETabState::OpenedTab)
+							FTabManager::NewSplitter()
+							->SetOrientation(Orient_Vertical)
+							->Split
+							(
+								FTabManager::NewStack()
+								->SetSizeCoefficient(0.5f)
+								->SetHideTabWell(false)
+								->AddTab(TabFactory::Names::Metasound, ETabState::OpenedTab)
+							)
+							->Split
+							(
+								FTabManager::NewStack()
+								->SetSizeCoefficient(0.5f)
+								->SetHideTabWell(false)
+								->AddTab(TabFactory::Names::Inspector, ETabState::OpenedTab)
+							)
 						)
 						->Split
 						(
@@ -193,9 +215,10 @@ namespace Metasound
 
 		void FEditor::SetSelection(const TArray<UObject*>& SelectedObjects)
 		{
-			if (MetasoundProperties.IsValid())
+			if (MetasoundDetails.IsValid())
 			{
-				MetasoundProperties->SetObjects(SelectedObjects);
+				MetasoundDetails->SetObjects(SelectedObjects);
+				MetasoundDetails->HideFilterArea(SelectedObjects.IsEmpty() || SelectedObjects.Num() > 1);
 			}
 		}
 
@@ -268,8 +291,11 @@ namespace Metasound
 			Args.NotifyHook = this;
 
 			FPropertyEditorModule& PropertyModule = FModuleManager::LoadModuleChecked<FPropertyEditorModule>("PropertyEditor");
-			MetasoundProperties = PropertyModule.CreateDetailView(Args);
-			MetasoundProperties->SetObject(Metasound);
+
+			MetasoundGeneral = PropertyModule.CreateDetailView(Args);
+			MetasoundGeneral->SetObject(Metasound);
+
+			MetasoundDetails = PropertyModule.CreateDetailView(Args);
 
 			Palette = SNew(SMetasoundPalette);
 		}
@@ -311,7 +337,7 @@ namespace Metasound
 			MetasoundMeterAnalyzer->StartAnalyzing(EditorWorld, MetasoundAudioBus.Get());
 
 			// Create the audio meter
-			MetasoundMeter = SNew(SAudioMeter);			
+			MetasoundMeter = SNew(SAudioMeter);
 
 			// TODO: THIS IS TEMP, WIP
 			MetasoundMeter->SetOrientation(EOrientation::Orient_Vertical);
@@ -703,24 +729,13 @@ namespace Metasound
 		{
 			TArray<UObject*> Selection;
 
-			if (NewSelection.Num())
+			if (!NewSelection.IsEmpty())
 			{
 				for (TSet<UObject*>::TConstIterator SetIt(NewSelection); SetIt; ++SetIt)
 				{
-					if (Cast<UMetasoundEditorGraphNode>(*SetIt))
-					{
-						Selection.Add(GetMetasoundObject());
-					}
-					else
-					{
-						Selection.Add(*SetIt);
-					}
+					Selection.Add(*SetIt);
 				}
 			}
-			else
-			{
-				Selection.Add(GetMetasoundObject());
-			}
 
 			SetSelection(Selection);
 		}
diff --git a/Engine/Plugins/Runtime/Metasound/Source/MetasoundEditor/Private/MetasoundEditor.h b/Engine/Plugins/Runtime/Metasound/Source/MetasoundEditor/Private/MetasoundEditor.h
--- a/Engine/Plugins/Runtime/Metasound/Source/MetasoundEditor/Private/MetasoundEditor.h
+++ b/Engine/Plugins/Runtime/Metasound/Source/MetasoundEditor/Private/MetasoundEditor.h
@@ -189,8 +189,11 @@ namespace Metasound
 			/** New Graph Editor */
 			TSharedPtr<SGraphEditor> MetasoundGraphEditor;
 
-			/** Properties tab */
-			TSharedPtr<IDetailsView> MetasoundProperties;
+			/** Details tab */
+			TSharedPtr<IDetailsView> MetasoundDetails;
+
+			/** General tab */
+			TSharedPtr<IDetailsView> MetasoundGeneral;
 
 			/** Metasound Output Meter widget */
 			TSharedPtr<SAudioMeter> MetasoundMeter;
diff --git a/Engine/Plugins/Runtime/Metasound/Source/MetasoundEditor/Private/MetasoundEditorGraph.cpp b/Engine/Plugins/Runtime/Metasound/Source/MetasoundEditor/Private/MetasoundEditorGraph.cpp
--- a/Engine/Plugins/Runtime/Metasound/Source/MetasoundEditor/Private/MetasoundEditorGraph.cpp
+++ b/Engine/Plugins/Runtime/Metasound/Source/MetasoundEditor/Private/MetasoundEditorGraph.cpp
@@ -2,8 +2,40 @@
 #include "MetasoundEditorGraph.h"
 
 #include "EdGraph/EdGraphNode.h"
+#include "MetasoundEditorGraphNode.h"
+#include "MetasoundEditorGraphInputNodes.h"
+#include "MetasoundEditorModule.h"
 
 
+UMetasoundEditorGraphInputNode* UMetasoundEditorGraph::CreateInputNode(EMetasoundFrontendLiteralType LiteralType, UClass* InLiteralObjectClass, bool bInSelectNewNode)
+{
+	using namespace Metasound::Editor;
+
+	IMetasoundEditorModule & EditorModule = FModuleManager::GetModuleChecked<IMetasoundEditorModule>("MetasoundEditor");
+
+	TSubclassOf<UMetasoundEditorGraphInputNode> NodeClass = EditorModule.FindNodeInputClass(LiteralType);
+
+	if (ensure(NodeClass))
+	{
+		UMetasoundEditorGraphInputNode* NewGraphNode = CastChecked<UMetasoundEditorGraphInputNode>(CreateNode(NodeClass, bInSelectNewNode));
+
+		NewGraphNode->CreateNewGuid();
+		NewGraphNode->PostPlacedNewNode();
+
+		if (NewGraphNode->Pins.IsEmpty())
+		{
+			NewGraphNode->AllocateDefaultPins();
+		}
+
+		PostEditChange();
+		MarkPackageDirty();
+
+		return NewGraphNode;
+	}
+
+	return nullptr;
+}
+
 UObject* UMetasoundEditorGraph::GetMetasound() const
 {
 	return ParentMetasound;
diff --git a/Engine/Plugins/Runtime/Metasound/Source/MetasoundEditor/Private/MetasoundEditorGraphBuilder.cpp b/Engine/Plugins/Runtime/Metasound/Source/MetasoundEditor/Private/MetasoundEditorGraphBuilder.cpp
--- a/Engine/Plugins/Runtime/Metasound/Source/MetasoundEditor/Private/MetasoundEditorGraphBuilder.cpp
+++ b/Engine/Plugins/Runtime/Metasound/Source/MetasoundEditor/Private/MetasoundEditorGraphBuilder.cpp
@@ -10,6 +10,7 @@
 #include "MetasoundAssetBase.h"
 #include "MetasoundEditorGraph.h"
 #include "MetasoundEditorGraphNode.h"
+#include "MetasoundEditorGraphInputNodes.h"
 #include "MetasoundEditorModule.h"
 #include "MetasoundFrontendRegistries.h"
 #include "MetasoundLiteral.h"
@@ -67,9 +68,41 @@ namespace Metasound
 			check(MetasoundAsset);
 
 			UEdGraph& Graph = MetasoundAsset->GetGraphChecked();
-			FGraphNodeCreator<UMetasoundEditorGraphNode> NodeCreator(Graph);
-			UMetasoundEditorGraphNode* NewGraphNode = NodeCreator.CreateNode(bInSelectNewNode);
-			NodeCreator.Finalize();
+
+			UMetasoundEditorGraphNode* NewGraphNode = nullptr;
+			switch (InNodeHandle->GetClassType())
+			{
+				case EMetasoundFrontendClassType::Input:
+				{
+					// Should use AddInput for this case to ensure child class is used for input.
+					checkNoEntry();
+				}
+				break;
+
+				case EMetasoundFrontendClassType::Output:
+				{
+					FGraphNodeCreator<UMetasoundEditorGraphOutputNode> NodeCreator(Graph);
+					NewGraphNode = NodeCreator.CreateNode(bInSelectNewNode);
+					NodeCreator.Finalize();
+				}
+				break;
+
+				case EMetasoundFrontendClassType::External:
+				{
+					FGraphNodeCreator<UMetasoundEditorGraphExternalNode> NodeCreator(Graph);
+					NewGraphNode = NodeCreator.CreateNode(bInSelectNewNode);
+					NodeCreator.Finalize();
+				}
+				break;
+
+				case EMetasoundFrontendClassType::Graph: // TODO: Implement
+				default:
+				{
+					static_assert(static_cast<int32>(EMetasoundFrontendClassType::Invalid) == 4, "Possible missing case coverage");
+					checkNoEntry();
+				}
+				break;
+			}
 
 			const FMetasoundFrontendNodeStyle& Style = InNodeHandle->GetNodeStyle();
 
@@ -219,10 +252,64 @@ namespace Metasound
 			return GetOutputHandleFromPin(InPin);
 		}
 
-		UEdGraphNode* FGraphBuilder::AddInput(UObject& InMetasound, const FString& InName, const FName InTypeName, const FMetasoundFrontendNodeStyle& InNodeStyle, const FText& InToolTip, bool bInSelectNewNode)
+		UMetasoundEditorGraphInputNode* FGraphBuilder::AddInput(UObject& InMetasound, Frontend::FNodeHandle InNodeHandle, bool bInSelectNewNode)
 		{
-			Frontend::FNodeHandle NodeHandle = AddInputNodeHandle(InMetasound, InName, InTypeName, InNodeStyle, InToolTip);
-			return AddNode(InMetasound, NodeHandle, bInSelectNewNode);
+			using namespace Metasound::Frontend;
+
+			FMetasoundAssetBase* MetasoundAsset = IMetasoundUObjectRegistry::Get().GetObjectAsAssetBase(&InMetasound);
+			check(MetasoundAsset);
+
+			const TArray<FOutputHandle>& Outputs = InNodeHandle->GetOutputs();
+			if (!ensure(!Outputs.IsEmpty()))
+			{
+				return nullptr;
+			}
+
+			FOutputHandle OutputHandle = InNodeHandle->GetOutputs()[0];
+			FGraphHandle Graph = InNodeHandle->GetOwningGraph();
+			TArray<FMetasoundFrontendVertexLiteral> DefaultLiterals = Graph->GetDefaultInputFrontendLiterals(OutputHandle->GetName());
+
+			EMetasoundFrontendLiteralType LiteralType = EMetasoundFrontendLiteralType::None;
+			if (!DefaultLiterals.IsEmpty())
+			{
+				const FMetasoundFrontendVertexLiteral& Literal = DefaultLiterals[0];
+				LiteralType = Literal.Value.GetType();
+			}
+
+			UMetasoundEditorGraph* MetasoundGraph = Cast<UMetasoundEditorGraph>(MetasoundAsset->GetGraph());
+			if (!ensure(MetasoundGraph))
+			{
+				return nullptr;
+			}
+
+			UMetasoundEditorGraphInputNode* NewGraphNode = MetasoundGraph->CreateInputNode(LiteralType, nullptr, bInSelectNewNode);
+			if (ensure(NewGraphNode))
+			{
+				RebuildNodePins(*NewGraphNode, InNodeHandle);
+
+				NewGraphNode->SetNodeID(InNodeHandle->GetID());
+				return NewGraphNode;
+			}
+
+			return nullptr;
+		}
+
+		UMetasoundEditorGraphInputNode* FGraphBuilder::AddInput(UObject& InMetasound, const FString& InName, const FName InTypeName, const FMetasoundFrontendNodeStyle& InNodeStyle, const FText& InToolTip, bool bInSelectNewNode)
+		{
+			using namespace Metasound::Frontend;
+
+			const FScopedTransaction Transaction(LOCTEXT("AddMetasoundGraphInputNode", "Add Metasound Input Node"));
+
+			FNodeHandle NodeHandle = AddInputNodeHandle(InMetasound, InName, InTypeName, InNodeStyle, InToolTip);
+			if (UMetasoundEditorGraphInputNode* NewGraphNode = AddInput(InMetasound, NodeHandle))
+			{
+				const FMetasoundFrontendNodeStyle& Style = NodeHandle->GetNodeStyle();
+				NewGraphNode->NodePosX = Style.Display.Location.X;
+				NewGraphNode->NodePosY = Style.Display.Location.Y;
+				return NewGraphNode;
+			}
+
+			return nullptr;
 		}
 
 		void FGraphBuilder::AddOrUpdateLiteralInput(UObject& InMetasound, Frontend::FNodeHandle InNodeHandle, UEdGraphPin& InInputPin)
@@ -304,8 +391,11 @@ namespace Metasound
 				}
 				break;
 
-				// TODO: Remove UObject proxy arrays.
-				// If no literal supported, no need for literal input
+				case ELiteralType::BooleanArray:
+				case ELiteralType::FloatArray:
+				case ELiteralType::IntegerArray:
+				case ELiteralType::NoneArray:
+				case ELiteralType::StringArray:
 				case ELiteralType::UObjectProxyArray:
 				case ELiteralType::None:
 				{
@@ -316,6 +406,7 @@ namespace Metasound
 				case ELiteralType::Invalid:
 				default:
 				{
+					static_assert(static_cast<int32>(ELiteralType::COUNT) == 13, "Possible missing ELiteralType case coverage.");
 					ensureMsgf(false, TEXT("Failed to set input node default: Literal type not supported"));
 					return;
 				}
@@ -341,7 +432,7 @@ namespace Metasound
 				{
 					InNodeHandle->GetClassDescription();
 					const TArray<FMetasoundFrontendVertexLiteral>& Defaults = InputHandle->GetDefaults();
-					if (!Defaults.IsEmpty() && Defaults[0].Value.Type != EMetasoundFrontendLiteralType::Invalid && Defaults[0].Value.Type != EMetasoundFrontendLiteralType::None)
+					if (!Defaults.IsEmpty() && Defaults[0].Value.GetType() != EMetasoundFrontendLiteralType::Invalid && Defaults[0].Value.GetType() != EMetasoundFrontendLiteralType::None)
 					{
 						const FString NewInputNodeName = NewInputNode->GetNodeName();
 						GraphHandle->SetDefaultInputToFrontendLiteral(NewInputNodeName, PointIDs[0], TypeName, Defaults[0].Value);
@@ -411,7 +502,7 @@ namespace Metasound
 
 				if (InDefaultValue)
 				{
-					if (InDefaultValue->Type != EMetasoundFrontendLiteralType::None && InDefaultValue->Type != EMetasoundFrontendLiteralType::Invalid)
+					if (InDefaultValue->GetType() != EMetasoundFrontendLiteralType::None && InDefaultValue->GetType() != EMetasoundFrontendLiteralType::Invalid)
 					{
 						GraphHandle->SetDefaultInputToFrontendLiteral(NodeHandle->GetNodeName(), PointID, InTypeName, *InDefaultValue);
 					}
@@ -548,7 +639,16 @@ namespace Metasound
 				}
 				NodeHandle->SetNodeStyle(Style);
 
-				UEdGraphNode* NewNode = AddNode(InMetasound, NodeHandle, false /* bInSelectNewNode */);
+				UEdGraphNode* NewNode = nullptr;
+				if (NodeHandle->GetClassType() == EMetasoundFrontendClassType::Input)
+				{
+					NewNode = Cast<UEdGraphNode>(AddInput(InMetasound, NodeHandle, false /* bInSelectNewNode */));
+				}
+				else
+				{
+					NewNode = AddNode(InMetasound, NodeHandle, false /* bInSelectNewNode */);
+				}
+
 				NewIdNodeMap.Add(NodeHandle->GetID(), FNodePair { NodeHandle, NewNode });
 			}
 
@@ -782,6 +882,34 @@ namespace Metasound
 			return false;
 		}
 
+		bool FGraphBuilder::IsRequiredInput(Frontend::FNodeHandle InNodeHandle)
+		{
+			if (InNodeHandle->GetClassType() != EMetasoundFrontendClassType::Input)
+			{
+				return false;
+			}
+
+			Frontend::FConstDocumentHandle DocumentHandle = InNodeHandle->GetOwningGraph()->GetOwningDocument();
+			return InNodeHandle->GetConstOutputs().ContainsByPredicate([&](const Frontend::FConstOutputHandle& OutputHandle)
+			{
+				return DocumentHandle->IsRequiredInput(OutputHandle->GetName());
+			});
+		}
+
+		bool FGraphBuilder::IsRequiredOutput(Frontend::FNodeHandle InNodeHandle)
+		{
+			if (InNodeHandle->GetClassType() != EMetasoundFrontendClassType::Output)
+			{
+				return false;
+			}
+
+			Frontend::FConstDocumentHandle DocumentHandle = InNodeHandle->GetOwningGraph()->GetOwningDocument();
+			return InNodeHandle->GetConstInputs().ContainsByPredicate([&](const Frontend::FConstInputHandle& InputHandle)
+			{
+				return DocumentHandle->IsRequiredOutput(InputHandle->GetName());
+			});
+		}
+
 		UEdGraphPin* FGraphBuilder::AddPinToNode(UMetasoundEditorGraphNode& InEditorNode, Frontend::FInputHandle InInputHandle)
 		{
 			using namespace Metasound::Editor;
diff --git a/Engine/Plugins/Runtime/Metasound/Source/MetasoundEditor/Private/MetasoundEditorGraphBuilder.h b/Engine/Plugins/Runtime/Metasound/Source/MetasoundEditor/Private/MetasoundEditorGraphBuilder.h
--- a/Engine/Plugins/Runtime/Metasound/Source/MetasoundEditor/Private/MetasoundEditorGraphBuilder.h
+++ b/Engine/Plugins/Runtime/Metasound/Source/MetasoundEditor/Private/MetasoundEditorGraphBuilder.h
@@ -5,6 +5,7 @@
 #include "EdGraph/EdGraph.h"
 #include "EdGraph/EdGraphNode.h"
 #include "EdGraph/EdGraphPin.h"
+#include "MetasoundEditorGraphInputNodes.h"
 #include "MetasoundFrontend.h"
 #include "MetasoundFrontendController.h"
 #include "MetasoundFrontendDocument.h"
@@ -76,13 +77,17 @@ namespace Metasound
 			static Frontend::FOutputHandle GetOutputHandleFromPin(const UEdGraphPin* InPin);
 			static Frontend::FConstOutputHandle GetConstOutputHandleFromPin(const UEdGraphPin* InPin);
 
-			static UEdGraphNode* AddInput(UObject& InMetasound, const FString& InName, const FName InTypeName, const FMetasoundFrontendNodeStyle& InNodeStyle, const FText& InToolTip, bool bInSelectNewNode = true);
+			static UMetasoundEditorGraphInputNode* AddInput(UObject& InMetasound, const FString& InName, const FName InTypeName, const FMetasoundFrontendNodeStyle& InNodeStyle, const FText& InToolTip, bool bInSelectNewNode = true);
+			static UMetasoundEditorGraphInputNode* AddInput(UObject& InMetasound, Frontend::FNodeHandle InNodeHandle, bool bInSelectNewNode = true);
 
 			// Adds or updates a hidden input that is set to the default literal value.  If the literal input is added, value defaults to
 			// 1. what is set on the input vertex of the node's inputer vertex interface, and if not defined there,
 			// 2. the pin's DataType literal default.
 			static void AddOrUpdateLiteralInput(UObject& InMetasound, Frontend::FNodeHandle InNodeHandle, UEdGraphPin& InInputPin);
 
+			static bool IsRequiredInput(Frontend::FNodeHandle InNodeHandle);
+			static bool IsRequiredOutput(Frontend::FNodeHandle InNodeHandle);
+
 			// Adds an input node handle with the given class info
 			static Frontend::FNodeHandle AddInputNodeHandle(UObject& InMetasound, const FString& InName, const FName InTypeName, const FMetasoundFrontendNodeStyle& InNodeStyle, const FText& InToolTip, const FMetasoundFrontendLiteral* InDefaultValue = nullptr);
 
diff --git a/Engine/Plugins/Runtime/Metasound/Source/MetasoundEditor/Private/MetasoundEditorGraphNode.cpp b/Engine/Plugins/Runtime/Metasound/Source/MetasoundEditor/Private/MetasoundEditorGraphNode.cpp
--- a/Engine/Plugins/Runtime/Metasound/Source/MetasoundEditor/Private/MetasoundEditorGraphNode.cpp
+++ b/Engine/Plugins/Runtime/Metasound/Source/MetasoundEditor/Private/MetasoundEditorGraphNode.cpp
@@ -122,7 +122,7 @@ void UMetasoundEditorGraphNode::AllocateDefaultPins()
 {
 	using namespace Metasound;
 
-	ensureAlways(Pins.Num() == 0);
+	ensureAlways(Pins.IsEmpty());
 	Editor::FGraphBuilder::RebuildNodePins(*this, GetNodeHandle());
 }
 
diff --git a/Engine/Plugins/Runtime/Metasound/Source/MetasoundEditor/Private/MetasoundEditorGraphSchema.cpp b/Engine/Plugins/Runtime/Metasound/Source/MetasoundEditor/Private/MetasoundEditorGraphSchema.cpp
--- a/Engine/Plugins/Runtime/Metasound/Source/MetasoundEditor/Private/MetasoundEditorGraphSchema.cpp
+++ b/Engine/Plugins/Runtime/Metasound/Source/MetasoundEditor/Private/MetasoundEditorGraphSchema.cpp
@@ -303,10 +303,11 @@ UEdGraphNode* FMetasoundGraphSchemaAction_NewInput::PerformAction(UEdGraph* Pare
 	FMetasoundFrontendNodeStyle Style;
 	Style.Display.Location = Location;
 
-	if (UEdGraphNode* NewGraphNode = FGraphBuilder::AddInput(ParentMetasound, NewNodeName, NodeTypeName, Style, FText::GetEmpty()))
+	if (UMetasoundEditorGraphInputNode* NewGraphNode = FGraphBuilder::AddInput(ParentMetasound, NewNodeName, NodeTypeName, Style, FText::GetEmpty()))
 	{
-		TryConnectNewNodeToPin(*NewGraphNode, FromPin);
-		return NewGraphNode;
+		UEdGraphNode* EdGraphNode = CastChecked<UEdGraphNode>(NewGraphNode);
+		TryConnectNewNodeToPin(*EdGraphNode, FromPin);
+		return EdGraphNode;
 	}
 
 	return nullptr;
diff --git a/Engine/Plugins/Runtime/Metasound/Source/MetasoundEditor/Private/MetasoundEditorModule.cpp b/Engine/Plugins/Runtime/Metasound/Source/MetasoundEditor/Private/MetasoundEditorModule.cpp
--- a/Engine/Plugins/Runtime/Metasound/Source/MetasoundEditor/Private/MetasoundEditorModule.cpp
+++ b/Engine/Plugins/Runtime/Metasound/Source/MetasoundEditor/Private/MetasoundEditorModule.cpp
@@ -14,7 +14,7 @@
 #include "MetasoundSource.h"
 #include "MetasoundAssetTypeActions.h"
 #include "MetasoundDetailCustomization.h"
-#include "MetasoundLiteralDescriptionDetailCustomization.h"
+#include "MetasoundInputNodeDetailCustomization.h"
 #include "MetasoundEditorGraphBuilder.h"
 #include "MetasoundEditorGraphNodeFactory.h"
 #include "MetasoundEditorGraphSchema.h"
@@ -100,6 +100,39 @@ namespace Metasound
 
 		class FModule : public IMetasoundEditorModule
 		{
+			void RegisterNodeInputClasses()
+			{
+				TSubclassOf<UMetasoundEditorGraphInputNode> NodeClass;
+				for (TObjectIterator<UClass> ClassIt; ClassIt; ++ClassIt)
+				{
+					UClass* Class = *ClassIt;
+					if (!Class->IsNative())
+					{
+						continue;
+					}
+
+					if (Class->HasAnyClassFlags(CLASS_Abstract | CLASS_Deprecated | CLASS_NewerVersionExists))
+					{
+						continue;
+					}
+
+					if (!ClassIt->IsChildOf(UMetasoundEditorGraphInputNode::StaticClass()))
+					{
+						continue;
+					}
+
+					if (const UMetasoundEditorGraphInputNode* InputNode = Class->GetDefaultObject<UMetasoundEditorGraphInputNode>())
+					{
+						NodeInputClassRegistry.Add(InputNode->GetLiteralType(), InputNode->GetClass());
+					}
+				}
+			}
+
+			const TSubclassOf<UMetasoundEditorGraphInputNode> FindNodeInputClass(EMetasoundFrontendLiteralType InLiteralType) const override
+			{
+				return NodeInputClassRegistry.FindRef(InLiteralType);
+			}
+
 			virtual const FEditorDataType& FindDataType(FName InDataTypeName) const override
 			{
 				return DataTypeInfo.FindChecked(InDataTypeName);
@@ -250,13 +283,18 @@ namespace Metasound
 					UMetasoundSource::StaticClass()->GetFName(),
 					FOnGetDetailCustomizationInstance::CreateLambda([]() { return MakeShared<FMetasoundDetailCustomization>(UMetasoundSource::GetDocumentPropertyName()); }));
 
+				PropertyModule.RegisterCustomClassLayout(
+					UMetasoundEditorGraphInputNode::StaticClass()->GetFName(),
+					FOnGetDetailCustomizationInstance::CreateLambda([]() { return MakeShared<FMetasoundInputNodeDetailCustomization>(); }));
+
 				PropertyModule.RegisterCustomPropertyTypeLayout(
-					"MetasoundFrontendLiteral",
-					FOnGetPropertyTypeCustomizationInstance::CreateLambda([]() { return MakeShared<FMetasoundFrontendLiteralDetailCustomization>(); }));
+					"MetasoundEditorGraphInputObject",
+					FOnGetPropertyTypeCustomizationInstance::CreateLambda([]() { return MakeShared<FMetasoundInputNodeObjectDetailCustomization>(); }));
 
 				StyleSet = MakeShared<FSlateStyle>();
 
 				RegisterCoreDataTypes();
+				RegisterNodeInputClasses();
 
 				GraphConnectionFactory = MakeShared<FGraphConnectionDrawingPolicyFactory>();
 				FEdGraphUtilities::RegisterVisualPinConnectionFactory(GraphConnectionFactory);
@@ -315,6 +353,7 @@ namespace Metasound
 
 			TArray<TSharedPtr<FAssetTypeActions_Base>> AssetActions;
 			TMap<FName, FEditorDataType> DataTypeInfo;
+			TMap<EMetasoundFrontendLiteralType, const TSubclassOf<UMetasoundEditorGraphInputNode>> NodeInputClassRegistry;
 
 			TSharedPtr<FMetasoundGraphNodeFactory> GraphNodeFactory;
 			TSharedPtr<FGraphPanelPinConnectionFactory> GraphConnectionFactory;
diff --git a/Engine/Plugins/Runtime/Metasound/Source/MetasoundEditor/Private/MetasoundEditorModule.h b/Engine/Plugins/Runtime/Metasound/Source/MetasoundEditor/Private/MetasoundEditorModule.h
--- a/Engine/Plugins/Runtime/Metasound/Source/MetasoundEditor/Private/MetasoundEditorModule.h
+++ b/Engine/Plugins/Runtime/Metasound/Source/MetasoundEditor/Private/MetasoundEditorModule.h
@@ -4,10 +4,14 @@
 #include "CoreMinimal.h"
 #include "EdGraph/EdGraphNode.h"
 #include "EdGraph/EdGraphPin.h"
-#include "Modules/ModuleInterface.h"
+#include "MetasoundFrontendLiteral.h"
 #include "MetasoundFrontendRegistries.h"
+#include "Modules/ModuleInterface.h"
 #include "Templates/Function.h"
 
+class UMetasoundEditorGraph;
+class UMetasoundEditorGraphInputNode;
+
 DECLARE_LOG_CATEGORY_EXTERN(LogMetasoundEditor, Log, All);
 
 
@@ -33,6 +37,8 @@ namespace Metasound
 			virtual void RegisterDataType(FName InPinCategoryName, FName InPinSubCategoryName, const FDataTypeRegistryInfo& InRegistryInfo) = 0;
 			virtual const FEditorDataType& FindDataType(FName InDataTypeName) const = 0;
 			virtual void IterateDataTypes(TUniqueFunction<void(const FEditorDataType&)> InDataTypeFunction) const = 0;
+
+			virtual const TSubclassOf<UMetasoundEditorGraphInputNode> FindNodeInputClass(EMetasoundFrontendLiteralType InLiteralType) const = 0;
 		};
 	} // namespace Editor
 } // namespace Metasound
diff --git a/Engine/Plugins/Runtime/Metasound/Source/MetasoundEditor/Private/MetasoundEditorTabFactory.h b/Engine/Plugins/Runtime/Metasound/Source/MetasoundEditor/Private/MetasoundEditorTabFactory.h
--- a/Engine/Plugins/Runtime/Metasound/Source/MetasoundEditor/Private/MetasoundEditorTabFactory.h
+++ b/Engine/Plugins/Runtime/Metasound/Source/MetasoundEditor/Private/MetasoundEditorTabFactory.h
@@ -20,9 +20,10 @@ namespace Metasound
 			namespace Names
 			{
 				const FName Analyzers = "MetasoundEditor_Analyzers";
+				const FName Inspector = "MetasoundEditor_Inspector";
 				const FName GraphCanvas = "MetasoundEditor_GraphCanvas";
+				const FName Metasound = "MetasoundEditor_Metasound";
 				const FName Palette = "MetasoundEditor_Palette";
-				const FName Properties = "MetasoundEditor_Properties";
 			}
 
 			TSharedRef<SDockTab> CreateAnalyzersTab(TSharedPtr<SAudioMeter> InAudioMeter, const FSpawnTabArgs& Args)
@@ -52,27 +53,39 @@ namespace Metasound
 				return SpawnedTab;
 			}
 
-			TSharedRef<SDockTab> CreatePaletteTab(TSharedPtr<SMetasoundPalette> Palette, const FSpawnTabArgs& Args)
+			TSharedRef<SDockTab> CreateInspectorTab(TSharedPtr<IDetailsView> DetailsView, const FSpawnTabArgs& Args)
 			{
-				check(Args.GetTabId() == Names::Palette);
+				check(Args.GetTabId() == Names::Inspector);
 
 				return SNew(SDockTab)
-					.Icon(FEditorStyle::GetBrush("Kismet.Tabs.Palette"))
-					.Label(LOCTEXT("MetasoundPaletteTitle", "Palette"))
+					.Icon(FEditorStyle::GetBrush("LevelEditor.Tabs.Details"))
+					.Label(LOCTEXT("MetasoundInspectorTitle", "Inspector"))
 					[
-						Palette.ToSharedRef()
+						DetailsView.ToSharedRef()
 					];
 			}
 
-			TSharedRef<SDockTab> CreatePropertiesTab(TSharedPtr<IDetailsView> PropertiesView, const FSpawnTabArgs& Args)
+			TSharedRef<SDockTab> CreateMetasoundTab(TSharedPtr<IDetailsView> DetailsView, const FSpawnTabArgs& Args)
 			{
-				check(Args.GetTabId() == Names::Properties);
+				check(Args.GetTabId() == Names::Metasound);
 
 				return SNew(SDockTab)
-					.Icon(FEditorStyle::GetBrush("LevelEditor.Tabs.Details"))
-					.Label(LOCTEXT("MetasoundDetailsTitle", "Details"))
+					.Icon(FAppStyle::Get().GetBrush("Icons.Settings"))
+					.Label(LOCTEXT("MetasoundGeneralTitle", "Metasound"))
 					[
-						PropertiesView.ToSharedRef()
+						DetailsView.ToSharedRef()
+					];
+			}
+
+			TSharedRef<SDockTab> CreatePaletteTab(TSharedPtr<SMetasoundPalette> Palette, const FSpawnTabArgs& Args)
+			{
+				check(Args.GetTabId() == Names::Palette);
+
+				return SNew(SDockTab)
+					.Icon(FEditorStyle::GetBrush("Kismet.Tabs.Palette"))
+					.Label(LOCTEXT("MetasoundPaletteTitle", "Palette"))
+					[
+						Palette.ToSharedRef()
 					];
 			}
 		}
diff --git a/Engine/Plugins/Runtime/Metasound/Source/MetasoundEditor/Private/MetasoundInputNodeDetailCustomization.cpp b/Engine/Plugins/Runtime/Metasound/Source/MetasoundEditor/Private/MetasoundInputNodeDetailCustomization.cpp
new file mode 100644
--- /dev/null
+++ b/Engine/Plugins/Runtime/Metasound/Source/MetasoundEditor/Private/MetasoundInputNodeDetailCustomization.cpp
@@ -0,0 +1,573 @@
+// Copyright Epic Games, Inc. All Rights Reserved.
+#include "MetasoundInputNodeDetailCustomization.h"
+
+#include "Containers/Set.h"
+#include "CoreMinimal.h"
+#include "Delegates/Delegate.h"
+#include "DetailCategoryBuilder.h"
+#include "DetailLayoutBuilder.h"
+#include "DetailWidgetRow.h"
+#include "Framework/Notifications/NotificationManager.h"
+#include "IDetailChildrenBuilder.h"
+#include "IDetailGroup.h"
+#include "Internationalization/Text.h"
+#include "IPropertyTypeCustomization.h"
+#include "MetasoundAssetBase.h"
+#include "MetasoundEditorGraphBuilder.h"
+#include "MetasoundEditorGraphNode.h"
+#include "MetasoundEditorGraphInputNodes.h"
+#include "MetasoundFrontend.h"
+#include "MetasoundFrontendController.h"
+#include "MetasoundFrontendRegistries.h"
+#include "MetasoundUObjectRegistry.h"
+#include "PropertyCustomizationHelpers.h"
+#include "PropertyEditorDelegates.h"
+#include "PropertyHandle.h"
+#include "PropertyRestriction.h"
+#include "SlateCore/Public/Styling/SlateColor.h"
+#include "Templates/Casts.h"
+#include "Templates/SharedPointer.h"
+#include "UObject/WeakObjectPtr.h"
+#include "UObject/WeakObjectPtrTemplates.h"
+#include "Widgets/Input/SCheckBox.h"
+#include "Widgets/Input/SMultiLineEditableTextBox.h"
+#include "Widgets/Notifications/SNotificationList.h"
+#include "Widgets/SWidget.h"
+#include "Widgets/Text/STextBlock.h"
+
+
+#define LOCTEXT_NAMESPACE "MetasoundEditor"
+
+namespace Metasound
+{
+	namespace Editor
+	{
+		namespace NodeCustomizationUtils
+		{
+			/** Minimum size of the details title panel */
+			static const float DetailsTitleMinWidth = 125.f;
+			/** Maximum size of the details title panel */
+			static const float DetailsTitleMaxWidth = 300.f;
+			/** magic number retrieved from SGraphNodeComment::GetWrapAt() */
+			static const float DetailsTitleWrapPadding = 32.0f;
+
+			static const FText NodeTooltipText = LOCTEXT("Node_Tooltip", "Tooltip");
+			static const FText DefaultPropertyText = LOCTEXT("Node_DefaultPropertyName", "Default Value");
+
+			static const FText InputNodeNameText = LOCTEXT("InputNode_Name", "Input Name");
+
+			static const FName ProxyGeneratorClassNameIdentifier = "GeneratorClass";
+
+			void UpdateInputLiteralFromNode(const TWeakObjectPtr<UMetasoundEditorGraphInputNode>& GraphNode)
+			{
+				using namespace Frontend;
+
+				if (!GraphNode.IsValid())
+				{
+					return;
+				}
+
+				const FNodeHandle NodeHandle = GraphNode->GetNodeHandle();
+				const FString& NodeName = NodeHandle->GetNodeName();
+				const FGraphHandle GraphHandle = GraphNode->GetRootGraphHandle();
+
+				const TArray<FGuid> PointIDs = GraphHandle->GetDefaultIDsForInputVertex(NodeName);
+				if (ensure(PointIDs.Num() == 1))
+				{
+					const FMetasoundFrontendLiteral Literal = GraphNode->GetLiteralDefault();
+					const TArray<FOutputHandle> Outputs = NodeHandle->GetOutputs();
+					if (ensure(Outputs.Num() == 1))
+					{
+						const FName TypeName = Outputs[0]->GetDataType();
+						GraphHandle->SetDefaultInputToFrontendLiteral(NodeName, PointIDs[0], TypeName, Literal);
+					}
+				}
+			}
+		}
+
+		void FMetasoundInputNodeObjectDetailCustomization::CustomizeHeader(TSharedRef<IPropertyHandle> StructPropertyHandle, FDetailWidgetRow& HeaderRow, IPropertyTypeCustomizationUtils& StructCustomizationUtils)
+		{
+		}
+
+		void FMetasoundInputNodeObjectDetailCustomization::CacheProxyClass(TSharedPtr<IPropertyHandle> InPropertyHandle)
+		{
+			ProxyGenClass.Reset();
+
+			const FString* MetadataProxyGenClass = InPropertyHandle->GetInstanceMetaData(NodeCustomizationUtils::ProxyGeneratorClassNameIdentifier);
+			if (!ensure(MetadataProxyGenClass))
+			{
+				return;
+			}
+
+			const FName ClassName = FName(*MetadataProxyGenClass);
+			for (TObjectIterator<UClass> ClassIt; ClassIt; ++ClassIt)
+			{
+				UClass* Class = *ClassIt;
+				if (!Class->IsNative())
+				{
+					continue;
+				}
+				
+				if (Class->HasAnyClassFlags(CLASS_Abstract | CLASS_Deprecated | CLASS_NewerVersionExists))
+				{
+					continue;
+				}
+
+				if (ClassIt->GetFName() != ClassName)
+				{
+					continue;
+				}
+
+				ProxyGenClass = *ClassIt;
+				return;
+			}
+
+			ensureMsgf(false, TEXT("Failed to find ProxyGeneratorClass. Class not set "));
+		}
+
+		TSharedRef<SWidget> FMetasoundInputNodeObjectDetailCustomization::CreateObjectPickerWidget(TSharedPtr<IPropertyHandle>& PropertyHandle) const
+		{
+			auto ValidateAsset = [InProxyGenClass = ProxyGenClass](const FAssetData& InAsset)
+			{
+				if (!InProxyGenClass.IsValid())
+				{
+					return false;
+				}
+
+				if (UObject* Object = InAsset.GetAsset())
+				{
+					if (UClass* Class = Object->GetClass())
+					{
+						return Class->IsChildOf(InProxyGenClass.Get());
+					}
+				}
+
+				return false;
+			};
+
+			auto CommitAsset = [InPropertyHandle = PropertyHandle, InProxyGenClass = ProxyGenClass](const FAssetData& InAssetData)
+			{
+				// if we've hit this code, the presumption is that the datatype for this parameter has already defined a corresponding UClass that can be used to set it.
+				ensureAlways(InProxyGenClass.IsValid());
+
+				UObject* InObject = InAssetData.GetAsset();
+
+				InPropertyHandle->SetValue(InObject);
+			};
+
+			auto GetAssetPath = [PropertyHandle = PropertyHandle]()
+			{
+				UObject* Object = nullptr;
+				if (PropertyHandle->GetValue(Object) == FPropertyAccess::Success)
+				{
+					return Object->GetPathName();
+				}
+				return FString();
+			};
+
+			TArray<const UClass*> AllowedClasses;
+			AllowedClasses.Add(ProxyGenClass.Get());
+
+			return SNew(SObjectPropertyEntryBox)
+				.ObjectPath_Lambda(GetAssetPath)
+				.AllowedClass(ProxyGenClass.Get())
+				.OnShouldSetAsset_Lambda(ValidateAsset)
+				.OnObjectChanged_Lambda(CommitAsset)
+				.AllowClear(false)
+				.DisplayUseSelected(true)
+				.DisplayBrowse(true)
+				.DisplayThumbnail(true)
+				.NewAssetFactories(PropertyCustomizationHelpers::GetNewAssetFactoriesForClasses(AllowedClasses));
+		}
+
+		TSharedRef<SWidget> FMetasoundInputNodeObjectDetailCustomization::CreateValueWidget(TSharedPtr<IPropertyHandleArray> ParentArrayProperty, TSharedPtr<IPropertyHandle> StructPropertyHandle, bool bIsInArray) const
+		{
+			TSharedPtr<IPropertyHandle> PropertyHandle = StructPropertyHandle->GetChildHandle(GET_MEMBER_NAME_CHECKED(FMetasoundEditorGraphInputObject, Object));
+			TSharedRef<SWidget> ObjectPickerWidget = CreateObjectPickerWidget(PropertyHandle);
+			if (!bIsInArray)
+			{
+				return ObjectPickerWidget;
+			}
+
+			TSharedPtr<IPropertyHandle> StructPropertyPtr = StructPropertyHandle;
+			FExecuteAction InsertAction = FExecuteAction::CreateLambda([ParentArrayProperty, StructPropertyPtr]
+			{
+				const int32 ArrayIndex = StructPropertyPtr.IsValid() ? StructPropertyPtr->GetIndexInArray() : INDEX_NONE;
+				if (ParentArrayProperty.IsValid() && ArrayIndex >= 0)
+				{
+					ParentArrayProperty->Insert(ArrayIndex);
+				}
+			});
+
+			FExecuteAction DeleteAction = FExecuteAction::CreateLambda([ParentArrayProperty, StructPropertyPtr]
+			{
+				const int32 ArrayIndex = StructPropertyPtr.IsValid() ? StructPropertyPtr->GetIndexInArray() : INDEX_NONE;
+				if (ParentArrayProperty.IsValid() && ArrayIndex >= 0)
+				{
+					ParentArrayProperty->DeleteItem(ArrayIndex);
+				}
+			});
+
+			FExecuteAction DuplicateAction = FExecuteAction::CreateLambda([ParentArrayProperty, StructPropertyPtr]
+			{
+				const int32 ArrayIndex = StructPropertyPtr.IsValid() ? StructPropertyPtr->GetIndexInArray() : INDEX_NONE;
+				if (ParentArrayProperty.IsValid() && ArrayIndex >= 0)
+				{
+					ParentArrayProperty->DuplicateItem(ArrayIndex);
+				}
+			});
+
+			return SNew(SHorizontalBox)
+				+ SHorizontalBox::Slot()
+				.FillWidth(0.95f)
+				.Padding(1.0f, 0.0f, 0.0f, 0.0f)
+				.VAlign(VAlign_Center)
+				[
+					ObjectPickerWidget
+				]
+				+ SHorizontalBox::Slot()
+				.FillWidth(0.05f)
+				.Padding(1.0f, 0.0f, 0.0f, 0.0f)
+				.VAlign(VAlign_Center)
+				[
+					PropertyCustomizationHelpers::MakeInsertDeleteDuplicateButton(InsertAction, DeleteAction, DuplicateAction)
+				];
+		}
+
+		void FMetasoundInputNodeObjectDetailCustomization::CustomizeChildren(TSharedRef<IPropertyHandle> StructPropertyHandle, IDetailChildrenBuilder& ChildBuilder, IPropertyTypeCustomizationUtils& StructCustomizationUtils)
+		{
+			bool bIsInArray = false;
+			TSharedPtr<IPropertyHandleArray> ParentArrayProperty;
+			TSharedPtr<IPropertyHandle> ProxyProperty = StructPropertyHandle;
+			{
+				TSharedPtr<IPropertyHandle> ParentProperty = ProxyProperty->GetParentHandle();
+				if (ProxyProperty.IsValid() && ParentProperty.IsValid())
+				{
+					ParentArrayProperty = ParentProperty->AsArray();
+					if (ParentArrayProperty.IsValid())
+					{
+						ProxyProperty = ParentProperty;
+						bIsInArray = true;
+					}
+				}
+			}
+
+			CacheProxyClass(ProxyProperty);
+
+			TSharedRef<SWidget> ValueWidget = CreateValueWidget(ParentArrayProperty, StructPropertyHandle, bIsInArray);
+
+			FDetailWidgetRow& ValueRow = ChildBuilder.AddCustomRow(NodeCustomizationUtils::DefaultPropertyText);
+			if (bIsInArray)
+			{
+				ValueRow.NameContent()
+				[
+					StructPropertyHandle->CreatePropertyNameWidget()
+				];
+			}
+			else
+			{
+				ValueRow.NameContent()
+				[
+					SNew(STextBlock)
+					.Text(NodeCustomizationUtils::DefaultPropertyText)
+					.Font(IDetailLayoutBuilder::GetDetailFont())
+				];
+			}
+
+			TArray<UObject*> OuterObjects;
+			StructPropertyHandle->GetOuterObjects(OuterObjects);
+			TArray<TWeakObjectPtr<UMetasoundEditorGraphInputNode>> Inputs;
+			for (UObject* Object : OuterObjects)
+			{
+				if (UMetasoundEditorGraphInputNode* InputNode = Cast<UMetasoundEditorGraphInputNode>(Object))
+				{
+					Inputs.Add(InputNode);
+				}
+			}
+
+			FSimpleDelegate OnLiteralChanged = FSimpleDelegate::CreateLambda([InInputs = Inputs]()
+			{
+				for (const TWeakObjectPtr<UMetasoundEditorGraphInputNode>& InputNode : InInputs)
+				{
+					NodeCustomizationUtils::UpdateInputLiteralFromNode(InputNode);
+				}
+			});
+			StructPropertyHandle->SetOnChildPropertyValueChanged(OnLiteralChanged);
+
+			ValueRow.ValueContent()
+			[
+				ValueWidget
+			];
+		}
+
+		bool FMetasoundInputNodeDetailCustomization::IsRequired() const
+		{
+			if (GraphNode.IsValid())
+			{
+				return FGraphBuilder::IsRequiredInput(GraphNode->GetNodeHandle());
+			}
+
+			return false;
+		}
+
+		void FMetasoundInputNodeDetailCustomization::CustomizeDetails(IDetailLayoutBuilder& DetailLayout)
+		{
+			using namespace Metasound::Frontend;
+
+			TMetasoundNodeVariableDetailCustomization<UMetasoundEditorGraphInputNode>::CustomizeDetails(DetailLayout);
+
+			if (!GraphNode.IsValid())
+			{
+				return;
+			}
+
+			IDetailCategoryBuilder& CategoryBuilder = DetailLayout.EditCategory("General");
+
+			const bool bIsRequired = IsRequired();
+			DisplayNameEditableTextBox = SNew(SEditableTextBox)
+				.Text(this, &FMetasoundInputNodeDetailCustomization::GetDisplayName)
+				.OnTextChanged(this, &FMetasoundInputNodeDetailCustomization::OnDisplayNameChanged)
+				.OnTextCommitted(this, &FMetasoundInputNodeDetailCustomization::OnDisplayNameCommitted)
+				.IsReadOnly(bIsRequired)
+				.Font(IDetailLayoutBuilder::GetDetailFont());
+
+			CategoryBuilder.AddCustomRow(NodeCustomizationUtils::InputNodeNameText)
+			.EditCondition(!bIsRequired, nullptr)
+			.NameContent()
+			[
+				SNew(STextBlock)
+				.Font(IDetailLayoutBuilder::GetDetailFontBold())
+				.Text(NodeCustomizationUtils::InputNodeNameText)
+				.ToolTipText(TAttribute<FText>::Create([GraphNode = this->GraphNode]()
+				{
+					if (GraphNode.IsValid())
+					{
+						Frontend::FNodeHandle NodeHandle = GraphNode->GetNodeHandle();
+						FMetasoundFrontendNodeStyle NodeStyle = NodeHandle->GetNodeStyle();
+						return NodeHandle->GetClassDescription();
+					}
+
+					return FText::GetEmpty();
+				}))
+			]
+			.ValueContent()
+			[
+				DisplayNameEditableTextBox.ToSharedRef()
+			];
+
+			CategoryBuilder.AddCustomRow(NodeCustomizationUtils::NodeTooltipText)
+			.EditCondition(!bIsRequired, nullptr)
+			.NameContent()
+			[
+				SNew(STextBlock)
+				.Font(IDetailLayoutBuilder::GetDetailFontBold())
+				.Text(NodeCustomizationUtils::NodeTooltipText)
+			]
+			.ValueContent()
+			[
+				SNew(SMultiLineEditableTextBox)
+				.Text(this, &FMetasoundInputNodeDetailCustomization::GetTooltip)
+				.OnTextCommitted(this, &FMetasoundInputNodeDetailCustomization::OnTooltipCommitted)
+				.IsReadOnly(bIsRequired)
+				.ModiferKeyForNewLine(EModifierKey::Shift)
+				.RevertTextOnEscape(true)
+				.WrapTextAt(NodeCustomizationUtils::DetailsTitleMaxWidth - NodeCustomizationUtils::DetailsTitleWrapPadding)
+				.Font(IDetailLayoutBuilder::GetDetailFont())
+			];
+
+			CategoryBuilder.AddCustomRow(LOCTEXT("InputPrivate", "Private"))
+			.Visibility(TAttribute<EVisibility>(this, &FMetasoundInputNodeDetailCustomization::ExposePrivateVisibility))
+			.NameContent()
+			[
+				SNew(STextBlock)
+				.Text(LOCTEXT("InputPrivate", "Private"))
+				.Font(IDetailLayoutBuilder::GetDetailFont())
+			]
+			.ValueContent()
+			[
+				SNew(SCheckBox)
+				.IsChecked(this, &FMetasoundInputNodeDetailCustomization::OnGetPrivateCheckboxState)
+				.OnCheckStateChanged(this, &FMetasoundInputNodeDetailCustomization::OnPrivateChanged)
+			];
+
+			FNodeHandle NodeHandle = GraphNode->GetNodeHandle();
+			const TArray<FOutputHandle>& Outputs = NodeHandle->GetOutputs();
+			if (!ensure(!Outputs.IsEmpty()))
+			{
+				return;
+			}
+
+			IDetailCategoryBuilder& DefaultCategoryBuilder = DetailLayout.EditCategory("DefaultValue");
+
+			TArray<TSharedRef<IPropertyHandle>> DefaultValueHandles;
+			DefaultCategoryBuilder.GetDefaultProperties(DefaultValueHandles);
+
+			for (TSharedRef<IPropertyHandle>& Handle : DefaultValueHandles)
+			{
+				const FName PropertyName = Handle->GetProperty()->GetFName();
+				if (PropertyName == GraphNode->GetLiteralDefaultPropertyFName())
+				{
+					SetDefaultPropertyMetaData(Handle);
+					FSimpleDelegate OnLiteralChanged = FSimpleDelegate::CreateLambda([InGraphNode = GraphNode]()
+					{
+						NodeCustomizationUtils::UpdateInputLiteralFromNode(InGraphNode);
+					});
+
+					Handle->SetOnPropertyValueChanged(OnLiteralChanged);
+
+					TSharedPtr<IPropertyHandleArray> PropertyArray = Handle->AsArray();
+					if (PropertyArray.IsValid())
+					{
+						PropertyArray->SetOnNumElementsChanged(OnLiteralChanged);
+					}
+				}
+
+				DefaultCategoryBuilder.AddProperty(Handle);
+			}
+		}
+
+		void FMetasoundInputNodeDetailCustomization::SetDefaultPropertyMetaData(TSharedRef<IPropertyHandle> InDefaultPropertyHandle) const
+		{
+			using namespace Metasound::Frontend;
+
+			if (!GraphNode.IsValid())
+			{
+				return;
+			}
+
+			FMetasoundFrontendRegistryContainer* Registry = FMetasoundFrontendRegistryContainer::Get();
+			if (!ensure(Registry))
+			{
+				return;
+			}
+
+			const FName TypeName = GetLiteralDataType();
+			if (TypeName.IsNone())
+			{
+				return;
+			}
+
+			FDataTypeRegistryInfo DataTypeInfo;
+			if (!ensure(Registry->GetInfoForDataType(TypeName, DataTypeInfo)))
+			{
+				return;
+			}
+
+			const EMetasoundFrontendLiteralType LiteralType = Frontend::GetMetasoundFrontendLiteralType(DataTypeInfo.PreferredLiteralType);
+			if (LiteralType != EMetasoundFrontendLiteralType::UObject && LiteralType != EMetasoundFrontendLiteralType::UObjectArray)
+			{
+				return;
+			}
+
+			UClass* ProxyGenClass = DataTypeInfo.ProxyGeneratorClass;
+			if (!ProxyGenClass)
+			{
+				return;
+			}
+
+			const FString ClassName = ProxyGenClass->GetName();
+			InDefaultPropertyHandle->SetInstanceMetaData(NodeCustomizationUtils::ProxyGeneratorClassNameIdentifier, ClassName);
+		}
+
+		FName FMetasoundInputNodeDetailCustomization::GetLiteralDataType() const
+		{
+			using namespace Metasound::Frontend;
+
+			FNodeHandle NodeHandle = GraphNode->GetNodeHandle();
+			const TArray<FOutputHandle> Outputs = NodeHandle->GetOutputs();
+
+			if (!ensure(Outputs.Num() == 1))
+			{
+				return FName();
+			}
+
+			return Outputs[0]->GetDataType();
+		}
+
+		FText FMetasoundInputNodeDetailCustomization::GetDisplayName() const
+		{
+			using namespace Metasound::Frontend;
+
+			if (GraphNode.IsValid())
+			{
+				FNodeHandle NodeHandle = GraphNode->GetNodeHandle();
+				const TArray<FOutputHandle>& OutputHandles = NodeHandle->GetOutputs();
+				if (ensure(!OutputHandles.IsEmpty()))
+				{
+					FOutputHandle OutputHandle = OutputHandles[0];
+					return OutputHandle->GetDisplayName();
+				}
+			}
+
+			return FText::GetEmpty();
+		}
+
+		void FMetasoundInputNodeDetailCustomization::OnDisplayNameChanged(const FText& InNewText)
+		{
+			bIsNameInvalid = false;
+
+			// Move validation here
+// 			if (bBadThings)
+// 			{
+// 				bIsNameInvalid = true;
+//	 			InputNameEditableTextBox->SetError(FText::Format(LOCTEXT("ComponentInputRenameFailed_InUse", "{0} is in use by another variable or function!"), InNewText));
+// 			}
+		}
+
+		void FMetasoundInputNodeDetailCustomization::OnDisplayNameCommitted(const FText& InNewName, ETextCommit::Type InTextCommit)
+		{
+			using namespace Metasound::Frontend;
+
+			if (!bIsNameInvalid && GraphNode.IsValid())
+			{
+				const FScopedTransaction Transaction(LOCTEXT("RenameInput", "Rename Metasound Input"));
+				FNodeHandle NodeHandle = GraphNode->GetNodeHandle();
+				const TArray<FOutputHandle>& Outputs = NodeHandle->GetOutputs();
+				if (ensure(!Outputs.IsEmpty()))
+				{
+					FOutputHandle OutputHandle = NodeHandle->GetOutputs()[0];
+					FGraphHandle Graph = NodeHandle->GetOwningGraph();
+					Graph->SetInputDisplayName(OutputHandle->GetName(), InNewName);
+				}
+			}
+
+			bIsNameInvalid = false;
+			DisplayNameEditableTextBox->SetError(FText::GetEmpty());
+		}
+
+		FText FMetasoundInputNodeDetailCustomization::GetTooltip() const
+		{
+			using namespace Metasound::Frontend;
+			if (GraphNode.IsValid())
+			{
+				FNodeHandle NodeHandle = GraphNode->GetNodeHandle();
+				const TArray<FOutputHandle>& OutputHandles = NodeHandle->GetOutputs();
+				if (ensure(!OutputHandles.IsEmpty()))
+				{
+					FOutputHandle OutputHandle = OutputHandles[0];
+					return NodeHandle->GetOwningGraph()->GetInputDescription(OutputHandle->GetName());
+				}
+			}
+
+			return FText::GetEmpty();
+		}
+
+		void FMetasoundInputNodeDetailCustomization::OnTooltipCommitted(const FText& InNewText, ETextCommit::Type InTextCommit)
+		{
+			using namespace Metasound::Frontend;
+
+			if (GraphNode.IsValid())
+			{
+				FNodeHandle NodeHandle = GraphNode->GetNodeHandle();
+				const FScopedTransaction Transaction(LOCTEXT("SetInputTooltip", "Set Metasound Input Tooltip"));
+				const TArray<FOutputHandle>& OutputHandles = NodeHandle->GetOutputs();
+				if (ensure(!OutputHandles.IsEmpty()))
+				{
+					Frontend::FOutputHandle OutputHandle = OutputHandles[0];
+					return NodeHandle->GetOwningGraph()->SetInputDescription(OutputHandle->GetName(), InNewText);
+				}
+			}
+		}
+	} // namespace Editor
+} // namespace Metasound
+#undef LOCTEXT_NAMESPACE
diff --git a/Engine/Plugins/Runtime/Metasound/Source/MetasoundEditor/Private/MetasoundInputNodeDetailCustomization.h b/Engine/Plugins/Runtime/Metasound/Source/MetasoundEditor/Private/MetasoundInputNodeDetailCustomization.h
new file mode 100644
--- /dev/null
+++ b/Engine/Plugins/Runtime/Metasound/Source/MetasoundEditor/Private/MetasoundInputNodeDetailCustomization.h
@@ -0,0 +1,162 @@
+// Copyright Epic Games, Inc. All Rights Reserved.
+#pragma once
+
+#include "Components/Widget.h"
+#include "CoreMinimal.h"
+#include "DetailLayoutBuilder.h"
+#include "IDetailCustomization.h"
+#include "Layout/Visibility.h"
+#include "MetasoundAssetBase.h"
+#include "MetasoundEditorGraphNode.h"
+#include "MetasoundEditorGraphInputNodes.h"
+#include "MetasoundUObjectRegistry.h"
+#include "Types/SlateEnums.h"
+#include "UObject/NameTypes.h"
+#include "Widgets/Input/SEditableTextBox.h"
+
+// Forward Declarations
+class FDetailWidgetRow;
+class FPropertyRestriction;
+class IDetailLayoutBuilder;
+class IPropertyHandle;
+class IPropertyTypeCustomizationUtils;
+
+namespace Metasound
+{
+	namespace Editor
+	{
+		class FMetasoundInputNodeObjectDetailCustomization : public IPropertyTypeCustomization
+		{
+		public:
+			//~ Begin IPropertyTypeCustomization
+			virtual void CustomizeHeader(TSharedRef<IPropertyHandle> StructPropertyHandle, FDetailWidgetRow& HeaderRow, IPropertyTypeCustomizationUtils& StructCustomizationUtils) override;
+			virtual void CustomizeChildren(TSharedRef<IPropertyHandle> StructPropertyHandle, IDetailChildrenBuilder& ChildBuilder, IPropertyTypeCustomizationUtils& StructCustomizationUtils) override;
+			//~ End IPropertyTypeCustomization
+
+		private:
+			void CacheProxyClass(TSharedPtr<IPropertyHandle> PropertyHandle);
+
+			TSharedRef<SWidget> CreateObjectPickerWidget(TSharedPtr<IPropertyHandle>& PropertyHandle) const;
+			TSharedRef<SWidget> CreateValueWidget(TSharedPtr<IPropertyHandleArray> ParentArrayProperty, TSharedPtr<IPropertyHandle> StructPropertyHandle, bool bIsInArray) const;
+
+			TWeakObjectPtr<UClass> ProxyGenClass;
+		};
+
+		// TODO: This is templated to easily make a customization for outputs once
+		// widget support is added.
+		template <typename NodeType>
+		class TMetasoundNodeVariableDetailCustomization : public IDetailCustomization
+		{
+		public:
+			TMetasoundNodeVariableDetailCustomization()
+				: IDetailCustomization()
+			{
+			}
+
+		protected:
+			TWeakObjectPtr<NodeType> GraphNode;
+			TSharedPtr<SEditableTextBox> DisplayNameEditableTextBox;
+			bool bIsNameInvalid = false;
+
+			// IDetailCustomization interface
+			virtual void CustomizeDetails(IDetailLayoutBuilder& DetailLayout) override
+			{
+				TArray<TWeakObjectPtr<UObject>> Objects;
+				DetailLayout.GetObjectsBeingCustomized(Objects);
+				if (Objects.IsEmpty())
+				{
+					return;
+				}
+
+				GraphNode = Cast<NodeType>(Objects[0].Get());
+			}
+			// End of IDetailCustomization interface
+
+
+			virtual bool IsRequired() const = 0;
+
+			virtual FText GetDisplayName() const = 0;
+			virtual void OnDisplayNameChanged(const FText& InNewText) = 0;
+			virtual void OnDisplayNameCommitted(const FText& InNewName, ETextCommit::Type InTextCommit) = 0;
+
+			virtual FText GetTooltip() const = 0;
+			virtual void OnTooltipCommitted(const FText& InNewName, ETextCommit::Type InTextCommit) = 0;
+
+			ECheckBoxState OnGetPrivateCheckboxState() const
+			{
+				if (GraphNode.IsValid())
+				{
+					return GraphNode->GetNodeHandle()->GetNodeStyle().bIsPrivate ? ECheckBoxState::Checked : ECheckBoxState::Unchecked;
+				}
+
+				return ECheckBoxState::Unchecked;
+			}
+
+			void OnPrivateChanged(ECheckBoxState InNewState)
+			{
+				if (GraphNode.IsValid())
+				{
+					const bool bIsChecked = InNewState == ECheckBoxState::Checked;
+					Frontend::FNodeHandle NodeHandle = GraphNode->GetNodeHandle();
+					FMetasoundFrontendNodeStyle NodeStyle = NodeHandle->GetNodeStyle();
+					NodeStyle.bIsPrivate = bIsChecked;
+					NodeHandle->SetNodeStyle(NodeStyle);
+				}
+			}
+
+			EVisibility ExposePrivateVisibility() const
+			{
+				if (IsRequired())
+				{
+					return EVisibility::Collapsed;
+				}
+
+				return EVisibility::Visible;
+			}
+
+			FMetasoundAssetBase* GetMetasoundAsset() const
+			{
+				if (!GraphNode.IsValid())
+				{
+					return nullptr;
+				}
+
+				UObject& MetasoundObj = GraphNode->GetMetasoundChecked();
+				if (FMetasoundAssetBase* MetasoundAsset = IMetasoundUObjectRegistry::Get().GetObjectAsAssetBase(&MetasoundObj))
+				{
+					return MetasoundAsset;
+				}
+
+				return nullptr;
+			}
+		};
+
+		class FMetasoundInputNodeDetailCustomization : public TMetasoundNodeVariableDetailCustomization<UMetasoundEditorGraphInputNode>
+		{
+		public:
+			FMetasoundInputNodeDetailCustomization()
+				: TMetasoundNodeVariableDetailCustomization<UMetasoundEditorGraphInputNode>()
+			{
+			}
+
+			// IDetailCustomization interface
+			virtual void CustomizeDetails(IDetailLayoutBuilder& DetailLayout) override;
+			// End of IDetailCustomization interface
+
+		protected:
+			virtual bool IsRequired() const override;
+
+			virtual FText GetDisplayName() const override;
+			virtual void OnDisplayNameChanged(const FText& InNewText) override;
+			virtual void OnDisplayNameCommitted(const FText& InNewName, ETextCommit::Type InTextCommit) override;
+
+			virtual FText GetTooltip() const override;
+			virtual void OnTooltipCommitted(const FText& InNewName, ETextCommit::Type InTextCommit) override;
+
+		private:
+			void SetDefaultPropertyMetaData(TSharedRef<IPropertyHandle> InDefaultPropertyHandle) const;
+
+			FName GetLiteralDataType() const;
+		};
+	} // namespace Editor
+} // namespace Metasound
\ No newline at end of file
diff --git a/Engine/Plugins/Runtime/Metasound/Source/MetasoundEditor/Private/MetasoundLiteralDescriptionDetailCustomization.cpp b/Engine/Plugins/Runtime/Metasound/Source/MetasoundEditor/Private/MetasoundLiteralDescriptionDetailCustomization.cpp
deleted file mode 100644
diff --git a/Engine/Plugins/Runtime/Metasound/Source/MetasoundEditor/Private/MetasoundLiteralDescriptionDetailCustomization.h b/Engine/Plugins/Runtime/Metasound/Source/MetasoundEditor/Private/MetasoundLiteralDescriptionDetailCustomization.h
deleted file mode 100644
diff --git a/Engine/Plugins/Runtime/Metasound/Source/MetasoundEditor/Public/MetasoundEditorGraph.h b/Engine/Plugins/Runtime/Metasound/Source/MetasoundEditor/Public/MetasoundEditorGraph.h
--- a/Engine/Plugins/Runtime/Metasound/Source/MetasoundEditor/Public/MetasoundEditorGraph.h
+++ b/Engine/Plugins/Runtime/Metasound/Source/MetasoundEditor/Public/MetasoundEditorGraph.h
@@ -4,12 +4,14 @@
 #include "CoreMinimal.h"
 #include "EdGraph/EdGraph.h"
 #include "MetasoundEditor.h"
+#include "MetasoundFrontendLiteral.h"
 #include "UObject/ObjectMacros.h"
 
 #include "MetasoundEditorGraph.generated.h"
 
 // Forward Declarations
 struct FMetasoundFrontendDocument;
+class UMetasoundEditorGraphInputNode;
 
 
 UCLASS(MinimalAPI)
@@ -18,6 +20,8 @@ class UMetasoundEditorGraph : public UEdGraph
 	GENERATED_BODY()
 
 public:
+	UMetasoundEditorGraphInputNode* CreateInputNode(EMetasoundFrontendLiteralType LiteralType, UClass* InLiteralObjectClass, bool bInSelectNewNode);
+
 	UObject* GetMetasound() const;
 	UObject& GetMetasoundChecked() const;
 
diff --git a/Engine/Plugins/Runtime/Metasound/Source/MetasoundEditor/Public/MetasoundEditorGraphInputNodes.h b/Engine/Plugins/Runtime/Metasound/Source/MetasoundEditor/Public/MetasoundEditorGraphInputNodes.h
new file mode 100644
--- /dev/null
+++ b/Engine/Plugins/Runtime/Metasound/Source/MetasoundEditor/Public/MetasoundEditorGraphInputNodes.h
@@ -0,0 +1,302 @@
+// Copyright Epic Games, Inc. All Rights Reserved.
+#pragma once
+
+#include "CoreMinimal.h"
+#include "EdGraph/EdGraphNode.h"
+#include "UObject/ObjectMacros.h"
+#include "MetasoundEditorGraphNode.h"
+#include "MetasoundFrontend.h"
+#include "MetasoundFrontendController.h"
+#include "MetasoundFrontendDocument.h"
+#include "MetasoundFrontendLiteral.h"
+#include "MetasoundFrontendRegistries.h"
+#include "Misc/Guid.h"
+#include "Sound/SoundWave.h"
+#include "UObject/SoftObjectPath.h"
+
+#include "MetasoundEditorGraphInputNodes.generated.h"
+
+// Forward Declarations
+class UEdGraphPin;
+class UMetasound;
+
+
+UCLASS()
+class METASOUNDEDITOR_API UMetasoundEditorGraphInputNode : public UMetasoundEditorGraphNode
+{
+	GENERATED_BODY()
+
+public:
+	virtual ~UMetasoundEditorGraphInputNode() = default;
+
+	virtual FMetasoundFrontendLiteral GetLiteralDefault() const
+	{
+		return FMetasoundFrontendLiteral();
+	}
+
+	FName GetLiteralDefaultPropertyFName() const
+	{
+		return "Default";
+	}
+
+	virtual EMetasoundFrontendLiteralType GetLiteralType() const
+	{
+		return EMetasoundFrontendLiteralType::None;
+	}
+};
+
+UCLASS(MinimalAPI)
+class UMetasoundEditorGraphInputBoolNode : public UMetasoundEditorGraphInputNode
+{
+	GENERATED_BODY()
+
+public:
+	virtual ~UMetasoundEditorGraphInputBoolNode() = default;
+
+	UPROPERTY(EditAnywhere, Category = DefaultValue)
+	bool Default = false;
+
+	virtual FMetasoundFrontendLiteral GetLiteralDefault() const override
+	{
+		FMetasoundFrontendLiteral Literal;
+		Literal.Set(Default);
+		return Literal;
+	}
+
+	virtual EMetasoundFrontendLiteralType GetLiteralType() const
+	{
+		return EMetasoundFrontendLiteralType::Boolean;
+	}
+};
+
+UCLASS(MinimalAPI)
+class UMetasoundEditorGraphInputBoolArrayNode : public UMetasoundEditorGraphInputNode
+{
+	GENERATED_BODY()
+
+public:
+	virtual ~UMetasoundEditorGraphInputBoolArrayNode() = default;
+
+	UPROPERTY(EditAnywhere, Category = DefaultValue)
+	TArray<bool> Default;
+
+	virtual FMetasoundFrontendLiteral GetLiteralDefault() const override
+	{
+		FMetasoundFrontendLiteral Literal;
+		Literal.Set(Default);
+		return Literal;
+	}
+
+	virtual EMetasoundFrontendLiteralType GetLiteralType() const
+	{
+		return EMetasoundFrontendLiteralType::BooleanArray;
+	}
+};
+
+UCLASS(MinimalAPI)
+class UMetasoundEditorGraphInputIntNode : public UMetasoundEditorGraphInputNode
+{
+	GENERATED_BODY()
+
+public:
+	virtual ~UMetasoundEditorGraphInputIntNode() = default;
+
+	UPROPERTY(EditAnywhere, Category = DefaultValue)
+	int32 Default = 0;
+
+	virtual FMetasoundFrontendLiteral GetLiteralDefault() const override
+	{
+		FMetasoundFrontendLiteral Literal;
+		Literal.Set(Default);
+		return Literal;
+	}
+
+	virtual EMetasoundFrontendLiteralType GetLiteralType() const
+	{
+		return EMetasoundFrontendLiteralType::Integer;
+	}
+};
+
+UCLASS(MinimalAPI)
+class UMetasoundEditorGraphInputIntArrayNode : public UMetasoundEditorGraphInputNode
+{
+	GENERATED_BODY()
+
+public:
+	virtual ~UMetasoundEditorGraphInputIntArrayNode() = default;
+
+	UPROPERTY(EditAnywhere, Category = DefaultValue)
+	TArray<int32> Default;
+
+	virtual FMetasoundFrontendLiteral GetLiteralDefault() const override
+	{
+		FMetasoundFrontendLiteral Literal;
+		Literal.Set(Default);
+		return Literal;
+	}
+
+	virtual EMetasoundFrontendLiteralType GetLiteralType() const
+	{
+		return EMetasoundFrontendLiteralType::IntegerArray;
+	}
+};
+
+UCLASS(MinimalAPI)
+class UMetasoundEditorGraphInputFloatNode : public UMetasoundEditorGraphInputNode
+{
+	GENERATED_BODY()
+
+public:
+	virtual ~UMetasoundEditorGraphInputFloatNode() = default;
+
+	UPROPERTY(EditAnywhere, Category = DefaultValue)
+	float Default = 0.f;
+
+	virtual FMetasoundFrontendLiteral GetLiteralDefault() const override
+	{
+		FMetasoundFrontendLiteral Literal;
+		Literal.Set(Default);
+		return Literal;
+	}
+
+	virtual EMetasoundFrontendLiteralType GetLiteralType() const
+	{
+		return EMetasoundFrontendLiteralType::Float;
+	}
+};
+
+UCLASS(MinimalAPI)
+class UMetasoundEditorGraphInputFloatArrayNode : public UMetasoundEditorGraphInputNode
+{
+	GENERATED_BODY()
+
+public:
+	virtual ~UMetasoundEditorGraphInputFloatArrayNode() = default;
+
+	UPROPERTY(EditAnywhere, Category = DefaultValue)
+	TArray<float> Default;
+
+	virtual FMetasoundFrontendLiteral GetLiteralDefault() const override
+	{
+		FMetasoundFrontendLiteral Literal;
+		Literal.Set(Default);
+		return Literal;
+	}
+
+	virtual EMetasoundFrontendLiteralType GetLiteralType() const
+	{
+		return EMetasoundFrontendLiteralType::FloatArray;
+	}
+};
+
+UCLASS(MinimalAPI)
+class UMetasoundEditorGraphInputStringNode : public UMetasoundEditorGraphInputNode
+{
+	GENERATED_BODY()
+
+public:
+	virtual ~UMetasoundEditorGraphInputStringNode() = default;
+
+	UPROPERTY(EditAnywhere, Category = DefaultValue)
+	FString Default;
+
+	virtual FMetasoundFrontendLiteral GetLiteralDefault() const override
+	{
+		FMetasoundFrontendLiteral Literal;
+		Literal.Set(Default);
+		return Literal;
+	}
+
+	virtual EMetasoundFrontendLiteralType GetLiteralType() const
+	{
+		return EMetasoundFrontendLiteralType::String;
+	}
+};
+
+UCLASS(MinimalAPI)
+class UMetasoundEditorGraphInputStringArrayNode : public UMetasoundEditorGraphInputNode
+{
+	GENERATED_BODY()
+
+public:
+	virtual ~UMetasoundEditorGraphInputStringArrayNode() = default;
+
+	UPROPERTY(EditAnywhere, Category = DefaultValue)
+	TArray<FString> Default;
+
+	virtual FMetasoundFrontendLiteral GetLiteralDefault() const override
+	{
+		FMetasoundFrontendLiteral Literal;
+		Literal.Set(Default);
+		return Literal;
+	}
+
+	virtual EMetasoundFrontendLiteralType GetLiteralType() const
+	{
+		return EMetasoundFrontendLiteralType::StringArray;
+	}
+};
+
+// Broken out to be able to customize and swap AllowedClass based on provided object proxy
+USTRUCT()
+struct FMetasoundEditorGraphInputObject
+{
+	GENERATED_BODY()
+
+	UPROPERTY(EditAnywhere, Category = DefaultValue)
+	UObject* Object = nullptr;
+};
+
+UCLASS(MinimalAPI)
+class UMetasoundEditorGraphInputObjectNode : public UMetasoundEditorGraphInputNode
+{
+	GENERATED_BODY()
+
+public:
+	virtual ~UMetasoundEditorGraphInputObjectNode() = default;
+
+	UPROPERTY(EditAnywhere, Category = DefaultValue)
+	FMetasoundEditorGraphInputObject Default;
+
+	virtual FMetasoundFrontendLiteral GetLiteralDefault() const override
+	{
+		FMetasoundFrontendLiteral Literal;
+		Literal.Set(Default.Object);
+		return Literal;
+	}
+
+	virtual EMetasoundFrontendLiteralType GetLiteralType() const
+	{
+		return EMetasoundFrontendLiteralType::UObject;
+	}
+};
+
+UCLASS(MinimalAPI)
+class UMetasoundEditorGraphInputObjectArrayNode : public UMetasoundEditorGraphInputNode
+{
+	GENERATED_BODY()
+
+public:
+	virtual ~UMetasoundEditorGraphInputObjectArrayNode() = default;
+
+	UPROPERTY(EditAnywhere, Category = DefaultValue)
+	TArray<FMetasoundEditorGraphInputObject> Default;
+
+	virtual FMetasoundFrontendLiteral GetLiteralDefault() const override
+	{
+		TArray<UObject*> Objects;
+		for (const FMetasoundEditorGraphInputObject& InputObject : Default)
+		{
+			Objects.Add(InputObject.Object);
+		}
+
+		FMetasoundFrontendLiteral Literal;
+		Literal.Set(Objects);
+		return Literal;
+	}
+
+	virtual EMetasoundFrontendLiteralType GetLiteralType() const
+	{
+		return EMetasoundFrontendLiteralType::UObjectArray;
+	}
+};
\ No newline at end of file
diff --git a/Engine/Plugins/Runtime/Metasound/Source/MetasoundEditor/Public/MetasoundEditorGraphNode.h b/Engine/Plugins/Runtime/Metasound/Source/MetasoundEditor/Public/MetasoundEditorGraphNode.h
--- a/Engine/Plugins/Runtime/Metasound/Source/MetasoundEditor/Public/MetasoundEditorGraphNode.h
+++ b/Engine/Plugins/Runtime/Metasound/Source/MetasoundEditor/Public/MetasoundEditorGraphNode.h
@@ -7,15 +7,17 @@
 #include "UObject/ObjectMacros.h"
 #include "MetasoundFrontend.h"
 #include "MetasoundFrontendController.h"
+#include "MetasoundFrontendDocument.h"
+#include "MetasoundFrontendLiteral.h"
 #include "MetasoundFrontendRegistries.h"
 #include "Misc/Guid.h"
+#include "Sound/SoundWave.h"
 
 #include "MetasoundEditorGraphNode.generated.h"
 
 // Forward Declarations
 class UEdGraphPin;
 class UMetasound;
-class USoundNode;
 
 
 UCLASS(MinimalAPI)
@@ -76,3 +78,15 @@ protected:
 	UPROPERTY()
 	FGuid NodeID;
 };
+
+UCLASS(MinimalAPI)
+class UMetasoundEditorGraphOutputNode : public UMetasoundEditorGraphNode
+{
+	GENERATED_BODY()
+};
+
+UCLASS(MinimalAPI)
+class UMetasoundEditorGraphExternalNode : public UMetasoundEditorGraphNode
+{
+	GENERATED_BODY()
+};
\ No newline at end of file
diff --git a/Engine/Plugins/Runtime/Metasound/Source/MetasoundEngine/Private/MetasoundSource.cpp b/Engine/Plugins/Runtime/Metasound/Source/MetasoundEngine/Private/MetasoundSource.cpp
--- a/Engine/Plugins/Runtime/Metasound/Source/MetasoundEngine/Private/MetasoundSource.cpp
+++ b/Engine/Plugins/Runtime/Metasound/Source/MetasoundEngine/Private/MetasoundSource.cpp
@@ -467,7 +467,7 @@ const FMetasoundFrontendArchetype& UMetasoundSource::GetBaseArchetype()
 	{
 		FMetasoundFrontendArchetype Archetype;
 		
-		FMetasoundFrontendClassVertex OnPlayTrigger;
+		FMetasoundFrontendClassInput OnPlayTrigger;
 		OnPlayTrigger.Name = UMetasoundSource::GetOnPlayInputName();
 		OnPlayTrigger.Metadata.DisplayName = FText::FromString(OnPlayTrigger.Name);
 		OnPlayTrigger.TypeName = Metasound::Frontend::GetDataTypeName<Metasound::FTrigger>();
@@ -476,7 +476,7 @@ const FMetasoundFrontendArchetype& UMetasoundSource::GetBaseArchetype()
 
 		Archetype.Interface.Inputs.Add(OnPlayTrigger);
 
-		FMetasoundFrontendClassVertex OnFinished;
+		FMetasoundFrontendClassOutput OnFinished;
 		OnFinished.Name = UMetasoundSource::GetIsFinishedOutputName();
 		OnFinished.Metadata.DisplayName = FText::FromString(OnFinished.Name);
 		OnFinished.TypeName = Metasound::Frontend::GetDataTypeName<Metasound::FTrigger>();
@@ -505,7 +505,7 @@ const FMetasoundFrontendArchetype& UMetasoundSource::GetMonoSourceArchetype()
 	auto CreateMonoArchetype = []() -> FMetasoundFrontendArchetype
 	{
 		FMetasoundFrontendArchetype Archetype = GetBaseArchetype();
-		Archetype.Name = FName(TEXT("MonoSource"));
+		Archetype.Name = "MonoSource";
 
 		FMetasoundFrontendClassVertex GeneratedAudio;
 		GeneratedAudio.Name = UMetasoundSource::GetAudioOutputName();
diff --git a/Engine/Plugins/Runtime/Metasound/Source/MetasoundEngine/Public/MetasoundSource.h b/Engine/Plugins/Runtime/Metasound/Source/MetasoundEngine/Public/MetasoundSource.h
--- a/Engine/Plugins/Runtime/Metasound/Source/MetasoundEngine/Public/MetasoundSource.h
+++ b/Engine/Plugins/Runtime/Metasound/Source/MetasoundEngine/Public/MetasoundSource.h
@@ -40,7 +40,7 @@ protected:
 	FMetasoundFrontendDocument RootMetasoundDocument;
 
 #if WITH_EDITORONLY_DATA
-	UPROPERTY()
+	UPROPERTY(EditAnywhere, Category = CustomView)
 	UEdGraph* Graph;
 #endif // WITH_EDITORONLY_DATA
 
diff --git a/Engine/Plugins/Runtime/Metasound/Source/MetasoundFrontend/Private/MetasoundFrontendDocument.cpp b/Engine/Plugins/Runtime/Metasound/Source/MetasoundFrontend/Private/MetasoundFrontendDocument.cpp
--- a/Engine/Plugins/Runtime/Metasound/Source/MetasoundFrontend/Private/MetasoundFrontendDocument.cpp
+++ b/Engine/Plugins/Runtime/Metasound/Source/MetasoundFrontend/Private/MetasoundFrontendDocument.cpp
@@ -110,7 +110,7 @@ FMetasoundFrontendClassInput::FMetasoundFrontendClassInput(const FMetasoundFront
 	{
 		FMetasoundFrontendVertexLiteral Default;
 		Default.PointID = PointID;
-		Default.Value.Type = DefaultType;
+		Default.Value.SetType(DefaultType);
 		Defaults.Add(Default);
 	}
 }
diff --git a/Engine/Plugins/Runtime/Metasound/Source/MetasoundFrontend/Private/MetasoundFrontendGraph.cpp b/Engine/Plugins/Runtime/Metasound/Source/MetasoundFrontend/Private/MetasoundFrontendGraph.cpp
--- a/Engine/Plugins/Runtime/Metasound/Source/MetasoundFrontend/Private/MetasoundFrontendGraph.cpp
+++ b/Engine/Plugins/Runtime/Metasound/Source/MetasoundFrontend/Private/MetasoundFrontendGraph.cpp
@@ -153,7 +153,7 @@ namespace Metasound
 			{
 				const FMetasoundFrontendVertex& InputVertex = InNode.Interface.Inputs[0];
 
-				const bool IsLiteralParsableByDataType = Frontend::DoesDataTypeSupportLiteralType(InputVertex.TypeName, VertexLiteral->Value.Type);
+				const bool IsLiteralParsableByDataType = Frontend::DoesDataTypeSupportLiteralType(InputVertex.TypeName, VertexLiteral->Value.GetType());
 
 				if (IsLiteralParsableByDataType)
 				{
diff --git a/Engine/Plugins/Runtime/Metasound/Source/MetasoundFrontend/Private/MetasoundFrontendInvalidController.h b/Engine/Plugins/Runtime/Metasound/Source/MetasoundFrontend/Private/MetasoundFrontendInvalidController.h
--- a/Engine/Plugins/Runtime/Metasound/Source/MetasoundFrontend/Private/MetasoundFrontendInvalidController.h
+++ b/Engine/Plugins/Runtime/Metasound/Source/MetasoundFrontend/Private/MetasoundFrontendInvalidController.h
@@ -222,12 +222,9 @@ namespace Metasound
 			TSharedRef<INodeController> GetOutputNodeWithName(const FString& InName) override { return FInvalidNodeController::GetInvalid(); }
 			TSharedRef<INodeController> GetInputNodeWithName(const FString& InName) override { return FInvalidNodeController::GetInvalid(); }
 
-
 			FConstClassInputAccessPtr FindClassInputWithName(const FString& InName) const override { return FConstClassInputAccessPtr(); }
 			FConstClassOutputAccessPtr FindClassOutputWithName(const FString& InName) const override { return FConstClassOutputAccessPtr(); }
 
-
-
 			TSharedRef<INodeController> AddInputVertex(const FMetasoundFrontendClassInput& InDescription) override { return FInvalidNodeController::GetInvalid(); }
 			bool RemoveInputVertex(const FString& InputName) override { return false; }
 
@@ -241,19 +238,23 @@ namespace Metasound
 			// For inputs whose preferred literal type is UObject or UObjectArray, this can be used to determine the UObject corresponding to that input's datatype.
 			UClass* GetSupportedClassForInputVertex(const FString& InInputName) override { return nullptr; }
 
-			TArray<FGuid> GetDefaultIDsForInputVertex(const FString& InInputName) const { return TArray<FGuid>(); }
-			TArray<FGuid> GetDefaultIDsForOutputVertex(const FString& InInputName) const { return TArray<FGuid>(); }
+			TArray<FGuid> GetDefaultIDsForInputVertex(const FString& InInputName) const override { return TArray<FGuid>(); }
+			TArray<FGuid> GetDefaultIDsForOutputVertex(const FString& InInputName) const override { return TArray<FGuid>(); }
+
+			TArray<FMetasoundFrontendVertexLiteral> GetDefaultInputFrontendLiterals(const FString& InInputName) const override { return TArray<FMetasoundFrontendVertexLiteral>(); }
 
 			// These can be used to set the default value for a given input on this graph.
 			// @returns false if the input name couldn't be found, or if the literal type was incompatible with the Data Type of this input.
 			bool SetDefaultInputToFrontendLiteral(const FString& InInputName, FGuid InPointID, FName InDataTypeName, const FMetasoundFrontendLiteral& InLiteral) override { return false; }
 			bool SetDefaultInputToTypeDefaultLiteral(const FString& InInputName, FGuid InPointID, FName InDataTypeName) override { return false; }
 
-			// Set the display name for the input with the given name
-			void SetInputDisplayName(FString InName, const FText& InDisplayName) override {}
+			const FText& GetInputDescription(const FString& InName) const override { return FText::GetEmpty(); }
+			const FText& GetOutputDescription(const FString& InName) const override { return FText::GetEmpty(); }
 
-			// Set the display name for the output with the given name
-			void SetOutputDisplayName(FString InName, const FText& InDisplayName) override {}
+			void SetInputDescription(const FString& InName, const FText& InDescription) override { }
+			void SetOutputDescription(const FString& InName, const FText& InDescription) override { }
+			void SetInputDisplayName(const FString& InName, const FText& InDisplayName) override { }
+			void SetOutputDisplayName(const FString& InName, const FText& InDisplayName) override { }
 
 			// This can be used to clear the current literal for a given input.
 			// @returns false if the input name couldn't be found.
@@ -263,7 +264,6 @@ namespace Metasound
 			TSharedRef<INodeController> AddNode(const FNodeRegistryKey& InNodeClass) override { return FInvalidNodeController::GetInvalid(); }
 			TSharedRef<INodeController> AddNode(const FMetasoundFrontendClassMetadata& InNodeClass) override { return FInvalidNodeController::GetInvalid(); }
 
-
 			// Remove the node corresponding to this node handle.
 			// On success, invalidates the received node handle.
 			bool RemoveNode(INodeController& InNode) override { return false; }
@@ -313,8 +313,8 @@ namespace Metasound
 				bool IsRequiredInput(const FString& InInputName) const override { return false; }
 				bool IsRequiredOutput(const FString& InOutputName) const override { return false; }
 
-				TArray<FMetasoundFrontendClassVertex> GetRequiredInputs() const override { return TArray<FMetasoundFrontendClassVertex>(); }
-				TArray<FMetasoundFrontendClassVertex> GetRequiredOutputs() const override { return TArray<FMetasoundFrontendClassVertex>(); }
+				const TArray<FMetasoundFrontendClassInput>& GetRequiredInputs() const override { const static TArray<FMetasoundFrontendClassInput> Empty; return Empty; }
+				const TArray<FMetasoundFrontendClassOutput>& GetRequiredOutputs() const override { const static TArray<FMetasoundFrontendClassOutput> Empty; return Empty; }
 
 				TArray<FMetasoundFrontendClass> GetDependencies() const override { return TArray<FMetasoundFrontendClass>(); }
 				TArray<FMetasoundFrontendGraphClass> GetSubgraphs() const override { return TArray<FMetasoundFrontendGraphClass>(); }
diff --git a/Engine/Plugins/Runtime/Metasound/Source/MetasoundFrontend/Private/MetasoundFrontendLiteral.cpp b/Engine/Plugins/Runtime/Metasound/Source/MetasoundFrontend/Private/MetasoundFrontendLiteral.cpp
--- a/Engine/Plugins/Runtime/Metasound/Source/MetasoundFrontend/Private/MetasoundFrontendLiteral.cpp
+++ b/Engine/Plugins/Runtime/Metasound/Source/MetasoundFrontend/Private/MetasoundFrontendLiteral.cpp
@@ -304,6 +304,98 @@ void FMetasoundFrontendLiteral::SetFromLiteral(const Metasound::FLiteral& InLite
 	}
 }
 
+void FMetasoundFrontendLiteral::SetType(EMetasoundFrontendLiteralType InType)
+{
+	switch (Type)
+	{
+		case EMetasoundFrontendLiteralType::NoneArray:
+		{
+			Set(FDefaultArray());
+		}
+		break;
+
+		case EMetasoundFrontendLiteralType::BooleanArray:
+		{
+			Set(TArray<bool>());
+		}
+		break;
+
+		case EMetasoundFrontendLiteralType::IntegerArray:
+		{
+			Set(TArray<int32>());
+		}
+		break;
+
+		case EMetasoundFrontendLiteralType::FloatArray:
+		{
+			Set(TArray<float>());
+		}
+		break;
+
+		case EMetasoundFrontendLiteralType::StringArray:
+		{
+			Set(TArray<FString>());
+		}
+		break;
+
+		case EMetasoundFrontendLiteralType::UObjectArray:
+		{
+			Set(TArray<UObject*>());
+		}
+		break;
+		
+		case EMetasoundFrontendLiteralType::None:
+		{
+			Set(FDefault());
+		}
+		break;
+
+		case EMetasoundFrontendLiteralType::Boolean:
+		{
+			Set(false);
+		}
+		break;
+
+		case EMetasoundFrontendLiteralType::Integer:
+		{
+			Set(0);
+		}
+		break;
+
+		case EMetasoundFrontendLiteralType::Float:
+		{
+			Set(0.f);
+		}
+		break;
+
+	
+		case EMetasoundFrontendLiteralType::String:
+		{
+			Set(FString());
+		}
+		break;
+		
+		case EMetasoundFrontendLiteralType::UObject:
+		{
+			Set(static_cast<UObject*>(nullptr));
+		}
+		break;
+		
+		case EMetasoundFrontendLiteralType::Invalid:
+		{
+			Set(FDefault());
+		}
+		break;
+		
+		default:
+		static_assert(static_cast<int32>(EMetasoundFrontendLiteralType::Invalid) == 12, "Possible missing EMetasoundFrontendLiteralType case coverage");
+		{
+			checkNoEntry();
+		}
+		break;
+	}
+}
+
 Metasound::FLiteral FMetasoundFrontendLiteral::ToLiteral(const FName& InMetasoundDataTypeName) const
 {
 	using namespace Metasound;
@@ -466,6 +558,11 @@ FString FMetasoundFrontendLiteral::ToString() const
 	}
 }
 
+EMetasoundFrontendLiteralType FMetasoundFrontendLiteral::GetType() const
+{
+	return Type;
+}
+
 void FMetasoundFrontendLiteral::Clear()
 {
 	Empty();
diff --git a/Engine/Plugins/Runtime/Metasound/Source/MetasoundFrontend/Private/MetasoundFrontendStandardController.cpp b/Engine/Plugins/Runtime/Metasound/Source/MetasoundFrontend/Private/MetasoundFrontendStandardController.cpp
--- a/Engine/Plugins/Runtime/Metasound/Source/MetasoundFrontend/Private/MetasoundFrontendStandardController.cpp
+++ b/Engine/Plugins/Runtime/Metasound/Source/MetasoundFrontend/Private/MetasoundFrontendStandardController.cpp
@@ -2310,6 +2310,16 @@ namespace Metasound
 			return nullptr;
 		}
 
+		TArray<FMetasoundFrontendVertexLiteral> FGraphController::GetDefaultInputFrontendLiterals(const FString& InInputName) const
+		{
+			if (const FMetasoundFrontendClassInput* Desc = FindInputDescriptionWithName(InInputName))
+			{
+				return Desc->Defaults;
+			}
+
+			return TArray<FMetasoundFrontendVertexLiteral>();
+		}
+
 		bool FGraphController::SetDefaultInputToFrontendLiteral(const FString& InInputName, FGuid InPointID, FName InDataTypeName, const FMetasoundFrontendLiteral& InLiteral)
 		{
 			if (FMetasoundFrontendClassInput* Desc = FindInputDescriptionWithName(InInputName))
@@ -2326,7 +2336,7 @@ namespace Metasound
 					VertexLiteral->PointID = InPointID;
 				}
 
-				if (ensure(DoesDataTypeSupportLiteralType(Desc->TypeName, InLiteral.Type)))
+				if (ensure(DoesDataTypeSupportLiteralType(Desc->TypeName, InLiteral.GetType())))
 				{
 					VertexLiteral->Value = InLiteral;
 					return true;
@@ -2350,8 +2360,27 @@ namespace Metasound
 			return SetDefaultInputToFrontendLiteral(InInputName, InPointID, InDataTypeName, FrontendLiteral);
 		}
 
-		// Set the display name for the input with the given name
-		void FGraphController::SetInputDisplayName(FString InName, const FText& InDisplayName)
+		const FText& FGraphController::GetInputDescription(const FString& InName) const
+		{
+			if (const FMetasoundFrontendClassInput* Desc = FindInputDescriptionWithName(InName))
+			{
+				return Desc->Metadata.Description;
+			}
+
+			return FText::GetEmpty();
+		}
+
+		const FText& FGraphController::GetOutputDescription(const FString& InName) const
+		{
+			if (const FMetasoundFrontendClassOutput* Desc = FindOutputDescriptionWithName(InName))
+			{
+				return Desc->Metadata.Description;
+			}
+
+			return FText::GetEmpty();
+		}
+
+		void FGraphController::SetInputDisplayName(const FString& InName, const FText& InDisplayName)
 		{
 			if (FMetasoundFrontendClassInput* Desc = FindInputDescriptionWithName(InName))
 			{
@@ -2359,8 +2388,7 @@ namespace Metasound
 			}
 		}
 
-		// Set the display name for the output with the given name
-		void FGraphController::SetOutputDisplayName(FString InName, const FText& InDisplayName)
+		void FGraphController::SetOutputDisplayName(const FString& InName, const FText& InDisplayName)
 		{
 			if (FMetasoundFrontendClassOutput* Desc = FindOutputDescriptionWithName(InName))
 			{
@@ -2368,6 +2396,22 @@ namespace Metasound
 			}
 		}
 
+		void FGraphController::SetInputDescription(const FString& InName, const FText& InDescription)
+		{
+			if (FMetasoundFrontendClassInput* Desc = FindInputDescriptionWithName(InName))
+			{
+				Desc->Metadata.Description = InDescription;
+			}
+		}
+
+		void FGraphController::SetOutputDescription(const FString& InName, const FText& InDescription)
+		{
+			if (FMetasoundFrontendClassOutput* Desc = FindOutputDescriptionWithName(InName))
+			{
+				Desc->Metadata.Description = InDescription;
+			}
+		}
+
 		FMetasoundFrontendClassMetadata FGraphController::CreateInputClassMetadata(const FMetasoundFrontendClassInput& InClassInput)
 		{
 			FMetasoundFrontendClassMetadata ClassMetadata;
@@ -3078,24 +3122,26 @@ namespace Metasound
 			return false;
 		}
 
-		TArray<FMetasoundFrontendClassVertex> FDocumentController::GetRequiredInputs() const
+		const TArray<FMetasoundFrontendClassInput>& FDocumentController::GetRequiredInputs() const
 		{
 			if (DocumentPtr.IsValid())
 			{
 				return DocumentPtr->Archetype.Interface.Inputs;
 			}
 
-			return TArray<FMetasoundFrontendClassVertex>();
+			static const TArray<FMetasoundFrontendClassInput> Empty;
+			return Empty;
 		}
 
-		TArray<FMetasoundFrontendClassVertex> FDocumentController::GetRequiredOutputs() const
+		const TArray<FMetasoundFrontendClassOutput>& FDocumentController::GetRequiredOutputs() const
 		{
 			if (DocumentPtr.IsValid())
 			{
 				return DocumentPtr->Archetype.Interface.Outputs;
 			}
 
-			return TArray<FMetasoundFrontendClassVertex>();
+			static const TArray<FMetasoundFrontendClassOutput> Empty;
+			return Empty;
 		}
 
 		TArray<FMetasoundFrontendClass> FDocumentController::GetDependencies() const
diff --git a/Engine/Plugins/Runtime/Metasound/Source/MetasoundFrontend/Private/MetasoundFrontendStandardController.h b/Engine/Plugins/Runtime/Metasound/Source/MetasoundFrontend/Private/MetasoundFrontendStandardController.h
--- a/Engine/Plugins/Runtime/Metasound/Source/MetasoundFrontend/Private/MetasoundFrontendStandardController.h
+++ b/Engine/Plugins/Runtime/Metasound/Source/MetasoundFrontend/Private/MetasoundFrontendStandardController.h
@@ -680,16 +680,30 @@ namespace Metasound
 			// For inputs whose preferred literal type is UObject or UObjectArray, this can be used to determine the UObject corresponding to that input's datatype.
 			UClass* GetSupportedClassForInputVertex(const FString& InInputName) override;
 
+			TArray<FMetasoundFrontendVertexLiteral> GetDefaultInputFrontendLiterals (const FString& InInputName) const override;
+
 			// These can be used to set the default value for a given input on this graph.
 			// @returns false if the input name couldn't be found, or if the literal type was incompatible with the Data Type of this input.
 			bool SetDefaultInputToFrontendLiteral(const FString& InInputName, FGuid InPointID, FName InDataTypeName, const FMetasoundFrontendLiteral& InLiteral) override;
 			bool SetDefaultInputToTypeDefaultLiteral(const FString& InInputName, FGuid InPointID, FName InDataTypeName) override;
 
+			// Get the description for the input with the given name.
+			const FText& GetInputDescription(const FString& InName) const override;
+
+			// Get the description for the output with the given name.
+			const FText& GetOutputDescription(const FString& InName) const override;
+
+			// Set the description for the input with the given name
+			void SetInputDescription(const FString& InName, const FText& InDescription) override;
+
+			// Set the description for the input with the given name
+			void SetOutputDescription(const FString& InName, const FText& InDescription) override;
+
 			// Set the display name for the input with the given name
-			void SetInputDisplayName(FString InName, const FText& InDisplayName) override;
+			void SetInputDisplayName(const FString& InName, const FText& InDisplayName) override;
 
 			// Set the display name for the output with the given name
-			void SetOutputDisplayName(FString InName, const FText& InDisplayName) override;
+			void SetOutputDisplayName(const FString& InName, const FText& InDisplayName) override;
 
 			// This can be used to clear the current literal for a given input.
 			// @returns false if the input name couldn't be found.
@@ -818,8 +832,8 @@ namespace Metasound
 			bool IsRequiredInput(const FString& InInputName) const override;
 			bool IsRequiredOutput(const FString& InOutputName) const override;
 
-			TArray<FMetasoundFrontendClassVertex> GetRequiredInputs() const override;
-			TArray<FMetasoundFrontendClassVertex> GetRequiredOutputs() const override;
+			const TArray<FMetasoundFrontendClassInput>& GetRequiredInputs() const override;
+			const TArray<FMetasoundFrontendClassOutput>& GetRequiredOutputs() const override;
 
 			TArray<FMetasoundFrontendClass> GetDependencies() const override;
 			TArray<FMetasoundFrontendGraphClass> GetSubgraphs() const override;
diff --git a/Engine/Plugins/Runtime/Metasound/Source/MetasoundFrontend/Private/MetasoundFrontendTransform.cpp b/Engine/Plugins/Runtime/Metasound/Source/MetasoundFrontend/Private/MetasoundFrontendTransform.cpp
--- a/Engine/Plugins/Runtime/Metasound/Source/MetasoundFrontend/Private/MetasoundFrontendTransform.cpp
+++ b/Engine/Plugins/Runtime/Metasound/Source/MetasoundFrontend/Private/MetasoundFrontendTransform.cpp
@@ -16,11 +16,11 @@ namespace Metasound
 			bool bDidEdit = false;
 			FGraphHandle Graph = InDocument->GetRootGraph();
 
-			TArray<FMetasoundFrontendClassVertex> RequiredInputs = InDocument->GetRequiredInputs();
-			TArray<FMetasoundFrontendClassVertex> RequiredOutputs = InDocument->GetRequiredOutputs();
+			TArray<FMetasoundFrontendClassInput> RequiredInputs = InDocument->GetRequiredInputs();
+			TArray<FMetasoundFrontendClassOutput> RequiredOutputs = InDocument->GetRequiredOutputs();
 
 			// Go through each input and add or swap if something is missing.
-			for (const FMetasoundFrontendClassVertex& RequiredInput : RequiredInputs)
+			for (const FMetasoundFrontendClassInput& RequiredInput : RequiredInputs)
 			{
 				const FMetasoundFrontendClassInput* ClassInput = Graph->FindClassInputWithName(RequiredInput.Name).Get();
 
@@ -42,7 +42,7 @@ namespace Metasound
 			}
 
 			// Go through each output and add or swap if something is missing.
-			for (const FMetasoundFrontendClassVertex& RequiredOutput : RequiredOutputs)
+			for (const FMetasoundFrontendClassOutput& RequiredOutput : RequiredOutputs)
 			{
 				const FMetasoundFrontendClassOutput* ClassOutput = Graph->FindClassOutputWithName(RequiredOutput.Name).Get();
 
diff --git a/Engine/Plugins/Runtime/Metasound/Source/MetasoundFrontend/Public/MetasoundFrontendController.h b/Engine/Plugins/Runtime/Metasound/Source/MetasoundFrontend/Public/MetasoundFrontendController.h
--- a/Engine/Plugins/Runtime/Metasound/Source/MetasoundFrontend/Public/MetasoundFrontendController.h
+++ b/Engine/Plugins/Runtime/Metasound/Source/MetasoundFrontend/Public/MetasoundFrontendController.h
@@ -536,6 +536,8 @@ namespace Metasound
 			virtual TArray<FGuid> GetDefaultIDsForInputVertex(const FString& InInputName) const = 0;
 			virtual TArray<FGuid> GetDefaultIDsForOutputVertex(const FString& InOutputName) const = 0;
 
+			virtual TArray<FMetasoundFrontendVertexLiteral> GetDefaultInputFrontendLiterals(const FString& InInputName) const = 0;
+
 			virtual bool SetDefaultInputToFrontendLiteral(const FString& InInputName, FGuid InPointID, FName InDataTypeName, const FMetasoundFrontendLiteral& InLiteral) = 0;
 
 			/** Set the default value for the graph input.
@@ -549,10 +551,22 @@ namespace Metasound
 			virtual bool SetDefaultInputToTypeDefaultLiteral(const FString& InInputName, FGuid InPointID, FName InDataTypeName) = 0;
 
 			/** Set the display name for the input with the given name. */
-			virtual void SetInputDisplayName(FString InName, const FText& InDisplayName) = 0;
+			virtual void SetInputDisplayName(const FString& InName, const FText& InDisplayName) = 0;
 
 			/** Set the display name for the output with the given name. */
-			virtual void SetOutputDisplayName(FString InName, const FText& InDisplayName) = 0;
+			virtual void SetOutputDisplayName(const FString& InName, const FText& InDisplayName) = 0;
+
+			/** Get the description for the input with the given name. */
+			virtual const FText& GetInputDescription(const FString& InName) const = 0;
+
+			/** Get the description for the output with the given name. */
+			virtual const FText& GetOutputDescription(const FString& InName) const = 0;
+
+			/** Set the description for the input with the given name. */
+			virtual void SetInputDescription(const FString& InName, const FText& InDescription) = 0;
+
+			/** Set the description for the output with the given name. */
+			virtual void SetOutputDescription(const FString& InName, const FText& InDescription) = 0;
 
 			/** Clear the current literal for a given input.
 			 *
@@ -660,10 +674,10 @@ namespace Metasound
 			virtual bool IsRequiredOutput(const FString& InOutputName) const = 0;
 
 			/** Returns an array of inputs which describe the required inputs needed to satisfy the document archetype. */
-			virtual TArray<FMetasoundFrontendClassVertex> GetRequiredInputs() const = 0;
+			virtual const TArray<FMetasoundFrontendClassInput>& GetRequiredInputs() const = 0;
 
 			/** Returns an array of outputs which describe the required outputs needed to satisfy the document archetype. */
-			virtual TArray<FMetasoundFrontendClassVertex> GetRequiredOutputs() const = 0;
+			virtual const TArray<FMetasoundFrontendClassOutput>& GetRequiredOutputs() const = 0;
 
 			// TODO: add info on environment variables. 
 			// TODO: consider find/add subgraph
diff --git a/Engine/Plugins/Runtime/Metasound/Source/MetasoundFrontend/Public/MetasoundFrontendDocument.h b/Engine/Plugins/Runtime/Metasound/Source/MetasoundFrontend/Public/MetasoundFrontendDocument.h
--- a/Engine/Plugins/Runtime/Metasound/Source/MetasoundFrontend/Public/MetasoundFrontendDocument.h
+++ b/Engine/Plugins/Runtime/Metasound/Source/MetasoundFrontend/Public/MetasoundFrontendDocument.h
@@ -164,6 +164,9 @@ struct METASOUNDFRONTEND_API FMetasoundFrontendNodeStyle
 	// Display style of a node
 	UPROPERTY()
 	FMetasoundFrontendNodeStyleDisplay Display;
+
+	UPROPERTY()
+	bool bIsPrivate = false;
 };
 
 // An FMetasoundFrontendNode represents a single instance of a FMetasoundFrontendClass
@@ -677,10 +680,10 @@ struct FMetasoundFrontendArchetypeInterface
 	GENERATED_BODY()
 
 	UPROPERTY()
-	TArray<FMetasoundFrontendClassVertex> Inputs;
+	TArray<FMetasoundFrontendClassInput> Inputs;
 
 	UPROPERTY()
-	TArray<FMetasoundFrontendClassVertex> Outputs;
+	TArray<FMetasoundFrontendClassOutput> Outputs;
 
 	UPROPERTY()
 	TArray<FMetasoundFrontendEnvironmentVariable> Environment;
diff --git a/Engine/Plugins/Runtime/Metasound/Source/MetasoundFrontend/Public/MetasoundFrontendLiteral.h b/Engine/Plugins/Runtime/Metasound/Source/MetasoundFrontend/Public/MetasoundFrontendLiteral.h
--- a/Engine/Plugins/Runtime/Metasound/Source/MetasoundFrontend/Public/MetasoundFrontendLiteral.h
+++ b/Engine/Plugins/Runtime/Metasound/Source/MetasoundFrontend/Public/MetasoundFrontendLiteral.h
@@ -50,31 +50,36 @@ struct METASOUNDFRONTEND_API FMetasoundFrontendLiteral
 	struct FDefault {};
 	struct FDefaultArray { int32 Num = 0; };
 
+private:
 	// The set type of this literal.
-	UPROPERTY(EditAnywhere, Category = Customized)
+	UPROPERTY()
 	EMetasoundFrontendLiteralType Type = EMetasoundFrontendLiteralType::None;
 
-	UPROPERTY(EditAnywhere, Category = Customized)
-	int32 AsNumDefault;
+	UPROPERTY()
+	int32 AsNumDefault = 0;
 
-	UPROPERTY(EditAnywhere, Category = Customized)
+	UPROPERTY()
 	TArray<bool> AsBoolean;
 
-	UPROPERTY(EditAnywhere, Category = Customized)
+	UPROPERTY()
 	TArray<int32> AsInteger;
 
-	UPROPERTY(EditAnywhere, Category = Customized)
+	UPROPERTY()
 	TArray<float> AsFloat;
 
-	UPROPERTY(EditAnywhere, Category = Customized)
+	UPROPERTY()
 	TArray<FString> AsString;
 
-	UPROPERTY(EditAnywhere, Category = Customized)
+	UPROPERTY()
 	TArray<UObject*> AsUObject;
 
+public:
 	// Returns true if the stored Type is an array type.
 	bool IsArray() const;
 
+	// Sets the literal to the given type and value to default;
+	void SetType(EMetasoundFrontendLiteralType InType);
+
 	void Set(FDefault InValue);
 	void Set(const FDefaultArray& InValue);
 	void Set(bool InValue);
@@ -90,6 +95,8 @@ struct METASOUNDFRONTEND_API FMetasoundFrontendLiteral
 
 	void SetFromLiteral(const Metasound::FLiteral& InLiteral);
 
+	EMetasoundFrontendLiteralType GetType() const;
+
 	// Return the literal description parsed into a init param. 
 	// @Returns an invalid init param if the data type couldn't be found, or if the literal type was incompatible with the data type.
 	Metasound::FLiteral ToLiteral(const FName& InMetasoundDataType) const;
