IWYU fix

diff --git a/Engine/Plugins/Runtime/Metasound/Source/MetasoundEditor/Private/MetasoundInputNodeDetailCustomization.cpp b/Engine/Plugins/Runtime/Metasound/Source/MetasoundEditor/Private/MetasoundInputNodeDetailCustomization.cpp
--- a/Engine/Plugins/Runtime/Metasound/Source/MetasoundEditor/Private/MetasoundInputNodeDetailCustomization.cpp
+++ b/Engine/Plugins/Runtime/Metasound/Source/MetasoundEditor/Private/MetasoundInputNodeDetailCustomization.cpp
@@ -24,6 +24,7 @@
 #include "PropertyEditorDelegates.h"
 #include "PropertyHandle.h"
 #include "PropertyRestriction.h"
+#include "ScopedTransaction.h"
 #include "SlateCore/Public/Styling/SlateColor.h"
 #include "Templates/Casts.h"
 #include "Templates/SharedPointer.h"
