Fix OnPlay not triggering due to wrong copy constructor being used to initialize array member

diff --git a/Engine/Plugins/Runtime/Metasound/Source/MetasoundEngine/Private/MetasoundSource.cpp b/Engine/Plugins/Runtime/Metasound/Source/MetasoundEngine/Private/MetasoundSource.cpp
--- a/Engine/Plugins/Runtime/Metasound/Source/MetasoundEngine/Private/MetasoundSource.cpp
+++ b/Engine/Plugins/Runtime/Metasound/Source/MetasoundEngine/Private/MetasoundSource.cpp
@@ -467,7 +467,7 @@ const FMetasoundFrontendArchetype& UMetasoundSource::GetBaseArchetype()
 	{
 		FMetasoundFrontendArchetype Archetype;
 		
-		FMetasoundFrontendClassInput OnPlayTrigger;
+		FMetasoundFrontendClassVertex OnPlayTrigger;
 		OnPlayTrigger.Name = UMetasoundSource::GetOnPlayInputName();
 		OnPlayTrigger.Metadata.DisplayName = FText::FromString(OnPlayTrigger.Name);
 		OnPlayTrigger.TypeName = Metasound::Frontend::GetDataTypeName<Metasound::FTrigger>();
@@ -476,7 +476,7 @@ const FMetasoundFrontendArchetype& UMetasoundSource::GetBaseArchetype()
 
 		Archetype.Interface.Inputs.Add(OnPlayTrigger);
 
-		FMetasoundFrontendClassOutput OnFinished;
+		FMetasoundFrontendClassVertex OnFinished;
 		OnFinished.Name = UMetasoundSource::GetIsFinishedOutputName();
 		OnFinished.Metadata.DisplayName = FText::FromString(OnFinished.Name);
 		OnFinished.TypeName = Metasound::Frontend::GetDataTypeName<Metasound::FTrigger>();
