MetaSounds PointIDs to VertexID

diff --git a/Engine/Plugins/Runtime/Metasound/Source/MetasoundEditor/Private/MetasoundEditorGraphBuilder.cpp b/Engine/Plugins/Runtime/Metasound/Source/MetasoundEditor/Private/MetasoundEditorGraphBuilder.cpp
--- a/Engine/Plugins/Runtime/Metasound/Source/MetasoundEditor/Private/MetasoundEditorGraphBuilder.cpp
+++ b/Engine/Plugins/Runtime/Metasound/Source/MetasoundEditor/Private/MetasoundEditorGraphBuilder.cpp
@@ -267,14 +267,10 @@ namespace Metasound
 
 			FOutputHandle OutputHandle = InNodeHandle->GetOutputs()[0];
 			FGraphHandle Graph = InNodeHandle->GetOwningGraph();
-			TArray<FMetasoundFrontendVertexLiteral> DefaultLiterals = Graph->GetDefaultInputFrontendLiterals(OutputHandle->GetName());
+			FGuid VertexID = Graph->GetVertexIDForInputVertex(OutputHandle->GetName());
+			FMetasoundFrontendLiteral DefaultLiteral = Graph->GetDefaultInput(VertexID);
 
-			EMetasoundFrontendLiteralType LiteralType = EMetasoundFrontendLiteralType::None;
-			if (!DefaultLiterals.IsEmpty())
-			{
-				const FMetasoundFrontendVertexLiteral& Literal = DefaultLiterals[0];
-				LiteralType = Literal.Value.GetType();
-			}
+			EMetasoundFrontendLiteralType LiteralType = DefaultLiteral.GetType();
 
 			UMetasoundEditorGraph* MetasoundGraph = Cast<UMetasoundEditorGraph>(MetasoundAsset->GetGraph());
 			if (!ensure(MetasoundGraph))
@@ -424,19 +420,17 @@ namespace Metasound
 				NewInputNode->SetNodeStyle(Style);
 
 				Frontend::FGraphHandle GraphHandle = InNodeHandle->GetOwningGraph();
-				const FString Name = InNodeHandle->GetNodeName();
 
-				const TArray<FGuid> PointIDs = GraphHandle->GetDefaultIDsForInputVertex(NewInputNode->GetNodeName());
-				// TODO: Support multiple literal point ids for arrays
-				if (ensure(PointIDs.Num() == 1))
+				const FGuid VertexID = GraphHandle->GetVertexIDForInputVertex(NewInputNode->GetNodeName());
+
+				const FMetasoundFrontendLiteral* DefaultLiteral = InputHandle->GetDefaultLiteral();
+
+				if (nullptr != DefaultLiteral)
 				{
-					InNodeHandle->GetClassDescription();
-					const TArray<FMetasoundFrontendVertexLiteral>& Defaults = InputHandle->GetDefaults();
-					if (!Defaults.IsEmpty() && Defaults[0].Value.GetType() != EMetasoundFrontendLiteralType::Invalid && Defaults[0].Value.GetType() != EMetasoundFrontendLiteralType::None)
+					if (DefaultLiteral->IsValid())
 					{
-						const FString NewInputNodeName = NewInputNode->GetNodeName();
-						GraphHandle->SetDefaultInputToFrontendLiteral(NewInputNodeName, PointIDs[0], TypeName, Defaults[0].Value);
-						InInputPin.DefaultValue = Defaults[0].Value.ToString();
+						GraphHandle->SetDefaultInput(VertexID, *DefaultLiteral);
+						InInputPin.DefaultValue = DefaultLiteral->ToString();
 					}
 				}
 
@@ -458,13 +452,9 @@ namespace Metasound
 				if (InputNode->IsValid() && Style.Display.Visibility == EMetasoundFrontendNodeStyleDisplayVisibility::Hidden)
 				{
 					Frontend::FGraphHandle GraphHandle = InputNode->GetOwningGraph();
-					const TArray<FGuid> PointIDs = GraphHandle->GetDefaultIDsForInputVertex(InputNode->GetNodeName());
+					const FGuid VertexID = GraphHandle->GetVertexIDForInputVertex(InputNode->GetNodeName());
 
-					// TODO: Support multiple literal point ids for arrays
-					if (ensure(PointIDs.Num() == 1))
-					{
-						GraphHandle->SetDefaultInputToFrontendLiteral(InputNode->GetNodeName(), PointIDs[0], TypeName, PinDataTypeDefaultLiteral);
-					}
+					GraphHandle->SetDefaultInput(VertexID, PinDataTypeDefaultLiteral);
 
 					InMetasound.PostEditChange();
 					InMetasound.MarkPackageDirty();
@@ -486,12 +476,9 @@ namespace Metasound
 
 			Description.Metadata.Description = InToolTip;
 
-			const FGuid PointID = GraphHandle->GetNewPointID();
+			const FGuid VertexID = GraphHandle->GetNewVertexID();
 
-			Description.PointIDs.Add(PointID);
-			FMetasoundFrontendVertexLiteral DefaultVertex;
-			DefaultVertex.PointID = PointID;
-			Description.Defaults.Add(DefaultVertex);
+			Description.VertexID = VertexID;
 
 			Frontend::FNodeHandle NodeHandle = GraphHandle->AddInputVertex(Description);
 			NodeHandle->SetNodeStyle(InNodeStyle);
@@ -504,12 +491,12 @@ namespace Metasound
 				{
 					if (InDefaultValue->GetType() != EMetasoundFrontendLiteralType::None && InDefaultValue->GetType() != EMetasoundFrontendLiteralType::Invalid)
 					{
-						GraphHandle->SetDefaultInputToFrontendLiteral(NodeHandle->GetNodeName(), PointID, InTypeName, *InDefaultValue);
+						GraphHandle->SetDefaultInput(VertexID, *InDefaultValue);
 					}
 				}
 				else
 				{
-					GraphHandle->SetDefaultInputToTypeDefaultLiteral(NodeHandle->GetNodeName(), PointID, InTypeName);
+					GraphHandle->SetDefaultInputToDefaultLiteralOfType(VertexID);
 				}
 			}
 
@@ -533,7 +520,7 @@ namespace Metasound
 			Description.Name = InName;
 			Description.TypeName = InTypeName;
 			Description.Metadata.Description = InToolTip;
-			Description.PointIDs.Add(GraphHandle->GetNewPointID());
+			Description.VertexID = GraphHandle->GetNewVertexID();
 
 			Frontend::FNodeHandle NodeHandle = GraphHandle->AddOutputVertex(Description);
 			NodeHandle->SetNodeStyle(InNodeStyle);
diff --git a/Engine/Plugins/Runtime/Metasound/Source/MetasoundEditor/Private/MetasoundInputNodeDetailCustomization.cpp b/Engine/Plugins/Runtime/Metasound/Source/MetasoundEditor/Private/MetasoundInputNodeDetailCustomization.cpp
--- a/Engine/Plugins/Runtime/Metasound/Source/MetasoundEditor/Private/MetasoundInputNodeDetailCustomization.cpp
+++ b/Engine/Plugins/Runtime/Metasound/Source/MetasoundEditor/Private/MetasoundInputNodeDetailCustomization.cpp
@@ -68,21 +68,15 @@ namespace Metasound
 					return;
 				}
 
-				const FNodeHandle NodeHandle = GraphNode->GetNodeHandle();
+				FConstNodeHandle NodeHandle = GraphNode->GetNodeHandle();
 				const FString& NodeName = NodeHandle->GetNodeName();
-				const FGraphHandle GraphHandle = GraphNode->GetRootGraphHandle();
+				FGraphHandle GraphHandle = GraphNode->GetRootGraphHandle();
 
-				const TArray<FGuid> PointIDs = GraphHandle->GetDefaultIDsForInputVertex(NodeName);
-				if (ensure(PointIDs.Num() == 1))
-				{
-					const FMetasoundFrontendLiteral Literal = GraphNode->GetLiteralDefault();
-					const TArray<FOutputHandle> Outputs = NodeHandle->GetOutputs();
-					if (ensure(Outputs.Num() == 1))
-					{
-						const FName TypeName = Outputs[0]->GetDataType();
-						GraphHandle->SetDefaultInputToFrontendLiteral(NodeName, PointIDs[0], TypeName, Literal);
-					}
-				}
+				FGuid VertexID = GraphHandle->GetVertexIDForInputVertex(NodeName);
+
+				const FMetasoundFrontendLiteral Literal = GraphNode->GetLiteralDefault();
+
+				GraphHandle->SetDefaultInput(VertexID, Literal);
 			}
 		}
 
diff --git a/Engine/Plugins/Runtime/Metasound/Source/MetasoundEngine/Private/MetasoundSource.cpp b/Engine/Plugins/Runtime/Metasound/Source/MetasoundEngine/Private/MetasoundSource.cpp
--- a/Engine/Plugins/Runtime/Metasound/Source/MetasoundEngine/Private/MetasoundSource.cpp
+++ b/Engine/Plugins/Runtime/Metasound/Source/MetasoundEngine/Private/MetasoundSource.cpp
@@ -241,17 +241,13 @@ bool UMetasoundSource::GetReceiveNodeMetadataForDataType(const FName& InTypeName
 Metasound::Frontend::FNodeHandle UMetasoundSource::AddInputPinForSendAddress(const Metasound::FMetasoundInstanceTransmitter::FSendInfo& InSendInfo, Metasound::Frontend::FGraphHandle InGraph) const
 {
 	FMetasoundFrontendClassInput Description;
-	FGuid PointID = InGraph->GetNewPointID();
+	FGuid VertexID = InGraph->GetNewVertexID();
 
 	Description.Name = InSendInfo.Address.ChannelName.ToString();
 	Description.TypeName = Metasound::GetMetasoundDataTypeName<Metasound::FSendAddress>();
 	Description.Metadata.Description = FText::GetEmpty();
-	Description.PointIDs.Add(PointID);
-
-	FMetasoundFrontendVertexLiteral DefaultValue;
-	DefaultValue.PointID = PointID;
-	DefaultValue.Value.Set(InSendInfo.Address.ChannelName.ToString());
-	Description.Defaults.Add(DefaultValue);
+	Description.VertexID = VertexID;
+	Description.DefaultLiteral.Set(InSendInfo.Address.ChannelName.ToString());
 	
 	return InGraph->AddInputVertex(Description);
 }
@@ -416,7 +412,7 @@ TArray<UMetasoundSource::FSendInfoAndVertexName> UMetasoundSource::GetSendInfos(
 		{
 			FSendInfoAndVertexName Info;
 
-			// TODO: incorporate PointID into address. But need to ensure that PointID
+			// TODO: incorporate VertexID into address. But need to ensure that VertexID
 			// will be maintained after injecting Receive nodes. 
 			Info.SendInfo.Address = CreateSendAddress(InInstanceID, InputHandle->GetName(), InputHandle->GetDataType());
 			Info.SendInfo.ParameterName = FName(InputHandle->GetDisplayName().ToString()); // TODO: display name hack. Need to have nameing consistent in editor for inputs
@@ -472,7 +468,7 @@ const FMetasoundFrontendArchetype& UMetasoundSource::GetBaseArchetype()
 		OnPlayTrigger.Metadata.DisplayName = FText::FromString(OnPlayTrigger.Name);
 		OnPlayTrigger.TypeName = Metasound::Frontend::GetDataTypeName<Metasound::FTrigger>();
 		OnPlayTrigger.Metadata.Description = LOCTEXT("OnPlayTriggerToolTip", "Trigger executed when this source is played.");
-		OnPlayTrigger.PointIDs.Add(FGuid::NewGuid());
+		OnPlayTrigger.VertexID = FGuid::NewGuid();
 
 		Archetype.Interface.Inputs.Add(OnPlayTrigger);
 
@@ -481,7 +477,7 @@ const FMetasoundFrontendArchetype& UMetasoundSource::GetBaseArchetype()
 		OnFinished.Metadata.DisplayName = FText::FromString(OnFinished.Name);
 		OnFinished.TypeName = Metasound::Frontend::GetDataTypeName<Metasound::FTrigger>();
 		OnFinished.Metadata.Description = LOCTEXT("OnFinishedToolTip", "Trigger executed to initiate stopping the source.");
-		OnFinished.PointIDs.Add(FGuid::NewGuid());
+		OnFinished.VertexID = FGuid::NewGuid();
 
 		Archetype.Interface.Outputs.Add(OnFinished);
 
@@ -512,7 +508,7 @@ const FMetasoundFrontendArchetype& UMetasoundSource::GetMonoSourceArchetype()
 		GeneratedAudio.TypeName = Metasound::Frontend::GetDataTypeName<Metasound::FMonoAudioFormat>();
 		GeneratedAudio.Metadata.DisplayName = LOCTEXT("GeneratedMono", "Mono");
 		GeneratedAudio.Metadata.Description = LOCTEXT("GeneratedAudioToolTip", "The resulting output audio from this source.");
-		GeneratedAudio.PointIDs.Add(FGuid::NewGuid());
+		GeneratedAudio.VertexID = FGuid::NewGuid();
 
 		Archetype.Interface.Outputs.Add(GeneratedAudio);
 
@@ -536,7 +532,7 @@ const FMetasoundFrontendArchetype& UMetasoundSource::GetStereoSourceArchetype()
 		GeneratedAudio.TypeName = Metasound::Frontend::GetDataTypeName<Metasound::FStereoAudioFormat>();
 		GeneratedAudio.Metadata.DisplayName = LOCTEXT("GeneratedStereo", "Stereo");
 		GeneratedAudio.Metadata.Description = LOCTEXT("GeneratedAudioToolTip", "The resulting output audio from this source.");
-		GeneratedAudio.PointIDs.Add(FGuid::NewGuid());
+		GeneratedAudio.VertexID = FGuid::NewGuid();
 
 		Archetype.Interface.Outputs.Add(GeneratedAudio);
 
diff --git a/Engine/Plugins/Runtime/Metasound/Source/MetasoundFrontend/Private/MetasoundFrontend.cpp b/Engine/Plugins/Runtime/Metasound/Source/MetasoundFrontend/Private/MetasoundFrontend.cpp
--- a/Engine/Plugins/Runtime/Metasound/Source/MetasoundFrontend/Private/MetasoundFrontend.cpp
+++ b/Engine/Plugins/Runtime/Metasound/Source/MetasoundFrontend/Private/MetasoundFrontend.cpp
@@ -76,17 +76,14 @@ namespace Metasound
 
 				ClassInput.Name = InputTuple.Value.GetVertexName();
 				ClassInput.TypeName = InputTuple.Value.GetDataTypeName();
-				ClassInput.PointIDs.Add(FGuid::NewGuid());
+				ClassInput.VertexID = FGuid::NewGuid();
 				ClassInput.Metadata.DisplayName = FText::FromString(InputTuple.Value.GetVertexName());
 				ClassInput.Metadata.Description = InputTuple.Value.GetDescription();
 
 				const FLiteral& DefaultLiteral = InputTuple.Value.GetDefaultValue();
 				if (DefaultLiteral.GetType() != ELiteralType::None)
 				{
-					FMetasoundFrontendVertexLiteral VertexLiteral;
-					VertexLiteral.Value.SetFromLiteral(DefaultLiteral);
-					VertexLiteral.PointID = FGuid::NewGuid();
-					ClassInput.Defaults.Add(VertexLiteral);
+					ClassInput.DefaultLiteral.SetFromLiteral(DefaultLiteral);
 				}
 
 				ClassDescription.Interface.Inputs.Add(ClassInput);
@@ -99,7 +96,7 @@ namespace Metasound
 
 				ClassOutput.Name = OutputTuple.Value.GetVertexName();
 				ClassOutput.TypeName = OutputTuple.Value.GetDataTypeName();
-				ClassOutput.PointIDs.Add(FGuid::NewGuid());
+				ClassOutput.VertexID = FGuid::NewGuid();
 				ClassOutput.Metadata.DisplayName = FText::FromString(OutputTuple.Value.GetVertexName());
 				ClassOutput.Metadata.Description = OutputTuple.Value.GetDescription();
 
diff --git a/Engine/Plugins/Runtime/Metasound/Source/MetasoundFrontend/Private/MetasoundFrontendDocument.cpp b/Engine/Plugins/Runtime/Metasound/Source/MetasoundFrontend/Private/MetasoundFrontendDocument.cpp
--- a/Engine/Plugins/Runtime/Metasound/Source/MetasoundFrontend/Private/MetasoundFrontendDocument.cpp
+++ b/Engine/Plugins/Runtime/Metasound/Source/MetasoundFrontend/Private/MetasoundFrontendDocument.cpp
@@ -43,17 +43,12 @@ FMetasoundFrontendNode::FMetasoundFrontendNode(const FMetasoundFrontendClass& In
 
 bool FMetasoundFrontendVertex::IsFunctionalEquivalent(const FMetasoundFrontendVertex& InLHS, const FMetasoundFrontendVertex& InRHS)
 {
-	return (InLHS.Name == InRHS.Name) && (InLHS.TypeName == InRHS.TypeName) && (InLHS.PointIDs.Num() == InRHS.PointIDs.Num());
-}
-
-bool FMetasoundFrontendVertexBehavior::IsFunctionalEquivalent(const FMetasoundFrontendVertexBehavior& InLHS, const FMetasoundFrontendVertexBehavior& InRHS)
-{
-	return (InLHS.Type == InRHS.Type) && (InLHS.ArrayMin == InRHS.ArrayMin) && (InLHS.ArrayMax == InRHS.ArrayMax);
+	return (InLHS.Name == InRHS.Name) && (InLHS.TypeName == InRHS.TypeName);
 }
 
 bool FMetasoundFrontendClassVertex::IsFunctionalEquivalent(const FMetasoundFrontendClassVertex& InLHS, const FMetasoundFrontendClassVertex& InRHS) 
 {
-	return FMetasoundFrontendVertex::IsFunctionalEquivalent(InLHS, InRHS) && FMetasoundFrontendVertexBehavior::IsFunctionalEquivalent(InLHS.Behavior, InRHS.Behavior);
+	return FMetasoundFrontendVertex::IsFunctionalEquivalent(InLHS, InRHS);
 }
 
 FMetasoundFrontendClassName::FMetasoundFrontendClassName(const FName& InNamespace, const FName& InName, const FName& InVariant)
@@ -106,13 +101,7 @@ FMetasoundFrontendClassInput::FMetasoundFrontendClassInput(const FMetasoundFront
 {
 	EMetasoundFrontendLiteralType DefaultType = Metasound::Frontend::GetMetasoundFrontendLiteralType(FMetasoundFrontendRegistryContainer::Get()->GetDesiredLiteralTypeForDataType(InOther.TypeName));
 
-	for (const FGuid& PointID : PointIDs)
-	{
-		FMetasoundFrontendVertexLiteral Default;
-		Default.PointID = PointID;
-		Default.Value.SetType(DefaultType);
-		Defaults.Add(Default);
-	}
+	DefaultLiteral.SetType(DefaultType);
 }
 
 FMetasoundFrontendGraphClass::FMetasoundFrontendGraphClass()
diff --git a/Engine/Plugins/Runtime/Metasound/Source/MetasoundFrontend/Private/MetasoundFrontendGraph.cpp b/Engine/Plugins/Runtime/Metasound/Source/MetasoundFrontend/Private/MetasoundFrontendGraph.cpp
--- a/Engine/Plugins/Runtime/Metasound/Source/MetasoundFrontend/Private/MetasoundFrontendGraph.cpp
+++ b/Engine/Plugins/Runtime/Metasound/Source/MetasoundFrontend/Private/MetasoundFrontendGraph.cpp
@@ -145,19 +145,19 @@ namespace Metasound
 
 	TUniquePtr<INode> FFrontendGraphBuilder::CreateInputNode(const FMetasoundFrontendNode& InNode, const FMetasoundFrontendClass& InClass, const FMetasoundFrontendClassInput& InOwningGraphClassInput) const
 	{
-		const FMetasoundFrontendVertexLiteral* VertexLiteral = FindInputLiteralForInputNode(InNode, InClass, InOwningGraphClassInput);
+		const FMetasoundFrontendLiteral* FrontendLiteral = FindInputLiteralForInputNode(InNode, InClass, InOwningGraphClassInput);
 
-		if (nullptr != VertexLiteral)
+		if (nullptr != FrontendLiteral)
 		{
 			if (ensure(InNode.Interface.Inputs.Num() == 1))
 			{
 				const FMetasoundFrontendVertex& InputVertex = InNode.Interface.Inputs[0];
 
-				const bool IsLiteralParsableByDataType = Frontend::DoesDataTypeSupportLiteralType(InputVertex.TypeName, VertexLiteral->Value.GetType());
+				const bool IsLiteralParsableByDataType = Frontend::DoesDataTypeSupportLiteralType(InputVertex.TypeName, FrontendLiteral->GetType());
 
 				if (IsLiteralParsableByDataType)
 				{
-					FLiteral Literal = VertexLiteral->Value.ToLiteral(InputVertex.TypeName);
+					FLiteral Literal = FrontendLiteral->ToLiteral(InputVertex.TypeName);
 
 					FInputNodeConstructorParams InitParams =
 					{
@@ -288,118 +288,57 @@ namespace Metasound
 		return nullptr;
 	}
 
-		/*
-	const FMetasoundFrontendVertexLiteral* FFrontendGraphBuilder::FindInputLiteralForInput(FGuid InPointID, const FMetasoundFrontendNode& InNode, const FMetasoundFrontendClass& InNodeClass) const
-	{
-		// TODO: Need utility for matching class input to node input. Will need to be cognisant of input type [Point, Array, ....]
-
-		// Default value priority is:
-		// 1. A value set directly on the node
-		// 3. A default value on the node class.
-
-		const FMetasoundFrontendVertexLiteral* VertexLiteral = nullptr;
-
-		// Check for default value directly on node.
-		if (ensure(InNode.Interface.Inputs.Num() == 1))
-		{
-			const FMetasoundFrontendVertex& InputVertex = InNode.Interface.Inputs[0];
-
-			// Currently only support one point on input vertex.
-			if (ensure(InputVertex.PointIDs.Num() == 1))
-			{
-				VertexLiteral = InNode.InputLiterals.FindByPredicate(
-					[&](const FMetasoundFrontendVertexLiteral& InVertexLiteral)
-					{
-						return InVertexLiteral.PointID == InputVertex.PointIDs[0];
-					}
-				);
-			}
-		}
-
-		// Check for default value on input node class
-		if (nullptr == VertexLiteral && ensure(InNodeClass.Interface.Inputs.Num() == 1))
-		{
-			const FMetasoundFrontendClassInput& InputNodeClassInput = InNodeClass.Interface.Inputs[0];
 
-			if (ensure(InputNodeClassInput.PointIDs.Num() == 1))
-			{
-				int32 PointID = InputNodeClassInput.PointIDs[0];
-
-				VertexLiteral = InputNodeClassInput.Defaults.FindByPredicate(
-					[&](const FMetasoundFrontendVertexLiteral& InVertexLiteral)
-					{
-						return InVertexLiteral.PointID == PointID;
-					}
-				);
-			}
-		}
-
-		return VertexLiteral;
-		return nullptr;
-	}
-		*/
-
-	const FMetasoundFrontendVertexLiteral* FFrontendGraphBuilder::FindInputLiteralForInputNode(const FMetasoundFrontendNode& InInputNode, const FMetasoundFrontendClass& InInputNodeClass, const FMetasoundFrontendClassInput& InOwningGraphClassInput) const
+	const FMetasoundFrontendLiteral* FFrontendGraphBuilder::FindInputLiteralForInputNode(const FMetasoundFrontendNode& InInputNode, const FMetasoundFrontendClass& InInputNodeClass, const FMetasoundFrontendClassInput& InOwningGraphClassInput) const
 	{
 		// Default value priority is:
 		// 1. A value set directly on the node
 		// 2. A default value of the owning graph
 		// 3. A default value on the input node class.
 
-		const FMetasoundFrontendVertexLiteral* VertexLiteral = nullptr;
+		const FMetasoundFrontendLiteral* Literal = nullptr;
 
 		// Check for default value directly on node.
 		if (ensure(InInputNode.Interface.Inputs.Num() == 1))
 		{
 			const FMetasoundFrontendVertex& InputVertex = InInputNode.Interface.Inputs[0];
 
-			// Currently only support one point on input vertex.
-			if (ensure(InputVertex.PointIDs.Num() == 1))
+			// Find input literal matching VerteXID
+			const FMetasoundFrontendVertexLiteral* VertexLiteral = InInputNode.InputLiterals.FindByPredicate(
+				[&](const FMetasoundFrontendVertexLiteral& InVertexLiteral)
+				{
+					return InVertexLiteral.VertexID == InputVertex.VertexID;
+				}
+			);
+
+			if (nullptr != VertexLiteral)
 			{
-				VertexLiteral = InInputNode.InputLiterals.FindByPredicate(
-					[&](const FMetasoundFrontendVertexLiteral& InVertexLiteral)
-					{
-						return InVertexLiteral.PointID == InputVertex.PointIDs[0];
-					}
-				);
+				Literal = &VertexLiteral->Value;
 			}
 		}
 
 		// Check for default value on owning graph.
-		if (nullptr == VertexLiteral)
+		if (nullptr == Literal)
 		{
-			if (ensure(InOwningGraphClassInput.PointIDs.Num() == 1))
+			// Find Class Default that is not invalid
+			if (InOwningGraphClassInput.DefaultLiteral.IsValid())
 			{
-				const FGuid& PointID = InOwningGraphClassInput.PointIDs[0];
-
-				VertexLiteral = InOwningGraphClassInput.Defaults.FindByPredicate(
-					[&](const FMetasoundFrontendVertexLiteral& InVertexLiteral)
-					{
-						return InVertexLiteral.PointID == PointID;
-					}
-				);
+				Literal = &InOwningGraphClassInput.DefaultLiteral;
 			}
 		}
 
 		// Check for default value on input node class
-		if (nullptr == VertexLiteral && ensure(InInputNodeClass.Interface.Inputs.Num() == 1))
+		if (nullptr == Literal && ensure(InInputNodeClass.Interface.Inputs.Num() == 1))
 		{
 			const FMetasoundFrontendClassInput& InputNodeClassInput = InInputNodeClass.Interface.Inputs[0];
 
-			if (ensure(InputNodeClassInput.PointIDs.Num() == 1))
+			if (InputNodeClassInput.DefaultLiteral.IsValid())
 			{
-				const FGuid& PointID = InputNodeClassInput.PointIDs[0];
-
-				VertexLiteral = InputNodeClassInput.Defaults.FindByPredicate(
-					[&](const FMetasoundFrontendVertexLiteral& InVertexLiteral)
-					{
-						return InVertexLiteral.PointID == PointID;
-					}
-				);
+				Literal = &InputNodeClassInput.DefaultLiteral;
 			}
 		}
 
-		return VertexLiteral;
+		return Literal;
 	}
 
 	// TODO: add errors here. Most will be a "PromptIfMissing"...
@@ -466,12 +405,12 @@ namespace Metasound
 		};
 
 		// TODO: add support for array vertices.
-		typedef TTuple<FGuid, FGuid> FNodeIDPointID;
+		typedef TTuple<FGuid, FGuid> FNodeIDVertexID;
 
-		TMap<FNodeIDPointID, FCoreNodeAndFrontendVertex> NodeSourcesByID;
-		TMap<FNodeIDPointID, FCoreNodeAndFrontendVertex> NodeDestinationsByID;
+		TMap<FNodeIDVertexID, FCoreNodeAndFrontendVertex> NodeSourcesByID;
+		TMap<FNodeIDVertexID, FCoreNodeAndFrontendVertex> NodeDestinationsByID;
 
-		// Add nodes to NodeID/PointID map
+		// Add nodes to NodeID/VertexID map
 		for (const FMetasoundFrontendNode& Node : InFrontendGraph.Nodes)
 		{
 			const INode* CoreNode = OutGraph.FindNode(Node.ID);
@@ -483,31 +422,25 @@ namespace Metasound
 
 			for (const FMetasoundFrontendVertex& Vertex : Node.Interface.Inputs)
 			{
-				for (const FGuid& PointID : Vertex.PointIDs)
-				{
-					NodeDestinationsByID.Add(FNodeIDPointID(Node.ID, PointID), FCoreNodeAndFrontendVertex({CoreNode, &Vertex}));
-				}
+				NodeDestinationsByID.Add(FNodeIDVertexID(Node.ID, Vertex.VertexID), FCoreNodeAndFrontendVertex({CoreNode, &Vertex}));
 			}
 
 			for (const FMetasoundFrontendVertex& Vertex : Node.Interface.Outputs)
 			{
-				for (const FGuid& PointID : Vertex.PointIDs)
-				{
-					NodeSourcesByID.Add(FNodeIDPointID(Node.ID, PointID), FCoreNodeAndFrontendVertex({CoreNode, &Vertex}));
-				}
+				NodeSourcesByID.Add(FNodeIDVertexID(Node.ID, Vertex.VertexID), FCoreNodeAndFrontendVertex({CoreNode, &Vertex}));
 			}
 		};
 		
 
 		for (const FMetasoundFrontendEdge& Edge : InFrontendGraph.Edges)
 		{
-			const FNodeIDPointID DestinationKey(Edge.ToNodeID, Edge.ToPointID);
+			const FNodeIDVertexID DestinationKey(Edge.ToNodeID, Edge.ToVertexID);
 			const FCoreNodeAndFrontendVertex* DestinationNodeAndVertex = NodeDestinationsByID.Find(DestinationKey);
 
 			if (nullptr == DestinationNodeAndVertex)
 			{
 				// TODO: bubble up error
-				UE_LOG(LogMetasound, Error, TEXT("Failed to add edge. Could not find destination [NodeID:%s, PointID:%s]"), *Edge.ToNodeID.ToString(), *Edge.ToPointID.ToString());
+				UE_LOG(LogMetasound, Error, TEXT("Failed to add edge. Could not find destination [NodeID:%s, VertexID:%s]"), *Edge.ToNodeID.ToString(), *Edge.ToVertexID.ToString());
 				continue;
 			}
 
@@ -518,12 +451,12 @@ namespace Metasound
 				continue;
 			}
 
-			const FNodeIDPointID SourceKey(Edge.FromNodeID, Edge.FromPointID);
+			const FNodeIDVertexID SourceKey(Edge.FromNodeID, Edge.FromVertexID);
 			const FCoreNodeAndFrontendVertex* SourceNodeAndVertex = NodeSourcesByID.Find(SourceKey);
 
 			if (nullptr == SourceNodeAndVertex)
 			{
-				UE_LOG(LogMetasound, Error, TEXT("Failed to add edge. Could not find source [NodeID:%s, PointID:%s]"), *Edge.FromNodeID.ToString(), *Edge.FromPointID.ToString());
+				UE_LOG(LogMetasound, Error, TEXT("Failed to add edge. Could not find source [NodeID:%s, VertexID:%s]"), *Edge.FromNodeID.ToString(), *Edge.FromVertexID.ToString());
 				continue;
 			}
 
@@ -544,7 +477,7 @@ namespace Metasound
 
 			if (!bSuccess)
 			{
-				UE_LOG(LogMetasound, Error, TEXT("Failed to connect edge from [NodeID:%s, PointID:%s] to [NodeID:%s, PointID:%s]"), *Edge.FromNodeID.ToString(), *Edge.FromPointID.ToString(), *Edge.ToNodeID.ToString(), *Edge.ToPointID.ToString());
+				UE_LOG(LogMetasound, Error, TEXT("Failed to connect edge from [NodeID:%s, VertexID:%s] to [NodeID:%s, VertexID:%s]"), *Edge.FromNodeID.ToString(), *Edge.FromVertexID.ToString(), *Edge.ToNodeID.ToString(), *Edge.ToVertexID.ToString());
 			}
 		}
 	}
diff --git a/Engine/Plugins/Runtime/Metasound/Source/MetasoundFrontend/Private/MetasoundFrontendGraph.h b/Engine/Plugins/Runtime/Metasound/Source/MetasoundFrontend/Private/MetasoundFrontendGraph.h
--- a/Engine/Plugins/Runtime/Metasound/Source/MetasoundFrontend/Private/MetasoundFrontendGraph.h
+++ b/Engine/Plugins/Runtime/Metasound/Source/MetasoundFrontend/Private/MetasoundFrontendGraph.h
@@ -122,8 +122,7 @@ namespace Metasound
 
 		const FMetasoundFrontendClassOutput* FindClassOutputForOutputNode(const FMetasoundFrontendGraphClass& InOwningGraph, const FMetasoundFrontendNode& InOutputNode, int32& OutClassOutputIndex) const;
 
-		const FMetasoundFrontendVertexLiteral* FindInputLiteralForInputNode(const FMetasoundFrontendNode& InInputNode, const FMetasoundFrontendClass& InInputNodeClass, const FMetasoundFrontendClassInput& InOwningGraphClassInput) const;
-		//const FMetasoundFrontendVertexLiteral* FindInputLiteralForInput(FGuid InPointID, const FMetasoundFrontendNode& InNode, const FMetasoundFrontendClass& InNodeClass) const;
+		const FMetasoundFrontendLiteral* FindInputLiteralForInputNode(const FMetasoundFrontendNode& InInputNode, const FMetasoundFrontendClass& InInputNodeClass, const FMetasoundFrontendClassInput& InOwningGraphClassInput) const;
 
 		TUniquePtr<INode> CreateInputNode(const FMetasoundFrontendNode& InNode, const FMetasoundFrontendClass& InClass, const FMetasoundFrontendClassInput& InOwningGraphClassInput) const;
 
diff --git a/Engine/Plugins/Runtime/Metasound/Source/MetasoundFrontend/Private/MetasoundFrontendInvalidController.h b/Engine/Plugins/Runtime/Metasound/Source/MetasoundFrontend/Private/MetasoundFrontendInvalidController.h
--- a/Engine/Plugins/Runtime/Metasound/Source/MetasoundFrontend/Private/MetasoundFrontendInvalidController.h
+++ b/Engine/Plugins/Runtime/Metasound/Source/MetasoundFrontend/Private/MetasoundFrontendInvalidController.h
@@ -91,11 +91,7 @@ namespace Metasound
 			const FString& GetName() const override { return MetasoundFrontendInvalidControllerPrivate::GetInvalid<FString>(); }
 			const FText& GetDisplayName() const override { return FText::GetEmpty(); }
 			const FText& GetTooltip() const override { return FText::GetEmpty(); }
-			const TArray<FMetasoundFrontendVertexLiteral>& GetDefaults() const override
-			{
-				static const TArray<FMetasoundFrontendVertexLiteral> InvalidDefaults;
-				return InvalidDefaults;
-			}
+			const FMetasoundFrontendLiteral* GetDefaultLiteral() const override { return nullptr; }
 			FGuid GetOwningNodeID() const override { return Metasound::FrontendInvalidID; }
 			TSharedRef<INodeController> GetOwningNode() override;
 			TSharedRef<const INodeController> GetOwningNode() const override;
@@ -141,21 +137,21 @@ namespace Metasound
 			TArray<FConstInputHandle> GetInputsWithVertexName(const FString& InName) const override { return TArray<FConstInputHandle>(); }
 			TArray<FOutputHandle> GetOutputsWithVertexName(const FString& InName) override { return TArray<FOutputHandle>(); }
 			TArray<FConstOutputHandle> GetOutputsWithVertexName(const FString& InName) const override { return TArray<FConstOutputHandle>(); }
-			TSharedRef<IInputController> GetInputWithID(FGuid InPointID) override { return FInvalidInputController::GetInvalid(); }
-			TSharedRef<IOutputController> GetOutputWithID(FGuid InPointID) override { return FInvalidOutputController::GetInvalid(); }
-			TSharedRef<const IInputController> GetInputWithID(FGuid InPointID) const override { return FInvalidInputController::GetInvalid(); }
-			TSharedRef<const IOutputController> GetOutputWithID(FGuid InPointID) const override { return FInvalidOutputController::GetInvalid(); }
+			TSharedRef<IInputController> GetInputWithID(FGuid InVertexID) override { return FInvalidInputController::GetInvalid(); }
+			TSharedRef<IOutputController> GetOutputWithID(FGuid InVertexID) override { return FInvalidOutputController::GetInvalid(); }
+			TSharedRef<const IInputController> GetInputWithID(FGuid InVertexID) const override { return FInvalidInputController::GetInvalid(); }
+			TSharedRef<const IOutputController> GetOutputWithID(FGuid InVertexID) const override { return FInvalidOutputController::GetInvalid(); }
 
 			const FMetasoundFrontendNodeStyle& GetNodeStyle() const override { static const FMetasoundFrontendNodeStyle Invalid; return Invalid; }
 			void SetNodeStyle(const FMetasoundFrontendNodeStyle& InNodeStyle) { };
 
 			bool CanAddInput(const FString& InVertexName) const override { return false; }
 			FInputHandle AddInput(const FString& InVertexName, const FMetasoundFrontendLiteral* InDefault) override { return FInvalidInputController::GetInvalid(); }
-			bool RemoveInput(FGuid InPointID) override { return false; }
+			bool RemoveInput(FGuid InVertexID) override { return false; }
 
 			bool CanAddOutput(const FString& InVertexName) const override { return false; }
 			FInputHandle AddOutput(const FString& InVertexName, const FMetasoundFrontendLiteral* InDefault) override { return FInvalidInputController::GetInvalid(); }
-			bool RemoveOutput(FGuid InPointID) override { return false; }
+			bool RemoveOutput(FGuid InVertexID) override { return false; }
 
 			EMetasoundFrontendClassType GetClassType() const override { return EMetasoundFrontendClassType::Invalid; }
 			const FMetasoundFrontendClassName& GetClassName() const override { return MetasoundFrontendInvalidControllerPrivate::GetInvalid<FMetasoundFrontendClassName>(); }
@@ -198,7 +194,7 @@ namespace Metasound
 
 			bool IsValid() const override { return false; }
 
-			FGuid GetNewPointID() const override { return Metasound::FrontendInvalidID; }
+			FGuid GetNewVertexID() const override { return Metasound::FrontendInvalidID; }
 			TArray<FString> GetInputVertexNames() const override { return TArray<FString>(); }
 			TArray<FString> GetOutputVertexNames() const override { return TArray<FString>(); }
 
@@ -238,15 +234,14 @@ namespace Metasound
 			// For inputs whose preferred literal type is UObject or UObjectArray, this can be used to determine the UObject corresponding to that input's datatype.
 			UClass* GetSupportedClassForInputVertex(const FString& InInputName) override { return nullptr; }
 
-			TArray<FGuid> GetDefaultIDsForInputVertex(const FString& InInputName) const override { return TArray<FGuid>(); }
-			TArray<FGuid> GetDefaultIDsForOutputVertex(const FString& InInputName) const override { return TArray<FGuid>(); }
-
-			TArray<FMetasoundFrontendVertexLiteral> GetDefaultInputFrontendLiterals(const FString& InInputName) const override { return TArray<FMetasoundFrontendVertexLiteral>(); }
+			FGuid GetVertexIDForInputVertex(const FString& InInputName) const { return Metasound::FrontendInvalidID; }
+			FGuid GetVertexIDForOutputVertex(const FString& InInputName) const { return Metasound::FrontendInvalidID; }
+			FMetasoundFrontendLiteral GetDefaultInput(const FGuid& InVertexID) const override { return FMetasoundFrontendLiteral{}; }
 
 			// These can be used to set the default value for a given input on this graph.
 			// @returns false if the input name couldn't be found, or if the literal type was incompatible with the Data Type of this input.
-			bool SetDefaultInputToFrontendLiteral(const FString& InInputName, FGuid InPointID, FName InDataTypeName, const FMetasoundFrontendLiteral& InLiteral) override { return false; }
-			bool SetDefaultInputToTypeDefaultLiteral(const FString& InInputName, FGuid InPointID, FName InDataTypeName) override { return false; }
+			bool SetDefaultInput(const FGuid& InVertexID, const FMetasoundFrontendLiteral& InLiteral) override { return false; }
+			bool SetDefaultInputToDefaultLiteralOfType(const FGuid& InVertexID) override { return false; }
 
 			const FText& GetInputDescription(const FString& InName) const override { return FText::GetEmpty(); }
 			const FText& GetOutputDescription(const FString& InName) const override { return FText::GetEmpty(); }
@@ -258,7 +253,7 @@ namespace Metasound
 
 			// This can be used to clear the current literal for a given input.
 			// @returns false if the input name couldn't be found.
-			bool ClearLiteralForInput(const FString& InInputName, FGuid InPointID) override { return false; }
+			bool ClearLiteralForInput(const FString& InInputName, FGuid InVertexID) override { return false; }
 
 			TSharedRef<INodeController> AddNode(const FNodeClassInfo& InNodeClass) override { return FInvalidNodeController::GetInvalid(); }
 			TSharedRef<INodeController> AddNode(const FNodeRegistryKey& InNodeClass) override { return FInvalidNodeController::GetInvalid(); }
@@ -313,8 +308,8 @@ namespace Metasound
 				bool IsRequiredInput(const FString& InInputName) const override { return false; }
 				bool IsRequiredOutput(const FString& InOutputName) const override { return false; }
 
-				const TArray<FMetasoundFrontendClassInput>& GetRequiredInputs() const override { const static TArray<FMetasoundFrontendClassInput> Empty; return Empty; }
-				const TArray<FMetasoundFrontendClassOutput>& GetRequiredOutputs() const override { const static TArray<FMetasoundFrontendClassOutput> Empty; return Empty; }
+				const TArray<FMetasoundFrontendClassVertex>& GetRequiredInputs() const override { const static TArray<FMetasoundFrontendClassVertex> Empty; return Empty; }
+				const TArray<FMetasoundFrontendClassVertex>& GetRequiredOutputs() const override { const static TArray<FMetasoundFrontendClassVertex> Empty; return Empty; }
 
 				TArray<FMetasoundFrontendClass> GetDependencies() const override { return TArray<FMetasoundFrontendClass>(); }
 				TArray<FMetasoundFrontendGraphClass> GetSubgraphs() const override { return TArray<FMetasoundFrontendGraphClass>(); }
diff --git a/Engine/Plugins/Runtime/Metasound/Source/MetasoundFrontend/Private/MetasoundFrontendLiteral.cpp b/Engine/Plugins/Runtime/Metasound/Source/MetasoundFrontend/Private/MetasoundFrontendLiteral.cpp
--- a/Engine/Plugins/Runtime/Metasound/Source/MetasoundFrontend/Private/MetasoundFrontendLiteral.cpp
+++ b/Engine/Plugins/Runtime/Metasound/Source/MetasoundFrontend/Private/MetasoundFrontendLiteral.cpp
@@ -103,6 +103,11 @@ bool FMetasoundFrontendLiteral::IsArray() const
 	}
 }
 
+bool FMetasoundFrontendLiteral::IsValid() const
+{
+	return (Type != EMetasoundFrontendLiteralType::Invalid);
+}
+
 void FMetasoundFrontendLiteral::Set(FDefault InValue)
 {
 	Empty();
diff --git a/Engine/Plugins/Runtime/Metasound/Source/MetasoundFrontend/Private/MetasoundFrontendStandardController.cpp b/Engine/Plugins/Runtime/Metasound/Source/MetasoundFrontend/Private/MetasoundFrontendStandardController.cpp
--- a/Engine/Plugins/Runtime/Metasound/Source/MetasoundFrontend/Private/MetasoundFrontendStandardController.cpp
+++ b/Engine/Plugins/Runtime/Metasound/Source/MetasoundFrontend/Private/MetasoundFrontendStandardController.cpp
@@ -181,39 +181,39 @@ namespace Metasound
 				};
 			}
 
-			auto FGetNodeInputFromNodeWithPoint(FGuid InPointID)
+			auto FGetNodeInputFromNodeWithVertexID(FGuid InVertexID)
 			{
 				return [=](FMetasoundFrontendNode& InNode) -> FMetasoundFrontendVertex*
 				{
-					auto IsVertexWithPoint = [=](const FMetasoundFrontendVertex& Vertex) { return Vertex.PointIDs.Contains(InPointID); };
-					return InNode.Interface.Inputs.FindByPredicate(IsVertexWithPoint);
+					auto IsVertexWithVertexID = [=](const FMetasoundFrontendVertex& Vertex) { return Vertex.VertexID == InVertexID; };
+					return InNode.Interface.Inputs.FindByPredicate(IsVertexWithVertexID);
 				};
 			}
 
-			auto FGetConstNodeInputFromNodeWithPoint(FGuid InPointID)
+			auto FGetConstNodeInputFromNodeWithVertexID(FGuid InVertexID)
 			{
 				return [=](const FMetasoundFrontendNode& InNode) -> const FMetasoundFrontendVertex*
 				{
-					auto IsVertexWithPoint = [=](const FMetasoundFrontendVertex& Vertex) { return Vertex.PointIDs.Contains(InPointID); };
-					return InNode.Interface.Inputs.FindByPredicate(IsVertexWithPoint);
+					auto IsVertexWithVertexID = [=](const FMetasoundFrontendVertex& Vertex) { return Vertex.VertexID == InVertexID; };
+					return InNode.Interface.Inputs.FindByPredicate(IsVertexWithVertexID);
 				};
 			}
 
-			auto FGetNodeOutputFromNodeWithPoint(FGuid InPointID)
+			auto FGetNodeOutputFromNodeWithVertexID(FGuid InVertexID)
 			{
 				return [=](FMetasoundFrontendNode& InNode) -> FMetasoundFrontendVertex*
 				{
-					auto IsVertexWithPoint = [=](const FMetasoundFrontendVertex& Vertex) { return Vertex.PointIDs.Contains(InPointID); };
-					return InNode.Interface.Outputs.FindByPredicate(IsVertexWithPoint);
+					auto IsVertexWithVertexID = [=](const FMetasoundFrontendVertex& Vertex) { return Vertex.VertexID == InVertexID; };
+					return InNode.Interface.Outputs.FindByPredicate(IsVertexWithVertexID);
 				};
 			}
 
-			auto FGetConstNodeOutputFromNodeWithPoint(FGuid InPointID)
+			auto FGetConstNodeOutputFromNodeWithVertexID(FGuid InVertexID)
 			{
 				return [=](const FMetasoundFrontendNode& InNode) -> const FMetasoundFrontendVertex*
 				{
-					auto IsVertexWithPoint = [=](const FMetasoundFrontendVertex& Vertex) { return Vertex.PointIDs.Contains(InPointID); };
-					return InNode.Interface.Outputs.FindByPredicate(IsVertexWithPoint);
+					auto IsVertexWithVertexID = [=](const FMetasoundFrontendVertex& Vertex) { return Vertex.VertexID == InVertexID; };
+					return InNode.Interface.Outputs.FindByPredicate(IsVertexWithVertexID);
 				};
 			}
 
@@ -408,7 +408,7 @@ namespace Metasound
 			{
 				FNodeHandle InputNode = Graph->GetNodeWithID(Edge.ToNodeID);
 
-				FInputHandle Input = InputNode->GetInputWithID(Edge.ToPointID);
+				FInputHandle Input = InputNode->GetInputWithID(Edge.ToVertexID);
 				if (Input->IsValid())
 				{
 					Inputs.Add(Input);
@@ -429,7 +429,7 @@ namespace Metasound
 			{
 				FConstNodeHandle InputNode = Graph->GetNodeWithID(Edge.ToNodeID);
 
-				FConstInputHandle Input = InputNode->GetInputWithID(Edge.ToPointID);
+				FConstInputHandle Input = InputNode->GetInputWithID(Edge.ToVertexID);
 				if (Input->IsValid())
 				{
 					Inputs.Add(Input);
@@ -477,11 +477,11 @@ namespace Metasound
 			if (GraphPtr.IsValid())
 			{
 				const FGuid NodeID = GetOwningNodeID();
-				FGuid PointID = GetID();
+				FGuid VertexID = GetID();
 
 				auto EdgeHasMatchingSource = [&](const FMetasoundFrontendEdge& Edge)
 				{
-					return (Edge.FromNodeID == NodeID) && (Edge.FromPointID == PointID);
+					return (Edge.FromNodeID == NodeID) && (Edge.FromVertexID == VertexID);
 				};
 
 				return GraphPtr->Edges.FilterByPredicate(EdgeHasMatchingSource);
@@ -684,7 +684,7 @@ namespace Metasound
 				// Create output handle from output node.
 				FGraphHandle Graph = OwningNode->GetOwningGraph();
 				FNodeHandle OutputNode = Graph->GetNodeWithID(Edge->FromNodeID);
-				return OutputNode->GetOutputWithID(Edge->FromPointID);
+				return OutputNode->GetOutputWithID(Edge->FromVertexID);
 			}
 
 			return FInvalidOutputController::GetInvalid();
@@ -697,7 +697,7 @@ namespace Metasound
 				// Create output handle from output node.
 				FConstGraphHandle Graph = OwningNode->GetOwningGraph();
 				FConstNodeHandle OutputNode = Graph->GetNodeWithID(Edge->FromNodeID);
-				return OutputNode->GetOutputWithID(Edge->FromPointID);
+				return OutputNode->GetOutputWithID(Edge->FromVertexID);
 			}
 
 			return FInvalidOutputController::GetInvalid();
@@ -749,11 +749,11 @@ namespace Metasound
 				{
 					Edge = &GraphPtr->Edges.AddDefaulted_GetRef();
 					Edge->ToNodeID = GetOwningNodeID();
-					Edge->ToPointID = GetID();
+					Edge->ToVertexID = GetID();
 				}
 
 				Edge->FromNodeID = InController.GetOwningNodeID();
-				Edge->FromPointID = InController.GetID();
+				Edge->FromVertexID = InController.GetID();
 
 				return true;
 			}
@@ -802,13 +802,13 @@ namespace Metasound
 			if (GraphPtr.IsValid())
 			{
 				FGuid FromNodeID = InController.GetOwningNodeID();
-				FGuid FromPointID = InController.GetID();
+				FGuid FromVertexID = InController.GetID();
 				FGuid ToNodeID = GetOwningNodeID();
-				FGuid ToPointID = GetID();
+				FGuid ToVertexID = GetID();
 
 				auto IsMatchingEdge = [&](const FMetasoundFrontendEdge& Edge)
 				{
-					return (Edge.FromNodeID == FromNodeID) && (Edge.FromPointID == FromPointID) && (Edge.ToNodeID == ToNodeID) && (Edge.ToPointID == ToPointID);
+					return (Edge.FromNodeID == FromNodeID) && (Edge.FromVertexID == FromVertexID) && (Edge.ToNodeID == ToNodeID) && (Edge.ToVertexID == ToVertexID);
 				};
 
 				int32 NumRemoved = GraphPtr->Edges.RemoveAllSwap(IsMatchingEdge);
@@ -823,11 +823,11 @@ namespace Metasound
 			if (GraphPtr.IsValid())
 			{
 				const FGuid NodeID = GetOwningNodeID();
-				FGuid PointID = GetID();
+				FGuid VertexID = GetID();
 
 				auto EdgeHasMatchingDestination = [&](const FMetasoundFrontendEdge& Edge)
 				{
-					return (Edge.ToNodeID == NodeID) && (Edge.ToPointID == PointID);
+					return (Edge.ToNodeID == NodeID) && (Edge.ToVertexID == VertexID);
 				};
 
 				int32 NumRemoved = GraphPtr->Edges.RemoveAllSwap(EdgeHasMatchingDestination);
@@ -843,11 +843,11 @@ namespace Metasound
 			if (GraphPtr.IsValid())
 			{
 				const FGuid NodeID = GetOwningNodeID();
-				FGuid PointID = GetID();
+				FGuid VertexID = GetID();
 
 				auto EdgeHasMatchingDestination = [&](const FMetasoundFrontendEdge& Edge)
 				{
-					return (Edge.ToNodeID == NodeID) && (Edge.ToPointID == PointID);
+					return (Edge.ToNodeID == NodeID) && (Edge.ToVertexID == VertexID);
 				};
 
 				return GraphPtr->Edges.FindByPredicate(EdgeHasMatchingDestination);
@@ -861,11 +861,11 @@ namespace Metasound
 			if (GraphPtr.IsValid())
 			{
 				const FGuid NodeID = GetOwningNodeID();
-				FGuid PointID = GetID();
+				FGuid VertexID = GetID();
 
 				auto EdgeHasMatchingDestination = [&](const FMetasoundFrontendEdge& Edge)
 				{
-					return (Edge.ToNodeID == NodeID) && (Edge.ToPointID == PointID);
+					return (Edge.ToNodeID == NodeID) && (Edge.ToVertexID == VertexID);
 				};
 
 				return GraphPtr->Edges.FindByPredicate(EdgeHasMatchingDestination);
@@ -919,15 +919,13 @@ namespace Metasound
 			return FrontendControllerIntrinsics::GetInvalidValueConstRef<FText>();
 		}
 
-		const TArray<FMetasoundFrontendVertexLiteral>& FNodeInputController::GetDefaults() const
+		const FMetasoundFrontendLiteral* FNodeInputController::GetDefaultLiteral() const 
 		{
-			if (OwningGraphClassInputPtr.IsValid())
+			if (const FMetasoundFrontendClassInput* OwningGraphClassInput = OwningGraphClassInputPtr.Get())
 			{
-				return OwningGraphClassInputPtr->Defaults;
+				return &(OwningGraphClassInput->DefaultLiteral);
 			}
-
-			static const TArray<FMetasoundFrontendVertexLiteral> InvalidDefaults;
-			return InvalidDefaults;
+			return nullptr;
 		}
 
 		const FText& FNodeInputController::GetTooltip() const
@@ -992,6 +990,11 @@ namespace Metasound
 			return FrontendControllerIntrinsics::GetInvalidValueConstRef<FText>();
 		}
 
+		const FMetasoundFrontendLiteral* FOutputNodeInputController::GetDefaultLiteral() const 
+		{
+			return nullptr;
+		}
+
 		FDocumentAccess FOutputNodeInputController::ShareAccess() 
 		{
 			FDocumentAccess Access = FBaseInputController::ShareAccess();
@@ -1165,7 +1168,7 @@ namespace Metasound
 			return FInvalidInputController::GetInvalid();
 		}
 
-		bool FBaseNodeController::RemoveInput(FGuid InPointID)
+		bool FBaseNodeController::RemoveInput(FGuid InVertexID)
 		{
 			checkNoEntry();
 			// TODO: not yet supported
@@ -1185,7 +1188,7 @@ namespace Metasound
 			return FInvalidInputController::GetInvalid();
 		}
 
-		bool FBaseNodeController::RemoveOutput(FGuid InPointID)
+		bool FBaseNodeController::RemoveOutput(FGuid InVertexID)
 		{
 			checkNoEntry();
 			// TODO: not yet supported
@@ -1200,7 +1203,7 @@ namespace Metasound
 
 			for (const FInputControllerParams& Params : GetInputControllerParams())
 			{
-				FInputHandle InputHandle = CreateInputController(Params.PointID, Params.NodeVertexPtr, Params.ClassInputPtr, ThisNode);
+				FInputHandle InputHandle = CreateInputController(Params.VertexID, Params.NodeVertexPtr, Params.ClassInputPtr, ThisNode);
 				if (InputHandle->IsValid())
 				{
 					Inputs.Add(InputHandle);
@@ -1218,7 +1221,7 @@ namespace Metasound
 
 			for (const FOutputControllerParams& Params : GetOutputControllerParams())
 			{
-				FOutputHandle OutputHandle = CreateOutputController(Params.PointID, Params.NodeVertexPtr, Params.ClassOutputPtr, ThisNode);
+				FOutputHandle OutputHandle = CreateOutputController(Params.VertexID, Params.NodeVertexPtr, Params.ClassOutputPtr, ThisNode);
 				if (OutputHandle->IsValid())
 				{
 					Outputs.Add(OutputHandle);
@@ -1244,7 +1247,7 @@ namespace Metasound
 
 			for (const FInputControllerParams& Params : GetInputControllerParams())
 			{
-				FConstInputHandle InputHandle = CreateInputController(Params.PointID, Params.NodeVertexPtr, Params.ClassInputPtr, ThisNode);
+				FConstInputHandle InputHandle = CreateInputController(Params.VertexID, Params.NodeVertexPtr, Params.ClassInputPtr, ThisNode);
 				if (InputHandle->IsValid())
 				{
 					Inputs.Add(InputHandle);
@@ -1263,7 +1266,7 @@ namespace Metasound
 
 			for (const FOutputControllerParams& Params : GetOutputControllerParams())
 			{
-				FConstOutputHandle OutputHandle = CreateOutputController(Params.PointID, Params.NodeVertexPtr, Params.ClassOutputPtr, ThisNode);
+				FConstOutputHandle OutputHandle = CreateOutputController(Params.VertexID, Params.NodeVertexPtr, Params.ClassOutputPtr, ThisNode);
 				if (OutputHandle->IsValid())
 				{
 					Outputs.Add(OutputHandle);
@@ -1281,7 +1284,7 @@ namespace Metasound
 
 			for (const FInputControllerParams& Params : GetInputControllerParamsWithVertexName(InName))
 			{
-				FInputHandle InputHandle = CreateInputController(Params.PointID, Params.NodeVertexPtr, Params.ClassInputPtr, ThisNode);
+				FInputHandle InputHandle = CreateInputController(Params.VertexID, Params.NodeVertexPtr, Params.ClassInputPtr, ThisNode);
 				if (InputHandle->IsValid())
 				{
 					Inputs.Add(InputHandle);
@@ -1300,7 +1303,7 @@ namespace Metasound
 
 			for (const FInputControllerParams& Params : GetInputControllerParamsWithVertexName(InName))
 			{
-				FInputHandle InputHandle = CreateInputController(Params.PointID, Params.NodeVertexPtr, Params.ClassInputPtr, ThisNode);
+				FInputHandle InputHandle = CreateInputController(Params.VertexID, Params.NodeVertexPtr, Params.ClassInputPtr, ThisNode);
 				if (InputHandle->IsValid())
 				{
 					Inputs.Add(InputHandle);
@@ -1318,7 +1321,7 @@ namespace Metasound
 
 			for (const FOutputControllerParams& Params : GetOutputControllerParamsWithVertexName(InName))
 			{
-				FOutputHandle OutputHandle = CreateOutputController(Params.PointID, Params.NodeVertexPtr, Params.ClassOutputPtr, ThisNode);
+				FOutputHandle OutputHandle = CreateOutputController(Params.VertexID, Params.NodeVertexPtr, Params.ClassOutputPtr, ThisNode);
 				if (OutputHandle->IsValid())
 				{
 					Outputs.Add(OutputHandle);
@@ -1337,7 +1340,7 @@ namespace Metasound
 
 			for (const FOutputControllerParams& Params : GetOutputControllerParamsWithVertexName(InName))
 			{
-				FConstOutputHandle OutputHandle = CreateOutputController(Params.PointID, Params.NodeVertexPtr, Params.ClassOutputPtr, ThisNode);
+				FConstOutputHandle OutputHandle = CreateOutputController(Params.VertexID, Params.NodeVertexPtr, Params.ClassOutputPtr, ThisNode);
 				if (OutputHandle->IsValid())
 				{
 					Outputs.Add(OutputHandle);
@@ -1347,55 +1350,55 @@ namespace Metasound
 			return Outputs;
 		}
 
-		FInputHandle FBaseNodeController::GetInputWithID(FGuid InPointID)
+		FInputHandle FBaseNodeController::GetInputWithID(FGuid InVertexID)
 		{
 			FInputControllerParams Params;
 
-			if (FindInputControllerParamsWithID(InPointID, Params))
+			if (FindInputControllerParamsWithID(InVertexID, Params))
 			{
 				FNodeHandle ThisNode = this->AsShared();
-				return CreateInputController(Params.PointID, Params.NodeVertexPtr, Params.ClassInputPtr, ThisNode);
+				return CreateInputController(Params.VertexID, Params.NodeVertexPtr, Params.ClassInputPtr, ThisNode);
 			}
 
 			return FInvalidInputController::GetInvalid();
 		}
 
-		FConstInputHandle FBaseNodeController::GetInputWithID(FGuid InPointID) const
+		FConstInputHandle FBaseNodeController::GetInputWithID(FGuid InVertexID) const
 		{
 			FInputControllerParams Params;
 
-			if (FindInputControllerParamsWithID(InPointID, Params))
+			if (FindInputControllerParamsWithID(InVertexID, Params))
 			{
 				// See early use of ConstCastSharedRef in this class for discussion.
 				FNodeHandle ThisNode = ConstCastSharedRef<INodeController>(this->AsShared());
-				return CreateInputController(Params.PointID, Params.NodeVertexPtr, Params.ClassInputPtr, ThisNode);
+				return CreateInputController(Params.VertexID, Params.NodeVertexPtr, Params.ClassInputPtr, ThisNode);
 			}
 
 			return FInvalidInputController::GetInvalid();
 		}
 
-		FOutputHandle FBaseNodeController::GetOutputWithID(FGuid InPointID)
+		FOutputHandle FBaseNodeController::GetOutputWithID(FGuid InVertexID)
 		{
 			FOutputControllerParams Params;
 
-			if (FindOutputControllerParamsWithID(InPointID, Params))
+			if (FindOutputControllerParamsWithID(InVertexID, Params))
 			{
 				FNodeHandle ThisNode = this->AsShared();
-				return CreateOutputController(Params.PointID, Params.NodeVertexPtr, Params.ClassOutputPtr, ThisNode);
+				return CreateOutputController(Params.VertexID, Params.NodeVertexPtr, Params.ClassOutputPtr, ThisNode);
 			}
 
 			return FInvalidOutputController::GetInvalid();
 		}
 
-		FConstOutputHandle FBaseNodeController::GetOutputWithID(FGuid InPointID) const
+		FConstOutputHandle FBaseNodeController::GetOutputWithID(FGuid InVertexID) const
 		{
 			FOutputControllerParams Params;
 
-			if (FindOutputControllerParamsWithID(InPointID, Params))
+			if (FindOutputControllerParamsWithID(InVertexID, Params))
 			{
 				// See early use of ConstCastSharedRef in this class for discussion.
 				FNodeHandle ThisNode = ConstCastSharedRef<INodeController>(this->AsShared());
-				return CreateOutputController(Params.PointID, Params.NodeVertexPtr, Params.ClassOutputPtr, ThisNode);
+				return CreateOutputController(Params.VertexID, Params.NodeVertexPtr, Params.ClassOutputPtr, ThisNode);
 			}
 
 			return FInvalidOutputController::GetInvalid();
@@ -1417,10 +1420,7 @@ namespace Metasound
 
 					FConstClassInputAccessPtr ClassInputPtr = ClassPtr.GetMemberAccessPtr<const FMetasoundFrontendClassInput>(FGetClassInputFromClassWithName(NodeInputVertex.Name));
 
-					for (const FGuid& PointID : NodeInputVertex.PointIDs)
-					{
-						Inputs.Add({PointID, NodeVertexPtr, ClassInputPtr});
-					}
+					Inputs.Add({NodeInputVertex.VertexID, NodeVertexPtr, ClassInputPtr});
 				}
 			}
 
@@ -1440,10 +1440,7 @@ namespace Metasound
 				FConstVertexAccessPtr NodeVertexPtr = NodePtr.GetMemberAccessPtr<const FMetasoundFrontendVertex>(FGetConstNodeOutputFromNodeWithName(VertexName));
 				FConstClassOutputAccessPtr ClassOutputPtr = ClassPtr.GetMemberAccessPtr<const FMetasoundFrontendClassOutput>(FGetClassOutputFromClassWithName(VertexName));
 
-				for (const FGuid& PointID : NodeOutputVertex.PointIDs)
-				{
-					Outputs.Add({PointID, NodeVertexPtr, ClassOutputPtr});
-				}
+				Outputs.Add({NodeOutputVertex.VertexID, NodeVertexPtr, ClassOutputPtr});
 			}
 
 			return Outputs;
@@ -1461,10 +1458,7 @@ namespace Metasound
 			{
 				FConstClassInputAccessPtr ClassInputPtr = ClassPtr.GetMemberAccessPtr<const FMetasoundFrontendClassInput>(FGetClassInputFromClassWithName(InName));
 
-				for (const FGuid& PointID : Vertex->PointIDs)
-				{
-					Inputs.Add({PointID, NodeVertexPtr, ClassInputPtr});
-				}
+				Inputs.Add({Vertex->VertexID, NodeVertexPtr, ClassInputPtr});
 			}
 
 			return Inputs;
@@ -1482,43 +1476,40 @@ namespace Metasound
 			{
 				FConstClassOutputAccessPtr ClassOutputPtr = ClassPtr.GetMemberAccessPtr<const FMetasoundFrontendClassOutput>(FGetClassOutputFromClassWithName(InName));
 
-				for (const FGuid& PointID : Vertex->PointIDs)
-				{
-					Outputs.Add({PointID, NodeVertexPtr, ClassOutputPtr});
-				}
+				Outputs.Add({Vertex->VertexID, NodeVertexPtr, ClassOutputPtr});
 			}
 
 			return Outputs;
 		}
 
-		bool FBaseNodeController::FindInputControllerParamsWithID(FGuid InPointID, FInputControllerParams& OutParams) const
+		bool FBaseNodeController::FindInputControllerParamsWithID(FGuid InVertexID, FInputControllerParams& OutParams) const
 		{
 			using namespace FrontendControllerIntrinsics;
 
-			FConstVertexAccessPtr NodeVertexPtr = NodePtr.GetMemberAccessPtr<const FMetasoundFrontendVertex>(FGetConstNodeInputFromNodeWithPoint(InPointID));
+			FConstVertexAccessPtr NodeVertexPtr = NodePtr.GetMemberAccessPtr<const FMetasoundFrontendVertex>(FGetConstNodeInputFromNodeWithVertexID(InVertexID));
 
 			if (const FMetasoundFrontendVertex* Vertex = NodeVertexPtr.Get())
 			{
 				FConstClassInputAccessPtr ClassInputPtr = ClassPtr.GetMemberAccessPtr<const FMetasoundFrontendClassInput>(FGetClassInputFromClassWithName(Vertex->Name));
 
-				OutParams = FInputControllerParams{InPointID, NodeVertexPtr, ClassInputPtr};
+				OutParams = FInputControllerParams{InVertexID, NodeVertexPtr, ClassInputPtr};
 				return true;
 			}
 
 			return false;
 		}
 
-		bool FBaseNodeController::FindOutputControllerParamsWithID(FGuid InPointID, FOutputControllerParams& OutParams) const
+		bool FBaseNodeController::FindOutputControllerParamsWithID(FGuid InVertexID, FOutputControllerParams& OutParams) const
 		{
 			using namespace FrontendControllerIntrinsics;
 
-			FConstVertexAccessPtr NodeVertexPtr = NodePtr.GetMemberAccessPtr<const FMetasoundFrontendVertex>(FGetConstNodeOutputFromNodeWithPoint(InPointID));
+			FConstVertexAccessPtr NodeVertexPtr = NodePtr.GetMemberAccessPtr<const FMetasoundFrontendVertex>(FGetConstNodeOutputFromNodeWithVertexID(InVertexID));
 
 			if (const FMetasoundFrontendVertex* Vertex = NodeVertexPtr.Get())
 			{
 				FConstClassOutputAccessPtr ClassOutputPtr = ClassPtr.GetMemberAccessPtr<const FMetasoundFrontendClassOutput>(FGetClassOutputFromClassWithName(Vertex->Name));
 
-				OutParams = FOutputControllerParams{InPointID, NodeVertexPtr, ClassOutputPtr};
+				OutParams = FOutputControllerParams{InVertexID, NodeVertexPtr, ClassOutputPtr};
 				return true;
 			}
 
@@ -1667,14 +1658,14 @@ namespace Metasound
 			return FBaseNodeController::IsValid() && GraphPtr.IsValid();
 		}
 
-		FInputHandle FNodeController::CreateInputController(FGuid InPointID, FConstVertexAccessPtr InNodeVertexPtr, FConstClassInputAccessPtr InClassInputPtr, FNodeHandle InOwningNode) const
+		FInputHandle FNodeController::CreateInputController(FGuid InVertexID, FConstVertexAccessPtr InNodeVertexPtr, FConstClassInputAccessPtr InClassInputPtr, FNodeHandle InOwningNode) const
 		{
-			return MakeShared<FNodeInputController>(FNodeInputController::FInitParams{InPointID, InNodeVertexPtr, InClassInputPtr, GraphPtr, InOwningNode});
+			return MakeShared<FNodeInputController>(FNodeInputController::FInitParams{InVertexID, InNodeVertexPtr, InClassInputPtr, GraphPtr, InOwningNode});
 		}
 
-		FOutputHandle FNodeController::CreateOutputController(FGuid InPointID, FConstVertexAccessPtr InNodeVertexPtr, FConstClassOutputAccessPtr InClassOutputPtr, FNodeHandle InOwningNode) const
+		FOutputHandle FNodeController::CreateOutputController(FGuid InVertexID, FConstVertexAccessPtr InNodeVertexPtr, FConstClassOutputAccessPtr InClassOutputPtr, FNodeHandle InOwningNode) const
 		{
-			return MakeShared<FNodeOutputController>(FNodeOutputController::FInitParams{InPointID, InNodeVertexPtr, InClassOutputPtr, GraphPtr, InOwningNode});
+			return MakeShared<FNodeOutputController>(FNodeOutputController::FInitParams{InVertexID, InNodeVertexPtr, InClassOutputPtr, GraphPtr, InOwningNode});
 		}
 
 		FDocumentAccess FNodeController::ShareAccess() 
@@ -1760,12 +1751,12 @@ namespace Metasound
 			return OwningGraph->IsValid() && OwningGraphClassOutputPtr.IsValid() && GraphPtr.IsValid();
 		}
 
-		FInputHandle FOutputNodeController::CreateInputController(FGuid InPointID, FConstVertexAccessPtr InNodeVertexPtr, FConstClassInputAccessPtr InClassInputPtr, FNodeHandle InOwningNode) const
+		FInputHandle FOutputNodeController::CreateInputController(FGuid InVertexID, FConstVertexAccessPtr InNodeVertexPtr, FConstClassInputAccessPtr InClassInputPtr, FNodeHandle InOwningNode) const
 		{
-			return MakeShared<FOutputNodeInputController>(FOutputNodeInputController::FInitParams{InPointID, InNodeVertexPtr, OwningGraphClassOutputPtr, GraphPtr, InOwningNode});
+			return MakeShared<FOutputNodeInputController>(FOutputNodeInputController::FInitParams{InVertexID, InNodeVertexPtr, OwningGraphClassOutputPtr, GraphPtr, InOwningNode});
 		}
 
-		FOutputHandle FOutputNodeController::CreateOutputController(FGuid InPointID, FConstVertexAccessPtr InNodeVertexPtr, FConstClassOutputAccessPtr InClassOutputPtr, FNodeHandle InOwningNode) const
+		FOutputHandle FOutputNodeController::CreateOutputController(FGuid InVertexID, FConstVertexAccessPtr InNodeVertexPtr, FConstClassOutputAccessPtr InClassOutputPtr, FNodeHandle InOwningNode) const
 		{
 			return FInvalidOutputController::GetInvalid();
 		}
@@ -1855,14 +1846,14 @@ namespace Metasound
 			return OwningGraph->IsValid() && OwningGraphClassInputPtr.IsValid() && GraphPtr.IsValid();
 		}
 
-		FInputHandle FInputNodeController::CreateInputController(FGuid InPointID, FConstVertexAccessPtr InNodeVertexPtr, FConstClassInputAccessPtr InClassInputPtr, FNodeHandle InOwningNode) const
+		FInputHandle FInputNodeController::CreateInputController(FGuid InVertexID, FConstVertexAccessPtr InNodeVertexPtr, FConstClassInputAccessPtr InClassInputPtr, FNodeHandle InOwningNode) const
 		{
 			return FInvalidInputController::GetInvalid();
 		}
 
-		FOutputHandle FInputNodeController::CreateOutputController(FGuid InPointID, FConstVertexAccessPtr InNodeVertexPtr, FConstClassOutputAccessPtr InClassOutputPtr, FNodeHandle InOwningNode) const
+		FOutputHandle FInputNodeController::CreateOutputController(FGuid InVertexID, FConstVertexAccessPtr InNodeVertexPtr, FConstClassOutputAccessPtr InClassOutputPtr, FNodeHandle InOwningNode) const
 		{
-			return MakeShared<FInputNodeOutputController>(FInputNodeOutputController::FInitParams{InPointID, InNodeVertexPtr, OwningGraphClassInputPtr, GraphPtr, InOwningNode});
+			return MakeShared<FInputNodeOutputController>(FInputNodeOutputController::FInitParams{InVertexID, InNodeVertexPtr, OwningGraphClassInputPtr, GraphPtr, InOwningNode});
 		}
 
 		FDocumentAccess FInputNodeController::ShareAccess() 
@@ -1943,7 +1934,7 @@ namespace Metasound
 			return Metasound::FrontendInvalidID;
 		}
 
-		FGuid FGraphController::GetNewPointID() const
+		FGuid FGraphController::GetNewVertexID() const
 		{
 			return FGuid::NewGuid();
 		}
@@ -1990,22 +1981,22 @@ namespace Metasound
 			return GraphClassPtr.GetMemberAccessPtr<FMetasoundFrontendClassOutput>(FGetClassOutputFromClassWithName(InName));
 		}
 
-		TArray<FGuid> FGraphController::GetDefaultIDsForInputVertex(const FString& InInputName) const
+		FGuid FGraphController::GetVertexIDForInputVertex(const FString& InInputName) const
 		{
 			if (const FMetasoundFrontendClassInput* Input = FindClassInputWithName(InInputName).Get())
 			{
-				return Input->PointIDs;
+				return Input->VertexID;
 			}
-			return TArray<FGuid>();
+			return Metasound::FrontendInvalidID;
 		}
 
-		TArray<FGuid> FGraphController::GetDefaultIDsForOutputVertex(const FString& InOutputName) const
+		FGuid FGraphController::GetVertexIDForOutputVertex(const FString& InOutputName) const
 		{
 			if (const FMetasoundFrontendClassOutput* Output = FindClassOutputWithName(InOutputName).Get())
 			{
-				return Output->PointIDs;
+				return Output->VertexID;
 			}
-			return TArray<FGuid>();
+			return Metasound::FrontendInvalidID;
 		}
 
 		TArray<FNodeHandle> FGraphController::GetNodes()
@@ -2310,54 +2301,45 @@ namespace Metasound
 			return nullptr;
 		}
 
-		TArray<FMetasoundFrontendVertexLiteral> FGraphController::GetDefaultInputFrontendLiterals(const FString& InInputName) const
+		FMetasoundFrontendLiteral FGraphController::GetDefaultInput(const FGuid& InVertexID) const
 		{
-			if (const FMetasoundFrontendClassInput* Desc = FindInputDescriptionWithName(InInputName))
+			if (const FMetasoundFrontendClassInput* Desc = FindInputDescriptionWithVertexID(InVertexID))
 			{
-				return Desc->Defaults;
+				return Desc->DefaultLiteral;
 			}
 
-			return TArray<FMetasoundFrontendVertexLiteral>();
+			return FMetasoundFrontendLiteral{};
 		}
 
-		bool FGraphController::SetDefaultInputToFrontendLiteral(const FString& InInputName, FGuid InPointID, FName InDataTypeName, const FMetasoundFrontendLiteral& InLiteral)
+		bool FGraphController::SetDefaultInput(const FGuid& InVertexID, const FMetasoundFrontendLiteral& InLiteral)
 		{
-			if (FMetasoundFrontendClassInput* Desc = FindInputDescriptionWithName(InInputName))
+			if (FMetasoundFrontendClassInput* Desc = FindInputDescriptionWithVertexID(InVertexID))
 			{
-				auto IsLiteralWithSamePointID = [&](const FMetasoundFrontendVertexLiteral& InVertexLiteral) 
-				{ 
-					return InVertexLiteral.PointID == InPointID; 
-				};
-
-				FMetasoundFrontendVertexLiteral* VertexLiteral = Desc->Defaults.FindByPredicate(IsLiteralWithSamePointID);
-				if (nullptr == VertexLiteral)
-				{
-					VertexLiteral = &Desc->Defaults.AddDefaulted_GetRef();
-					VertexLiteral->PointID = InPointID;
-				}
-
 				if (ensure(DoesDataTypeSupportLiteralType(Desc->TypeName, InLiteral.GetType())))
 				{
-					VertexLiteral->Value = InLiteral;
+					Desc->DefaultLiteral = InLiteral;
 					return true;
 				}
 				else
 				{
-					VertexLiteral->Value.Clear();
+					SetDefaultInputToDefaultLiteralOfType(InVertexID);
 				}
 			}
 
 			return false;
 		}
 
-		bool FGraphController::SetDefaultInputToTypeDefaultLiteral(const FString& InInputName, FGuid InPointID, FName InDataTypeName)
+		bool FGraphController::SetDefaultInputToDefaultLiteralOfType(const FGuid& InVertexID)
 		{
-			FMetasoundFrontendLiteral FrontendLiteral;
+			if (FMetasoundFrontendClassInput* Desc = FindInputDescriptionWithVertexID(InVertexID))
+			{
+				Metasound::FLiteral Literal = Frontend::GetDefaultParamForDataType(Desc->TypeName);
+				Desc->DefaultLiteral.SetFromLiteral(Literal);
 
-			Metasound::FLiteral Literal = Frontend::GetDefaultParamForDataType(InDataTypeName);
-			FrontendLiteral.SetFromLiteral(Literal);
+				return Desc->DefaultLiteral.IsValid();
+			}
 
-			return SetDefaultInputToFrontendLiteral(InInputName, InPointID, InDataTypeName, FrontendLiteral);
+			return false;
 		}
 
 		const FText& FGraphController::GetInputDescription(const FString& InName) const
@@ -2429,20 +2411,11 @@ namespace Metasound
 
 		// This can be used to clear the current literal for a given input.
 		// @returns false if the input name couldn't be found.
-		bool FGraphController::ClearLiteralForInput(const FString& InInputName, FGuid InPointID)
+		bool FGraphController::ClearLiteralForInput(const FString& InInputName, FGuid InVertexID)
 		{
 			if (FMetasoundFrontendClassInput* Desc = FindInputDescriptionWithName(InInputName))
 			{
-				auto IsLiteralWithSamePointID = [&](const FMetasoundFrontendVertexLiteral& InVertexLiteral) 
-				{ 
-					return InVertexLiteral.PointID == InPointID; 
-				};
-
-				if (FMetasoundFrontendVertexLiteral* VertexLiteral = Desc->Defaults.FindByPredicate(IsLiteralWithSamePointID))
-				{
-					VertexLiteral->Value.Clear();
-					return true;
-				}
+				Desc->DefaultLiteral.Clear();
 			}
 
 			return false;
@@ -3040,6 +3013,42 @@ namespace Metasound
 			return nullptr;
 		}
 
+		FMetasoundFrontendClassInput* FGraphController::FindInputDescriptionWithVertexID(const FGuid& InVertexID)
+		{
+			if (GraphClassPtr.IsValid())
+			{
+				return GraphClassPtr->Interface.Inputs.FindByPredicate([&](const FMetasoundFrontendClassInput& Desc) { return Desc.VertexID == InVertexID; });
+			}
+			return nullptr;
+		}
+
+		const FMetasoundFrontendClassInput* FGraphController::FindInputDescriptionWithVertexID(const FGuid& InVertexID) const
+		{
+			if (GraphClassPtr.IsValid())
+			{
+				return GraphClassPtr->Interface.Inputs.FindByPredicate([&](const FMetasoundFrontendClassInput& Desc) { return Desc.VertexID == InVertexID; });
+			}
+			return nullptr;
+		}
+
+		FMetasoundFrontendClassOutput* FGraphController::FindOutputDescriptionWithVertexID(const FGuid& InVertexID)
+		{
+			if (GraphClassPtr.IsValid())
+			{
+				return GraphClassPtr->Interface.Outputs.FindByPredicate([&](const FMetasoundFrontendClassOutput& Desc) { return Desc.VertexID == InVertexID; });
+			}
+			return nullptr;
+		}
+
+		const FMetasoundFrontendClassOutput* FGraphController::FindOutputDescriptionWithVertexID(const FGuid& InVertexID) const
+		{
+			if (GraphClassPtr.IsValid())
+			{
+				return GraphClassPtr->Interface.Outputs.FindByPredicate([&](const FMetasoundFrontendClassOutput& Desc) { return Desc.VertexID == InVertexID; });
+			}
+			return nullptr;
+		}
+
 		FClassInputAccessPtr FGraphController::FindInputDescriptionWithNodeID(FGuid InNodeID)
 		{
 			using namespace FrontendControllerIntrinsics;
@@ -3122,25 +3131,25 @@ namespace Metasound
 			return false;
 		}
 
-		const TArray<FMetasoundFrontendClassInput>& FDocumentController::GetRequiredInputs() const
+		const TArray<FMetasoundFrontendClassVertex>& FDocumentController::GetRequiredInputs() const
 		{
 			if (DocumentPtr.IsValid())
 			{
 				return DocumentPtr->Archetype.Interface.Inputs;
 			}
 
-			static const TArray<FMetasoundFrontendClassInput> Empty;
+			static const TArray<FMetasoundFrontendClassVertex> Empty;
 			return Empty;
 		}
 
-		const TArray<FMetasoundFrontendClassOutput>& FDocumentController::GetRequiredOutputs() const
+		const TArray<FMetasoundFrontendClassVertex>& FDocumentController::GetRequiredOutputs() const
 		{
 			if (DocumentPtr.IsValid())
 			{
 				return DocumentPtr->Archetype.Interface.Outputs;
 			}
 
-			static const TArray<FMetasoundFrontendClassOutput> Empty;
+			static const TArray<FMetasoundFrontendClassVertex> Empty;
 			return Empty;
 		}
 
diff --git a/Engine/Plugins/Runtime/Metasound/Source/MetasoundFrontend/Private/MetasoundFrontendStandardController.h b/Engine/Plugins/Runtime/Metasound/Source/MetasoundFrontend/Private/MetasoundFrontendStandardController.h
--- a/Engine/Plugins/Runtime/Metasound/Source/MetasoundFrontend/Private/MetasoundFrontendStandardController.h
+++ b/Engine/Plugins/Runtime/Metasound/Source/MetasoundFrontend/Private/MetasoundFrontendStandardController.h
@@ -180,11 +180,6 @@ namespace Metasound
 			const FName& GetDataType() const override;
 			const FString& GetName() const override;
 
-			const TArray<FMetasoundFrontendVertexLiteral>& GetDefaults() const override
-			{
-				static const TArray<FMetasoundFrontendVertexLiteral> InvalidDefaults;
-				return InvalidDefaults;
-			}
 
 			// Owning node info
 			FGuid GetOwningNodeID() const override;
@@ -240,7 +235,7 @@ namespace Metasound
 
 			bool IsValid() const override;
 
-			const TArray<FMetasoundFrontendVertexLiteral>& GetDefaults() const;
+			const FMetasoundFrontendLiteral* GetDefaultLiteral() const override;
 
 			// Input metadata
 			const FText& GetDisplayName() const override;
@@ -282,6 +277,7 @@ namespace Metasound
 			// Input metadata
 			const FText& GetDisplayName() const override;
 			const FText& GetTooltip() const override;
+			const FMetasoundFrontendLiteral* GetDefaultLiteral() const override;
 
 		protected:
 
@@ -334,11 +330,11 @@ namespace Metasound
 
 			bool CanAddInput(const FString& InVertexName) const override;
 			FInputHandle AddInput(const FString& InVertexName, const FMetasoundFrontendLiteral* InDefault) override;
-			bool RemoveInput(FGuid InPointID) override;
+			bool RemoveInput(FGuid InVertexID) override;
 
 			bool CanAddOutput(const FString& InVertexName) const override;
 			FInputHandle AddOutput(const FString& InVertexName, const FMetasoundFrontendLiteral* InDefault) override;
-			bool RemoveOutput(FGuid InPointID) override;
+			bool RemoveOutput(FGuid InVertexID) override;
 
 			const FMetasoundFrontendClassName& GetClassName() const override;
 
@@ -364,25 +360,25 @@ namespace Metasound
 			 *
 			 * If the input does not exist, an invalid handle is returned. 
 			 */
-			FInputHandle GetInputWithID(FGuid InPointID) override;
+			FInputHandle GetInputWithID(FGuid InVertexID) override;
 
 			/** Returns an input with the given name. 
 			 *
 			 * If the input does not exist, an invalid handle is returned. 
 			 */
-			FConstInputHandle GetInputWithID(FGuid InPointID) const override;
+			FConstInputHandle GetInputWithID(FGuid InVertexID) const override;
 
 			/** Returns an output with the given name. 
 			 *
 			 * If the output does not exist, an invalid handle is returned. 
 			 */
-			FOutputHandle GetOutputWithID(FGuid InPointID) override;
+			FOutputHandle GetOutputWithID(FGuid InVertexID) override;
 
 			/** Returns an output with the given name. 
 			 *
 			 * If the output does not exist, an invalid handle is returned. 
 			 */
-			FConstOutputHandle GetOutputWithID(FGuid InPointID) const override;
+			FConstOutputHandle GetOutputWithID(FGuid InVertexID) const override;
 
 			FGraphHandle AsGraph() override;
 			FConstGraphHandle AsGraph() const override;
@@ -400,14 +396,14 @@ namespace Metasound
 
 			struct FInputControllerParams
 			{
-				FGuid PointID;
+				FGuid VertexID;
 				FConstVertexAccessPtr NodeVertexPtr;
 				FConstClassInputAccessPtr ClassInputPtr;
 			};
 
 			struct FOutputControllerParams
 			{
-				FGuid PointID;
+				FGuid VertexID;
 				FConstVertexAccessPtr NodeVertexPtr;
 				FConstClassOutputAccessPtr ClassOutputPtr;
 			};
@@ -418,8 +414,8 @@ namespace Metasound
 			TArray<FInputControllerParams> GetInputControllerParamsWithVertexName(const FString& InName) const;
 			TArray<FOutputControllerParams> GetOutputControllerParamsWithVertexName(const FString& InName) const;
 
-			bool FindInputControllerParamsWithID(FGuid InPointID, FInputControllerParams& OutParams) const;
-			bool FindOutputControllerParamsWithID(FGuid InPointID, FOutputControllerParams& OutParams) const;
+			bool FindInputControllerParamsWithID(FGuid InVertexID, FInputControllerParams& OutParams) const;
+			bool FindOutputControllerParamsWithID(FGuid InVertexID, FOutputControllerParams& OutParams) const;
 
 			const FMetasoundFrontendClassInput* FindClassInputWithName(const FString& InName) const;
 			const FMetasoundFrontendClassOutput* FindClassOutputWithName(const FString& InName) const;
@@ -429,9 +425,9 @@ namespace Metasound
 
 
 
-			virtual FInputHandle CreateInputController(FGuid InPointID, FConstVertexAccessPtr InNodeVertexPtr, FConstClassInputAccessPtr InClassInputPtr, FNodeHandle InOwningNode) const = 0;
+			virtual FInputHandle CreateInputController(FGuid InVertexID, FConstVertexAccessPtr InNodeVertexPtr, FConstClassInputAccessPtr InClassInputPtr, FNodeHandle InOwningNode) const = 0;
 
-			virtual FOutputHandle CreateOutputController(FGuid InPointID, FConstVertexAccessPtr InNodeVertexPtr, FConstClassOutputAccessPtr InClassOutputPtr, FNodeHandle InOwningNode) const = 0;
+			virtual FOutputHandle CreateOutputController(FGuid InVertexID, FConstVertexAccessPtr InNodeVertexPtr, FConstClassOutputAccessPtr InClassOutputPtr, FNodeHandle InOwningNode) const = 0;
 		};
 
 		/** FNodeController represents a external or subgraph node. */
@@ -480,9 +476,9 @@ namespace Metasound
 
 		private:
 
-			FInputHandle CreateInputController(FGuid InPointID, FConstVertexAccessPtr InNodeVertexPtr, FConstClassInputAccessPtr InClassInputPtr, FNodeHandle InOwningNode) const override;
+			FInputHandle CreateInputController(FGuid InVertexID, FConstVertexAccessPtr InNodeVertexPtr, FConstClassInputAccessPtr InClassInputPtr, FNodeHandle InOwningNode) const override;
 
-			FOutputHandle CreateOutputController(FGuid InPointID, FConstVertexAccessPtr InNodeVertexPtr, FConstClassOutputAccessPtr InClassOutputPtr, FNodeHandle InOwningNode) const override;
+			FOutputHandle CreateOutputController(FGuid InVertexID, FConstVertexAccessPtr InNodeVertexPtr, FConstClassOutputAccessPtr InClassOutputPtr, FNodeHandle InOwningNode) const override;
 
 			FGraphAccessPtr GraphPtr;
 
@@ -534,9 +530,9 @@ namespace Metasound
 
 
 		private:
-			FInputHandle CreateInputController(FGuid InPointID, FConstVertexAccessPtr InNodeVertexPtr, FConstClassInputAccessPtr InClassInputPtr, FNodeHandle InOwningNode) const override;
+			FInputHandle CreateInputController(FGuid InVertexID, FConstVertexAccessPtr InNodeVertexPtr, FConstClassInputAccessPtr InClassInputPtr, FNodeHandle InOwningNode) const override;
 
-			FOutputHandle CreateOutputController(FGuid InPointID, FConstVertexAccessPtr InNodeVertexPtr, FConstClassOutputAccessPtr InClassOutputPtr, FNodeHandle InOwningNode) const override;
+			FOutputHandle CreateOutputController(FGuid InVertexID, FConstVertexAccessPtr InNodeVertexPtr, FConstClassOutputAccessPtr InClassOutputPtr, FNodeHandle InOwningNode) const override;
 
 			FGraphAccessPtr GraphPtr;
 			FConstClassOutputAccessPtr OwningGraphClassOutputPtr; 
@@ -589,9 +585,9 @@ namespace Metasound
 
 
 		private:
-			FInputHandle CreateInputController(FGuid InPointID, FConstVertexAccessPtr InNodeVertexPtr, FConstClassInputAccessPtr InClassInputPtr, FNodeHandle InOwningNode) const override;
+			FInputHandle CreateInputController(FGuid InVertexID, FConstVertexAccessPtr InNodeVertexPtr, FConstClassInputAccessPtr InClassInputPtr, FNodeHandle InOwningNode) const override;
 
-			FOutputHandle CreateOutputController(FGuid InPointID, FConstVertexAccessPtr InNodeVertexPtr, FConstClassOutputAccessPtr InClassOutputPtr, FNodeHandle InOwningNode) const override;
+			FOutputHandle CreateOutputController(FGuid InVertexID, FConstVertexAccessPtr InNodeVertexPtr, FConstClassOutputAccessPtr InClassOutputPtr, FNodeHandle InOwningNode) const override;
 
 			FConstClassInputAccessPtr OwningGraphClassInputPtr;
 			FGraphAccessPtr GraphPtr;
@@ -637,14 +633,14 @@ namespace Metasound
 
 			FGuid GetClassID() const override;
 
-			FGuid GetNewPointID() const override;
+			FGuid GetNewVertexID() const override;
 
 			TArray<FString> GetInputVertexNames() const override;
 			TArray<FString> GetOutputVertexNames() const override;
 			FConstClassInputAccessPtr FindClassInputWithName(const FString& InName) const override;
 			FConstClassOutputAccessPtr FindClassOutputWithName(const FString& InName) const override;
-			TArray<FGuid> GetDefaultIDsForInputVertex(const FString& InInputName) const override;
-			TArray<FGuid> GetDefaultIDsForOutputVertex(const FString& InOutputName) const override;
+			FGuid GetVertexIDForInputVertex(const FString& InInputName) const override;
+			FGuid GetVertexIDForOutputVertex(const FString& InOutputName) const override;
 
 			TArray<FNodeHandle> GetNodes() override;
 			TArray<FConstNodeHandle> GetConstNodes() const override;
@@ -680,12 +676,12 @@ namespace Metasound
 			// For inputs whose preferred literal type is UObject or UObjectArray, this can be used to determine the UObject corresponding to that input's datatype.
 			UClass* GetSupportedClassForInputVertex(const FString& InInputName) override;
 
-			TArray<FMetasoundFrontendVertexLiteral> GetDefaultInputFrontendLiterals (const FString& InInputName) const override;
+			FMetasoundFrontendLiteral GetDefaultInput(const FGuid& InVertexID) const override;
 
 			// These can be used to set the default value for a given input on this graph.
 			// @returns false if the input name couldn't be found, or if the literal type was incompatible with the Data Type of this input.
-			bool SetDefaultInputToFrontendLiteral(const FString& InInputName, FGuid InPointID, FName InDataTypeName, const FMetasoundFrontendLiteral& InLiteral) override;
-			bool SetDefaultInputToTypeDefaultLiteral(const FString& InInputName, FGuid InPointID, FName InDataTypeName) override;
+			bool SetDefaultInput(const FGuid& InVertexID, const FMetasoundFrontendLiteral& InLiteral) override;
+			bool SetDefaultInputToDefaultLiteralOfType(const FGuid& InVertexID) override;
 
 			// Get the description for the input with the given name.
 			const FText& GetInputDescription(const FString& InName) const override;
@@ -707,7 +703,7 @@ namespace Metasound
 
 			// This can be used to clear the current literal for a given input.
 			// @returns false if the input name couldn't be found.
-			bool ClearLiteralForInput(const FString& InInputName, FGuid InPointID) override;
+			bool ClearLiteralForInput(const FString& InInputName, FGuid InVertexID) override;
 
 			FNodeHandle AddNode(const FNodeClassInfo& InNodeClass) override;
 			FNodeHandle AddNode(const FNodeRegistryKey& InNodeClass) override;
@@ -786,9 +782,15 @@ namespace Metasound
 			FMetasoundFrontendClassInput* FindInputDescriptionWithName(const FString& InName);
 			const FMetasoundFrontendClassInput* FindInputDescriptionWithName(const FString& InName) const;
 
+			FMetasoundFrontendClassInput* FindInputDescriptionWithVertexID(const FGuid& InVertexID);
+			const FMetasoundFrontendClassInput* FindInputDescriptionWithVertexID(const FGuid& InVertexID) const;
+
 			FMetasoundFrontendClassOutput* FindOutputDescriptionWithName(const FString& InName);
 			const FMetasoundFrontendClassOutput* FindOutputDescriptionWithName(const FString& InName) const;
 
+			FMetasoundFrontendClassOutput* FindOutputDescriptionWithVertexID(const FGuid& InVertexID);
+			const FMetasoundFrontendClassOutput* FindOutputDescriptionWithVertexID(const FGuid& InVertexID) const;
+
 			FClassInputAccessPtr FindInputDescriptionWithNodeID(FGuid InNodeID);
 			FConstClassInputAccessPtr FindInputDescriptionWithNodeID(FGuid InNodeID) const;
 
@@ -832,8 +834,8 @@ namespace Metasound
 			bool IsRequiredInput(const FString& InInputName) const override;
 			bool IsRequiredOutput(const FString& InOutputName) const override;
 
-			const TArray<FMetasoundFrontendClassInput>& GetRequiredInputs() const override;
-			const TArray<FMetasoundFrontendClassOutput>& GetRequiredOutputs() const override;
+			const TArray<FMetasoundFrontendClassVertex>& GetRequiredInputs() const override;
+			const TArray<FMetasoundFrontendClassVertex>& GetRequiredOutputs() const override;
 
 			TArray<FMetasoundFrontendClass> GetDependencies() const override;
 			TArray<FMetasoundFrontendGraphClass> GetSubgraphs() const override;
diff --git a/Engine/Plugins/Runtime/Metasound/Source/MetasoundFrontend/Private/MetasoundFrontendTransform.cpp b/Engine/Plugins/Runtime/Metasound/Source/MetasoundFrontend/Private/MetasoundFrontendTransform.cpp
--- a/Engine/Plugins/Runtime/Metasound/Source/MetasoundFrontend/Private/MetasoundFrontendTransform.cpp
+++ b/Engine/Plugins/Runtime/Metasound/Source/MetasoundFrontend/Private/MetasoundFrontendTransform.cpp
@@ -16,11 +16,11 @@ namespace Metasound
 			bool bDidEdit = false;
 			FGraphHandle Graph = InDocument->GetRootGraph();
 
-			TArray<FMetasoundFrontendClassInput> RequiredInputs = InDocument->GetRequiredInputs();
-			TArray<FMetasoundFrontendClassOutput> RequiredOutputs = InDocument->GetRequiredOutputs();
+			TArray<FMetasoundFrontendClassVertex> RequiredInputs = InDocument->GetRequiredInputs();
+			TArray<FMetasoundFrontendClassVertex> RequiredOutputs = InDocument->GetRequiredOutputs();
 
 			// Go through each input and add or swap if something is missing.
-			for (const FMetasoundFrontendClassInput& RequiredInput : RequiredInputs)
+			for (const FMetasoundFrontendClassVertex& RequiredInput : RequiredInputs)
 			{
 				const FMetasoundFrontendClassInput* ClassInput = Graph->FindClassInputWithName(RequiredInput.Name).Get();
 
@@ -42,7 +42,7 @@ namespace Metasound
 			}
 
 			// Go through each output and add or swap if something is missing.
-			for (const FMetasoundFrontendClassOutput& RequiredOutput : RequiredOutputs)
+			for (const FMetasoundFrontendClassVertex& RequiredOutput : RequiredOutputs)
 			{
 				const FMetasoundFrontendClassOutput* ClassOutput = Graph->FindClassOutputWithName(RequiredOutput.Name).Get();
 
diff --git a/Engine/Plugins/Runtime/Metasound/Source/MetasoundFrontend/Public/MetasoundFrontendController.h b/Engine/Plugins/Runtime/Metasound/Source/MetasoundFrontend/Public/MetasoundFrontendController.h
--- a/Engine/Plugins/Runtime/Metasound/Source/MetasoundFrontend/Public/MetasoundFrontendController.h
+++ b/Engine/Plugins/Runtime/Metasound/Source/MetasoundFrontend/Public/MetasoundFrontendController.h
@@ -244,7 +244,7 @@ namespace Metasound
 			virtual const FText& GetDisplayName() const = 0;
 
 			/** Returns the default value of the given input. */
-			virtual const TArray<FMetasoundFrontendVertexLiteral>& GetDefaults() const = 0;
+			virtual const FMetasoundFrontendLiteral* GetDefaultLiteral() const = 0;
 
 			/** Returns the data type name associated with this input. */
 			virtual const FText& GetTooltip() const = 0;
@@ -333,34 +333,34 @@ namespace Metasound
 			 *
 			 * If the input does not exist, an invalid handle is returned. 
 			 */
-			virtual FInputHandle GetInputWithID(FGuid InPointID) = 0;
+			virtual FInputHandle GetInputWithID(FGuid InVertexID) = 0;
 
 			/** Returns an input with the given name. 
 			 *
 			 * If the input does not exist, an invalid handle is returned. 
 			 */
-			virtual FConstInputHandle GetInputWithID(FGuid InPointID) const = 0;
+			virtual FConstInputHandle GetInputWithID(FGuid InVertexID) const = 0;
 
 			/** Returns an output with the given name. 
 			 *
 			 * If the output does not exist, an invalid handle is returned. 
 			 */
-			virtual FOutputHandle GetOutputWithID(FGuid InPointID) = 0;
+			virtual FOutputHandle GetOutputWithID(FGuid InVertexID) = 0;
 
 			/** Returns an output with the given name. 
 			 *
 			 * If the output does not exist, an invalid handle is returned. 
 			 */
-			virtual FConstOutputHandle GetOutputWithID(FGuid InPointID) const = 0;
+			virtual FConstOutputHandle GetOutputWithID(FGuid InVertexID) const = 0;
 
 			virtual bool CanAddInput(const FString& InVertexName) const = 0;
 			virtual FInputHandle AddInput(const FString& InVertexName, const FMetasoundFrontendLiteral* InDefault) = 0;
-			virtual bool RemoveInput(FGuid InPointID) = 0;
+			virtual bool RemoveInput(FGuid InVertexID) = 0;
 
-			// TODO: PointID -> PinID? PointID? SlotID? RefID? 
+			// TODO: VertexID -> PinID? VertexID? SlotID? RefID? 
 			virtual bool CanAddOutput(const FString& InVertexName) const = 0;
 			virtual FInputHandle AddOutput(const FString& InVertexName, const FMetasoundFrontendLiteral* InDefault) = 0;
-			virtual bool RemoveOutput(FGuid InPointID) = 0;
+			virtual bool RemoveOutput(FGuid InVertexID) = 0;
 
 			/** Returns the class type for this node. */
 			virtual EMetasoundFrontendClassType GetClassType() const = 0;
@@ -416,7 +416,7 @@ namespace Metasound
 			/** Returns true if the controller is in a valid state. */
 			virtual bool IsValid() const = 0;
 
-			virtual FGuid GetNewPointID() const = 0;
+			virtual FGuid GetNewVertexID() const = 0;
 
 			virtual TArray<FString> GetInputVertexNames() const = 0;
 			virtual TArray<FString> GetOutputVertexNames() const = 0;
@@ -533,22 +533,21 @@ namespace Metasound
 			 */
 			virtual UClass* GetSupportedClassForInputVertex(const FString& InInputName) = 0;
 
-			virtual TArray<FGuid> GetDefaultIDsForInputVertex(const FString& InInputName) const = 0;
-			virtual TArray<FGuid> GetDefaultIDsForOutputVertex(const FString& InOutputName) const = 0;
+			virtual FGuid GetVertexIDForInputVertex(const FString& InInputName) const = 0;
+			virtual FGuid GetVertexIDForOutputVertex(const FString& InOutputName) const = 0;
 
-			virtual TArray<FMetasoundFrontendVertexLiteral> GetDefaultInputFrontendLiterals(const FString& InInputName) const = 0;
-
-			virtual bool SetDefaultInputToFrontendLiteral(const FString& InInputName, FGuid InPointID, FName InDataTypeName, const FMetasoundFrontendLiteral& InLiteral) = 0;
+			virtual FMetasoundFrontendLiteral GetDefaultInput(const FGuid& InVertexID) const = 0;
+			virtual bool SetDefaultInput(const FGuid& InVertexID, const FMetasoundFrontendLiteral& InLiteral) = 0;
 
 			/** Set the default value for the graph input.
 			 *
 			 * @param InInputName - Name of the graph input.
-			 * @param InPointID - Point to set to DataType default.
+			 * @param InVertexID - Vertex to set to DataType default.
 			 * @param InDataTypeName - Name of datatype to set to default.
 			 *
 			 * @return True on success. False if the input does not exist or if the literal type was incompatible with the input.
 			 */
-			virtual bool SetDefaultInputToTypeDefaultLiteral(const FString& InInputName, FGuid InPointID, FName InDataTypeName) = 0;
+			virtual bool SetDefaultInputToDefaultLiteralOfType(const FGuid& InVertexID) = 0;
 
 			/** Set the display name for the input with the given name. */
 			virtual void SetInputDisplayName(const FString& InName, const FText& InDisplayName) = 0;
@@ -572,7 +571,7 @@ namespace Metasound
 			 *
 			 * @return True on success, false on failure.
 			 */
-			virtual bool ClearLiteralForInput(const FString& InInputName, FGuid InPointID) = 0;
+			virtual bool ClearLiteralForInput(const FString& InInputName, FGuid InVertexID) = 0;
 
 			/** Add a new node to this graph.
 			 *
@@ -674,10 +673,10 @@ namespace Metasound
 			virtual bool IsRequiredOutput(const FString& InOutputName) const = 0;
 
 			/** Returns an array of inputs which describe the required inputs needed to satisfy the document archetype. */
-			virtual const TArray<FMetasoundFrontendClassInput>& GetRequiredInputs() const = 0;
+			virtual const TArray<FMetasoundFrontendClassVertex>& GetRequiredInputs() const = 0;
 
 			/** Returns an array of outputs which describe the required outputs needed to satisfy the document archetype. */
-			virtual const TArray<FMetasoundFrontendClassOutput>& GetRequiredOutputs() const = 0;
+			virtual const TArray<FMetasoundFrontendClassVertex>& GetRequiredOutputs() const = 0;
 
 			// TODO: add info on environment variables. 
 			// TODO: consider find/add subgraph
diff --git a/Engine/Plugins/Runtime/Metasound/Source/MetasoundFrontend/Public/MetasoundFrontendDocument.h b/Engine/Plugins/Runtime/Metasound/Source/MetasoundFrontend/Public/MetasoundFrontendDocument.h
--- a/Engine/Plugins/Runtime/Metasound/Source/MetasoundFrontend/Public/MetasoundFrontendDocument.h
+++ b/Engine/Plugins/Runtime/Metasound/Source/MetasoundFrontend/Public/MetasoundFrontendDocument.h
@@ -88,9 +88,9 @@ struct METASOUNDFRONTEND_API FMetasoundFrontendVertex
 	UPROPERTY(EditAnywhere, Category = Parameters)
 	FName TypeName;
 
-	// IDs of connection points supported by the vertex.
+	// ID of vertex
 	UPROPERTY()	
-	TArray<FGuid> PointIDs;
+	FGuid VertexID;
 
 	// Returns true if vertices have equal name, type and number of IDs.
 	static bool IsFunctionalEquivalent(const FMetasoundFrontendVertex& InLHS, const FMetasoundFrontendVertex& InRHS);
@@ -104,7 +104,7 @@ struct FMetasoundFrontendVertexLiteral
 
 	// ID of vertex.
 	UPROPERTY(EditAnywhere, Category = Parameters)
-	FGuid PointID = Metasound::FrontendInvalidID;
+	FGuid VertexID = Metasound::FrontendInvalidID;
 
 	// Value to use when constructing input. 
 	UPROPERTY(EditAnywhere, Category = Parameters)
@@ -218,7 +218,7 @@ struct FMetasoundFrontendEdge
 
 	// ID of source point on source node.
 	UPROPERTY()
-	FGuid FromPointID = Metasound::FrontendInvalidID;
+	FGuid FromVertexID = Metasound::FrontendInvalidID;
 
 	// ID of destination node.
 	UPROPERTY()
@@ -226,7 +226,7 @@ struct FMetasoundFrontendEdge
 
 	// ID of destination point on destination node.
 	UPROPERTY()
-	FGuid ToPointID = Metasound::FrontendInvalidID;
+	FGuid ToVertexID = Metasound::FrontendInvalidID;
 };
 
 // Display style for an edge.
@@ -292,36 +292,6 @@ struct FMetasoundFrontendGraph
 	FMetasoundFrontendGraphStyle Style;
 };
 
-// The type of the vertex.
-UENUM()
-enum class EMetasoundFrontendVertexType : uint8
-{
-	Point, //< Vertex represents a single value.
-	//TODO: Array   //< Vertex represents a variable sized array of values.
-};
-
-
-// Defines the behavior of the vertex.
-USTRUCT()
-struct METASOUNDFRONTEND_API FMetasoundFrontendVertexBehavior
-{
-	GENERATED_BODY()
-
-	// TODO: should all these be editable? 
-	UPROPERTY()
-	EMetasoundFrontendVertexType Type = EMetasoundFrontendVertexType::Point;
-
-	// Minimum number of connection points. Only used if the vertex type is EMetasoundFrontendVertexType::Array
-	UPROPERTY()
-	int32 ArrayMin = 1;
-
-	// Maximum number of connection points. Only used if the vertex type is EMetasoundFrontendVertexType::Array
-	UPROPERTY()
-	int32 ArrayMax = 1;
-
-	static bool IsFunctionalEquivalent(const FMetasoundFrontendVertexBehavior& InLHS, const FMetasoundFrontendVertexBehavior& InRHS);
-};
-
 // Metadata associated with a vertex.
 USTRUCT() 
 struct FMetasoundFrontendVertexMetadata
@@ -379,10 +349,6 @@ struct METASOUNDFRONTEND_API FMetasoundFrontendClassVertex : public FMetasoundFr
 	UPROPERTY(EditAnywhere, Category = CustomView)
 	FMetasoundFrontendVertexMetadata Metadata;
 
-	// Behavior of input vertex.
-	UPROPERTY()
-	FMetasoundFrontendVertexBehavior Behavior;
-
 	static bool IsFunctionalEquivalent(const FMetasoundFrontendClassVertex& InLHS, const FMetasoundFrontendClassVertex& InRHS);
 };
 
@@ -427,9 +393,9 @@ struct METASOUNDFRONTEND_API FMetasoundFrontendClassInput : public FMetasoundFro
 
 	virtual ~FMetasoundFrontendClassInput() = default;
 
-	// Default values for vertex IDs in this input.
+	// Default value for this input.
 	UPROPERTY(EditAnywhere, Category = Parameters)
-	TArray<FMetasoundFrontendVertexLiteral> Defaults;
+	FMetasoundFrontendLiteral DefaultLiteral;
 };
 
 USTRUCT()
@@ -680,10 +646,10 @@ struct FMetasoundFrontendArchetypeInterface
 	GENERATED_BODY()
 
 	UPROPERTY()
-	TArray<FMetasoundFrontendClassInput> Inputs;
+	TArray<FMetasoundFrontendClassVertex> Inputs;
 
 	UPROPERTY()
-	TArray<FMetasoundFrontendClassOutput> Outputs;
+	TArray<FMetasoundFrontendClassVertex> Outputs;
 
 	UPROPERTY()
 	TArray<FMetasoundFrontendEnvironmentVariable> Environment;
diff --git a/Engine/Plugins/Runtime/Metasound/Source/MetasoundFrontend/Public/MetasoundFrontendLiteral.h b/Engine/Plugins/Runtime/Metasound/Source/MetasoundFrontend/Public/MetasoundFrontendLiteral.h
--- a/Engine/Plugins/Runtime/Metasound/Source/MetasoundFrontend/Public/MetasoundFrontendLiteral.h
+++ b/Engine/Plugins/Runtime/Metasound/Source/MetasoundFrontend/Public/MetasoundFrontendLiteral.h
@@ -77,6 +77,9 @@ public:
 	// Returns true if the stored Type is an array type.
 	bool IsArray() const;
 
+	// Returns true if the literal is in a valid state (Type != EMetasoundFrontendLiteralType::Invalid)
+	bool IsValid() const;
+
 	// Sets the literal to the given type and value to default;
 	void SetType(EMetasoundFrontendLiteralType InType);
 
