Fix pin ordering to use that which is provided by the vertex interface in the node definition

diff --git a/Engine/Plugins/Runtime/Metasound/Source/MetasoundEditor/Private/MetasoundEditorGraphBuilder.cpp b/Engine/Plugins/Runtime/Metasound/Source/MetasoundEditor/Private/MetasoundEditorGraphBuilder.cpp
--- a/Engine/Plugins/Runtime/Metasound/Source/MetasoundEditor/Private/MetasoundEditorGraphBuilder.cpp
+++ b/Engine/Plugins/Runtime/Metasound/Source/MetasoundEditor/Private/MetasoundEditorGraphBuilder.cpp
@@ -1,6 +1,7 @@
 // Copyright Epic Games, Inc. All Rights Reserved.
 #include "MetasoundEditorGraphBuilder.h"
 
+#include "Algo/Sort.h"
 #include "AssetRegistry/AssetRegistryModule.h"
 #include "EdGraph/EdGraph.h"
 #include "EdGraph/EdGraphNode.h"
@@ -789,6 +790,7 @@ namespace Metasound
 
 		void FGraphBuilder::RebuildNodePins(UMetasoundEditorGraphNode& InGraphNode, Frontend::FNodeHandle InNodeHandle)
 		{
+			using namespace Frontend;
 			const FScopedTransaction Transaction(LOCTEXT("RebuildMetasoundGraphNodePins", "Rebuild Metasound Pins"));
 
 			for (int32 i = InGraphNode.Pins.Num() - 1; i >= 0; i--)
@@ -797,23 +799,25 @@ namespace Metasound
 				InGraphNode.RemovePin(InGraphNode.Pins[i]);
 			}
 
-			IMetasoundEditorModule& EditorModule = FModuleManager::GetModuleChecked<IMetasoundEditorModule>("MetasoundEditor");
-
-			FMetasoundAssetBase* MetasoundAsset = IMetasoundUObjectRegistry::Get().GetObjectAsAssetBase(&InGraphNode.GetMetasoundChecked());
-			check(MetasoundAsset);
-
-			Frontend::FGraphHandle GraphHandle = MetasoundAsset->GetRootGraphHandle();
+			TArray<FInputHandle> InputHandles = InNodeHandle->GetInputs();
+			Algo::SortBy(InputHandles, [](FInputHandle InputHandle)
+			{
+				return InputHandle->GetDisplayIndex();
+			});
 
-			TArray<Frontend::FInputHandle> InputHandles = InNodeHandle->GetInputs();
-			for (int32 i = 0; i < InputHandles.Num(); ++i)
+			for (const FInputHandle& InputHandle : InputHandles)
 			{
-				AddPinToNode(InGraphNode, InputHandles[i]);
+				AddPinToNode(InGraphNode, InputHandle);
 			}
 
-			TArray<Frontend::FOutputHandle> OutputHandles = InNodeHandle->GetOutputs();
-			for (int32 i = 0; i < OutputHandles.Num(); ++i)
+			TArray<FOutputHandle> OutputHandles = InNodeHandle->GetOutputs();
+			Algo::SortBy(OutputHandles, [](FOutputHandle OutputHandle)
+			{
+				return OutputHandle->GetDisplayIndex();
+			});
+			for (const FOutputHandle& OutputHandle : OutputHandles)
 			{
-				AddPinToNode(InGraphNode, OutputHandles[i]);
+				AddPinToNode(InGraphNode, OutputHandle);
 			}
 
 			InGraphNode.MarkPackageDirty();
@@ -998,6 +1002,7 @@ namespace Metasound
 			}
 
 			// Add missing editor nodes marked as visible.
+			// TODO: Synchronize Input/Output nodes which are own types now.
 			bIsEditorGraphDirty |= (FrontendNodes.Num() > 0);
 			for (FNodeHandle Node : FrontendNodes)
 			{
diff --git a/Engine/Plugins/Runtime/Metasound/Source/MetasoundFrontend/Private/MetasoundFrontend.cpp b/Engine/Plugins/Runtime/Metasound/Source/MetasoundFrontend/Private/MetasoundFrontend.cpp
--- a/Engine/Plugins/Runtime/Metasound/Source/MetasoundFrontend/Private/MetasoundFrontend.cpp
+++ b/Engine/Plugins/Runtime/Metasound/Source/MetasoundFrontend/Private/MetasoundFrontend.cpp
@@ -69,7 +69,8 @@ namespace Metasound
 
 			// External metasounds aren't dependent on any other nodes by definition, so all we need to do
 			// is populate the Input and Output sets.
-			for (auto& InputTuple : InNodeMetadata.DefaultInterface.GetInputInterface())
+			const FInputVertexInterface& InputInterface = InNodeMetadata.DefaultInterface.GetInputInterface();
+			for (auto& InputTuple : InputInterface)
 			{
 				// TODO: lots to be added here. 
 				FMetasoundFrontendClassInput ClassInput;
@@ -79,7 +80,7 @@ namespace Metasound
 				ClassInput.VertexID = FGuid::NewGuid();
 				ClassInput.Metadata.DisplayName = FText::FromString(InputTuple.Value.GetVertexName());
 				ClassInput.Metadata.Description = InputTuple.Value.GetDescription();
-
+				ClassInput.Metadata.DisplayIndex = InputInterface.GetOrderIndex(InputTuple.Key);
 				const FLiteral& DefaultLiteral = InputTuple.Value.GetDefaultValue();
 				if (DefaultLiteral.GetType() != ELiteralType::None)
 				{
@@ -89,7 +90,8 @@ namespace Metasound
 				ClassDescription.Interface.Inputs.Add(ClassInput);
 			}
 
-			for (auto& OutputTuple : InNodeMetadata.DefaultInterface.GetOutputInterface())
+			const FOutputVertexInterface& OutputInterface = InNodeMetadata.DefaultInterface.GetOutputInterface();
+			for (auto& OutputTuple : OutputInterface)
 			{
 				// TODO: lots to be added here. 
 				FMetasoundFrontendClassOutput ClassOutput;
@@ -99,6 +101,7 @@ namespace Metasound
 				ClassOutput.VertexID = FGuid::NewGuid();
 				ClassOutput.Metadata.DisplayName = FText::FromString(OutputTuple.Value.GetVertexName());
 				ClassOutput.Metadata.Description = OutputTuple.Value.GetDescription();
+				ClassOutput.Metadata.DisplayIndex = OutputInterface.GetOrderIndex(OutputTuple.Key);
 
 				ClassDescription.Interface.Outputs.Add(ClassOutput);
 			}
diff --git a/Engine/Plugins/Runtime/Metasound/Source/MetasoundFrontend/Private/MetasoundFrontendInvalidController.h b/Engine/Plugins/Runtime/Metasound/Source/MetasoundFrontend/Private/MetasoundFrontendInvalidController.h
--- a/Engine/Plugins/Runtime/Metasound/Source/MetasoundFrontend/Private/MetasoundFrontendInvalidController.h
+++ b/Engine/Plugins/Runtime/Metasound/Source/MetasoundFrontend/Private/MetasoundFrontendInvalidController.h
@@ -46,6 +46,7 @@ namespace Metasound
 			FGuid GetID() const override { return Metasound::FrontendInvalidID; }
 			const FName& GetDataType() const override { return MetasoundFrontendInvalidControllerPrivate::GetInvalid<FName>(); }  
 			const FString& GetName() const override { return MetasoundFrontendInvalidControllerPrivate::GetInvalid<FString>(); }
+			int32 GetDisplayIndex() const override { return INDEX_NONE; }
 			const FText& GetDisplayName() const override { return FText::GetEmpty(); }
 			const FText& GetTooltip() const override { return FText::GetEmpty(); }
 
@@ -87,6 +88,7 @@ namespace Metasound
 			bool IsValid() const override { return false; }
 			FGuid GetID() const override { return Metasound::FrontendInvalidID; }
 			bool IsConnected() const override { return false; }
+			int32 GetDisplayIndex() const override { return INDEX_NONE; }
 			const FName& GetDataType() const override { return MetasoundFrontendInvalidControllerPrivate::GetInvalid<FName>(); }
 			const FString& GetName() const override { return MetasoundFrontendInvalidControllerPrivate::GetInvalid<FString>(); }
 			const FText& GetDisplayName() const override { return FText::GetEmpty(); }
diff --git a/Engine/Plugins/Runtime/Metasound/Source/MetasoundFrontend/Private/MetasoundFrontendStandardController.cpp b/Engine/Plugins/Runtime/Metasound/Source/MetasoundFrontend/Private/MetasoundFrontendStandardController.cpp
--- a/Engine/Plugins/Runtime/Metasound/Source/MetasoundFrontend/Private/MetasoundFrontendStandardController.cpp
+++ b/Engine/Plugins/Runtime/Metasound/Source/MetasoundFrontend/Private/MetasoundFrontendStandardController.cpp
@@ -526,6 +526,17 @@ namespace Metasound
 			return FBaseOutputController::IsValid() && ClassOutputPtr.IsValid();
 		}
 
+		/** Returns index to display this input in parent structure. */
+		int32 FNodeOutputController::GetDisplayIndex() const
+		{
+			if (ClassOutputPtr.IsValid())
+			{
+				return ClassOutputPtr->Metadata.DisplayIndex;
+			}
+
+			return INDEX_NONE;
+		}
+
 		const FText& FNodeOutputController::GetDisplayName() const
 		{
 			if (ClassOutputPtr.IsValid())
@@ -578,6 +589,16 @@ namespace Metasound
 			return FBaseOutputController::IsValid() && OwningGraphClassInputPtr.IsValid();
 		}
 
+		int32 FInputNodeOutputController::GetDisplayIndex() const
+		{
+			if (OwningGraphClassInputPtr.IsValid())
+			{
+				return OwningGraphClassInputPtr->Metadata.DisplayIndex;
+			}
+
+			return INDEX_NONE;
+		}
+
 		const FText& FInputNodeOutputController::GetDisplayName() const
 		{
 			if (OwningGraphClassInputPtr.IsValid())
@@ -909,6 +930,16 @@ namespace Metasound
 			return FBaseInputController::IsValid() && OwningGraphClassInputPtr.IsValid();
 		}
 
+		int32 FNodeInputController::GetDisplayIndex() const
+		{
+			if (OwningGraphClassInputPtr.IsValid())
+			{
+				return OwningGraphClassInputPtr->Metadata.DisplayIndex;
+			}
+
+			return INDEX_NONE;
+		}
+
 		const FText& FNodeInputController::GetDisplayName() const
 		{
 			if (OwningGraphClassInputPtr.IsValid())
@@ -970,6 +1001,16 @@ namespace Metasound
 			return FBaseInputController::IsValid() && OwningGraphClassOutputPtr.IsValid();
 		}
 
+		int32 FOutputNodeInputController::GetDisplayIndex() const
+		{
+			if (OwningGraphClassOutputPtr.IsValid())
+			{
+				return OwningGraphClassOutputPtr->Metadata.DisplayIndex;
+			}
+
+			return INDEX_NONE;
+		}
+
 		const FText& FOutputNodeInputController::GetDisplayName() const
 		{
 			if (OwningGraphClassOutputPtr.IsValid())
@@ -1415,9 +1456,7 @@ namespace Metasound
 			{
 				for (const FMetasoundFrontendVertex& NodeInputVertex : NodePtr->Interface.Inputs)
 				{
-
 					FConstVertexAccessPtr NodeVertexPtr = NodePtr.GetMemberAccessPtr<const FMetasoundFrontendVertex>(FGetNodeInputFromNodeWithName(NodeInputVertex.Name));
-
 					FConstClassInputAccessPtr ClassInputPtr = ClassPtr.GetMemberAccessPtr<const FMetasoundFrontendClassInput>(FGetClassInputFromClassWithName(NodeInputVertex.Name));
 
 					Inputs.Add({NodeInputVertex.VertexID, NodeVertexPtr, ClassInputPtr});
diff --git a/Engine/Plugins/Runtime/Metasound/Source/MetasoundFrontend/Private/MetasoundFrontendStandardController.h b/Engine/Plugins/Runtime/Metasound/Source/MetasoundFrontend/Private/MetasoundFrontendStandardController.h
--- a/Engine/Plugins/Runtime/Metasound/Source/MetasoundFrontend/Private/MetasoundFrontendStandardController.h
+++ b/Engine/Plugins/Runtime/Metasound/Source/MetasoundFrontend/Private/MetasoundFrontendStandardController.h
@@ -99,6 +99,7 @@ namespace Metasound
 			bool IsValid() const override;
 
 			// Output metadata
+			int32 GetDisplayIndex() const override;
 			const FText& GetDisplayName() const override;
 			const FText& GetTooltip() const override;
 
@@ -140,6 +141,7 @@ namespace Metasound
 			bool IsValid() const override;
 
 			// Output metadata
+			int32 GetDisplayIndex() const override;
 			const FText& GetDisplayName() const override;
 			const FText& GetTooltip() const override;
 
@@ -238,6 +240,7 @@ namespace Metasound
 			const FMetasoundFrontendLiteral* GetDefaultLiteral() const override;
 
 			// Input metadata
+			int32 GetDisplayIndex() const override;
 			const FText& GetDisplayName() const override;
 			const FText& GetTooltip() const override;
 
@@ -275,6 +278,7 @@ namespace Metasound
 			bool IsValid() const override;
 
 			// Input metadata
+			int32 GetDisplayIndex() const override;
 			const FText& GetDisplayName() const override;
 			const FText& GetTooltip() const override;
 			const FMetasoundFrontendLiteral* GetDefaultLiteral() const override;
diff --git a/Engine/Plugins/Runtime/Metasound/Source/MetasoundFrontend/Public/MetasoundFrontendController.h b/Engine/Plugins/Runtime/Metasound/Source/MetasoundFrontend/Public/MetasoundFrontendController.h
--- a/Engine/Plugins/Runtime/Metasound/Source/MetasoundFrontend/Public/MetasoundFrontendController.h
+++ b/Engine/Plugins/Runtime/Metasound/Source/MetasoundFrontend/Public/MetasoundFrontendController.h
@@ -172,6 +172,9 @@ namespace Metasound
 			/** Returns the name associated with this output. */
 			virtual const FString& GetName() const = 0;
 			
+			/** Returns index to display this input in parent structure. */
+			virtual int32 GetDisplayIndex() const = 0;
+
 			/** Returns the human readable name associated with this output. */
 			virtual const FText& GetDisplayName() const = 0;
 			
@@ -240,6 +243,9 @@ namespace Metasound
 			/** Returns the data type name associated with this input. */
 			virtual const FString& GetName() const = 0;
 
+			/** Returns index to display this input in parent structure. */
+			virtual int32 GetDisplayIndex() const = 0;
+
 			/** Returns the data type name associated with this input. */
 			virtual const FText& GetDisplayName() const = 0;
 
diff --git a/Engine/Plugins/Runtime/Metasound/Source/MetasoundFrontend/Public/MetasoundFrontendDocument.h b/Engine/Plugins/Runtime/Metasound/Source/MetasoundFrontend/Public/MetasoundFrontendDocument.h
--- a/Engine/Plugins/Runtime/Metasound/Source/MetasoundFrontend/Public/MetasoundFrontendDocument.h
+++ b/Engine/Plugins/Runtime/Metasound/Source/MetasoundFrontend/Public/MetasoundFrontendDocument.h
@@ -297,7 +297,11 @@ USTRUCT()
 struct FMetasoundFrontendVertexMetadata
 {
 	GENERATED_BODY()
-		
+
+	// Index to inform where the vertex should be displayed in parent collection.
+	UPROPERTY()
+	int32 DisplayIndex = 0;
+
 	// Display name for a vertex
 	UPROPERTY(EditAnywhere, Category = Parameters, meta = (DisplayName = "Name"))
 	FText DisplayName;
diff --git a/Engine/Plugins/Runtime/Metasound/Source/MetasoundGraphCore/Public/MetasoundVertex.h b/Engine/Plugins/Runtime/Metasound/Source/MetasoundGraphCore/Public/MetasoundVertex.h
--- a/Engine/Plugins/Runtime/Metasound/Source/MetasoundGraphCore/Public/MetasoundVertex.h
+++ b/Engine/Plugins/Runtime/Metasound/Source/MetasoundGraphCore/Public/MetasoundVertex.h
@@ -553,19 +553,19 @@ namespace Metasound
 	};
 
 
-	/** Create an FVertexKey from a FDataVertex. */
+	/** Create an FVertexKey from an FOutDataVertex. */
 	FORCEINLINE FVertexKey MakeDataVertexKey(const FInputDataVertex& InVertex)
 	{
 		return InVertex.GetVertexName();
 	}
 
-	/** Create an FVertexKey from a FDataVertex. */
+	/** Create an FVertexKey from an FOutputDataVertex. */
 	FORCEINLINE FVertexKey MakeDataVertexKey(const FOutputDataVertex& InVertex)
 	{
 		return InVertex.GetVertexName();
 	}
 
-	/** Create an FVertexKey from a FEnvironmentVariableVertex. */
+	/** Create an FVertexKey from an FEnvironmentVertex. */
 	FORCEINLINE FVertexKey MakeEnvironmentVertexKey(const FEnvironmentVertex& InVertex)
 	{
 		return InVertex.GetVertexName();
@@ -585,9 +585,10 @@ namespace Metasound
 			static_assert(bIsSupportedVertexType, "VertexType must be derived from FInputDataVertex, FOutputDataVertex, or FEnvironmentVertex");
 
 			using FContainerType = TSortedMap<FVertexKey, VertexType>;
+			using FOrderContainerType = TArray<FVertexKey>;
 
 			// Required for end of recursion.
-			static void CopyInputs(FContainerType& InStorage)
+			static void CopyInputs(FContainerType& InStorage, FOrderContainerType& InOrder)
 			{
 			}
 
@@ -598,18 +599,20 @@ namespace Metasound
 			// Assume that InputType is either a sublcass of FDataVertexModel
 			// and that VertexType is either FInputDataVertex or FOutputDataVertex
 			template<typename VertexModelType, typename... RemainingVertexModelTypes>
-			static void CopyInputs(FContainerType& InStorage, const VertexModelType& InInput, const RemainingVertexModelTypes&... InRemainingInputs)
+			static void CopyInputs(FContainerType& InStorage, FOrderContainerType& InKeyOrder, const VertexModelType& InInput, const RemainingVertexModelTypes&... InRemainingInputs)
 			{
 				// Create vertex out of vertex model
 				VertexType Vertex(InInput);
 
-				InStorage.Add(MakeDataVertexKey(Vertex), Vertex);
-				CopyInputs(InStorage, InRemainingInputs...);
+				const FVertexKey VertexKey = MakeDataVertexKey(Vertex);
+				InStorage.Add(VertexKey, Vertex);
+				InKeyOrder.Add(VertexKey);
+				CopyInputs(InStorage, InKeyOrder, InRemainingInputs...);
 			}
 
 		public:
 
-			using RangedForConstIteratorType = typename TSortedMap<FVertexKey, VertexType>::RangedForConstIteratorType;
+			using RangedForConstIteratorType = typename FContainerType::RangedForConstIteratorType;
 
 			/** TVertexInterfaceGroup constructor with variadic list of vertex
 			 * models.
@@ -617,7 +620,7 @@ namespace Metasound
 			template<typename... VertexModelTypes>
 			TVertexInterfaceGroup(VertexModelTypes&&... InVertexModels)
 			{
-				CopyInputs(Vertices, Forward<VertexModelTypes>(InVertexModels)...);
+				CopyInputs(Vertices, OrderedKeys, Forward<VertexModelTypes>(InVertexModels)...);
 			}
 
 			/** Add a vertex model to the group. */
@@ -631,13 +634,16 @@ namespace Metasound
 			/** Add a vertex to the group. */
 			void Add(const VertexType& InVertex)
 			{
-				Vertices.Add(MakeDataVertexKey(InVertex), InVertex);
+				const FVertexKey VertexKey = MakeDataVertexKey(InVertex);
+				Vertices.Add(VertexKey, InVertex);
+				OrderedKeys.Add(VertexKey);
 			}
 
 			/** Remove a vertex by key. */
 			bool Remove(const FVertexKey& InKey)
 			{
-				int32 NumRemoved = Vertices.Remove();
+				int32 NumRemoved = Vertices.Remove(InKey);
+				OrderedKeys.Remove(InKey);
 				return (NumRemoved > 0);
 			}
 
@@ -647,6 +653,13 @@ namespace Metasound
 				return Vertices.Contains(InKey);
 			}
 
+			int32 GetOrderIndex(const FVertexKey& InKey) const
+			{
+				int32 OutIndex = INDEX_NONE;
+				OrderedKeys.Find(InKey, OutIndex);
+				return OutIndex;
+			}
+
 			/** Return the vertex for a given vertex key. */
 			const VertexType& operator[](const FVertexKey& InKey) const
 			{
@@ -686,6 +699,7 @@ namespace Metasound
 		private:
 
 			FContainerType Vertices;
+			FOrderContainerType OrderedKeys;
 	};
 
 	
diff --git a/Engine/Plugins/Runtime/Metasound/Source/MetasoundStandardNodes/Private/MetasoundADSRNode.cpp b/Engine/Plugins/Runtime/Metasound/Source/MetasoundStandardNodes/Private/MetasoundADSRNode.cpp
--- a/Engine/Plugins/Runtime/Metasound/Source/MetasoundStandardNodes/Private/MetasoundADSRNode.cpp
+++ b/Engine/Plugins/Runtime/Metasound/Source/MetasoundStandardNodes/Private/MetasoundADSRNode.cpp
@@ -233,18 +233,18 @@ namespace Metasound
 	{
 		static const FVertexInterface Interface(
 			FInputVertexInterface(
+				TInputDataVertexModel<FTrigger>(TEXT("Trigger Attack"), LOCTEXT("TriggerAttackTooltip", "Trigger the envelope's attack.")),
+				TInputDataVertexModel<FTrigger>(TEXT("Trigger Release"), LOCTEXT("TriggerReleaseTooltip", "Trigger the envelope's release.")),
 				TInputDataVertexModel<FTime>(TEXT("Attack"), LOCTEXT("AttackDurationTooltip", "Attack duration."), 0.01f),
 				TInputDataVertexModel<FTime>(TEXT("Decay"), LOCTEXT("DecayDurationTooltip", "Decay duration."), 0.02f),
-				TInputDataVertexModel<FTime>(TEXT("Release"), LOCTEXT("ReleaseDurationTooltip", "Release duration."), 1.0f),
 				TInputDataVertexModel<float>(TEXT("Sustain Level"), LOCTEXT("SustainLevelTooltip", "Sustain level [0.0f, 1.0f]."), 0.7f),
-				TInputDataVertexModel<FTrigger>(TEXT("Trigger Attack"), LOCTEXT("TriggerAttackTooltip", "Trigger the envelope's attack.")),
-				TInputDataVertexModel<FTrigger>(TEXT("Trigger Release"), LOCTEXT("TriggerReleaseTooltip", "Trigger the envelope's release."))
+				TInputDataVertexModel<FTime>(TEXT("Release"), LOCTEXT("ReleaseDurationTooltip", "Release duration."), 1.0f)
 			),
 			FOutputVertexInterface(
-				TOutputDataVertexModel<FAudioBuffer>(TEXT("Envelope"), LOCTEXT("EnvelopeTooltip", "The output envelope")),
 				TOutputDataVertexModel<FTrigger>(TEXT("Attack Complete"), LOCTEXT("AttackCompleteTooltip", "Triggered when the envelope attack is complete")),
 				TOutputDataVertexModel<FTrigger>(TEXT("Decay Complete"), LOCTEXT("DecayCompleteTooltip", "Triggered when the envelope decay is complete")),
-				TOutputDataVertexModel<FTrigger>(TEXT("Release Complete"), LOCTEXT("ReleaseCompleteTooltip", "Triggered when the envelope is released"))
+				TOutputDataVertexModel<FTrigger>(TEXT("Release Complete"), LOCTEXT("ReleaseCompleteTooltip", "Triggered when the envelope is released")),
+				TOutputDataVertexModel<FAudioBuffer>(TEXT("Envelope"), LOCTEXT("EnvelopeTooltip", "The output envelope"))
 			)
 		);
 
diff --git a/Engine/Plugins/Runtime/Metasound/Source/MetasoundStandardNodes/Private/MetasoundRandomNode.cpp b/Engine/Plugins/Runtime/Metasound/Source/MetasoundStandardNodes/Private/MetasoundRandomNode.cpp
--- a/Engine/Plugins/Runtime/Metasound/Source/MetasoundStandardNodes/Private/MetasoundRandomNode.cpp
+++ b/Engine/Plugins/Runtime/Metasound/Source/MetasoundStandardNodes/Private/MetasoundRandomNode.cpp
@@ -168,11 +168,11 @@ namespace Metasound
 	{
 		static const FVertexInterface Interface(
 			FInputVertexInterface(
+				TInputDataVertexModel<FTrigger>(InParamNameNextTrigger, LOCTEXT("NextTooltip", "Trigger to generate the next random integer.")),
+				TInputDataVertexModel<FTrigger>(InParamNameResetTrigger, LOCTEXT("ResetTooltip", "Trigger to reset the random sequence with the supplied seed. Useful to get randomized repetition.")),
 				TInputDataVertexModel<int32>(InParamNameSeed, LOCTEXT("SeedTooltip", "The seed value to use for the random node. Set to 0 to use a random seed."), 0),
 				TInputDataVertexModel<int32>(InParamNameMin, LOCTEXT("MinTooltip", "Min random integer."), 0),
-				TInputDataVertexModel<int32>(InParamNameMax, LOCTEXT("MaxTooltip", "Max random integer."), 100),
-				TInputDataVertexModel<FTrigger>(InParamNameNextTrigger, LOCTEXT("NextTooltip", "Trigger to generate the next random integer.")),
-				TInputDataVertexModel<FTrigger>(InParamNameResetTrigger, LOCTEXT("ResetTooltip", "Trigger to reset the random sequence with the supplied seed. Useful to get randomized repetition."))
+				TInputDataVertexModel<int32>(InParamNameMax, LOCTEXT("MaxTooltip", "Max random integer."), 100)
 			),
 			FOutputVertexInterface(
 				TOutputDataVertexModel<int32>(OutParamNameValue, LOCTEXT("ValueTooltip", "The current randomly generated random integer value."))
@@ -344,11 +344,11 @@ namespace Metasound
 	{
 		static const FVertexInterface Interface(
 			FInputVertexInterface(
+				TInputDataVertexModel<FTrigger>(InParamNameNextTrigger, LOCTEXT("NextTooltip", "Trigger to generate the next random float.")),
+				TInputDataVertexModel<FTrigger>(InParamNameResetTrigger, LOCTEXT("ResetTooltip", "Trigger to reset the random sequence with the supplied seed. Useful to get randomized repetition.")),
 				TInputDataVertexModel<int32>(InParamNameSeed, LOCTEXT("SeedTooltip", "The seed value to use for the random node. Set to 0 to use a random seed."), 0),
 				TInputDataVertexModel<float>(InParamNameMin, LOCTEXT("MinTooltip", "Min random float."), 0.0f),
-				TInputDataVertexModel<float>(InParamNameMax, LOCTEXT("MaxTooltip", "Max random float."), 1.0f),
-				TInputDataVertexModel<FTrigger>(InParamNameNextTrigger, LOCTEXT("NextTooltip", "Trigger to generate the next random float.")),
-				TInputDataVertexModel<FTrigger>(InParamNameResetTrigger, LOCTEXT("ResetTooltip", "Trigger to reset the random sequence with the supplied seed. Useful to get randomized repetition."))
+				TInputDataVertexModel<float>(InParamNameMax, LOCTEXT("MaxTooltip", "Max random float."), 1.0f)
 			),
 			FOutputVertexInterface(
 				TOutputDataVertexModel<float>(OutParamNameValue, LOCTEXT("ValueTooltip", "The current randomly generated float value."))
