Metasound Oscillators Nodes/Operators first pass (saw, triangle, square, sine)

diff --git a/Engine/Plugins/Runtime/Metasound/Source/MetasoundFrontend/Public/MetasoundEnum.h b/Engine/Plugins/Runtime/Metasound/Source/MetasoundFrontend/Public/MetasoundEnum.h
--- a/Engine/Plugins/Runtime/Metasound/Source/MetasoundFrontend/Public/MetasoundEnum.h
+++ b/Engine/Plugins/Runtime/Metasound/Source/MetasoundFrontend/Public/MetasoundEnum.h
@@ -103,8 +103,7 @@ namespace Metasound
 	{
 	public:
 		using InnerType = EnumType;
-		static constexpr EnumType Default = DefaultValue;
-
+	
 		static_assert(TIsEnum<EnumType>::Value, "Expecting an Enum type");
 		
 		// Allow Non-enum class Enum or Enum class that's are derived from int32
diff --git a/Engine/Plugins/Runtime/Metasound/Source/MetasoundStandardNodes/Private/MetasoundOscillatorNodes.cpp b/Engine/Plugins/Runtime/Metasound/Source/MetasoundStandardNodes/Private/MetasoundOscillatorNodes.cpp
new file mode 100644
--- /dev/null
+++ b/Engine/Plugins/Runtime/Metasound/Source/MetasoundStandardNodes/Private/MetasoundOscillatorNodes.cpp
@@ -0,0 +1,1015 @@
+// Copyright Epic Games, Inc. All Rights Reserved.
+#include "MetasoundOscillatorNodes.h"
+
+#include "DSP/Dsp.h"
+#include "MetasoundAudioBuffer.h"
+#include "MetasoundExecutableOperator.h"
+#include "MetasoundFacade.h"
+#include "MetasoundNodeRegistrationMacro.h"
+#include "MetasoundDataTypeRegistrationMacro.h"
+#include "MetasoundPrimitives.h"
+#include "MetasoundStandardNodesNames.h"
+#include "MetasoundTrigger.h"
+#include "MetasoundVertex.h"
+
+#define LOCTEXT_NAMESPACE "MetasoundStandardNodes"
+
+namespace Metasound
+{			
+#pragma region Common
+
+	// Contained here are a small selection of block generators and sample generators.
+	// They are templatized to allow quick construction of new oscillators, but aren't considered
+	// optimal (as they contain many branches) and they should be only as a default 
+	// until Vectorized versions replace them.
+
+	namespace Generators
+	{
+		// Standard params passed to the Generate Block templates below
+		struct FGeneratorArgs
+		{
+			float SampleRate = 0.f;
+			float FrequencyHz = 0.f;
+			float PulseWidth = 0.f;
+			TArrayView<float> AlignedBuffer;
+			TArrayView<const float> FM;
+		};
+
+		// Functor used when we need to wrap the phase of some of the oscillators. Sinf for example does not care
+		struct FWrapPhase
+		{
+			FORCEINLINE void operator()(float& InPhase) const
+			{
+				while (InPhase >= 1.f)
+				{
+					InPhase -= 1.f;
+				}
+			}
+		};
+
+		// Functor that does nothing to the phase and lets it climb indefinitely 
+		struct FPhaseLetClimb
+		{
+			FORCEINLINE void operator()(const float&) const {};
+		};
+
+		// Smooth out the edges of the saw based on its current frequency
+		// using a polynomial to smooth it at the discontinuity. This 
+		// limits aliasing by avoiding the infinite frequency at the discontinuity.
+		FORCEINLINE float PolySmoothSaw(const float InPhase, const float InPhaseDelta)
+		{	
+			float Output = 0.0f;
+
+			// The current phase is on the left side of discontinuity
+			if (InPhase > 1.0f - InPhaseDelta)
+			{
+				const float Dist = (InPhase - 1.0f) / InPhaseDelta;
+				Output = -Dist * Dist - 2.0f * Dist - 1.0f;
+			}
+			// The current phase is on the right side of the discontinuity
+			else if (InPhase < InPhaseDelta)
+			{
+				// Distance into polynomial
+				const float Dist = InPhase / InPhaseDelta;
+				Output = Dist * Dist - 2.0f * Dist + 1.0f;
+			}
+			return Output;
+		}
+
+		// Functor that does the a generic generate block with a supplied Oscillator.
+		// The PhaseWrap functor controls how the phase is accumulated and wrapped
+		template<
+			typename Oscillator,
+			typename PhaseWrap = FWrapPhase
+		>
+		struct TGenerateBlock
+		{
+			float Phase = 0.f;
+			Oscillator Osc;
+			PhaseWrap Wrap;
+			void operator()(const FGeneratorArgs& InArgs)
+			{
+				int32 RemainingSamplesInBlock = InArgs.AlignedBuffer.Num();
+				float* Out = InArgs.AlignedBuffer.GetData();
+				float OneOverSampleRate = 1.f / InArgs.SampleRate;
+				float PulseWidth = InArgs.PulseWidth;
+				float DeltaPhase = InArgs.FrequencyHz * OneOverSampleRate;
+
+				while (RemainingSamplesInBlock > 0)
+				{
+					Wrap(Phase);
+
+					*Out++ = Osc(Phase, DeltaPhase, PulseWidth);
+
+					Phase += DeltaPhase;
+
+					RemainingSamplesInBlock--;
+				}
+			}
+		};
+
+		// Functor that does the a generic generate block with a supplied Oscillator and applies an FM signal
+		// to the frequency per sample.
+		template<
+			typename Oscillator,
+			typename PhaseWrap = FWrapPhase
+		>
+		struct TGenerateBlockFM
+		{
+			float Phase = 0.f;
+			Oscillator Osc;
+			PhaseWrap Wrap;
+			void operator()(const FGeneratorArgs& InArgs)
+			{
+				int32 RemainingSamplesInBlock = InArgs.AlignedBuffer.Num();
+				float* Out = InArgs.AlignedBuffer.GetData();
+				const float* FM = InArgs.FM.GetData();
+				float OneOverSampleRate = 1.f / InArgs.SampleRate;
+				float PulseWidth = InArgs.PulseWidth;
+
+				while (RemainingSamplesInBlock > 0)
+				{
+					float PerSampleFreq = InArgs.FrequencyHz + *FM++;
+					float DeltaPhase = PerSampleFreq * OneOverSampleRate;
+
+					*Out++ = Osc(Phase, DeltaPhase, PulseWidth);
+
+					Phase += DeltaPhase;
+
+					Wrap(Phase);
+
+					RemainingSamplesInBlock--;
+				}
+			}
+		};	
+		
+		// Sine types and generators.
+		struct FSinfGenerator
+		{
+			FORCEINLINE float operator()(float InPhase, float, float)
+			{
+				static constexpr float TwoPi = 2.f * PI;
+				return FMath::Sin(InPhase * TwoPi);
+			}
+		};
+		struct FBhaskaraGenerator
+		{
+			float operator()(float InPhase, float, float) const
+			{
+				static const float TwoPi = 2.f * PI;
+				float PhaseRadians = InPhase * TwoPi;
+				return Audio::FastSin3(PhaseRadians - PI); // Expects [-PI, PI] 
+			}
+		};
+		struct F2DRotatorGenerateBlock
+		{
+			Audio::FSinOsc2DRotation Rotator;
+			F2DRotatorGenerateBlock(float InStartingPhase) { Rotator.LastPhasePerSample = InStartingPhase; }
+			void operator()(const FGeneratorArgs& Args)
+			{
+				Rotator.GenerateBuffer(Args.SampleRate, Args.FrequencyHz, Args.AlignedBuffer.GetData(), Args.AlignedBuffer.Num());
+			}
+		};		
+		
+		using FSinf						= TGenerateBlock<FSinfGenerator, FPhaseLetClimb>;
+		using FSinfWithFm				= TGenerateBlockFM<FSinfGenerator, FPhaseLetClimb>;
+		
+		using FBhaskara					= TGenerateBlock<FBhaskaraGenerator>;
+		using FBhaskaraWithFm			= TGenerateBlockFM<FBhaskaraGenerator>;
+
+		// Saws.
+		struct FSawGenerator
+		{
+			FORCEINLINE float operator()(float InPhase, float, float)
+			{
+				return -1.f + (2.f * InPhase);
+			}
+		};
+		struct FSawPolySmoothGenerator
+		{
+			FORCEINLINE float operator()(float InPhase, float InPhaseDelta, float)
+			{
+				// Two-sided wave-shaped sawtooth
+				static const float A = Audio::FastTanh(1.5f);
+				float Output = Audio::GetBipolar(InPhase);
+				Output = Audio::FastTanh(1.5f * Output) / A;
+				Output += PolySmoothSaw(InPhase, InPhaseDelta);
+				return Output;
+			}
+		};
+
+		using FSaw							= TGenerateBlock<FSawGenerator>;
+		using FSawWithFm					= TGenerateBlockFM<FSawGenerator>;
+
+		using FSawPolysmooth				= TGenerateBlock<FSawPolySmoothGenerator>;
+		using FSawPolysmoothWithFm			= TGenerateBlockFM<FSawPolySmoothGenerator>;
+
+		// Square.
+		struct FSquareGenerator
+		{
+			FORCEINLINE float operator()(float InPhase, float, float)
+			{
+				// Does not obey pulse width.
+				return InPhase >= 0.5f ? 1.f : -1.f;
+			}
+		};
+		struct FSquarePolysmoothGenerator
+		{
+			FORCEINLINE float operator()(float InPhase, float InPhaseDelta, float InCurrentPulseWidth)
+			{
+				// Taken piecemeal from Osc.cpp
+				// Lots of branches.
+
+				// First generate a smoothed sawtooth
+				float SquareSaw1 = Audio::GetBipolar(InPhase);
+				SquareSaw1 += PolySmoothSaw(InPhase, InPhaseDelta);
+
+				// Create a second sawtooth that is phase-shifted based on the pulsewidth
+				float NewPhase = 0.0f;
+				if (InPhaseDelta > 0.0f)
+				{
+					NewPhase = InPhase + InCurrentPulseWidth;
+					if (NewPhase >= 1.0f)
+					{
+						NewPhase -= 1.0f;
+					}
+				}
+				else
+				{
+					NewPhase = InPhase - InCurrentPulseWidth;
+					if (NewPhase <= 0.0f)
+					{
+						NewPhase += 1.0f;
+					}
+				}
+
+				float SquareSaw2 = Audio::GetBipolar(NewPhase);
+				SquareSaw2 += PolySmoothSaw(NewPhase, InPhaseDelta);
+
+				// Subtracting 2 saws creates a square wave!
+				float Output = 0.5f * SquareSaw1 - 0.5f * SquareSaw2;
+
+				// Apply DC correction
+				const float Correction = (InCurrentPulseWidth < 0.5f) ? 
+					(1.0f / (1.0f - InCurrentPulseWidth)) : (1.0f / InCurrentPulseWidth);
+
+				Output *= Correction;
+				return Output;
+			}
+		};
+
+		using FSquare						= TGenerateBlock<FSquareGenerator>;
+		using FSquareWithFm					= TGenerateBlockFM<FSquareGenerator>;
+
+		using FSquarePolysmooth				= TGenerateBlock<FSquarePolysmoothGenerator>;
+		using FSquarePolysmoothWithFm		= TGenerateBlockFM<FSquarePolysmoothGenerator>;
+
+		// Triangle.
+		struct FTriangleGenerator
+		{
+			FORCEINLINE float operator()(float InPhase, float, float)
+			{
+				return InPhase <= 0.5f ? InPhase - 0.25f * 4.f : ((1.f - InPhase) - 0.25f) * 4.f;
+			}
+		};
+		using FTriangle						= TGenerateBlock<FTriangleGenerator>;
+		using FTriangleWithFm				= TGenerateBlockFM<FTriangleGenerator>;
+	}
+
+	// Base class of Oscillator factories which holds common the interface.
+	class FOscilatorFactoryBase : public IOperatorFactory
+	{
+	public:
+		// Common pins
+		static constexpr const TCHAR* EnabledPinName = TEXT("Enabled");
+		static constexpr const TCHAR* FrequencyModPinName = TEXT("Modulation");
+		static constexpr const TCHAR* BaseFrequencyPinName = TEXT("Frequency");
+		static constexpr const TCHAR* PhaseResetPinName = TEXT("Sync");
+		static constexpr const TCHAR* PhaseOffsetPinName = TEXT("Phase Offset");
+		static constexpr const TCHAR* AudioOutPinName = TEXT("Audio");
+		static constexpr const TCHAR* BiPolarPinName = TEXT("Bi Polar");
+
+		// Common to all Oscillators.
+		static FVertexInterface GetCommmonVertexInterface()
+		{
+			static const FVertexInterface Interface
+			{
+				FInputVertexInterface{
+					TInputDataVertexModel<bool>(EnabledPinName, LOCTEXT("OscActivateDescription", "Enable the oscillator."), true),
+					TInputDataVertexModel<bool>(BiPolarPinName, LOCTEXT("OscBaseBiPolarDescription", "If the output is Bi-Polar (-1..1) or Uni-Polar (0..1)"), true),
+					TInputDataVertexModel<float>(BaseFrequencyPinName, LOCTEXT("OscBaseFrequencyDescription", "Base Frequency of Oscillator in HZ"), 440.f),
+					TInputDataVertexModel<FAudioBuffer>(FrequencyModPinName, LOCTEXT("OscModFrequencyDescription", "Modulation Frequency Input (for doing FM)")),
+					TInputDataVertexModel<FTrigger>(PhaseResetPinName, LOCTEXT("OscPhaseResetDescription", "Phase Reset")),
+					TInputDataVertexModel<float>(PhaseOffsetPinName, LOCTEXT("OscPhaseOffsetDescription", "Phase Offset In Degrees (0..360)"), 0.f)
+				},
+				FOutputVertexInterface{
+					TOutputDataVertexModel<FAudioBuffer>(TEXT("Audio"), LOCTEXT("AudioTooltip", "The output audio"))
+				}
+			};
+			return Interface;
+		}
+	};
+
+	// Common set of construct params for each of the Oscillators.
+	struct FOscillatorOperatorConstructParams
+	{
+		const FOperatorSettings& Settings;
+		FBoolReadRef Enabled;
+		FFloatReadRef BaseFrequency;
+		FFloatReadRef PhaseOffset; 
+		FTriggerReadRef PhaseReset;
+		FBoolReadRef BiPolar;
+	};
+
+	// Base Oscillator Operator CRTP.
+	// Expects your Derived class to implement a Generate(int32 InStartFrame, int32 InEndFrame, float InFreq)
+	template<typename GeneratorPolicy, typename Derived>
+	class TOscillatorOperatorBase : public TExecutableOperator<Derived>
+	{		
+	public:
+		TOscillatorOperatorBase(const FOscillatorOperatorConstructParams& InConstructParams)
+			: Generator{*InConstructParams.PhaseOffset}
+			, SampleRate(InConstructParams.Settings.GetSampleRate())
+			, Nyquist(InConstructParams.Settings.GetSampleRate() / 2.0f)
+			, Enabled(InConstructParams.Enabled)
+			, BaseFrequency(InConstructParams.BaseFrequency)
+			, PhaseReset(InConstructParams.PhaseReset)
+			, PhaseOffset(InConstructParams.PhaseOffset)
+			, BiPolar(InConstructParams.BiPolar)
+			, AudioBuffer(FAudioBufferWriteRef::CreateNew(InConstructParams.Settings))
+		{
+			check(AudioBuffer->Num() == InConstructParams.Settings.GetNumFramesPerBlock());
+		}
+
+		FDataReferenceCollection GetInputs() const override
+		{
+			FDataReferenceCollection InputDataReferences;
+			InputDataReferences.AddDataReadReference(FOscilatorFactoryBase::EnabledPinName, Enabled);
+			InputDataReferences.AddDataReadReference(FOscilatorFactoryBase::BaseFrequencyPinName, BaseFrequency);
+			InputDataReferences.AddDataReadReference(FOscilatorFactoryBase::PhaseOffsetPinName, PhaseOffset);
+			InputDataReferences.AddDataReadReference(FOscilatorFactoryBase::PhaseResetPinName, PhaseReset);
+			InputDataReferences.AddDataReadReference(FOscilatorFactoryBase::BiPolarPinName, BiPolar);
+			return InputDataReferences;
+		}
+
+		FDataReferenceCollection GetOutputs() const override
+		{
+			FDataReferenceCollection OutputDataReferences;
+			OutputDataReferences.AddDataReadReference(FOscilatorFactoryBase::AudioOutPinName, AudioBuffer);
+			return OutputDataReferences;
+		}
+
+		void ResetPhase(float InPhaseInDegrees)
+		{
+			float LinearPhase = FMath::Clamp(InPhaseInDegrees, 0.f, 360.f) / 360.f;
+			
+			// Recreate the generator type with the phase requested.
+			Generator = GeneratorPolicy{ LinearPhase };
+		}
+
+		void Execute()
+		{
+			// Clamp frequencies into Nyquist range.
+			const float ClampedFreq = FMath::Clamp(*BaseFrequency, -Nyquist, Nyquist);
+			
+			AudioBuffer->Zero();
+
+			Derived* Self = static_cast<Derived*>(this);
+			PhaseReset->ExecuteBlock
+			(
+				[Self, ClampedFreq](int32 InFrameStart, int32 InFrameEnd) 
+				{
+					Self->Generate(InFrameStart, InFrameEnd, ClampedFreq);
+				},
+				[Self, ClampedFreq](int32 InFrameStart, int32 InFrameEnd) 
+				{
+					Self->ResetPhase(*Self->PhaseOffset);
+					Self->Generate(InFrameStart, InFrameEnd, ClampedFreq);
+				}
+			);
+
+			// Convert to Unipolar
+			if (*Enabled && !*BiPolar)
+			{
+				Audio::ConvertBipolarBufferToUnipolar(AudioBuffer->GetData(), AudioBuffer->Num());
+			}
+		}
+
+	protected:
+		GeneratorPolicy Generator;
+		float SampleRate;
+		float Nyquist;
+
+		FBoolReadRef Enabled;
+		FFloatReadRef BaseFrequency;
+		FTriggerReadRef PhaseReset;
+		FFloatReadRef PhaseOffset;
+		FBoolReadRef BiPolar;
+
+		FAudioBufferWriteRef AudioBuffer;
+	};
+
+	// Generic Oscillator operator for NON-FM Operators. 
+	template<typename GeneratorPolicy>
+	class TOscillatorOperator final : public TOscillatorOperatorBase<GeneratorPolicy, TOscillatorOperator<GeneratorPolicy>>
+	{
+		using Super = TOscillatorOperatorBase<GeneratorPolicy, TOscillatorOperator<GeneratorPolicy>>;
+	public:
+		TOscillatorOperator(const FOscillatorOperatorConstructParams& InConstructParams)
+			: Super(InConstructParams)
+		{}
+
+		void Generate(int32 InStartFrame, int32 InEndFrame, float InClampedFreq)
+		{		
+			int32 NumFrames = InEndFrame - InStartFrame;
+
+			if (*this->Enabled && NumFrames > 0)
+			{
+				this->Generator(
+				{
+					this->SampleRate, 
+					InClampedFreq,
+					0.f,
+					MakeArrayView(this->AudioBuffer->GetData() + InStartFrame, NumFrames), // Not aligned.
+				});
+			}			
+		}
+	};
+
+	// Generic Oscillator operator for FM Operators.
+	template<typename TGeneratorPolicy>
+	class TOscillatorOperatorFM final : public TOscillatorOperatorBase<TGeneratorPolicy, TOscillatorOperatorFM<TGeneratorPolicy>>
+	{
+		using Super = TOscillatorOperatorBase<TGeneratorPolicy, TOscillatorOperatorFM<TGeneratorPolicy>>;
+	public:
+		TOscillatorOperatorFM(const FOscillatorOperatorConstructParams& InCommonParams, const FAudioBufferReadRef& InFmData)
+			: Super(InCommonParams), Fm(InFmData)
+		{}
+
+		FDataReferenceCollection GetInputs() const override
+		{
+			FDataReferenceCollection Inputs = Super::GetInputs();
+			Inputs.AddDataReadReference(FOscilatorFactoryBase::FrequencyModPinName, Fm);
+			return Inputs;
+		}
+
+		void Generate(int32 InStartFrame, int32 InEndFrame, float InClampedFreq)
+		{
+			int32 NumFrames = InEndFrame - InStartFrame;
+			if (*this->Enabled && NumFrames > 0)
+			{
+				this->Generator(
+				{
+					this->SampleRate, 
+					InClampedFreq,
+					0.f,
+					MakeArrayView(this->AudioBuffer->GetData() + InStartFrame, NumFrames), // Not aligned.
+					MakeArrayView(this->Fm->GetData() + InStartFrame, NumFrames) // Not aligned.
+				});
+			}
+		}		
+
+	private:
+		FAudioBufferReadRef Fm;
+	};
+
+	FOscilatorNodeBase::FOscilatorNodeBase(const FString& InInstanceName, const FGuid& InInstanceID, const FNodeClassMetadata& InInfo, const TSharedRef<IOperatorFactory, ESPMode::ThreadSafe>& InFactory, float InDefaultFrequency, bool bInDefaultEnablement)
+		: FNode(InInstanceName, InInstanceID, InInfo)
+		, Factory(InFactory)
+		, VertexInterface(GetMetadata().DefaultInterface)
+		, DefaultFrequency(InDefaultFrequency)
+		, bDefaultEnablement(bInDefaultEnablement)
+	{}
+	
+#pragma endregion Common
+
+#pragma region Sine
+
+	enum class ESineGenerationType
+	{
+		Sinf,
+		FilterResonanceFeedback,
+		Rotation,
+		Bhaskara,
+		Wavetable,
+	};
+
+	DECLARE_METASOUND_ENUM(ESineGenerationType, ESineGenerationType::Sinf, METASOUNDSTANDARDNODES_API,
+	FEnumSineGenerationType, FEnumSineGenerationTypeInfo, FEnumSineGenerationTypeReadRef, FEnumSineGenerationTypeWriteRef);
+
+	DEFINE_METASOUND_ENUM_BEGIN(ESineGenerationType, FEnumSineGenerationType, "SineGenerationType")
+		DEFINE_METASOUND_ENUM_ENTRY(ESineGenerationType::Sinf, LOCTEXT("SinfDescription", "Pure Math"), LOCTEXT("SinfDescriptionTT", "Uses Pure Math (Sinf) to generate the Sine Wave (most expensive)")),
+		DEFINE_METASOUND_ENUM_ENTRY(ESineGenerationType::Rotation, LOCTEXT("RotationDescription", "2D Rotation"), LOCTEXT("RotationDescriptionTT", "Rotates around the unit circle generate the sinewave")),
+		DEFINE_METASOUND_ENUM_ENTRY(ESineGenerationType::Bhaskara, LOCTEXT("BhaskaraDescription", "Bhaskara"), LOCTEXT("BhaskaraDescriptionTT", "Sine approximation using Bhaskara technique discovered in 7th century")),
+		DEFINE_METASOUND_ENUM_ENTRY(ESineGenerationType::Wavetable, LOCTEXT("WavetableDescription", "Wavetable"), LOCTEXT("WavetableDescriptionTT", "Uses wavetable to generate the sinewave")),
+		DEFINE_METASOUND_ENUM_ENTRY(ESineGenerationType::FilterResonanceFeedback, LOCTEXT("FRFDescription", "Filter resonance"), LOCTEXT("FRFDescriptionTT", "Filter resonance/feedback")),
+	DEFINE_METASOUND_ENUM_END()
+
+	class FSineOscilatorNode::FFactory : public FOscilatorFactoryBase
+	{
+	public:
+		FFactory() = default;
+		
+		static constexpr const TCHAR* SineTypePin = TEXT("Type");
+
+		static const FNodeClassMetadata& GetNodeInfo()
+		{
+			auto InitNodeInfo = []() -> FNodeClassMetadata
+			{
+				FNodeClassMetadata Info;
+				Info.ClassName = { Metasound::StandardNodes::Namespace, TEXT("Sine"), Metasound::StandardNodes::AudioVariant };
+				Info.MajorVersion = 1;
+				Info.MinorVersion = 0;
+				Info.DisplayName = LOCTEXT("Metasound_SineNodeDisplayName", "Sine");
+				Info.Description = LOCTEXT("Metasound_SineNodeDescription", "Emits an audio signal of a sinusoid.");
+				Info.Author = PluginAuthor;
+				Info.PromptIfMissing = PluginNodeMissingPrompt;
+				Info.DefaultInterface = GetVertexInterface();
+				return Info;
+			};
+			static const FNodeClassMetadata Info = InitNodeInfo();
+			return Info;
+		}
+		static const FVertexInterface& GetVertexInterface()
+		{
+			auto MakeInterface = []() -> FVertexInterface
+			{
+				FVertexInterface Interface = GetCommmonVertexInterface();
+				Interface.GetInputInterface().Add(
+					TInputDataVertexModel<FEnumSineGenerationType>(SineTypePin, LOCTEXT("SineTypeDescription", "Type of the Sinewave Generator"))
+				);
+				return Interface;
+			};
+			static const FVertexInterface Interface = MakeInterface();
+			return Interface;
+		}
+		
+		TUniquePtr<IOperator> CreateOperator(const FCreateOperatorParams& InParams, FBuildErrorArray& OutErrors) override
+		{
+			const FSineOscilatorNode& SineNode = static_cast<const FSineOscilatorNode&>(InParams.Node);
+			const FDataReferenceCollection& InputCol = InParams.InputDataReferences;
+			const FOperatorSettings& Settings = InParams.OperatorSettings;
+			using namespace Generators;
+			
+			FOscillatorOperatorConstructParams OpParams
+			{
+				Settings,
+				InputCol.GetDataReadReferenceOrConstruct<bool>(EnabledPinName, SineNode.GetDefaultEnablement()),
+				InputCol.GetDataReadReferenceOrConstruct<float>(BaseFrequencyPinName, SineNode.GetDefaultFrequency()),
+				InputCol.GetDataReadReferenceOrConstruct<float>(PhaseOffsetPinName, SineNode.GetDefaultPhaseOffset()),
+				InputCol.GetDataReadReferenceOrConstruct<FTrigger>(PhaseResetPinName, Settings),
+				InputCol.GetDataReadReferenceOrConstruct<bool>(BiPolarPinName, true)
+			};
+
+			// TODO: Make this a static prop. For now its a pin.
+			FEnumSineGenerationTypeReadRef Type = InputCol.GetDataReadReferenceOrConstruct<FEnumSineGenerationType>(SineTypePin);
+			
+			// Check to see if we have an FM input connected.
+			bool bHasFM = InputCol.ContainsDataReadReference<FAudioBuffer>(FrequencyModPinName);			
+			if (bHasFM)
+			{
+				// FM Oscillators.
+				FAudioBufferReadRef FmBuffer = InputCol.GetDataReadReference<FAudioBuffer>(FrequencyModPinName);
+				switch (*Type)
+				{
+				default:
+				case ESineGenerationType::Sinf: return MakeUnique<TOscillatorOperatorFM<FSinfWithFm>>(OpParams, FmBuffer);
+				case ESineGenerationType::Bhaskara: return MakeUnique<TOscillatorOperatorFM<FBhaskaraWithFm>>(OpParams, FmBuffer);
+				}
+			}
+			else //HasFM
+			{
+				switch (*Type)
+				{
+				default:
+				case ESineGenerationType::Sinf: return MakeUnique<TOscillatorOperator<FSinf>>(OpParams);
+				case ESineGenerationType::Rotation: return MakeUnique<TOscillatorOperator<F2DRotatorGenerateBlock>>(OpParams);
+				case ESineGenerationType::Bhaskara: return MakeUnique<TOscillatorOperator<FBhaskara>>(OpParams);
+				}
+			} // HasFM
+			return nullptr;
+		}
+	private:
+	};
+
+	FSineOscilatorNode::FSineOscilatorNode(const FString& InInstanceName, const FGuid& InInstanceID, float InDefaultFrequency, bool bInDefaultEnablement)
+		: FOscilatorNodeBase(InInstanceName,InInstanceID, FFactory::GetNodeInfo(), MakeShared<FFactory, ESPMode::ThreadSafe>(), InDefaultFrequency, bInDefaultEnablement )
+	{}
+
+	FSineOscilatorNode::FSineOscilatorNode(const FNodeInitData& InInitData)
+		: FSineOscilatorNode(InInitData.InstanceName, InInitData.InstanceID, 440.0f, true)
+	{}
+
+	METASOUND_REGISTER_NODE(FSineOscilatorNode);
+
+#pragma endregion Sine
+
+#pragma region Saw
+
+	enum class ESawGenerationType
+	{
+		PolySmooth,
+		Trivial,
+		Wavetable
+	};
+
+	DECLARE_METASOUND_ENUM(ESawGenerationType, ESawGenerationType::PolySmooth, METASOUNDSTANDARDNODES_API,
+		FEnumSawGenerationType, FEnumSawGenerationTypeInfo, FSawGenerationTypeReadRef, FEnumSawGenerationTypeWriteRef);
+
+	DEFINE_METASOUND_ENUM_BEGIN(ESawGenerationType, FEnumSawGenerationType, "SawGenerationType")
+		DEFINE_METASOUND_ENUM_ENTRY(ESawGenerationType::PolySmooth, LOCTEXT("SawPolySmoothDescription", "Poly Smooth"), LOCTEXT("PolySmoothDescriptionTT", "PolySmooth (i.e. BLEP)")),
+		DEFINE_METASOUND_ENUM_ENTRY(ESawGenerationType::Trivial, LOCTEXT("SawTrivialDescription", "Trivial"), LOCTEXT("TrivialDescriptionTT", "The most basic raw implementation")),
+		DEFINE_METASOUND_ENUM_ENTRY(ESawGenerationType::Wavetable, LOCTEXT("SawWavetableDescription", "Wavetable"), LOCTEXT("SawWavetableDescriptionTT", "Use a Wavetable iterpolation to generate the Waveform"))
+	DEFINE_METASOUND_ENUM_END()
+
+	class FSawOscilatorNode::FFactory : public FOscilatorFactoryBase
+	{
+	public:
+		static constexpr const TCHAR* SawTypePin = TEXT("Type");
+		
+		FFactory() = default;
+
+		TUniquePtr<IOperator> CreateOperator(const FCreateOperatorParams& InParams, FBuildErrorArray& OutErrors) override;
+
+		static const FNodeClassMetadata& GetNodeInfo()
+		{
+			auto InitNodeInfo = []() -> FNodeClassMetadata
+			{
+				FNodeClassMetadata Info;
+				Info.ClassName = { Metasound::StandardNodes::Namespace, TEXT("Saw"), Metasound::StandardNodes::AudioVariant };
+				Info.MajorVersion = 1;
+				Info.MinorVersion = 0;
+				Info.DisplayName = LOCTEXT("Metasound_SawNodeDisplayName", "Saw");
+				Info.Description = LOCTEXT("Metasound_SawNodeDescription", "Emits a Saw wave");
+				Info.Author = PluginAuthor;
+				Info.PromptIfMissing = PluginNodeMissingPrompt;
+				Info.DefaultInterface = GetVertexInterface();
+				return Info;
+			};
+			static const FNodeClassMetadata Info = InitNodeInfo();
+			return Info;
+		}
+		static const FVertexInterface& GetVertexInterface()
+		{
+			auto MakeInterface = []() -> FVertexInterface
+			{
+				FVertexInterface Interface = GetCommmonVertexInterface();
+				Interface.GetInputInterface().Add(
+					TInputDataVertexModel<FEnumSawGenerationType>(SawTypePin, LOCTEXT("Saw Type", "Which type of Saw Generator to Create"))
+				);
+				return Interface;
+			};
+			static const FVertexInterface Interface = MakeInterface();
+			return Interface;
+		}
+
+	private:
+	};
+
+	TUniquePtr<Metasound::IOperator> FSawOscilatorNode::FFactory::CreateOperator(const FCreateOperatorParams& InParams, FBuildErrorArray& OutErrors)
+	{
+		const FSawOscilatorNode& Node = static_cast<const FSawOscilatorNode&>(InParams.Node);
+		const FDataReferenceCollection& InputCol = InParams.InputDataReferences;
+		const FOperatorSettings& Settings = InParams.OperatorSettings;
+		using namespace Generators;		
+
+		FSawGenerationTypeReadRef Type = InputCol.GetDataReadReferenceOrConstruct<FEnumSawGenerationType>(SawTypePin);
+
+		FOscillatorOperatorConstructParams OpParams
+		{
+			Settings,
+			InputCol.GetDataReadReferenceOrConstruct<bool>(EnabledPinName, Node.GetDefaultEnablement()),
+			InputCol.GetDataReadReferenceOrConstruct<float>(BaseFrequencyPinName, Node.GetDefaultFrequency()),
+			InputCol.GetDataReadReferenceOrConstruct<float>(PhaseOffsetPinName, Node.GetDefaultPhaseOffset()),
+			InputCol.GetDataReadReferenceOrConstruct<FTrigger>(PhaseResetPinName, Settings),
+			InputCol.GetDataReadReferenceOrConstruct<bool>(BiPolarPinName, true)
+		};
+
+		bool bHasFM = InputCol.ContainsDataReadReference<FAudioBuffer>(FrequencyModPinName);
+
+		if (bHasFM)
+		{
+			FAudioBufferReadRef FmBuffer = InputCol.GetDataReadReference<FAudioBuffer>(FrequencyModPinName);
+			switch (*Type)
+			{
+			default:
+			case ESawGenerationType::Trivial: return MakeUnique<TOscillatorOperatorFM<FSawWithFm>>(OpParams, FmBuffer);
+			case ESawGenerationType::PolySmooth: return MakeUnique<TOscillatorOperatorFM<FSawPolysmooth>>(OpParams, FmBuffer);
+			}
+		}
+		else
+		{
+			switch (*Type)
+			{
+			default:
+			case ESawGenerationType::Trivial: return MakeUnique<TOscillatorOperator<FTriangle>>(OpParams);
+			case ESawGenerationType::PolySmooth: return MakeUnique<TOscillatorOperator<FSawPolysmooth>>(OpParams);
+			}
+		}
+		return nullptr;
+	}
+
+	FSawOscilatorNode::FSawOscilatorNode(const FString& InInstanceName, const FGuid& InInstanceID, float InDefaultFrequency, bool bInDefaultEnablement)
+		: FOscilatorNodeBase(InInstanceName, InInstanceID, FFactory::GetNodeInfo(), MakeShared<FFactory, ESPMode::ThreadSafe>(), InDefaultFrequency, bInDefaultEnablement)
+	{}
+
+	FSawOscilatorNode::FSawOscilatorNode(const FNodeInitData& InInitData)
+		: FSawOscilatorNode(InInitData.InstanceName, InInitData.InstanceID, 440.0f, true)
+	{}
+
+	METASOUND_REGISTER_NODE(FSawOscilatorNode);
+
+#pragma endregion Saw
+
+#pragma region Square
+
+	enum class ESquareGenerationType
+	{
+		PolySmooth,
+		Trivial,
+		Wavetable
+	};
+
+	DECLARE_METASOUND_ENUM(ESquareGenerationType, ESquareGenerationType::PolySmooth, METASOUNDSTANDARDNODES_API,
+	FEnumSquareGenerationType, FEnumSquareGenerationTypeInfo, FSquareGenerationTypeReadRef, FEnumSquareGenerationTypeWriteRef);
+
+	DEFINE_METASOUND_ENUM_BEGIN(ESquareGenerationType, FEnumSquareGenerationType, "SquareGenerationType")
+		DEFINE_METASOUND_ENUM_ENTRY(ESquareGenerationType::PolySmooth, LOCTEXT("SquarePolySmoothDescription", "Poly Smooth"), LOCTEXT("PolySmoothDescriptionTT", "PolySmooth (i.e. BLEP)")),
+		DEFINE_METASOUND_ENUM_ENTRY(ESquareGenerationType::Trivial, LOCTEXT("SquareTrivialDescription", "Trivial"), LOCTEXT("TrivialDescriptionTT", "The most basic raw implementation")),
+		DEFINE_METASOUND_ENUM_ENTRY(ESquareGenerationType::Wavetable, LOCTEXT("SquareWavetableDescription", "Wavetable"), LOCTEXT("SquareWavetableDescriptionTT", "Use a Wavetable interpolation to generate the Waveform"))
+	DEFINE_METASOUND_ENUM_END()
+	
+	template<typename GeneratorPolicy>
+	class FSquareOperator final : public TOscillatorOperatorBase<GeneratorPolicy, FSquareOperator<GeneratorPolicy>>
+	{
+		using Super = TOscillatorOperatorBase<GeneratorPolicy, FSquareOperator<GeneratorPolicy>>;
+	public:
+		FSquareOperator(const FOscillatorOperatorConstructParams& InConstructParams, const FFloatReadRef& InPulseWidth)
+			: Super(InConstructParams)
+			, PulseWidth(InPulseWidth)
+		{}
+		void Generate(int32 InStartFrame, int32 InEndFrame, float InClampedFreq)
+		{
+			int32 NumFrames = InEndFrame - InStartFrame;
+
+			if (*this->Enabled && NumFrames > 0)
+			{
+				this->Generator(
+				{
+					this->SampleRate,
+					InClampedFreq,
+					*PulseWidth,
+					MakeArrayView(this->AudioBuffer->GetData() + InStartFrame, NumFrames), // Not aligned.
+				});
+			}
+		}
+
+	private:
+		FFloatReadRef PulseWidth;
+	};
+
+	template<typename GeneratorPolicy>
+	class FSquareOperatorFM final : public TOscillatorOperatorBase<GeneratorPolicy, FSquareOperatorFM<GeneratorPolicy>>
+	{
+		using Super = TOscillatorOperatorBase<GeneratorPolicy, FSquareOperatorFM<GeneratorPolicy>>;
+	public:
+		FSquareOperatorFM(const FOscillatorOperatorConstructParams& InConstructParams, const FFloatReadRef& InPulseWidth, const FAudioBufferReadRef& InFm)
+			: Super(InConstructParams)
+			, PulseWidth(InPulseWidth)
+			, FM(InFm)
+		{
+			check(InFm->GetData());
+			check(InConstructParams.Settings.GetNumFramesPerBlock() == InFm->Num());
+		}
+		void Generate(int32 InStartFrame, int32 InEndFrame, float InClampedFreq)
+		{
+			int32 NumFrames = InEndFrame - InStartFrame;
+
+			if (*this->Enabled && NumFrames > 0)
+			{
+				this->Generator(
+				{
+					this->SampleRate,
+					InClampedFreq,
+					*PulseWidth,
+					MakeArrayView(this->AudioBuffer->GetData() + InStartFrame, NumFrames), // Not aligned.
+					MakeArrayView(this->FM->GetData() + InStartFrame, NumFrames), // Not aligned.
+				});
+			}
+		}
+
+	private:
+		FFloatReadRef PulseWidth;
+		FAudioBufferReadRef FM;
+	};
+
+	class FSquareOscilatorNode::FFactory : public FOscilatorFactoryBase
+	{
+	public:
+		FFactory() = default;
+
+		static constexpr const TCHAR* PulseWidthPinName = TEXT("Pulse Width");
+		static constexpr const TCHAR* SquareTypePinName = TEXT("Type");
+
+		TUniquePtr<IOperator> CreateOperator(const FCreateOperatorParams& InParams, FBuildErrorArray& OutErrors) override
+		{
+			const FSquareOscilatorNode& Node = static_cast<const FSquareOscilatorNode&>(InParams.Node);
+			const FDataReferenceCollection& InputCol = InParams.InputDataReferences;
+			const FOperatorSettings& Settings = InParams.OperatorSettings;
+			using namespace Generators;
+
+			FOscillatorOperatorConstructParams OpParams
+			{
+				Settings,
+				InputCol.GetDataReadReferenceOrConstruct<bool>(EnabledPinName, Node.GetDefaultEnablement()),
+				InputCol.GetDataReadReferenceOrConstruct<float>(BaseFrequencyPinName, Node.GetDefaultFrequency()),
+				InputCol.GetDataReadReferenceOrConstruct<float>(PhaseOffsetPinName, Node.GetDefaultPhaseOffset()),
+				InputCol.GetDataReadReferenceOrConstruct<FTrigger>(PhaseResetPinName, Settings),
+				InputCol.GetDataReadReferenceOrConstruct<bool>(BiPolarPinName, true)
+			};
+
+			FSquareGenerationTypeReadRef Type = InputCol.GetDataReadReferenceOrConstruct<FEnumSquareGenerationType>(SquareTypePinName);
+
+			bool bHasFM = InputCol.ContainsDataReadReference<FAudioBuffer>(FrequencyModPinName);
+			FFloatReadRef PulseWidth = InputCol.GetDataReadReferenceOrConstruct<float>(PulseWidthPinName, 0.5f);
+
+			if (bHasFM)
+			{
+				FAudioBufferReadRef FmBuffer = InputCol.GetDataReadReference<FAudioBuffer>(FrequencyModPinName);
+				switch (*Type)
+				{
+				default:				
+				case ESquareGenerationType::Trivial: return MakeUnique<FSquareOperatorFM<FSquareWithFm>>(OpParams, PulseWidth, FmBuffer);
+				case ESquareGenerationType::PolySmooth: return MakeUnique<FSquareOperatorFM<FSquarePolysmoothWithFm>>(OpParams, PulseWidth, FmBuffer);
+				}
+			}
+			else
+			{
+				switch (*Type)
+				{
+				default:
+				case ESquareGenerationType::Trivial: return MakeUnique<FSquareOperator<FSquare>>(OpParams, PulseWidth);
+				case ESquareGenerationType::PolySmooth: return MakeUnique<FSquareOperator<FSquarePolysmooth>>(OpParams, PulseWidth);
+				}
+			}
+			return nullptr;
+		}
+
+		static const FNodeClassMetadata& GetNodeInfo()
+		{
+			auto InitNodeInfo = []() -> FNodeClassMetadata
+			{
+				FNodeClassMetadata Info;
+				Info.ClassName = { Metasound::StandardNodes::Namespace, TEXT("Square"), Metasound::StandardNodes::AudioVariant };
+				Info.MajorVersion = 1;
+				Info.MinorVersion = 0;
+				Info.DisplayName = LOCTEXT("Metasound_SquareNodeDisplayName", "Square");
+				Info.Description = LOCTEXT("Metasound_SquareNodeDescription", "Emits a Square wave");
+				Info.Author = PluginAuthor;
+				Info.PromptIfMissing = PluginNodeMissingPrompt;
+				Info.DefaultInterface = GetVertexInterface();
+				return Info;
+			};
+			static const FNodeClassMetadata Info = InitNodeInfo();
+			return Info;
+		}
+		static const FVertexInterface& GetVertexInterface()
+		{
+			auto MakeInterface = []() -> FVertexInterface
+			{
+				FVertexInterface Interface = GetCommmonVertexInterface();
+				Interface.GetInputInterface().Add(TInputDataVertexModel<FEnumSquareGenerationType>(SquareTypePinName, LOCTEXT("SquareTypeDescription", "The generator type to make the squarewave")));
+				Interface.GetInputInterface().Add(TInputDataVertexModel<float>(PulseWidthPinName, LOCTEXT("PhaseWidthDescription", "The Width of the square part of the wave"), 1.f));
+				return Interface;
+			};
+			static const FVertexInterface Interface = MakeInterface();
+			return Interface;
+		}
+	private:
+	};
+	
+	FSquareOscilatorNode::FSquareOscilatorNode(const FString& InInstanceName, const FGuid& InInstanceID, float InDefaultFrequency, bool bInDefaultEnablement)
+		: FOscilatorNodeBase(InInstanceName, InInstanceID, FFactory::GetNodeInfo(), MakeShared<FFactory, ESPMode::ThreadSafe>(), InDefaultFrequency, bInDefaultEnablement)
+	{}
+
+	FSquareOscilatorNode::FSquareOscilatorNode(const FNodeInitData& InInitData)
+		: FSquareOscilatorNode(InInitData.InstanceName, InInitData.InstanceID, 440.0f, true)
+	{}
+
+	METASOUND_REGISTER_NODE(FSquareOscilatorNode)
+#pragma endregion Square
+
+#pragma region Triangle
+
+	enum class ETriangleGenerationType
+	{
+		PolySmooth,
+		Trivial,
+		Wavetable
+	};
+
+	DECLARE_METASOUND_ENUM(ETriangleGenerationType, ETriangleGenerationType::PolySmooth, METASOUNDSTANDARDNODES_API,
+	FEnumTriangleGenerationType, FEnumTriangleGenerationTypeInfo, FTriangleGenerationTypeReadRef, FEnumTriangleGenerationTypeWriteRef);
+
+	DEFINE_METASOUND_ENUM_BEGIN(ETriangleGenerationType, FEnumTriangleGenerationType, "TriangleGenerationType")
+		DEFINE_METASOUND_ENUM_ENTRY(ETriangleGenerationType::PolySmooth, LOCTEXT("TrianglePolySmoothDescription", "Poly Smooth"), LOCTEXT("PolySmoothDescriptionTT", "PolySmooth (i.e. BLEP)")),
+		DEFINE_METASOUND_ENUM_ENTRY(ETriangleGenerationType::Trivial, LOCTEXT("TriangleTrivialDescription", "Trivial"), LOCTEXT("TrivialDescriptionTT", "The most basic raw implementation")),
+		DEFINE_METASOUND_ENUM_ENTRY(ETriangleGenerationType::Wavetable, LOCTEXT("TriangleWavetableDescription", "Wavetable"), LOCTEXT("TriangleWavetableDescriptionTT", "Use a Wavetable iterpolation to generate the Waveform"))
+	DEFINE_METASOUND_ENUM_END()
+
+	class FTriangleOscilatorNode::FFactory : public FOscilatorFactoryBase
+	{
+	public:
+		FFactory() = default;
+
+		static constexpr const TCHAR* TriangeTypePinName = TEXT("Type");
+
+		TUniquePtr<IOperator> CreateOperator(const FCreateOperatorParams& InParams, FBuildErrorArray& OutErrors) override
+		{
+			const FTriangleOscilatorNode& Node = static_cast<const FTriangleOscilatorNode&>(InParams.Node);
+			const FDataReferenceCollection& InputCol = InParams.InputDataReferences;
+			const FOperatorSettings& Settings = InParams.OperatorSettings;
+			using namespace Generators;
+			
+			FTriangleGenerationTypeReadRef Type = InputCol.GetDataReadReferenceOrConstruct<FEnumTriangleGenerationType>(TriangeTypePinName);
+		
+			FOscillatorOperatorConstructParams OpParams
+			{
+				Settings,
+				InputCol.GetDataReadReferenceOrConstruct<bool>(EnabledPinName, Node.GetDefaultEnablement()),
+				InputCol.GetDataReadReferenceOrConstruct<float>(BaseFrequencyPinName, Node.GetDefaultFrequency()),
+				InputCol.GetDataReadReferenceOrConstruct<float>(PhaseOffsetPinName, Node.GetDefaultPhaseOffset()),
+				InputCol.GetDataReadReferenceOrConstruct<FTrigger>(PhaseResetPinName, Settings),
+				InputCol.GetDataReadReferenceOrConstruct<bool>(BiPolarPinName, true)
+			};
+
+			bool bHasFM = InputCol.ContainsDataReadReference<FAudioBuffer>(FrequencyModPinName);
+			
+			if (bHasFM)
+			{
+				FAudioBufferReadRef FmBuffer = InputCol.GetDataReadReference<FAudioBuffer>(FrequencyModPinName);
+				switch (*Type)
+				{
+				default:
+				case ETriangleGenerationType::Trivial: return MakeUnique<TOscillatorOperatorFM<FTriangleWithFm>>(OpParams, FmBuffer);
+				}
+			}
+			else
+			{
+				switch (*Type)
+				{
+				default:
+				case ETriangleGenerationType::Trivial: return MakeUnique<TOscillatorOperator<FTriangle>>(OpParams);
+				}
+			}
+			return nullptr;
+		}
+
+		static const FNodeClassMetadata& GetNodeInfo()
+		{
+			auto InitNodeInfo = []() -> FNodeClassMetadata
+			{
+				FNodeClassMetadata Info;
+				Info.ClassName = { Metasound::StandardNodes::Namespace, TEXT("Triangle"), Metasound::StandardNodes::AudioVariant };
+				Info.MajorVersion = 1;
+				Info.MinorVersion = 0;
+				Info.DisplayName = LOCTEXT("Metasound_TriangleNodeDisplayName", "Triangle");
+				Info.Description = LOCTEXT("Metasound_TriangleNodeDescription", "Emits a Triangle wave");
+				Info.Author = PluginAuthor;
+				Info.PromptIfMissing = PluginNodeMissingPrompt;
+				Info.DefaultInterface = GetVertexInterface();
+				return Info;
+			};
+			static const FNodeClassMetadata Info = InitNodeInfo();
+			return Info;
+		}
+		static const FVertexInterface& GetVertexInterface()
+		{
+			auto MakeInterface = []() -> FVertexInterface
+			{
+				FVertexInterface Interface = GetCommmonVertexInterface();
+				Interface.GetInputInterface().Add(
+					TInputDataVertexModel<FEnumTriangleGenerationType>(TriangeTypePinName, LOCTEXT("TriangleTypeDescription", "The generator type to make the triangle wave"))
+				);
+				return Interface;
+			};
+			static const FVertexInterface Interface = MakeInterface();
+			return Interface;
+		}
+	private:
+	};
+	
+	FTriangleOscilatorNode::FTriangleOscilatorNode(const FString& InInstanceName, const FGuid& InInstanceID, float InDefaultFrequency, bool bInDefaultEnablement)
+		: FOscilatorNodeBase(InInstanceName, InInstanceID, FFactory::GetNodeInfo(), MakeShared<FFactory, ESPMode::ThreadSafe>(), InDefaultFrequency, bInDefaultEnablement)
+	{
+
+	}
+
+	FTriangleOscilatorNode::FTriangleOscilatorNode(const FNodeInitData& InInitData)
+		: FTriangleOscilatorNode(InInitData.InstanceName, InInitData.InstanceID, 440.0f, true)
+	{}
+
+	METASOUND_REGISTER_NODE(FTriangleOscilatorNode);
+#pragma endregion Triangle
+}
+#undef LOCTEXT_NAMESPACE //MetasoundStandardNodes
\ No newline at end of file
diff --git a/Engine/Plugins/Runtime/Metasound/Source/MetasoundStandardNodes/Public/MetasoundOscillatorNodes.h b/Engine/Plugins/Runtime/Metasound/Source/MetasoundStandardNodes/Public/MetasoundOscillatorNodes.h
new file mode 100644
--- /dev/null
+++ b/Engine/Plugins/Runtime/Metasound/Source/MetasoundStandardNodes/Public/MetasoundOscillatorNodes.h
@@ -0,0 +1,102 @@
+// Copyright Epic Games, Inc. All Rights Reserved.
+
+#pragma once
+
+#include "CoreMinimal.h"
+
+#include "MetasoundAudioBuffer.h"
+#include "MetasoundBuilderInterface.h"
+#include "MetasoundDataReferenceCollection.h"
+#include "MetasoundExecutableOperator.h"
+#include "MetasoundFacade.h"
+#include "MetasoundNode.h"
+#include "MetasoundOperatorInterface.h"
+
+namespace Metasound
+{
+	class METASOUNDSTANDARDNODES_API FOscilatorNodeBase : public FNode
+	{
+	protected:
+		FOscilatorNodeBase(const FString& InInstanceName, const FGuid& InInstanceID, const FNodeClassMetadata& InInfo, const TSharedRef<IOperatorFactory, ESPMode::ThreadSafe>& InFactory, float InDefaultFrequency, bool bInDefaultEnablement);
+		FOscilatorNodeBase(const FNodeInitData& InInitData);
+
+		float GetDefaultPhaseOffset() const 
+		{
+			return DefaultPhaseOffset; 
+		}
+
+		float GetDefaultFrequency() const 
+		{
+			return DefaultFrequency;
+		}
+
+		bool GetDefaultEnablement() const 
+		{
+			return bDefaultEnablement;
+		}
+
+		FOperatorFactorySharedRef GetDefaultOperatorFactory() const override
+		{
+			return Factory;
+		}
+
+		const FVertexInterface& GetVertexInterface() const override
+		{
+			return VertexInterface;
+		}
+
+		bool SetVertexInterface(const FVertexInterface& InInterface) override
+		{
+			if (IsVertexInterfaceSupported(InInterface))
+			{
+				VertexInterface = InInterface;
+				return true;
+			}
+			return false;
+		}
+
+		bool IsVertexInterfaceSupported(const FVertexInterface& InInterface) const override
+		{
+			return VertexInterface == InInterface;
+		}
+
+		TSharedRef<IOperatorFactory, ESPMode::ThreadSafe> Factory;
+		FVertexInterface VertexInterface;
+
+		float DefaultPhaseOffset = 0.f;
+		float DefaultFrequency = 440.f;
+		bool bDefaultEnablement = true;
+	};
+
+	class METASOUNDSTANDARDNODES_API FSineOscilatorNode : public FOscilatorNodeBase
+	{
+	public:
+		class FFactory;		
+		FSineOscilatorNode(const FString& InInstanceName, const FGuid& InInstanceID, float InDefaultFrequency, bool bInDefaultEnablement);
+		FSineOscilatorNode(const FNodeInitData& InInitData);
+	};
+
+	class METASOUNDSTANDARDNODES_API FSawOscilatorNode : public FOscilatorNodeBase
+	{
+	public:
+		class FFactory;
+		FSawOscilatorNode(const FString& InInstanceName, const FGuid& InInstanceID, float InDefaultFrequency, bool bInDefaultEnablement);
+		FSawOscilatorNode(const FNodeInitData& InInitData);
+	};
+
+	class METASOUNDSTANDARDNODES_API FTriangleOscilatorNode : public FOscilatorNodeBase
+	{
+	public:
+		class FFactory;
+		FTriangleOscilatorNode(const FString& InInstanceName, const FGuid& InInstanceID, float InDefaultFrequency, bool bInDefaultEnablement);
+		FTriangleOscilatorNode(const FNodeInitData& InInitData);
+	};
+
+	class METASOUNDSTANDARDNODES_API FSquareOscilatorNode : public FOscilatorNodeBase
+	{
+	public:
+		class FFactory;
+		FSquareOscilatorNode(const FString& InInstanceName, const FGuid& InInstanceID, float InDefaultFrequency, bool bInDefaultEnablement);
+		FSquareOscilatorNode(const FNodeInitData& InInitData);
+	};
+}
diff --git a/Engine/Source/Runtime/SignalProcessing/Public/DSP/Dsp.h b/Engine/Source/Runtime/SignalProcessing/Public/DSP/Dsp.h
--- a/Engine/Source/Runtime/SignalProcessing/Public/DSP/Dsp.h
+++ b/Engine/Source/Runtime/SignalProcessing/Public/DSP/Dsp.h
@@ -270,6 +270,26 @@ namespace Audio
 		return 0.5f * X + 0.5f;
 	}
 
+	// Converts in place a buffer into Unipolar (0..1)
+	// This is a Vector op version of the GetUnipolar function. (0.5f * X + 0.5f).
+	static void ConvertBipolarBufferToUnipolar(float* InAlignedBuffer, int32 NumSamples)
+	{
+		// Make sure buffers are aligned and we can do a whole number of loops.
+		check(IsAligned(InAlignedBuffer, sizeof(VectorRegister)));
+		check(NumSamples % 4 == 0);
+
+		const VectorRegister Half = VectorSetFloat1(0.5f);
+
+		// Process buffer 1 vector (4 floats) at a time.
+		for(int32 i = NumSamples / 4; i; --i, InAlignedBuffer += 4)
+		{
+			VectorRegister V = VectorLoadAligned(InAlignedBuffer);
+			V = VectorMultiply(V, Half);
+			V = VectorAdd(V, Half);
+			VectorStoreAligned(V, InAlignedBuffer);
+		}
+	}
+
 	// Using midi tuning standard, compute frequency in hz from midi value
 	static FORCEINLINE float GetFrequencyFromMidi(const float InMidiNote)
 	{
