Add back Enum editor support to Metasound inputs

diff --git a/Engine/Plugins/Runtime/Metasound/Source/MetasoundEditor/Private/MetasoundEditorGraphBuilder.h b/Engine/Plugins/Runtime/Metasound/Source/MetasoundEditor/Private/MetasoundEditorGraphBuilder.h
--- a/Engine/Plugins/Runtime/Metasound/Source/MetasoundEditor/Private/MetasoundEditorGraphBuilder.h
+++ b/Engine/Plugins/Runtime/Metasound/Source/MetasoundEditor/Private/MetasoundEditorGraphBuilder.h
@@ -99,6 +99,7 @@ namespace Metasound
 			// Constructs graph with default inputs & outputs.
 			static void ConstructGraph(UObject& InMetasound);
 
+			// Rebuilds all editor node pins based on the provided node handle's class definition.
 			static void RebuildNodePins(UMetasoundEditorGraphNode& InGraphNode, Frontend::FNodeHandle InNodeHandle);
 
 			// Removes all literal inputs connected to the given node
diff --git a/Engine/Plugins/Runtime/Metasound/Source/MetasoundEditor/Private/MetasoundEditorModule.cpp b/Engine/Plugins/Runtime/Metasound/Source/MetasoundEditor/Private/MetasoundEditorModule.cpp
--- a/Engine/Plugins/Runtime/Metasound/Source/MetasoundEditor/Private/MetasoundEditorModule.cpp
+++ b/Engine/Plugins/Runtime/Metasound/Source/MetasoundEditor/Private/MetasoundEditorModule.cpp
@@ -287,6 +287,10 @@ namespace Metasound
 					UMetasoundEditorGraphInputNode::StaticClass()->GetFName(),
 					FOnGetDetailCustomizationInstance::CreateLambda([]() { return MakeShared<FMetasoundInputNodeDetailCustomization>(); }));
 
+				PropertyModule.RegisterCustomPropertyTypeLayout(
+					"MetasoundEditorGraphInputInt",
+					FOnGetPropertyTypeCustomizationInstance::CreateLambda([]() { return MakeShared<FMetasoundInputNodeIntDetailCustomization>(); }));
+
 				PropertyModule.RegisterCustomPropertyTypeLayout(
 					"MetasoundEditorGraphInputObject",
 					FOnGetPropertyTypeCustomizationInstance::CreateLambda([]() { return MakeShared<FMetasoundInputNodeObjectDetailCustomization>(); }));
diff --git a/Engine/Plugins/Runtime/Metasound/Source/MetasoundEditor/Private/MetasoundInputNodeDetailCustomization.cpp b/Engine/Plugins/Runtime/Metasound/Source/MetasoundEditor/Private/MetasoundInputNodeDetailCustomization.cpp
--- a/Engine/Plugins/Runtime/Metasound/Source/MetasoundEditor/Private/MetasoundInputNodeDetailCustomization.cpp
+++ b/Engine/Plugins/Runtime/Metasound/Source/MetasoundEditor/Private/MetasoundInputNodeDetailCustomization.cpp
@@ -11,7 +11,6 @@
 #include "IDetailChildrenBuilder.h"
 #include "IDetailGroup.h"
 #include "Internationalization/Text.h"
-#include "IPropertyTypeCustomization.h"
 #include "MetasoundAssetBase.h"
 #include "MetasoundEditorGraphBuilder.h"
 #include "MetasoundEditorGraphNode.h"
@@ -43,7 +42,7 @@ namespace Metasound
 {
 	namespace Editor
 	{
-		namespace NodeCustomizationUtils
+		namespace InputNodeCustomizationPrivate
 		{
 			/** Minimum size of the details title panel */
 			static const float DetailsTitleMinWidth = 125.f;
@@ -57,6 +56,7 @@ namespace Metasound
 
 			static const FText InputNodeNameText = LOCTEXT("InputNode_Name", "Input Name");
 
+			static const FName DataTypeNameIdentifier = "DataTypeName";
 			static const FName ProxyGeneratorClassNameIdentifier = "GeneratorClass";
 
 			void UpdateInputLiteralFromNode(const TWeakObjectPtr<UMetasoundEditorGraphInputNode>& GraphNode)
@@ -80,15 +80,103 @@ namespace Metasound
 			}
 		}
 
-		void FMetasoundInputNodeObjectDetailCustomization::CustomizeHeader(TSharedRef<IPropertyHandle> StructPropertyHandle, FDetailWidgetRow& HeaderRow, IPropertyTypeCustomizationUtils& StructCustomizationUtils)
+		void FMetasoundInputNodeIntDetailCustomization::CacheProxyData(TSharedPtr<IPropertyHandle> ProxyHandle)
 		{
+			DataTypeName = FName();
+
+			if (const FString* MetadataDataTypeName = ProxyHandle->GetInstanceMetaData(InputNodeCustomizationPrivate::DataTypeNameIdentifier))
+			{
+				DataTypeName = **MetadataDataTypeName;
+			}
 		}
 
-		void FMetasoundInputNodeObjectDetailCustomization::CacheProxyClass(TSharedPtr<IPropertyHandle> InPropertyHandle)
+		TSharedRef<SWidget> FMetasoundInputNodeIntDetailCustomization::CreateStructureWidget(TSharedPtr<IPropertyHandle>& StructPropertyHandle) const
+		{
+			using namespace Frontend;
+
+			if (FMetasoundFrontendRegistryContainer* Registry = FMetasoundFrontendRegistryContainer::Get())
+			{
+				TSharedPtr<IPropertyHandle> ValueProperty = StructPropertyHandle->GetChildHandle(GET_MEMBER_NAME_CHECKED(FMetasoundEditorGraphInputInt, Value));
+				if (ValueProperty.IsValid())
+				{
+					TSharedPtr<const IEnumDataTypeInterface> EnumInterface = Registry->GetEnumInterfaceForDataType(DataTypeName);
+
+					// Not an enum, so just display as underlying type (int32)
+					if (!EnumInterface.IsValid())
+					{
+						return ValueProperty->CreatePropertyValueWidget();
+					}
+
+					auto GetAll = [Interface = EnumInterface](TArray<TSharedPtr<FString>>& OutStrings, TArray<TSharedPtr<SToolTip>>& OutTooltips, TArray<bool>&)
+					{
+						for (const IEnumDataTypeInterface::FGenericInt32Entry& i : Interface->GetAllEntries())
+						{
+							OutTooltips.Emplace(SNew(SToolTip).Text(i.Tooltip));
+							OutStrings.Emplace(MakeShared<FString>(i.DisplayName.ToString()));
+						}
+					};
+					auto GetValue = [Interface = EnumInterface, Prop = ValueProperty]()
+					{
+						int32 IntValue;
+						if (Prop->GetValue(IntValue) != FPropertyAccess::Success)
+						{
+							IntValue = Interface->GetDefaultValue();
+							UE_LOG(LogTemp, Warning, TEXT("Failed to read int Property '%s', defaulting."), *GetNameSafe(Prop->GetProperty()));
+						}
+						if (TOptional<IEnumDataTypeInterface::FGenericInt32Entry> Result = Interface->FindByValue(IntValue))
+						{
+							return Result->DisplayName.ToString();
+						}
+						UE_LOG(LogTemp, Warning, TEXT("Failed to resolve int value '%d' to a valid enum value for enum '%s'"),
+							IntValue, *Interface->GetNamespace().ToString());
+
+						// Return default (should always succeed as we can't have empty Enums and we must have a default).
+						return Interface->FindByValue(Interface->GetDefaultValue())->DisplayName.ToString();
+					};
+					auto SelectedValue = [Interface = EnumInterface, Prop = ValueProperty](const FString& InSelected)
+					{
+						TOptional<IEnumDataTypeInterface::FGenericInt32Entry> Found =
+							Interface->FindEntryBy([TextSelected = FText::FromString(InSelected)](const IEnumDataTypeInterface::FGenericInt32Entry& i)
+						{
+							return i.DisplayName.EqualTo(TextSelected);
+						});
+
+						if (Found)
+						{
+							// Only save the changes if its different and we can read the old value to check that.
+							int32 CurrentValue;
+							bool bReadCurrentValue = Prop->GetValue(CurrentValue) == FPropertyAccess::Success;
+							if ((bReadCurrentValue && CurrentValue != Found->Value) || !bReadCurrentValue)
+							{
+								ensure(Prop->SetValue(Found->Value) == FPropertyAccess::Success);
+							}
+						}
+						else
+						{
+							UE_LOG(LogTemp, Warning, TEXT("Failed to Set Valid Value for Property '%s' with Value of '%s', writing default."),
+								*GetNameSafe(Prop->GetProperty()), *InSelected);
+
+							ensure(Prop->SetValue(Interface->GetDefaultValue()) == FPropertyAccess::Success);
+						}
+					};
+
+					return PropertyCustomizationHelpers::MakePropertyComboBox(
+						nullptr,
+						FOnGetPropertyComboBoxStrings::CreateLambda(GetAll),
+						FOnGetPropertyComboBoxValue::CreateLambda(GetValue),
+						FOnPropertyComboBoxValueSelected::CreateLambda(SelectedValue)
+					);
+				}
+			}
+
+			return SNullWidget::NullWidget;
+		}
+
+		void FMetasoundInputNodeObjectDetailCustomization::CacheProxyData(TSharedPtr<IPropertyHandle> ProxyHandle)
 		{
 			ProxyGenClass.Reset();
 
-			const FString* MetadataProxyGenClass = InPropertyHandle->GetInstanceMetaData(NodeCustomizationUtils::ProxyGeneratorClassNameIdentifier);
+			const FString* MetadataProxyGenClass = ProxyHandle->GetInstanceMetaData(InputNodeCustomizationPrivate::ProxyGeneratorClassNameIdentifier);
 			if (!ensure(MetadataProxyGenClass))
 			{
 				return;
@@ -120,8 +208,10 @@ namespace Metasound
 			ensureMsgf(false, TEXT("Failed to find ProxyGeneratorClass. Class not set "));
 		}
 
-		TSharedRef<SWidget> FMetasoundInputNodeObjectDetailCustomization::CreateObjectPickerWidget(TSharedPtr<IPropertyHandle>& PropertyHandle) const
+		TSharedRef<SWidget> FMetasoundInputNodeObjectDetailCustomization::CreateStructureWidget(TSharedPtr<IPropertyHandle>& StructPropertyHandle) const
 		{
+			TSharedPtr<IPropertyHandle> PropertyHandle = StructPropertyHandle->GetChildHandle(GET_MEMBER_NAME_CHECKED(FMetasoundEditorGraphInputObject, Object));
+
 			auto ValidateAsset = [InProxyGenClass = ProxyGenClass](const FAssetData& InAsset)
 			{
 				if (!InProxyGenClass.IsValid())
@@ -175,10 +265,9 @@ namespace Metasound
 				.NewAssetFactories(PropertyCustomizationHelpers::GetNewAssetFactoriesForClasses(AllowedClasses));
 		}
 
-		TSharedRef<SWidget> FMetasoundInputNodeObjectDetailCustomization::CreateValueWidget(TSharedPtr<IPropertyHandleArray> ParentArrayProperty, TSharedPtr<IPropertyHandle> StructPropertyHandle, bool bIsInArray) const
+		TSharedRef<SWidget> FMetasoundInputNodeArrayDetailCustomizationBase::CreateValueWidget(TSharedPtr<IPropertyHandleArray> ParentArrayProperty, TSharedPtr<IPropertyHandle> StructPropertyHandle, bool bIsInArray) const
 		{
-			TSharedPtr<IPropertyHandle> PropertyHandle = StructPropertyHandle->GetChildHandle(GET_MEMBER_NAME_CHECKED(FMetasoundEditorGraphInputObject, Object));
-			TSharedRef<SWidget> ObjectPickerWidget = CreateObjectPickerWidget(PropertyHandle);
+			TSharedRef<SWidget> ObjectPickerWidget = CreateStructureWidget(StructPropertyHandle);
 			if (!bIsInArray)
 			{
 				return ObjectPickerWidget;
@@ -229,7 +318,7 @@ namespace Metasound
 				];
 		}
 
-		void FMetasoundInputNodeObjectDetailCustomization::CustomizeChildren(TSharedRef<IPropertyHandle> StructPropertyHandle, IDetailChildrenBuilder& ChildBuilder, IPropertyTypeCustomizationUtils& StructCustomizationUtils)
+		void FMetasoundInputNodeArrayDetailCustomizationBase::CustomizeChildren(TSharedRef<IPropertyHandle> StructPropertyHandle, IDetailChildrenBuilder& ChildBuilder, IPropertyTypeCustomizationUtils& StructCustomizationUtils)
 		{
 			bool bIsInArray = false;
 			TSharedPtr<IPropertyHandleArray> ParentArrayProperty;
@@ -247,11 +336,11 @@ namespace Metasound
 				}
 			}
 
-			CacheProxyClass(ProxyProperty);
+			CacheProxyData(ProxyProperty);
 
 			TSharedRef<SWidget> ValueWidget = CreateValueWidget(ParentArrayProperty, StructPropertyHandle, bIsInArray);
 
-			FDetailWidgetRow& ValueRow = ChildBuilder.AddCustomRow(NodeCustomizationUtils::DefaultPropertyText);
+			FDetailWidgetRow& ValueRow = ChildBuilder.AddCustomRow(InputNodeCustomizationPrivate::DefaultPropertyText);
 			if (bIsInArray)
 			{
 				ValueRow.NameContent()
@@ -264,7 +353,7 @@ namespace Metasound
 				ValueRow.NameContent()
 				[
 					SNew(STextBlock)
-					.Text(NodeCustomizationUtils::DefaultPropertyText)
+					.Text(InputNodeCustomizationPrivate::DefaultPropertyText)
 					.Font(IDetailLayoutBuilder::GetDetailFont())
 				];
 			}
@@ -284,7 +373,7 @@ namespace Metasound
 			{
 				for (const TWeakObjectPtr<UMetasoundEditorGraphInputNode>& InputNode : InInputs)
 				{
-					NodeCustomizationUtils::UpdateInputLiteralFromNode(InputNode);
+					InputNodeCustomizationPrivate::UpdateInputLiteralFromNode(InputNode);
 				}
 			});
 			StructPropertyHandle->SetOnChildPropertyValueChanged(OnLiteralChanged);
@@ -295,6 +384,10 @@ namespace Metasound
 			];
 		}
 
+		void FMetasoundInputNodeArrayDetailCustomizationBase::CustomizeHeader(TSharedRef<IPropertyHandle> StructPropertyHandle, FDetailWidgetRow& HeaderRow, IPropertyTypeCustomizationUtils& StructCustomizationUtils)
+		{
+		}
+
 		bool FMetasoundInputNodeDetailCustomization::IsRequired() const
 		{
 			if (GraphNode.IsValid())
@@ -307,7 +400,7 @@ namespace Metasound
 
 		void FMetasoundInputNodeDetailCustomization::CustomizeDetails(IDetailLayoutBuilder& DetailLayout)
 		{
-			using namespace Metasound::Frontend;
+			using namespace Frontend;
 
 			TMetasoundNodeVariableDetailCustomization<UMetasoundEditorGraphInputNode>::CustomizeDetails(DetailLayout);
 
@@ -326,18 +419,18 @@ namespace Metasound
 				.IsReadOnly(bIsRequired)
 				.Font(IDetailLayoutBuilder::GetDetailFont());
 
-			CategoryBuilder.AddCustomRow(NodeCustomizationUtils::InputNodeNameText)
+			CategoryBuilder.AddCustomRow(InputNodeCustomizationPrivate::InputNodeNameText)
 			.EditCondition(!bIsRequired, nullptr)
 			.NameContent()
 			[
 				SNew(STextBlock)
 				.Font(IDetailLayoutBuilder::GetDetailFontBold())
-				.Text(NodeCustomizationUtils::InputNodeNameText)
+				.Text(InputNodeCustomizationPrivate::InputNodeNameText)
 				.ToolTipText(TAttribute<FText>::Create([GraphNode = this->GraphNode]()
 				{
 					if (GraphNode.IsValid())
 					{
-						Frontend::FNodeHandle NodeHandle = GraphNode->GetNodeHandle();
+						FNodeHandle NodeHandle = GraphNode->GetNodeHandle();
 						FMetasoundFrontendNodeStyle NodeStyle = NodeHandle->GetNodeStyle();
 						return NodeHandle->GetClassDescription();
 					}
@@ -350,13 +443,13 @@ namespace Metasound
 				DisplayNameEditableTextBox.ToSharedRef()
 			];
 
-			CategoryBuilder.AddCustomRow(NodeCustomizationUtils::NodeTooltipText)
+			CategoryBuilder.AddCustomRow(InputNodeCustomizationPrivate::NodeTooltipText)
 			.EditCondition(!bIsRequired, nullptr)
 			.NameContent()
 			[
 				SNew(STextBlock)
 				.Font(IDetailLayoutBuilder::GetDetailFontBold())
-				.Text(NodeCustomizationUtils::NodeTooltipText)
+				.Text(InputNodeCustomizationPrivate::NodeTooltipText)
 			]
 			.ValueContent()
 			[
@@ -366,7 +459,7 @@ namespace Metasound
 				.IsReadOnly(bIsRequired)
 				.ModiferKeyForNewLine(EModifierKey::Shift)
 				.RevertTextOnEscape(true)
-				.WrapTextAt(NodeCustomizationUtils::DetailsTitleMaxWidth - NodeCustomizationUtils::DetailsTitleWrapPadding)
+				.WrapTextAt(InputNodeCustomizationPrivate::DetailsTitleMaxWidth - InputNodeCustomizationPrivate::DetailsTitleWrapPadding)
 				.Font(IDetailLayoutBuilder::GetDetailFont())
 			];
 
@@ -405,7 +498,7 @@ namespace Metasound
 					SetDefaultPropertyMetaData(Handle);
 					FSimpleDelegate OnLiteralChanged = FSimpleDelegate::CreateLambda([InGraphNode = GraphNode]()
 					{
-						NodeCustomizationUtils::UpdateInputLiteralFromNode(InGraphNode);
+						InputNodeCustomizationPrivate::UpdateInputLiteralFromNode(InGraphNode);
 					});
 
 					Handle->SetOnPropertyValueChanged(OnLiteralChanged);
@@ -423,7 +516,7 @@ namespace Metasound
 
 		void FMetasoundInputNodeDetailCustomization::SetDefaultPropertyMetaData(TSharedRef<IPropertyHandle> InDefaultPropertyHandle) const
 		{
-			using namespace Metasound::Frontend;
+			using namespace Frontend;
 
 			if (!GraphNode.IsValid())
 			{
@@ -442,31 +535,37 @@ namespace Metasound
 				return;
 			}
 
+			static const FString ArrayIdentifier = TEXT(":Array");
+			FString TypeNameString = TypeName.ToString();
+			if (TypeNameString.EndsWith(ArrayIdentifier))
+			{
+				TypeNameString = TypeNameString.LeftChop(ArrayIdentifier.Len());
+			}
+			InDefaultPropertyHandle->SetInstanceMetaData(InputNodeCustomizationPrivate::DataTypeNameIdentifier, TypeNameString);
+
 			FDataTypeRegistryInfo DataTypeInfo;
 			if (!ensure(Registry->GetInfoForDataType(TypeName, DataTypeInfo)))
 			{
 				return;
 			}
 
-			const EMetasoundFrontendLiteralType LiteralType = Frontend::GetMetasoundFrontendLiteralType(DataTypeInfo.PreferredLiteralType);
+			const EMetasoundFrontendLiteralType LiteralType = GetMetasoundFrontendLiteralType(DataTypeInfo.PreferredLiteralType);
 			if (LiteralType != EMetasoundFrontendLiteralType::UObject && LiteralType != EMetasoundFrontendLiteralType::UObjectArray)
 			{
 				return;
 			}
 
 			UClass* ProxyGenClass = DataTypeInfo.ProxyGeneratorClass;
-			if (!ProxyGenClass)
+			if (ProxyGenClass)
 			{
-				return;
+				const FString ClassName = ProxyGenClass->GetName();
+				InDefaultPropertyHandle->SetInstanceMetaData(InputNodeCustomizationPrivate::ProxyGeneratorClassNameIdentifier, ClassName);
 			}
-
-			const FString ClassName = ProxyGenClass->GetName();
-			InDefaultPropertyHandle->SetInstanceMetaData(NodeCustomizationUtils::ProxyGeneratorClassNameIdentifier, ClassName);
 		}
 
 		FName FMetasoundInputNodeDetailCustomization::GetLiteralDataType() const
 		{
-			using namespace Metasound::Frontend;
+			using namespace Frontend;
 
 			FNodeHandle NodeHandle = GraphNode->GetNodeHandle();
 			const TArray<FOutputHandle> Outputs = NodeHandle->GetOutputs();
@@ -481,7 +580,7 @@ namespace Metasound
 
 		FText FMetasoundInputNodeDetailCustomization::GetDisplayName() const
 		{
-			using namespace Metasound::Frontend;
+			using namespace Frontend;
 
 			if (GraphNode.IsValid())
 			{
@@ -511,7 +610,7 @@ namespace Metasound
 
 		void FMetasoundInputNodeDetailCustomization::OnDisplayNameCommitted(const FText& InNewName, ETextCommit::Type InTextCommit)
 		{
-			using namespace Metasound::Frontend;
+			using namespace Frontend;
 
 			if (!bIsNameInvalid && GraphNode.IsValid())
 			{
@@ -532,7 +631,7 @@ namespace Metasound
 
 		FText FMetasoundInputNodeDetailCustomization::GetTooltip() const
 		{
-			using namespace Metasound::Frontend;
+			using namespace Frontend;
 			if (GraphNode.IsValid())
 			{
 				FNodeHandle NodeHandle = GraphNode->GetNodeHandle();
@@ -549,7 +648,7 @@ namespace Metasound
 
 		void FMetasoundInputNodeDetailCustomization::OnTooltipCommitted(const FText& InNewText, ETextCommit::Type InTextCommit)
 		{
-			using namespace Metasound::Frontend;
+			using namespace Frontend;
 
 			if (GraphNode.IsValid())
 			{
@@ -558,7 +657,7 @@ namespace Metasound
 				const TArray<FOutputHandle>& OutputHandles = NodeHandle->GetOutputs();
 				if (ensure(!OutputHandles.IsEmpty()))
 				{
-					Frontend::FOutputHandle OutputHandle = OutputHandles[0];
+					FOutputHandle OutputHandle = OutputHandles[0];
 					return NodeHandle->GetOwningGraph()->SetInputDescription(OutputHandle->GetName(), InNewText);
 				}
 			}
diff --git a/Engine/Plugins/Runtime/Metasound/Source/MetasoundEditor/Private/MetasoundInputNodeDetailCustomization.h b/Engine/Plugins/Runtime/Metasound/Source/MetasoundEditor/Private/MetasoundInputNodeDetailCustomization.h
--- a/Engine/Plugins/Runtime/Metasound/Source/MetasoundEditor/Private/MetasoundInputNodeDetailCustomization.h
+++ b/Engine/Plugins/Runtime/Metasound/Source/MetasoundEditor/Private/MetasoundInputNodeDetailCustomization.h
@@ -5,6 +5,7 @@
 #include "CoreMinimal.h"
 #include "DetailLayoutBuilder.h"
 #include "IDetailCustomization.h"
+#include "IPropertyTypeCustomization.h"
 #include "Layout/Visibility.h"
 #include "MetasoundAssetBase.h"
 #include "MetasoundEditorGraphNode.h"
@@ -19,13 +20,12 @@ class FDetailWidgetRow;
 class FPropertyRestriction;
 class IDetailLayoutBuilder;
 class IPropertyHandle;
-class IPropertyTypeCustomizationUtils;
 
 namespace Metasound
 {
 	namespace Editor
 	{
-		class FMetasoundInputNodeObjectDetailCustomization : public IPropertyTypeCustomization
+		class FMetasoundInputNodeArrayDetailCustomizationBase : public IPropertyTypeCustomization
 		{
 		public:
 			//~ Begin IPropertyTypeCustomization
@@ -33,12 +33,31 @@ namespace Metasound
 			virtual void CustomizeChildren(TSharedRef<IPropertyHandle> StructPropertyHandle, IDetailChildrenBuilder& ChildBuilder, IPropertyTypeCustomizationUtils& StructCustomizationUtils) override;
 			//~ End IPropertyTypeCustomization
 
-		private:
-			void CacheProxyClass(TSharedPtr<IPropertyHandle> PropertyHandle);
+		protected:
+			virtual TSharedRef<SWidget> CreateStructureWidget(TSharedPtr<IPropertyHandle>& PropertyHandle) const = 0;
+			virtual void CacheProxyData(TSharedPtr<IPropertyHandle> ProxyHandle) { }
 
-			TSharedRef<SWidget> CreateObjectPickerWidget(TSharedPtr<IPropertyHandle>& PropertyHandle) const;
+		private:
 			TSharedRef<SWidget> CreateValueWidget(TSharedPtr<IPropertyHandleArray> ParentArrayProperty, TSharedPtr<IPropertyHandle> StructPropertyHandle, bool bIsInArray) const;
+		};
+
+		class FMetasoundInputNodeIntDetailCustomization : public FMetasoundInputNodeArrayDetailCustomizationBase
+		{
+		protected:
+			virtual TSharedRef<SWidget> CreateStructureWidget(TSharedPtr<IPropertyHandle>& StructPropertyHandle) const override;
+			virtual void CacheProxyData(TSharedPtr<IPropertyHandle> ProxyHandle) override;
 
+		private:
+			FName DataTypeName;
+		};
+
+		class FMetasoundInputNodeObjectDetailCustomization : public FMetasoundInputNodeArrayDetailCustomizationBase
+		{
+		protected:
+			virtual TSharedRef<SWidget> CreateStructureWidget(TSharedPtr<IPropertyHandle>& StructPropertyHandle) const override;
+			virtual void CacheProxyData(TSharedPtr<IPropertyHandle> ProxyHandle) override;
+
+		private:
 			TWeakObjectPtr<UClass> ProxyGenClass;
 		};
 
diff --git a/Engine/Plugins/Runtime/Metasound/Source/MetasoundEditor/Public/MetasoundEditorGraphInputNodes.h b/Engine/Plugins/Runtime/Metasound/Source/MetasoundEditor/Public/MetasoundEditorGraphInputNodes.h
--- a/Engine/Plugins/Runtime/Metasound/Source/MetasoundEditor/Public/MetasoundEditorGraphInputNodes.h
+++ b/Engine/Plugins/Runtime/Metasound/Source/MetasoundEditor/Public/MetasoundEditorGraphInputNodes.h
@@ -93,6 +93,16 @@ public:
 	}
 };
 
+// Broken out to be able to customize and swap enum behavior for basic integer literal behavior
+USTRUCT()
+struct FMetasoundEditorGraphInputInt
+{
+	GENERATED_BODY()
+
+	UPROPERTY(EditAnywhere, Category = DefaultValue)
+	int32 Value = 0;
+};
+
 UCLASS(MinimalAPI)
 class UMetasoundEditorGraphInputIntNode : public UMetasoundEditorGraphInputNode
 {
@@ -102,12 +112,12 @@ public:
 	virtual ~UMetasoundEditorGraphInputIntNode() = default;
 
 	UPROPERTY(EditAnywhere, Category = DefaultValue)
-	int32 Default = 0;
+	FMetasoundEditorGraphInputInt Default;
 
 	virtual FMetasoundFrontendLiteral GetLiteralDefault() const override
 	{
 		FMetasoundFrontendLiteral Literal;
-		Literal.Set(Default);
+		Literal.Set(Default.Value);
 		return Literal;
 	}
 
@@ -126,12 +136,18 @@ public:
 	virtual ~UMetasoundEditorGraphInputIntArrayNode() = default;
 
 	UPROPERTY(EditAnywhere, Category = DefaultValue)
-	TArray<int32> Default;
+	TArray<FMetasoundEditorGraphInputInt> Default;
 
 	virtual FMetasoundFrontendLiteral GetLiteralDefault() const override
 	{
+		TArray<int32> Values;
+		for (const FMetasoundEditorGraphInputInt& Value : Default)
+		{
+			Values.Add(Value.Value);
+		}
+
 		FMetasoundFrontendLiteral Literal;
-		Literal.Set(Default);
+		Literal.Set(Values);
 		return Literal;
 	}
 
