Fix detuning oscillators

diff --git a/Engine/Plugins/Runtime/Metasound/Source/MetasoundStandardNodes/Private/MetasoundOscillatorNodes.cpp b/Engine/Plugins/Runtime/Metasound/Source/MetasoundStandardNodes/Private/MetasoundOscillatorNodes.cpp
--- a/Engine/Plugins/Runtime/Metasound/Source/MetasoundStandardNodes/Private/MetasoundOscillatorNodes.cpp
+++ b/Engine/Plugins/Runtime/Metasound/Source/MetasoundStandardNodes/Private/MetasoundOscillatorNodes.cpp
@@ -171,8 +171,8 @@ namespace Metasound
 			}
 		};		
 		
-		using FSinf						= TGenerateBlock<FSinfGenerator, FPhaseLetClimb>;
-		using FSinfWithFm				= TGenerateBlockFM<FSinfGenerator, FPhaseLetClimb>;
+		using FSinf						= TGenerateBlock<FSinfGenerator>;
+		using FSinfWithFm				= TGenerateBlockFM<FSinfGenerator>;
 		
 		using FBhaskara					= TGenerateBlock<FBhaskaraGenerator>;
 		using FBhaskaraWithFm			= TGenerateBlockFM<FBhaskaraGenerator>;
