Fix default MetaSound source blockrate

diff --git a/Engine/Plugins/Runtime/Metasound/Source/MetasoundEngine/Private/MetasoundSource.cpp b/Engine/Plugins/Runtime/Metasound/Source/MetasoundEngine/Private/MetasoundSource.cpp
--- a/Engine/Plugins/Runtime/Metasound/Source/MetasoundEngine/Private/MetasoundSource.cpp
+++ b/Engine/Plugins/Runtime/Metasound/Source/MetasoundEngine/Private/MetasoundSource.cpp
@@ -23,7 +23,7 @@
 #include "EdGraph/EdGraph.h"
 #endif // WITH_EDITORONLY_DATA
 
-static float MetaSoundBlockRateCVar = 0;
+static float MetaSoundBlockRateCVar = 100.f;
 FAutoConsoleVariableRef CVarMetaSoundBlockRate(
 	TEXT("au.MetaSound.BlockRate"),
 	MetaSoundBlockRateCVar,
@@ -375,7 +375,7 @@ TArray<FString> UMetasoundSource::GetTransmittableInputVertexNames() const
 Metasound::FOperatorSettings UMetasoundSource::GetOperatorSettings(Metasound::FSampleRate InSampleRate) const
 {
 	// Metasound graph gets evaluated 100 times per second by default.
-	float BlockRate = FMath::Clamp(MetaSoundBlockRateCVar, 1.0f, 200.0f); 
+	float BlockRate = FMath::Clamp(MetaSoundBlockRateCVar, 1.0f, 1000.0f); 
 	return Metasound::FOperatorSettings(InSampleRate, BlockRate);
 }
 
