Fix for mono delay dry level issue

diff --git a/Engine/Plugins/Runtime/Metasound/Source/MetasoundStandardNodes/Private/MetasoundDelayNode.cpp b/Engine/Plugins/Runtime/Metasound/Source/MetasoundStandardNodes/Private/MetasoundDelayNode.cpp
--- a/Engine/Plugins/Runtime/Metasound/Source/MetasoundStandardNodes/Private/MetasoundDelayNode.cpp
+++ b/Engine/Plugins/Runtime/Metasound/Source/MetasoundStandardNodes/Private/MetasoundDelayNode.cpp
@@ -152,7 +152,7 @@ namespace Metasound
 
 			for (int32 FrameIndex = 0; FrameIndex < NumFrames; ++FrameIndex)
 			{
-				OutputAudio[FrameIndex] = CurrentWetLevel * DelayBuffer.ProcessAudioSample(InputAudio[FrameIndex]) + CurrentDryLevel * OutputAudio[FrameIndex];
+				OutputAudio[FrameIndex] = CurrentWetLevel * DelayBuffer.ProcessAudioSample(InputAudio[FrameIndex]) + CurrentDryLevel * InputAudio[FrameIndex];
 			}
 		}
 		else
@@ -160,7 +160,7 @@ namespace Metasound
 			// There is some amount of feedback so we do the feedback mixing
 			for (int32 FrameIndex = 0; FrameIndex < NumFrames; ++FrameIndex)
 			{
-				OutputAudio[FrameIndex] = CurrentWetLevel * DelayBuffer.ProcessAudioSample(InputAudio[FrameIndex] + FeedbackSample * FeedbackAmount) + CurrentDryLevel * OutputAudio[FrameIndex];
+				OutputAudio[FrameIndex] = CurrentWetLevel * DelayBuffer.ProcessAudioSample(InputAudio[FrameIndex] + FeedbackSample * FeedbackAmount) + CurrentDryLevel * InputAudio[FrameIndex];
 				FeedbackSample = OutputAudio[FrameIndex];
 			}
 		}
