Adding minor features to stereo delay node.

- Enum to pick delay mode type
- Add delay ratio

diff --git a/Engine/Plugins/Runtime/Metasound/Source/MetasoundStandardNodes/Private/MetasoundStereoDelayNode.cpp b/Engine/Plugins/Runtime/Metasound/Source/MetasoundStandardNodes/Private/MetasoundStereoDelayNode.cpp
--- a/Engine/Plugins/Runtime/Metasound/Source/MetasoundStandardNodes/Private/MetasoundStereoDelayNode.cpp
+++ b/Engine/Plugins/Runtime/Metasound/Source/MetasoundStandardNodes/Private/MetasoundStereoDelayNode.cpp
@@ -5,6 +5,7 @@
 #include "Internationalization/Text.h"
 #include "MetasoundExecutableOperator.h"
 #include "MetasoundNodeRegistrationMacro.h"
+#include "MetasoundDataTypeRegistrationMacro.h"
 #include "MetasoundPrimitives.h"
 #include "MetasoundStandardNodesNames.h"
 #include "MetasoundTrigger.h"
@@ -20,7 +21,9 @@ namespace Metasound
 	{
 		static const TCHAR* InParamNameAudioInputLeft = TEXT("In Left");
 		static const TCHAR* InParamNameAudioInputRight = TEXT("In Right");
+		static const TCHAR* InParamNameDelayMode = TEXT("Delay Mode");
 		static const TCHAR* InParamNameDelayTime = TEXT("Delay Time");
+		static const TCHAR* InParamNameDelayRatio = TEXT("Delay Ratio");
 		static const TCHAR* InParamNameDryLevel = TEXT("Dry Level");
 		static const TCHAR* InParamNameWetLevel = TEXT("Wet Level");
 		static const TCHAR* InParamNameFeedbackAmount = TEXT("Feedback");
@@ -30,6 +33,22 @@ namespace Metasound
 		static float MaxDelaySeconds = 5.0f;
 	}
 
+	enum class EStereoDelayMode
+	{
+		Normal = 0,
+		Cross,
+		PingPong,
+	};
+
+	DECLARE_METASOUND_ENUM(EStereoDelayMode, EStereoDelayMode::Normal, METASOUNDSTANDARDNODES_API,
+		FEnumStereoDelayMode, FEnumStereoDelayModeInfo, FStereoDelayModeReadRef, FEnumStereoDelayModeWriteRef);
+
+	DEFINE_METASOUND_ENUM_BEGIN(EStereoDelayMode, FEnumStereoDelayMode, "StereoDelayMode")
+		DEFINE_METASOUND_ENUM_ENTRY(EStereoDelayMode::Normal, LOCTEXT("StereoDelayModeNormalDescription", "Normal"), LOCTEXT("StereoDelayModeNormalDescriptionTT", "Left input mixes with left delay output and feeds to left output.")),
+		DEFINE_METASOUND_ENUM_ENTRY(EStereoDelayMode::Cross, LOCTEXT("StereoDelayModeCrossDescription", "Cross"), LOCTEXT("StereoDelayModeCrossDescriptionTT", "Left input mixes with right delay output and feeds to right output.")),
+		DEFINE_METASOUND_ENUM_ENTRY(EStereoDelayMode::PingPong, LOCTEXT("StereoDelayModePingPongDescription", "Ping Pong"), LOCTEXT("StereoDelayModePingPongDescriptionTT", "Left input mixes with left delay output and feeds to right output.")),
+		DEFINE_METASOUND_ENUM_END()
+
 	class FStereoDelayOperator : public TExecutableOperator<FStereoDelayOperator>
 	{
 	public:
@@ -41,7 +60,9 @@ namespace Metasound
 		FStereoDelayOperator(const FOperatorSettings& InSettings, 
 			const FAudioBufferReadRef& InLeftAudioInput, 
 			const FAudioBufferReadRef& InRightAudioInput,
+			const FStereoDelayModeReadRef& InStereoDelayMode,
 			const FTimeReadRef& InDelayTime,
+			const FFloatReadRef& InDelayRatio,
 			const FFloatReadRef& InDryLevel, 
 			const FFloatReadRef& InWetLevel, 
 			const FFloatReadRef& InFeedback);
@@ -51,15 +72,22 @@ namespace Metasound
 		void Execute();
 
 	private:
-		float GetInputDelayTimeMsec() const;
+		float GetInputDelayTimeMsecClamped() const;
+		float GetInputDelayRatioClamped() const;
 
 		// The input audio buffer
 		FAudioBufferReadRef LeftAudioInput;
 		FAudioBufferReadRef RightAudioInput;
 
+		// Which stereo delay mode to render the audio delay with
+		FStereoDelayModeReadRef StereoDelayMode;
+
 		// The amount of delay time
 		FTimeReadRef DelayTime;
 
+		// The stereo delay ratio
+		FFloatReadRef DelayRatio;
+
 		// The the dry level
 		FFloatReadRef DryLevel;
 
@@ -77,32 +105,38 @@ namespace Metasound
 		Audio::FDelay LeftDelayBuffer;
 		Audio::FDelay RightDelayBuffer;
 
-		// The current delay time
+		// The current delay time and delay ratio
 		float PrevDelayTimeMsec;
+		float PrevDelayRatio;
 	};
 
 	FStereoDelayOperator::FStereoDelayOperator(const FOperatorSettings& InSettings, 
 		const FAudioBufferReadRef& InLeftAudioInput, 
 		const FAudioBufferReadRef& InRightAudioInput,
+		const FStereoDelayModeReadRef& InStereoDelayMode,
 		const FTimeReadRef& InDelayTime,
-		const FFloatReadRef& InDryLevel, 
+		const FFloatReadRef& InDelayRatio,
+		const FFloatReadRef& InDryLevel,
 		const FFloatReadRef& InWetLevel, 
 		const FFloatReadRef& InFeedback)
 
 		: LeftAudioInput(InLeftAudioInput)
 		, RightAudioInput(InRightAudioInput)
+		, StereoDelayMode(InStereoDelayMode)
 		, DelayTime(InDelayTime)
+		, DelayRatio(InDelayRatio)
 		, DryLevel(InDryLevel)
 		, WetLevel(InWetLevel)
 		, Feedback(InFeedback)
 		, LeftAudioOutput(FAudioBufferWriteRef::CreateNew(InSettings))
 		, RightAudioOutput(FAudioBufferWriteRef::CreateNew(InSettings))
-		, PrevDelayTimeMsec(GetInputDelayTimeMsec())
+		, PrevDelayTimeMsec(GetInputDelayTimeMsecClamped())
+		, PrevDelayRatio(GetInputDelayRatioClamped())
 	{
 		LeftDelayBuffer.Init(InSettings.GetSampleRate(), StereoDelay::MaxDelaySeconds);
-		LeftDelayBuffer.SetDelayMsec(PrevDelayTimeMsec);
+		LeftDelayBuffer.SetDelayMsec(PrevDelayTimeMsec * (1.0f + PrevDelayRatio));
 		RightDelayBuffer.Init(InSettings.GetSampleRate(), StereoDelay::MaxDelaySeconds);
-		RightDelayBuffer.SetDelayMsec(PrevDelayTimeMsec);
+		RightDelayBuffer.SetDelayMsec(PrevDelayTimeMsec * (1.0f - PrevDelayRatio));
 	}
 
 	FDataReferenceCollection FStereoDelayOperator::GetInputs() const
@@ -110,7 +144,9 @@ namespace Metasound
 		FDataReferenceCollection InputDataReferences;
 		InputDataReferences.AddDataReadReference(StereoDelay::InParamNameAudioInputLeft, FAudioBufferReadRef(LeftAudioInput));
 		InputDataReferences.AddDataReadReference(StereoDelay::InParamNameAudioInputRight, FAudioBufferReadRef(RightAudioInput));
+		InputDataReferences.AddDataReadReference(StereoDelay::InParamNameDelayMode, FStereoDelayModeReadRef(StereoDelayMode));
 		InputDataReferences.AddDataReadReference(StereoDelay::InParamNameDelayTime, FTimeReadRef(DelayTime));
+		InputDataReferences.AddDataReadReference(StereoDelay::InParamNameDelayRatio, FFloatReadRef(DelayRatio));
 		InputDataReferences.AddDataReadReference(StereoDelay::InParamNameDryLevel, FFloatReadRef(DryLevel));
 		InputDataReferences.AddDataReadReference(StereoDelay::InParamNameWetLevel, FFloatReadRef(WetLevel));
 		InputDataReferences.AddDataReadReference(StereoDelay::InParamNameFeedbackAmount, FFloatReadRef(Feedback));
@@ -126,23 +162,30 @@ namespace Metasound
 		return OutputDataReferences;
 	}
 
-	float FStereoDelayOperator::GetInputDelayTimeMsec() const
+	float FStereoDelayOperator::GetInputDelayTimeMsecClamped() const
 	{
 		// Clamp the delay time to the max delay allowed
 		return 1000.0f * FMath::Clamp((float)DelayTime->GetSeconds(), 0.0f, StereoDelay::MaxDelaySeconds);
 	}
 
+	float FStereoDelayOperator::GetInputDelayRatioClamped() const
+	{
+		return FMath::Clamp(*DelayRatio, -1.0f, 1.0f);
+	}
+
 	void FStereoDelayOperator::Execute()
 	{
 		// Get clamped delay time
-		float CurrentInputDelayTime = GetInputDelayTimeMsec();
+		float CurrentInputDelayTime = GetInputDelayTimeMsecClamped();
+		float CurrentDelayRatio = GetInputDelayRatioClamped();
 
 		// Check to see if our delay amount has changed
-		if (!FMath::IsNearlyEqual(PrevDelayTimeMsec, CurrentInputDelayTime))
+		if (!FMath::IsNearlyEqual(PrevDelayTimeMsec, CurrentInputDelayTime) || !FMath::IsNearlyEqual(PrevDelayRatio, CurrentDelayRatio))
 		{
 			PrevDelayTimeMsec = CurrentInputDelayTime;
-			LeftDelayBuffer.SetEasedDelayMsec(PrevDelayTimeMsec);
-			RightDelayBuffer.SetEasedDelayMsec(PrevDelayTimeMsec);
+			PrevDelayRatio = CurrentDelayRatio;
+			LeftDelayBuffer.SetEasedDelayMsec(PrevDelayTimeMsec * (1.0f + PrevDelayRatio));
+			RightDelayBuffer.SetEasedDelayMsec(PrevDelayTimeMsec * (1.0f - PrevDelayRatio));
 		}
 
 		const float* LeftInput = LeftAudioInput->GetData();
@@ -161,34 +204,77 @@ namespace Metasound
 		if (FMath::IsNearlyZero(FeedbackAmount))
 		{
 			// if pingpong
+			switch (*StereoDelayMode)
 			{
-				for (int32 FrameIndex = 0; FrameIndex < NumFrames; ++FrameIndex)
+				// Normal feeds left to left and right to right
+				case EStereoDelayMode::Normal:
 				{
-					// Ping pong feeds right to left and left to right
-					LeftOutput[FrameIndex] = CurrentWetLevel * LeftDelayBuffer.ProcessAudioSample(RightInput[FrameIndex]) + CurrentDryLevel * LeftInput[FrameIndex];
-					RightOutput[FrameIndex] = CurrentWetLevel * RightDelayBuffer.ProcessAudioSample(LeftInput[FrameIndex]) + CurrentDryLevel * RightInput[FrameIndex];
+					for (int32 FrameIndex = 0; FrameIndex < NumFrames; ++FrameIndex)
+					{
+						LeftOutput[FrameIndex] = CurrentWetLevel * LeftDelayBuffer.ProcessAudioSample(LeftInput[FrameIndex]) + CurrentDryLevel * LeftInput[FrameIndex];
+						RightOutput[FrameIndex] = CurrentWetLevel * RightDelayBuffer.ProcessAudioSample(RightInput[FrameIndex]) + CurrentDryLevel * RightInput[FrameIndex];
+					}
 				}
+				break;
+
+				// No-feedback Cross and ping-pong feeds right input to left and left input to right
+				case EStereoDelayMode::Cross:
+				case EStereoDelayMode::PingPong:
+				{
+					for (int32 FrameIndex = 0; FrameIndex < NumFrames; ++FrameIndex)
+					{
+						// Ping pong feeds right to left and left to right
+						LeftOutput[FrameIndex] = CurrentWetLevel * LeftDelayBuffer.ProcessAudioSample(RightInput[FrameIndex]) + CurrentDryLevel * LeftInput[FrameIndex];
+						RightOutput[FrameIndex] = CurrentWetLevel * RightDelayBuffer.ProcessAudioSample(LeftInput[FrameIndex]) + CurrentDryLevel * RightInput[FrameIndex];
+					}
+				}
+				break;
 			}
 		}
 		else
 		{
 			// TODO: support different delay cross-modes via enum, currently default to pingpong
-		
-			// if pingpong
+			switch (*StereoDelayMode)
 			{
-				// There is some amount of feedback so we do the feedback mixing
-				for (int32 FrameIndex = 0; FrameIndex < NumFrames; ++FrameIndex)
+				case EStereoDelayMode::Normal:
 				{
-					float LeftDelayOut = LeftDelayBuffer.Read();
-					float RightDelayOut = RightDelayBuffer.Read();
+						
+					for (int32 FrameIndex = 0; FrameIndex < NumFrames; ++FrameIndex)
+					{
+						float LeftDelayIn = LeftInput[FrameIndex] + FeedbackAmount * LeftDelayBuffer.Read();
+						float RightDelayIn = RightInput[FrameIndex] + FeedbackAmount * RightDelayBuffer.Read();
+
+						LeftOutput[FrameIndex] = CurrentWetLevel * LeftDelayBuffer.ProcessAudioSample(LeftDelayIn) + CurrentDryLevel * LeftInput[FrameIndex];
+						RightOutput[FrameIndex] = CurrentWetLevel * RightDelayBuffer.ProcessAudioSample(RightDelayIn) + CurrentDryLevel * RightInput[FrameIndex];
+					}
+				}
+				break;
 
-					// Pingpong algorithm feeds right output into left and left output to right input
-					float LeftDelayIn = RightInput[FrameIndex] + RightDelayOut  * FeedbackAmount;
-					float RightDelayIn = LeftInput[FrameIndex] + LeftDelayOut * FeedbackAmount;
+				case EStereoDelayMode::Cross:
+				{
+					for (int32 FrameIndex = 0; FrameIndex < NumFrames; ++FrameIndex)
+					{
+						float LeftDelayIn = RightInput[FrameIndex] + FeedbackAmount * LeftDelayBuffer.Read();
+						float RightDelayIn = LeftInput[FrameIndex] + FeedbackAmount * RightDelayBuffer.Read();
+
+						LeftOutput[FrameIndex] = CurrentWetLevel * LeftDelayBuffer.ProcessAudioSample(LeftDelayIn) + CurrentDryLevel * LeftInput[FrameIndex];
+						RightOutput[FrameIndex] = CurrentWetLevel * RightDelayBuffer.ProcessAudioSample(RightDelayIn) + CurrentDryLevel * RightInput[FrameIndex];
+					}
+				}
+				break;
 
-					LeftOutput[FrameIndex] = CurrentWetLevel * LeftDelayBuffer.ProcessAudioSample(LeftDelayIn) + CurrentDryLevel * LeftInput[FrameIndex];
-					RightOutput[FrameIndex] = CurrentWetLevel * RightDelayBuffer.ProcessAudioSample(RightDelayIn) + CurrentDryLevel * RightInput[FrameIndex];
+				case EStereoDelayMode::PingPong:
+				{
+					for (int32 FrameIndex = 0; FrameIndex < NumFrames; ++FrameIndex)
+					{
+						float LeftDelayIn = RightInput[FrameIndex] + FeedbackAmount * RightDelayBuffer.Read();
+						float RightDelayIn = LeftInput[FrameIndex] + FeedbackAmount * LeftDelayBuffer.Read();
+
+						LeftOutput[FrameIndex] = CurrentWetLevel * LeftDelayBuffer.ProcessAudioSample(LeftDelayIn) + CurrentDryLevel * LeftInput[FrameIndex];
+						RightOutput[FrameIndex] = CurrentWetLevel * RightDelayBuffer.ProcessAudioSample(RightDelayIn) + CurrentDryLevel * RightInput[FrameIndex];
+					}
 				}
+				break;
 			}
 		}
 	}
@@ -199,7 +285,9 @@ namespace Metasound
 			FInputVertexInterface(
 				TInputDataVertexModel<FAudioBuffer>(StereoDelay::InParamNameAudioInputLeft, LOCTEXT("LeftAudioInputTooltip", "Left channel audio input.")),
 				TInputDataVertexModel<FAudioBuffer>(StereoDelay::InParamNameAudioInputRight, LOCTEXT("RightAudioInputTooltip", "Right channel audio input.")),
+				TInputDataVertexModel<FEnumStereoDelayMode>(StereoDelay::InParamNameDelayMode, LOCTEXT("DelayModeTooltip", "Delay mode.")),
 				TInputDataVertexModel<FTime>(StereoDelay::InParamNameDelayTime, LOCTEXT("DelayTimeTooltip", "The amount of time to delay the audio."), 1.0f),
+				TInputDataVertexModel<float>(StereoDelay::InParamNameDelayRatio, LOCTEXT("DelayRatioTooltip", "Delay spread for left and right channels. Allows left and right channels to have differential delay amounts. Useful for stereo channel decorrelation"), 0.0f),
 				TInputDataVertexModel<float>(StereoDelay::InParamNameDryLevel, LOCTEXT("DryLevelTooltip", "The dry level of the delay."), 0.0f),
 				TInputDataVertexModel<float>(StereoDelay::InParamNameWetLevel, LOCTEXT("FeedbackTooltip", "The wet level of the delay."), 1.0f),
 				TInputDataVertexModel<float>(StereoDelay::InParamNameFeedbackAmount, LOCTEXT("FeedbackTooltip", "Feedback amount."), 0.0f)
@@ -243,12 +331,14 @@ namespace Metasound
 
 		FAudioBufferReadRef LeftAudioIn = InputCollection.GetDataReadReferenceOrConstruct<FAudioBuffer>(StereoDelay::InParamNameAudioInputLeft, InParams.OperatorSettings);
 		FAudioBufferReadRef RightAudioIn = InputCollection.GetDataReadReferenceOrConstruct<FAudioBuffer>(StereoDelay::InParamNameAudioInputRight, InParams.OperatorSettings);
+		FStereoDelayModeReadRef StereoDelayMode = InputCollection.GetDataReadReferenceOrConstruct<FEnumStereoDelayMode>(StereoDelay::InParamNameDelayMode);
 		FTimeReadRef DelayTime = InputCollection.GetDataReadReferenceOrConstructWithVertexDefault<FTime, float>(InputInterface, StereoDelay::InParamNameDelayTime);
+		FFloatReadRef DelayRatio = InputCollection.GetDataReadReferenceOrConstructWithVertexDefault<float>(InputInterface, StereoDelay::InParamNameDelayRatio);
 		FFloatReadRef DryLevel = InputCollection.GetDataReadReferenceOrConstructWithVertexDefault<float>(InputInterface, StereoDelay::InParamNameDryLevel);
 		FFloatReadRef WetLevel = InputCollection.GetDataReadReferenceOrConstructWithVertexDefault<float>(InputInterface, StereoDelay::InParamNameWetLevel);
 		FFloatReadRef Feedback = InputCollection.GetDataReadReferenceOrConstructWithVertexDefault<float>(InputInterface, StereoDelay::InParamNameFeedbackAmount);
 
-		return MakeUnique<FStereoDelayOperator>(InParams.OperatorSettings, LeftAudioIn, RightAudioIn, DelayTime, DryLevel, WetLevel, Feedback);
+		return MakeUnique<FStereoDelayOperator>(InParams.OperatorSettings, LeftAudioIn, RightAudioIn, StereoDelayMode, DelayTime, DelayRatio, DryLevel, WetLevel, Feedback);
 	}
 
 	FStereoDelayNode::FStereoDelayNode(const FNodeInitData& InitData)
