IWYU Fix

diff --git a/Engine/Plugins/Runtime/Metasound/Source/MetasoundEditor/Private/MetasoundInputNodeDetailCustomization.cpp b/Engine/Plugins/Runtime/Metasound/Source/MetasoundEditor/Private/MetasoundInputNodeDetailCustomization.cpp
--- a/Engine/Plugins/Runtime/Metasound/Source/MetasoundEditor/Private/MetasoundInputNodeDetailCustomization.cpp
+++ b/Engine/Plugins/Runtime/Metasound/Source/MetasoundEditor/Private/MetasoundInputNodeDetailCustomization.cpp
@@ -32,6 +32,7 @@
 #include "Widgets/Input/SCheckBox.h"
 #include "Widgets/Input/SMultiLineEditableTextBox.h"
 #include "Widgets/Notifications/SNotificationList.h"
+#include "Widgets/SToolTip.h"
 #include "Widgets/SWidget.h"
 #include "Widgets/Text/STextBlock.h"
 
