Improvements to Metasound Wave decoding

diff --git a/Engine/Plugins/Runtime/Metasound/Source/MetasoundEngine/Private/MetasoundWave.cpp b/Engine/Plugins/Runtime/Metasound/Source/MetasoundEngine/Private/MetasoundWave.cpp
--- a/Engine/Plugins/Runtime/Metasound/Source/MetasoundEngine/Private/MetasoundWave.cpp
+++ b/Engine/Plugins/Runtime/Metasound/Source/MetasoundEngine/Private/MetasoundWave.cpp
@@ -29,103 +29,85 @@ namespace Audio
 		// initialize input/output data
 		InputSampleRate = InWave->GetSampleRate();
 		OutputSampleRate = InInitParams.OutputSampleRate;
-		FsInToFsOutRatio = (InputSampleRate / OutputSampleRate);
+		FsOutToInRatio = (OutputSampleRate / InputSampleRate);
 		MaxPitchShiftRatio = FMath::Pow(2.f, InInitParams.MaxPitchShiftMagnitudeAllowedInOctaves);
 		MaxPitchShiftCents = InInitParams.MaxPitchShiftMagnitudeAllowedInOctaves * 1200.f;
 
 		NumChannels = InWave->GetNumChannels();
-		DecodeBlockSizeInFrames = InInitParams.OutputBlockSizeInFrames;
+		DecodeBlockSizeInFrames = 64;
 		DecodeBlockSizeInSamples = DecodeBlockSizeInFrames * NumChannels;
 
 		// set Circular Buffer capacity
-		CircularDecoderOutputBuffer.SetCapacity(DecodeBlockSizeInSamples * (1.f + FsInToFsOutRatio * MaxPitchShiftRatio));
+		OutputCircularBuffer.SetCapacity(InInitParams.OutputBlockSizeInFrames * NumChannels * (1.f + FsOutToInRatio * MaxPitchShiftRatio) * 2);
 
 		TotalNumFramesOutput = 0;
 		TotalNumFramesDecoded = 0;
-		
+
 		// try to initialize decoders
 		bool bSuccessful = InitializeDecodersInternal(InWave);
 		bDecoderIsDone = !bSuccessful;
 
 		// initialize SRC object (will be re-initialized for pitch shifting)
-		Resampler.Init(Audio::EResamplingMethod::Linear, 1.f / FsInToFsOutRatio, NumChannels);
+		Resampler.Init(Audio::EResamplingMethod::Linear, 1.f / FsOutToInRatio, NumChannels);
 
 		return bSuccessful;
 	}
 
 	uint32 FSimpleDecoderWrapper::GenerateAudio(float* OutputDest, int32 NumOutputFrames, float PitchShiftInCents)
 	{
-		PitchShiftInCents = FMath::Clamp(PitchShiftInCents, -MaxPitchShiftCents, MaxPitchShiftCents);
-		const float PitchShiftRatio = FMath::Pow(2.f, PitchShiftInCents / 1200.f);
-		const float TotalInSamplesOutSamplesRatio = PitchShiftRatio * FsInToFsOutRatio;
-
 		const uint32 NumOutputSamples = NumOutputFrames * NumChannels;
-		uint32 NumSamplesToDecode = NumChannels * FMath::CeilToInt(NumOutputFrames * TotalInSamplesOutSamplesRatio + 1.f);
-
-		// zero the output
-		FMemory::Memzero(OutputDest, NumOutputSamples * sizeof(float));
 
-		// Decode audio until we have enough to satisfy the output (post SRC)
-		while (!bDecoderIsDone && (CircularDecoderOutputBuffer.Num() < NumSamplesToDecode))
+		if (OutputCircularBuffer.Num() < NumOutputSamples)
 		{
-			// get more audio from the decoder
-			Audio::IDecoderOutput::FPushedAudioDetails Details;
-			bDecoderIsDone = (Decoder->Decode() == Audio::IDecoder::EDecodeResult::Finished);
+			const float SampleRateRatio = GetSampleRateRatio(PitchShiftInCents);
 
-			PreSrcBuffer.Reset(DecodeBlockSizeInSamples);
-			PreSrcBuffer.AddZeroed(DecodeBlockSizeInSamples);
-			const int32 NumSamplesDecoded = Output->PopAudio(MakeArrayView(PreSrcBuffer.GetData(), DecodeBlockSizeInSamples), Details);
-
-			// push that (interleaved) audio to the (interleaved) circular buffer
-			int32 NumPushed = CircularDecoderOutputBuffer.Push(PreSrcBuffer.GetData(), NumSamplesDecoded);
-			ensure(NumPushed == NumSamplesDecoded); // there will be a discontinuity in the output because our CircularDecoderOutputBuffer was not large enough!
+			// (multiply by two just to be safe we can handle SRC output)
+			const int32 MaxNumResamplerOutputFramesPerBlock = FMath::CeilToInt(SampleRateRatio * DecodeBlockSizeInFrames) * 2;
+			const int32 MaxNumResamplerOutputSamplesPerBlock = MaxNumResamplerOutputFramesPerBlock * NumChannels;
 
-			// for debugging
-			TotalNumFramesDecoded += (NumPushed / NumChannels);
+			PreSrcBuffer.Reset(DecodeBlockSizeInSamples);
+			PreSrcBuffer.AddUninitialized(DecodeBlockSizeInSamples);
+
+			PostSrcBuffer.Reset(MaxNumResamplerOutputSamplesPerBlock);
+			PostSrcBuffer.AddUninitialized(MaxNumResamplerOutputSamplesPerBlock);
+
+			// combine pitch shift and SRC ratios
+			Resampler.SetSampleRateRatio(SampleRateRatio);
+
+			// perform SRC and push to circular buffer until we have enough frames for the output
+			while (Decoder && !bDecoderIsDone && (OutputCircularBuffer.Num() < NumOutputSamples))
+			{
+				// get more audio from the decoder
+				Audio::IDecoderOutput::FPushedAudioDetails Details;
+				bDecoderIsDone = (Decoder->Decode() == Audio::IDecoder::EDecodeResult::Finished);
+				const int32 NumFramesDecoded = Output->PopAudio(PreSrcBuffer, Details) / NumChannels;
+				int32 NumResamplerOutputFrames = 0;
+				int32 Error = Resampler.ProcessAudio(PreSrcBuffer.GetData(), NumFramesDecoded, bDecoderIsDone, PostSrcBuffer.GetData(), MaxNumResamplerOutputFramesPerBlock, NumResamplerOutputFrames);
+				ensure(Error == 0);
+				OutputCircularBuffer.Push(PostSrcBuffer.GetData(), NumResamplerOutputFrames * NumChannels);
+			}
 		}
 
-		// now that we have enough audio decoded, pop off the circular buffer into an interleaved, pre-src temp buffer
-		// (It is possible that CircularDecoderOutputBuffer.Num() < NumOutputSamples if the decoder is dry...)
-		const uint32 NumSamplesToPop = FMath::Min(CircularDecoderOutputBuffer.Num(), NumSamplesToDecode);
-		const uint32 NumPostSrcSamples = NumOutputFrames * NumChannels;
-
-		PreSrcBuffer.Reset(NumSamplesToPop);
-		PreSrcBuffer.AddZeroed(NumSamplesToPop);
-
-		// pop into pre-src buffer
-		uint32 NumSamplesPopped = CircularDecoderOutputBuffer.Pop(PreSrcBuffer.GetData(), NumSamplesToPop);
-		ensure(NumSamplesPopped == NumSamplesToPop); // Circular buffer is too small! there will be discontinuities in the output
-
-		// holds result of sample rate conversion
-		float* PostSrcBufferPtr = PreSrcBuffer.GetData(); // assume no SRC
-
-		// perform SRC 
-		if (!FMath::IsNearlyEqual(LastPitchShiftCents, PitchShiftInCents))
+		if (OutputCircularBuffer.Num() >= NumOutputSamples)
 		{
-			Resampler.Init(Audio::EResamplingMethod::Linear, 1.f / TotalInSamplesOutSamplesRatio, NumChannels);
-			LastPitchShiftCents = PitchShiftInCents;
+			OutputCircularBuffer.Pop(OutputDest, NumOutputSamples);
 		}
-
-		int32 NumFramesConverted;
-		int32 Error = Resampler.ProcessAudio(PreSrcBuffer.GetData(), NumSamplesToPop / NumChannels, false, OutputDest, NumOutputFrames, NumFramesConverted);
-		ensure(!Error);
-
-		if (Error)
+		else if (ensure(bDecoderIsDone))
 		{
-			bDecoderIsDone = false;
-			Input.Reset();
-			Output.Reset();
-			Decoder.Reset();
+			const int32 NumSamplesToPop = OutputCircularBuffer.Num();
+			const int32 NumSamplesRemaining = NumOutputSamples - NumSamplesToPop;
+			OutputCircularBuffer.Pop(OutputDest, OutputCircularBuffer.Num());
+			FMemory::Memzero(&OutputDest[NumSamplesToPop], sizeof(float) * NumSamplesRemaining);
+			return NumSamplesToPop;
 		}
-
-
-		if (NumFramesConverted < (int32)NumOutputFrames)
+		else
 		{
+			ensureMsgf(false, TEXT("Something went wrong with decoding."));
 			bDecoderIsDone = true;
+			return 0;
 		}
-		
-		TotalNumFramesOutput += NumFramesConverted;
-		return NumFramesConverted;
+
+		return NumOutputSamples; // update once we are aware of partial decode on last buffer
 	}
 
 	bool FSimpleDecoderWrapper::InitializeDecodersInternal(const FSoundWaveProxyPtr& Wave)
@@ -143,8 +125,8 @@ namespace Audio
 		{
 			return false;
 		}
-		
-		// aquire codec:
+
+		// acquire codec:
 		ICodecRegistry::FCodecPtr Codec = ICodecRegistry::Get().FindCodecByParsingInput(Input.Get());
 		if (!Codec)
 		{
@@ -165,8 +147,16 @@ namespace Audio
 		// Decoder:
 		Decoder = Codec->CreateDecoder(Input.Get(), Output.Get());
 
-		// return true if all the components were succesfully create
+		// return true if all the components were successfully create
 		return Input.IsValid() && Output.IsValid() && Decoder.IsValid();
 	}
 
-} // namesapce Audio
+	float FSimpleDecoderWrapper::GetSampleRateRatio(float PitchShiftInCents) const
+	{
+		PitchShiftInCents = FMath::Clamp(PitchShiftInCents, -MaxPitchShiftCents, MaxPitchShiftCents);
+
+		// combine pitch shift due to actual pitch shift and SRC
+		return FsOutToInRatio / FMath::Pow(2.f, PitchShiftInCents / 1200.f);
+	}
+
+} // namespace Audio
diff --git a/Engine/Plugins/Runtime/Metasound/Source/MetasoundEngine/Private/MetasoundWavePlayerNode.cpp b/Engine/Plugins/Runtime/Metasound/Source/MetasoundEngine/Private/MetasoundWavePlayerNode.cpp
--- a/Engine/Plugins/Runtime/Metasound/Source/MetasoundEngine/Private/MetasoundWavePlayerNode.cpp
+++ b/Engine/Plugins/Runtime/Metasound/Source/MetasoundEngine/Private/MetasoundWavePlayerNode.cpp
@@ -71,8 +71,6 @@ namespace Metasound
 			{
 				CurrentSoundWaveName = Wave->SoundWaveProxy->GetFName();
 			}
-
-
 		}
 
 		virtual FDataReferenceCollection GetInputs() const override
@@ -179,7 +177,7 @@ namespace Metasound
 				PostSrcBuffer.AddZeroed(NumSamplesToGenerate);
 				float* PostSrcBufferPtr = PostSrcBuffer.GetData();
 
-				int32 NumFramesDecoded = Decoder.GenerateAudio(PostSrcBufferPtr, NumOutputFrames, *PitchShiftCents); // TODO: pitch shift
+				int32 NumFramesDecoded = Decoder.GenerateAudio(PostSrcBufferPtr, NumOutputFrames, *PitchShiftCents) / NumInputChannels;
 
 				// TODO: handle decoder having completed during it's decode
 				if (!Decoder.CanGenerateAudio() ||  (NumFramesDecoded < NumOutputFrames))
@@ -192,12 +190,12 @@ namespace Metasound
 				{
 					// TODO: attenuate by -3 dB?
 					// copy to Left & Right output buffers
-					FMemory::Memcpy(FinalOutputLeft, PostSrcBufferPtr, sizeof(float) * NumOutputFrames);
-					FMemory::Memcpy(FinalOutputRight, PostSrcBufferPtr, sizeof(float) * NumOutputFrames);
+					FMemory::Memcpy(FinalOutputLeft, PostSrcBufferPtr, sizeof(float) * NumFramesDecoded);
+					FMemory::Memcpy(FinalOutputRight, PostSrcBufferPtr, sizeof(float) * NumFramesDecoded);
 				}
 				else if (bNeedsDeinterleave)
 				{
-					for (int32 i = 0; i < NumOutputFrames; ++i)
+					for (int32 i = 0; i < NumFramesDecoded; ++i)
 					{
 						// de-interleave each stereo frame into output buffers
 						FinalOutputLeft[i] = PostSrcBufferPtr[(i << 1)];
@@ -255,9 +253,9 @@ namespace Metasound
 			AddBuildError<FWavePlayerError>(OutErrors, WaveNode, LOCTEXT("NoSoundWave", "No Sound Wave"));
 
 		}
-		else if (Wave->SoundWaveProxy->GetNumChannels() != 1)
+		else if (Wave->SoundWaveProxy->GetNumChannels() > 2)
 		{
-			AddBuildError<FWavePlayerError>(OutErrors, WaveNode, LOCTEXT("WavePlayerCurrentlyOnlySuportsMonoAssets", "Wave Player Currently Only Supports Mono Assets"));
+			AddBuildError<FWavePlayerError>(OutErrors, WaveNode, LOCTEXT("WavePlayerCurrentlyOnlySuportsMonoAssets", "Wave Player Currently Only Supports Mono Or Stereo Assets"));
 		}
 
 		return MakeUnique<FWavePlayerOperator>(
diff --git a/Engine/Plugins/Runtime/Metasound/Source/MetasoundEngine/Public/MetasoundWave.h b/Engine/Plugins/Runtime/Metasound/Source/MetasoundEngine/Public/MetasoundWave.h
--- a/Engine/Plugins/Runtime/Metasound/Source/MetasoundEngine/Public/MetasoundWave.h
+++ b/Engine/Plugins/Runtime/Metasound/Source/MetasoundEngine/Public/MetasoundWave.h
@@ -89,6 +89,8 @@ namespace Audio
 		uint32 GenerateAudio(float* OutputDest, int32 NumOutputFrames, float PitchShiftInCents = 0.f);
 
 	private:
+		float GetSampleRateRatio(float PitchShiftInCents) const;
+
 		// actual decoder objects
 		TUniquePtr<Audio::IDecoder> Decoder;
 		TUniquePtr<Audio::IDecoderOutput> Output;
@@ -102,12 +104,13 @@ namespace Audio
 
 		// buffers
 		TArray<float> PreSrcBuffer;
-		Audio::TCircularAudioBuffer<float> CircularDecoderOutputBuffer;
+		TArray<float> PostSrcBuffer;
+		Audio::TCircularAudioBuffer<float> OutputCircularBuffer;
 
 		// meta data:
 		float InputSampleRate;
 		float OutputSampleRate;
-		float FsInToFsOutRatio;
+		float FsOutToInRatio;
 		float MaxPitchShiftCents;
 		float MaxPitchShiftRatio;
 
diff --git a/Engine/Source/Runtime/Android/AndroidAudio/Private/AndroidAudioBuffer.cpp b/Engine/Source/Runtime/Android/AndroidAudio/Private/AndroidAudioBuffer.cpp
--- a/Engine/Source/Runtime/Android/AndroidAudio/Private/AndroidAudioBuffer.cpp
+++ b/Engine/Source/Runtime/Android/AndroidAudio/Private/AndroidAudioBuffer.cpp
@@ -296,9 +296,10 @@ bool FSLESSoundBuffer::ReadCompressedData( uint8* Destination, int32 NumFramesTo
 {
 	ensure(DecompressionState);
 	
+	int32 NumFramesDecoded = DecompressionState->GetStreamBufferSize() * NumChannels;
 	if(Format == SoundFormat_Streaming)
 	{
-		return(DecompressionState->StreamCompressedData(Destination, bLooping, DecompressionState->GetStreamBufferSize() * NumChannels));
+		return(DecompressionState->StreamCompressedData(Destination, bLooping, NumFramesDecoded, NumFramesDecoded));
 	}
 
 	return(DecompressionState->ReadCompressedData(Destination, bLooping, DecompressionState->GetStreamBufferSize() * NumChannels));
diff --git a/Engine/Source/Runtime/AudioCodecEngine/Private/DecoderBackCompat.cpp b/Engine/Source/Runtime/AudioCodecEngine/Private/DecoderBackCompat.cpp
--- a/Engine/Source/Runtime/AudioCodecEngine/Private/DecoderBackCompat.cpp
+++ b/Engine/Source/Runtime/AudioCodecEngine/Private/DecoderBackCompat.cpp
@@ -68,19 +68,21 @@ namespace Audio
 
 		while(!bFinished && NumFramesRemaining > 0)
 		{
+			int32 NumBytesStreamed = BuffSizeInBytes;
 			if (BackCompatSrc.Wave->IsStreaming())
 			{
-				bFinished = Info->StreamCompressedData(Buff, false, BuffSizeInBytes);
+				bFinished = Info->StreamCompressedData(Buff, false, BuffSizeInBytes, NumBytesStreamed);
 			}
 			else
 			{
 				bFinished = Info->ReadCompressedData(Buff, false, BuffSizeInBytes);
 			}						
 			PushedDetails.SampleFramesStartOffset = FrameOffset;
+			ResidualBuffer.SetNum(NumBytesStreamed / sizeof(int16));
 			
 			Dst->PushAudio( 
 				PushedDetails,
-				MakeArrayView(ResidualBuffer) 
+				MakeArrayView(ResidualBuffer)
 			);
 			
 			FrameOffset += BuffSizeInFrames;
diff --git a/Engine/Source/Runtime/AudioMixer/Private/AudioMixerSourceDecode.cpp b/Engine/Source/Runtime/AudioMixer/Private/AudioMixerSourceDecode.cpp
--- a/Engine/Source/Runtime/AudioMixer/Private/AudioMixerSourceDecode.cpp
+++ b/Engine/Source/Runtime/AudioMixer/Private/AudioMixerSourceDecode.cpp
@@ -129,11 +129,12 @@ public:
 				if (DecodeTaskData.bSkipFirstBuffer)
 				{
 					const int32 kPCMBufferSize = NumChannels * DecodeTaskData.NumPrecacheFrames * sizeof(int16);
+					int32 NumBytesStreamed = kPCMBufferSize;
 					if (DecodeTaskData.BufferType == EBufferType::Streaming)
 					{
 						for (int32 NumberOfBuffersToSkip = 0; NumberOfBuffersToSkip < PLATFORM_NUM_AUDIODECOMPRESSION_PRECACHE_BUFFERS; NumberOfBuffersToSkip++)
 						{
-							DecodeTaskData.DecompressionState->StreamCompressedData(DecodeBuffer.GetData(), DecodeTaskData.bLoopingMode, kPCMBufferSize);
+							DecodeTaskData.DecompressionState->StreamCompressedData(DecodeBuffer.GetData(), DecodeTaskData.bLoopingMode, kPCMBufferSize, NumBytesStreamed);
 						}
 					}
 					else
@@ -146,9 +147,10 @@ public:
 				}
 
 				const int32 kPCMBufferSize = NumChannels * DecodeTaskData.NumFramesToDecode * sizeof(int16);
+				int32 NumBytesStreamed = kPCMBufferSize;
 				if (DecodeTaskData.BufferType == EBufferType::Streaming)
 				{
-					DecodeResult.bLooped = DecodeTaskData.DecompressionState->StreamCompressedData(DecodeBuffer.GetData(), DecodeTaskData.bLoopingMode, kPCMBufferSize);
+					DecodeResult.bLooped = DecodeTaskData.DecompressionState->StreamCompressedData(DecodeBuffer.GetData(), DecodeTaskData.bLoopingMode, kPCMBufferSize, NumBytesStreamed);
 				}
 				else
 				{
diff --git a/Engine/Source/Runtime/Engine/Private/ADPCMAudioInfo.cpp b/Engine/Source/Runtime/Engine/Private/ADPCMAudioInfo.cpp
--- a/Engine/Source/Runtime/Engine/Private/ADPCMAudioInfo.cpp
+++ b/Engine/Source/Runtime/Engine/Private/ADPCMAudioInfo.cpp
@@ -581,9 +581,11 @@ bool FADPCMAudioInfo::StreamCompressedInfoInternal(const FSoundWaveProxyPtr& InW
 	return true;
 }
 
-bool FADPCMAudioInfo::StreamCompressedData(uint8* Destination, bool bLooping, uint32 BufferSize)
+bool FADPCMAudioInfo::StreamCompressedData(uint8* Destination, bool bLooping, uint32 BufferSize, int32& OutNumBytesStreamed)
 {
 	FScopeLock ScopeLock(&CurCompressedChunkHandleCriticalSection);
+	int32 OriginalBufferSize = (int32)BufferSize;
+	OutNumBytesStreamed = 0;
 
 	// Initial sanity checks:
 	if (bDecompressorReleased)
@@ -772,6 +774,7 @@ bool FADPCMAudioInfo::StreamCompressedData(uint8* Destination, bool bLooping, ui
 
 						// zero out remaining data and bail
 						FMemory::Memset(OutData, 0, BufferSize);
+						OutNumBytesStreamed = OriginalBufferSize;
 						return false;
 					}
 
@@ -818,6 +821,7 @@ bool FADPCMAudioInfo::StreamCompressedData(uint8* Destination, bool bLooping, ui
 				{
 					uint16 Value = *(int16*)(UncompressedBlockData + ChannelItr * UncompressedBlockSize + (CurrentUncompressedBlockSampleIndex + SampleItr) * sizeof(int16));
 					*OutData++ = Value;
+					OutNumBytesStreamed += sizeof(int16);
 				}
 			}
 
@@ -895,6 +899,7 @@ bool FADPCMAudioInfo::StreamCompressedData(uint8* Destination, bool bLooping, ui
 					NumConsecutiveReadFailiures++;
 					const bool bReadAttemptTimedOut = NumConsecutiveReadFailiures > ADPCMReadFailiureTimeoutCVar;
 					UE_CLOG(bReadAttemptTimedOut, LogAudio, Warning, TEXT("ADPCM Audio Decode timed out."), bReadAttemptTimedOut);
+					OutNumBytesStreamed = OriginalBufferSize;
 					return NumConsecutiveReadFailiures > ADPCMReadFailiureTimeoutCVar;
 				}
 
@@ -927,8 +932,9 @@ bool FADPCMAudioInfo::StreamCompressedData(uint8* Destination, bool bLooping, ui
 
 			const uint32 SampleSize = DecompressedSamplesToCopy * ChannelSampleSize;
 			FMemory::Memcpy(OutData, CurCompressedChunkData + CurrentChunkBufferOffset, SampleSize);
-
 			OutData += DecompressedSamplesToCopy * NumChannels;
+			OutNumBytesStreamed += DecompressedSamplesToCopy * NumChannels * sizeof(int16);
+
 			CurrentChunkBufferOffset += SampleSize;
 			BufferSize -= SampleSize;
 			TotalSamplesStreamed += DecompressedSamplesToCopy;
diff --git a/Engine/Source/Runtime/Engine/Private/AudioDecompress.cpp b/Engine/Source/Runtime/Engine/Private/AudioDecompress.cpp
--- a/Engine/Source/Runtime/Engine/Private/AudioDecompress.cpp
+++ b/Engine/Source/Runtime/Engine/Private/AudioDecompress.cpp
@@ -177,7 +177,7 @@ bool IStreamedCompressedInfo::StreamCompressedInfoInternal(const FSoundWaveProxy
 	return false;
 }
 
-bool IStreamedCompressedInfo::StreamCompressedData(uint8* Destination, bool bLooping, uint32 BufferSize)
+bool IStreamedCompressedInfo::StreamCompressedData(uint8* Destination, bool bLooping, uint32 BufferSize, int32& OutNumBytesStreamed)
 {
 	check(Destination);
 	SCOPED_NAMED_EVENT(IStreamedCompressedInfo_StreamCompressedData, FColor::Blue);
@@ -209,6 +209,7 @@ bool IStreamedCompressedInfo::StreamCompressedData(uint8* Destination, bool bLoo
 				bPrintChunkFailMessage = false;
 			}
 			ZeroBuffer(Destination + RawPCMOffset, BufferSize - RawPCMOffset);
+			OutNumBytesStreamed = BufferSize - RawPCMOffset;
 			return false;
 		}
 	}
@@ -234,6 +235,7 @@ bool IStreamedCompressedInfo::StreamCompressedData(uint8* Destination, bool bLoo
 		{
 			LastPCMByteSize = 0;
 			ZeroBuffer(Destination + RawPCMOffset, BufferSize - RawPCMOffset);
+			OutNumBytesStreamed = BufferSize - RawPCMOffset;
 			return false;
 		}
 		else
@@ -270,6 +272,7 @@ bool IStreamedCompressedInfo::StreamCompressedData(uint8* Destination, bool bLoo
 					else
 					{
 						RawPCMOffset += ZeroBuffer(Destination + RawPCMOffset, BufferSize - RawPCMOffset);
+						OutNumBytesStreamed = BufferSize - RawPCMOffset;
 					}
 				}
 				else
@@ -292,6 +295,7 @@ bool IStreamedCompressedInfo::StreamCompressedData(uint8* Destination, bool bLoo
 		}
 	}
 
+	OutNumBytesStreamed = BufferSize;
 	return bLooped;
 }
 
diff --git a/Engine/Source/Runtime/Engine/Private/VorbisAudioInfo.cpp b/Engine/Source/Runtime/Engine/Private/VorbisAudioInfo.cpp
--- a/Engine/Source/Runtime/Engine/Private/VorbisAudioInfo.cpp
+++ b/Engine/Source/Runtime/Engine/Private/VorbisAudioInfo.cpp
@@ -612,8 +612,10 @@ const uint8* FVorbisAudioInfo::GetLoadedChunk(FSoundWaveProxyPtr InSoundWave, ui
 	return nullptr;
 }
 
-bool FVorbisAudioInfo::StreamCompressedData(uint8* InDestination, bool bLooping, uint32 BufferSize)
+bool FVorbisAudioInfo::StreamCompressedData(uint8* InDestination, bool bLooping, uint32 BufferSize, int32& OutNumBytesStreamed)
 {
+	const uint32 DestinationSize = BufferSize;
+
 	if (!bDllLoaded)
 	{
 		UE_LOG(LogAudio, Error, TEXT("FVorbisAudioInfo::StreamCompressedData failed due to vorbis DLL not being loaded."));
@@ -653,6 +655,7 @@ bool FVorbisAudioInfo::StreamCompressedData(uint8* InDestination, bool bLooping,
 			{
 				// zero out the rest of the buffer
 				FMemory::Memzero(InDestination, BufferSize);
+				OutNumBytesStreamed = DestinationSize;
 				return false;
 			}
 
@@ -708,6 +711,7 @@ bool FVorbisAudioInfo::StreamCompressedData(uint8* InDestination, bool bLooping,
 		BufferSize -= BytesActuallyRead;
 	}
 
+	OutNumBytesStreamed = DestinationSize - BufferSize;
 	return( bLooped );
 }
 
diff --git a/Engine/Source/Runtime/Engine/Public/ADPCMAudioInfo.h b/Engine/Source/Runtime/Engine/Public/ADPCMAudioInfo.h
--- a/Engine/Source/Runtime/Engine/Public/ADPCMAudioInfo.h
+++ b/Engine/Source/Runtime/Engine/Public/ADPCMAudioInfo.h
@@ -107,7 +107,7 @@ public:
 	// Additional overrides for streaming
 	virtual bool SupportsStreaming() const override {return true;}
 	virtual bool StreamCompressedInfoInternal(const FSoundWaveProxyPtr& InWaveProxy, struct FSoundQualityInfo* QualityInfo) override;
-	virtual bool StreamCompressedData(uint8* Destination, bool bLooping, uint32 BufferSize) override;
+	virtual bool StreamCompressedData(uint8* Destination, bool bLooping, uint32 BufferSize, int32& OutNumBytesStreamed) override;
 	virtual int32 GetCurrentChunkIndex() const override
 	{
 		return CurrentChunkIndex;
diff --git a/Engine/Source/Runtime/Engine/Public/AudioDecompress.h b/Engine/Source/Runtime/Engine/Public/AudioDecompress.h
--- a/Engine/Source/Runtime/Engine/Public/AudioDecompress.h
+++ b/Engine/Source/Runtime/Engine/Public/AudioDecompress.h
@@ -140,7 +140,7 @@ public:
 	*
 	* @return	bool		true if the end of the data was reached (for both single shot and looping sounds)
 	*/
-	virtual bool StreamCompressedData(uint8* Destination, bool bLooping, uint32 BufferSize) {return false;}
+	virtual bool StreamCompressedData(uint8* Destination, bool bLooping, uint32 BufferSize, int32& OutNumBytesStreamed) { OutNumBytesStreamed = -1; return false; }
 
 	/**
 	 * Gets the chunk index that was last read from (for Streaming Manager requests)
@@ -198,7 +198,7 @@ public:
 	virtual int GetStreamBufferSize() const override { return  MONO_PCM_BUFFER_SIZE; }
 	virtual bool SupportsStreaming() const override { return true; }
 	virtual bool StreamCompressedInfoInternal(const FSoundWaveProxyPtr& InWaveProxy, FSoundQualityInfo* QualityInfo) override;
-	virtual bool StreamCompressedData(uint8* Destination, bool bLooping, uint32 BufferSize) override;
+	virtual bool StreamCompressedData(uint8* Destination, bool bLooping, uint32 BufferSize, int32& OutNumBytesStreamed) override;
 	virtual int32 GetCurrentChunkIndex() const override { return CurrentChunkIndex; }
 	virtual int32 GetCurrentChunkOffset() const override { return SrcBufferOffset; }
 	//~ End ICompressedInfo Interface
diff --git a/Engine/Source/Runtime/Engine/Public/VorbisAudioInfo.h b/Engine/Source/Runtime/Engine/Public/VorbisAudioInfo.h
--- a/Engine/Source/Runtime/Engine/Public/VorbisAudioInfo.h
+++ b/Engine/Source/Runtime/Engine/Public/VorbisAudioInfo.h
@@ -99,7 +99,7 @@ public:
 
 	// Additional overrides for streaming
 	virtual bool SupportsStreaming() const override {return true;}
-	virtual bool StreamCompressedData(uint8* InDestination, bool bLooping, uint32 BufferSize) override;
+	virtual bool StreamCompressedData(uint8* InDestination, bool bLooping, uint32 BufferSize, int32& OutNumBytesStreamed) override;
 	virtual int32 GetCurrentChunkIndex() const override {return CurrentStreamingChunkIndex;}
 	virtual int32 GetCurrentChunkOffset() const override {return BufferOffset % CurrentStreamingChunksSize;}
 	// End of ICompressedAudioInfo Interface
diff --git a/Engine/Source/Runtime/IOS/IOSAudio/Private/IOSAudioBuffer.cpp b/Engine/Source/Runtime/IOS/IOSAudio/Private/IOSAudioBuffer.cpp
--- a/Engine/Source/Runtime/IOS/IOSAudio/Private/IOSAudioBuffer.cpp
+++ b/Engine/Source/Runtime/IOS/IOSAudio/Private/IOSAudioBuffer.cpp
@@ -138,9 +138,10 @@ bool FIOSAudioSoundBuffer::ReadCompressedInfo(USoundWave* InWave)
 
 bool FIOSAudioSoundBuffer::ReadCompressedData( uint8* Destination, int32 NumFramesToDecode, bool bLooping )
 {
+	int32 NumFramesDecoded = NumFramesToDecode;
 	if(bStreaming)
 	{
-		return DecompressionState->StreamCompressedData(Destination, bLooping, RenderCallbackBufferSize);
+		return DecompressionState->StreamCompressedData(Destination, bLooping, RenderCallbackBufferSize, NumFramesDecoded);
 	}
 	else
 	{
diff --git a/Engine/Source/Runtime/Mac/CoreAudio/Private/CoreAudioBuffer.cpp b/Engine/Source/Runtime/Mac/CoreAudio/Private/CoreAudioBuffer.cpp
--- a/Engine/Source/Runtime/Mac/CoreAudio/Private/CoreAudioBuffer.cpp
+++ b/Engine/Source/Runtime/Mac/CoreAudio/Private/CoreAudioBuffer.cpp
@@ -126,9 +126,10 @@ void FCoreAudioSoundBuffer::InitAudioStreamBasicDescription( UInt32 FormatID, US
 bool FCoreAudioSoundBuffer::ReadCompressedData(uint8* Destination, int32 NumFramesToDecode, bool bLooping)
 {
 	const uint32 kPCMBufferSize = MONO_PCM_BUFFER_SIZE * NumChannels;
+	int32 NumBytesStreamed = (int32)kPCMBUfferSize;
 	if (SoundFormat == SoundFormat_Streaming)
 	{
-		return DecompressionState->StreamCompressedData(Destination, bLooping, kPCMBufferSize);
+		return DecompressionState->StreamCompressedData(Destination, bLooping, kPCMBufferSize, NumBytesStreamed);
 	}
 	else
 	{
diff --git a/Engine/Source/Runtime/Windows/XAudio2/Private/XAudio2Buffer.cpp b/Engine/Source/Runtime/Windows/XAudio2/Private/XAudio2Buffer.cpp
--- a/Engine/Source/Runtime/Windows/XAudio2/Private/XAudio2Buffer.cpp
+++ b/Engine/Source/Runtime/Windows/XAudio2/Private/XAudio2Buffer.cpp
@@ -338,9 +338,10 @@ bool FXAudio2SoundBuffer::ReadCompressedData( uint8* Destination, int32 NumFrame
 	}
 	
 	int32 NumBytesToDecode = NumFramesToDecode * sizeof(int16) * NumChannels;
+	int32 NumBytesDecoded = NumBytesToDecode;
 	if (SoundFormat == SoundFormat_Streaming)
 	{
-		return DecompressionState->StreamCompressedData(Destination, bLooping, NumBytesToDecode);
+		return DecompressionState->StreamCompressedData(Destination, bLooping, NumBytesToDecode, NumBytesDecoded);
 	}
 	else
 	{
