Disable default UMetasoundFactory for EA as asset type as it's not fleshed out yet

diff --git a/Engine/Plugins/Runtime/Metasound/Source/MetasoundEditor/Private/MetasoundFactory.cpp b/Engine/Plugins/Runtime/Metasound/Source/MetasoundEditor/Private/MetasoundFactory.cpp
--- a/Engine/Plugins/Runtime/Metasound/Source/MetasoundEditor/Private/MetasoundFactory.cpp
+++ b/Engine/Plugins/Runtime/Metasound/Source/MetasoundEditor/Private/MetasoundFactory.cpp
@@ -7,34 +7,35 @@
 #include "MetasoundEditorGraphNode.h"
 #include "MetasoundEditorGraphBuilder.h"
 
-UMetasoundFactory::UMetasoundFactory(const FObjectInitializer& ObjectInitializer)
-	: Super(ObjectInitializer)
-{
-	SupportedClass = UMetasound::StaticClass();
-
-	bCreateNew = true;
-	bEditorImport = false;
-	bEditAfterNew = true;
-}
-
-UObject* UMetasoundFactory::FactoryCreateNew(UClass* InClass, UObject* InParent, FName Name, EObjectFlags Flags, UObject* InContext, FFeedbackContext* InFeedbackContext)
-{
-	UMetasound* NewMetasound = NewObject<UMetasound>(InParent, Name, Flags);
-
-	FMetasoundFrontendClassMetadata Metadata;
-
-	Metadata.ClassName = FMetasoundFrontendClassName{TEXT(""), *NewMetasound->GetName(), TEXT("")};
-	Metadata.Version.Major = 1;
-	Metadata.Version.Minor = 0;
-	Metadata.Type = EMetasoundFrontendClassType::Graph;
-	Metadata.Author = FText::FromString(UKismetSystemLibrary::GetPlatformUserName());
-
-	NewMetasound->SetMetadata(Metadata);
-
-	NewMetasound->ConformDocumentToMetasoundArchetype();
-
-	return NewMetasound;
-}
+// TODO: Re-enable and potentially rename once composition is supported
+// UMetasoundFactory::UMetasoundFactory(const FObjectInitializer& ObjectInitializer)
+// 	: Super(ObjectInitializer)
+// {
+// 	SupportedClass = UMetasound::StaticClass();
+// 
+// 	bCreateNew = true;
+// 	bEditorImport = false;
+// 	bEditAfterNew = true;
+// }
+// 
+// UObject* UMetasoundFactory::FactoryCreateNew(UClass* InClass, UObject* InParent, FName Name, EObjectFlags Flags, UObject* InContext, FFeedbackContext* InFeedbackContext)
+// {
+// 	UMetasound* NewMetasound = NewObject<UMetasound>(InParent, Name, Flags);
+// 
+// 	FMetasoundFrontendClassMetadata Metadata;
+// 
+// 	Metadata.ClassName = FMetasoundFrontendClassName{TEXT(""), *NewMetasound->GetName(), TEXT("")};
+// 	Metadata.Version.Major = 1;
+// 	Metadata.Version.Minor = 0;
+// 	Metadata.Type = EMetasoundFrontendClassType::Graph;
+// 	Metadata.Author = FText::FromString(UKismetSystemLibrary::GetPlatformUserName());
+// 
+// 	NewMetasound->SetMetadata(Metadata);
+// 
+// 	NewMetasound->ConformDocumentToMetasoundArchetype();
+// 
+// 	return NewMetasound;
+// }
 
 UMetasoundSourceFactory::UMetasoundSourceFactory(const FObjectInitializer& ObjectInitializer)
 	: Super(ObjectInitializer)
diff --git a/Engine/Plugins/Runtime/Metasound/Source/MetasoundEditor/Public/MetasoundFactory.h b/Engine/Plugins/Runtime/Metasound/Source/MetasoundEditor/Public/MetasoundFactory.h
--- a/Engine/Plugins/Runtime/Metasound/Source/MetasoundEditor/Public/MetasoundFactory.h
+++ b/Engine/Plugins/Runtime/Metasound/Source/MetasoundEditor/Public/MetasoundFactory.h
@@ -8,15 +8,16 @@
 
 #include "MetasoundFactory.generated.h"
 
-UCLASS(hidecategories=Object, MinimalAPI)
-class UMetasoundFactory : public UFactory
-{
-	GENERATED_UCLASS_BODY()
-
-	//~ Begin UFactory Interface
-	virtual UObject* FactoryCreateNew(UClass* InClass, UObject* InParent, FName Name, EObjectFlags Flags, UObject* InContext, FFeedbackContext* InFeedbackContext) override;
-	//~ Begin UFactory Interface
-};
+// TODO: Re-enable and potentially rename once composition is supported
+// UCLASS(hidecategories=Object, MinimalAPI)
+// class UMetasoundFactory : public UFactory
+// {
+// 	GENERATED_UCLASS_BODY()
+// 
+// 	//~ Begin UFactory Interface
+// 	virtual UObject* FactoryCreateNew(UClass* InClass, UObject* InParent, FName Name, EObjectFlags Flags, UObject* InContext, FFeedbackContext* InFeedbackContext) override;
+// 	//~ Begin UFactory Interface
+// };
 
 UCLASS(hidecategories = Object, MinimalAPI)
 class UMetasoundSourceFactory : public UFactory
