- Fix bug/crashes with undo desync after updating input literals from graph input node
- Fix minor bug with MS doc sync incorrectly marking package dirty
- Add MS sync when undo'ing transaction involving MS as a safeguard to ensure graph/document parity is preserved (and sync code is maintained)
- Make Metasound version history visible but not editable by end-user

diff --git a/Engine/Plugins/Runtime/Metasound/Source/MetasoundEditor/Private/MetasoundEditor.cpp b/Engine/Plugins/Runtime/Metasound/Source/MetasoundEditor/Private/MetasoundEditor.cpp
--- a/Engine/Plugins/Runtime/Metasound/Source/MetasoundEditor/Private/MetasoundEditor.cpp
+++ b/Engine/Plugins/Runtime/Metasound/Source/MetasoundEditor/Private/MetasoundEditor.cpp
@@ -229,7 +229,7 @@ namespace Metasound
 
 		FName FEditor::GetToolkitFName() const
 		{
-			return FName("MetasoundEditor");
+			return "MetasoundEditor";
 		}
 
 		FText FEditor::GetBaseToolkitName() const
@@ -258,9 +258,15 @@ namespace Metasound
 			{
 				MetasoundGraphEditor->ClearSelectionSet();
 				MetasoundGraphEditor->NotifyGraphChanged();
-				FSlateApplication::Get().DismissAllMenus();
 			}
 
+			FSlateApplication::Get().DismissAllMenus();
+
+			if (Metasound)
+			{
+				UEdGraphPin::ResolveAllPinReferences();
+				Metasound::Editor::FGraphBuilder::SynchronizeGraph(*Metasound);
+			}
 		}
 
 		void FEditor::NotifyPostChange(const FPropertyChangedEvent& PropertyChangedEvent, FProperty* PropertyThatChanged)
@@ -767,7 +773,7 @@ namespace Metasound
 				// Some nodes may not be metasound nodes (ex. comments and perhaps aliases eventually), but can be safely deleted.
 				if (UEdGraphNode* Node = Cast<UEdGraphNode>(*NodeIt))
 				{
-					if (!FGraphBuilder::DeleteNode(*Node, true /* bInRecordTransaction */))
+					if (!FGraphBuilder::DeleteNode(*Node))
 					{
 						MetasoundGraphEditor->SetNodeSelection(Node, true /* bSelect */);
 					}
diff --git a/Engine/Plugins/Runtime/Metasound/Source/MetasoundEditor/Private/MetasoundEditorGraph.cpp b/Engine/Plugins/Runtime/Metasound/Source/MetasoundEditor/Private/MetasoundEditorGraph.cpp
--- a/Engine/Plugins/Runtime/Metasound/Source/MetasoundEditor/Private/MetasoundEditorGraph.cpp
+++ b/Engine/Plugins/Runtime/Metasound/Source/MetasoundEditor/Private/MetasoundEditorGraph.cpp
@@ -2,6 +2,7 @@
 #include "MetasoundEditorGraph.h"
 
 #include "EdGraph/EdGraphNode.h"
+#include "MetasoundEditorGraphBuilder.h"
 #include "MetasoundEditorGraphNode.h"
 #include "MetasoundEditorGraphInputNodes.h"
 #include "MetasoundEditorModule.h"
@@ -46,3 +47,11 @@ UObject& UMetasoundEditorGraph::GetMetasoundChecked() const
 	check(ParentMetasound);
 	return *ParentMetasound;
 }
+
+void UMetasoundEditorGraph::Synchronize()
+{
+	if (ParentMetasound)
+	{
+		Metasound::Editor::FGraphBuilder::SynchronizeGraph(*ParentMetasound);
+	}
+}
\ No newline at end of file
diff --git a/Engine/Plugins/Runtime/Metasound/Source/MetasoundEditor/Private/MetasoundEditorGraphBuilder.cpp b/Engine/Plugins/Runtime/Metasound/Source/MetasoundEditor/Private/MetasoundEditorGraphBuilder.cpp
--- a/Engine/Plugins/Runtime/Metasound/Source/MetasoundEditor/Private/MetasoundEditorGraphBuilder.cpp
+++ b/Engine/Plugins/Runtime/Metasound/Source/MetasoundEditor/Private/MetasoundEditorGraphBuilder.cpp
@@ -17,7 +17,6 @@
 #include "MetasoundLiteral.h"
 #include "MetasoundUObjectRegistry.h"
 #include "Modules/ModuleManager.h"
-#include "ScopedTransaction.h"
 #include "Templates/Tuple.h"
 #include "Widgets/Notifications/SNotificationList.h"
 
@@ -63,8 +62,6 @@ namespace Metasound
 
 		UEdGraphNode* FGraphBuilder::AddNode(UObject& InMetasound, Frontend::FNodeHandle& InNodeHandle, bool bInSelectNewNode)
 		{
-			const FScopedTransaction Transaction(LOCTEXT("AddMetasoundGraphNode", "Add Metasound Node"));
-
 			FMetasoundAssetBase* MetasoundAsset = IMetasoundUObjectRegistry::Get().GetObjectAsAssetBase(&InMetasound);
 			check(MetasoundAsset);
 
@@ -255,7 +252,7 @@ namespace Metasound
 
 		UMetasoundEditorGraphInputNode* FGraphBuilder::AddInput(UObject& InMetasound, Frontend::FNodeHandle InNodeHandle, bool bInSelectNewNode)
 		{
-			using namespace Metasound::Frontend;
+			using namespace Frontend;
 
 			FMetasoundAssetBase* MetasoundAsset = IMetasoundUObjectRegistry::Get().GetObjectAsAssetBase(&InMetasound);
 			check(MetasoundAsset);
@@ -293,9 +290,7 @@ namespace Metasound
 
 		UMetasoundEditorGraphInputNode* FGraphBuilder::AddInput(UObject& InMetasound, const FString& InName, const FName InTypeName, const FMetasoundFrontendNodeStyle& InNodeStyle, const FText& InToolTip, bool bInSelectNewNode)
 		{
-			using namespace Metasound::Frontend;
-
-			const FScopedTransaction Transaction(LOCTEXT("AddMetasoundGraphInputNode", "Add Metasound Input Node"));
+			using namespace Frontend;
 
 			FNodeHandle NodeHandle = AddInputNodeHandle(InMetasound, InName, InTypeName, InNodeStyle, InToolTip);
 			if (UMetasoundEditorGraphInputNode* NewGraphNode = AddInput(InMetasound, NodeHandle))
@@ -311,9 +306,7 @@ namespace Metasound
 
 		void FGraphBuilder::AddOrUpdateLiteralInput(UObject& InMetasound, Frontend::FNodeHandle InNodeHandle, UEdGraphPin& InInputPin)
 		{
-			using namespace Metasound::Frontend;
-
-			const FScopedTransaction Transaction(LOCTEXT("SetMetasoundGraphNode", "Add/Update Metasound Literal Input"));
+			using namespace Frontend;
 
 			const FString InInputName = InInputPin.GetName();
 			TArray<FInputHandle> InputHandles = InNodeHandle->GetInputsWithVertexName(InInputName);
@@ -323,9 +316,65 @@ namespace Metasound
 			}
 
 			FInputHandle InputHandle = InputHandles[0];
-			IMetasoundEditorModule& EditorModule = FModuleManager::GetModuleChecked<IMetasoundEditorModule>("MetasoundEditor");
-
 			FMetasoundFrontendLiteral PinDataTypeDefaultLiteral;
+			if (!GetPinDefaultLiteral(InInputPin, InputHandle, PinDataTypeDefaultLiteral))
+			{
+				return;
+			}
+
+			FOutputHandle OutputHandle = InputHandle->GetCurrentlyConnectedOutput();
+			if (!OutputHandle->IsValid())
+			{
+				FMetasoundFrontendNodeStyle Style;
+				Style.Display.Visibility = EMetasoundFrontendNodeStyleDisplayVisibility::Hidden;
+
+				const FName TypeName = InputHandle->GetDataType();
+				const FString NewInputName = GenerateUniqueInputName(InMetasound, "LiteralInput");
+				Frontend::FNodeHandle NewInputNode = AddInputNodeHandle(InMetasound, NewInputName, TypeName, Style, FText::GetEmpty(), &PinDataTypeDefaultLiteral);
+				NewInputNode->SetNodeStyle(Style);
+
+				const FMetasoundFrontendLiteral* DefaultLiteral = InputHandle->GetDefaultLiteral();
+				if (DefaultLiteral && DefaultLiteral->IsValid())
+				{
+					Frontend::FGraphHandle GraphHandle = InNodeHandle->GetOwningGraph();
+					const FGuid VertexID = GraphHandle->GetVertexIDForInputVertex(NewInputNode->GetNodeName());
+					GraphHandle->SetDefaultInput(VertexID, *DefaultLiteral);
+					InInputPin.DefaultValue = DefaultLiteral->ToString();
+				}
+
+				TArray<FOutputHandle> OutputHandles = NewInputNode->GetOutputs();
+				if (ensure(OutputHandles.Num() == 1))
+				{
+					OutputHandle = OutputHandles[0];
+				}
+
+				ensure(InputHandle->Connect(*OutputHandle));
+
+				InMetasound.PostEditChange();
+				InMetasound.MarkPackageDirty();
+			}
+			else
+			{
+				FNodeHandle InputNode = OutputHandle->GetOwningNode();
+				const FMetasoundFrontendNodeStyle& Style = InputNode->GetNodeStyle();
+				if (InputNode->IsValid() && Style.Display.Visibility == EMetasoundFrontendNodeStyleDisplayVisibility::Hidden)
+				{
+					Frontend::FGraphHandle GraphHandle = InputNode->GetOwningGraph();
+					const FGuid VertexID = GraphHandle->GetVertexIDForInputVertex(InputNode->GetNodeName());
+
+					GraphHandle->SetDefaultInput(VertexID, PinDataTypeDefaultLiteral);
+
+					InMetasound.PostEditChange();
+					InMetasound.MarkPackageDirty();
+				}
+			}
+		}
+
+		bool FGraphBuilder::GetPinDefaultLiteral(UEdGraphPin& InInputPin, Frontend::FInputHandle InputHandle, FMetasoundFrontendLiteral& OutDefaultLiteral)
+		{
+			using namespace Frontend;
+
+			IMetasoundEditorModule& EditorModule = FModuleManager::GetModuleChecked<IMetasoundEditorModule>("MetasoundEditor");
 
 			const FString& InStringValue = InInputPin.DefaultValue;
 			const FName TypeName = InputHandle->GetDataType();
@@ -334,25 +383,25 @@ namespace Metasound
 			{
 				case ELiteralType::Boolean:
 				{
-					PinDataTypeDefaultLiteral.Set(FCString::ToBool(*InStringValue));
+					OutDefaultLiteral.Set(FCString::ToBool(*InStringValue));
 				}
 				break;
 
 				case ELiteralType::Float:
 				{
-					PinDataTypeDefaultLiteral.Set(FCString::Atof(*InStringValue));
+					OutDefaultLiteral.Set(FCString::Atof(*InStringValue));
 				}
 				break;
 
 				case ELiteralType::Integer:
 				{
-					PinDataTypeDefaultLiteral.Set(FCString::Atoi(*InStringValue));
+					OutDefaultLiteral.Set(FCString::Atoi(*InStringValue));
 				}
 				break;
 
 				case ELiteralType::String:
 				{
-					PinDataTypeDefaultLiteral.Set(InStringValue);
+					OutDefaultLiteral.Set(InStringValue);
 				}
 				break;
 
@@ -376,15 +425,10 @@ namespace Metasound
 
 							if (AssetData.Num() > 0)
 							{
-								PinDataTypeDefaultLiteral.Set(AssetData[0].GetAsset());
-								bObjectFound = true;
+								OutDefaultLiteral.Set(AssetData[0].GetAsset());
 							}
 						}
 					}
-					if (!bObjectFound)
-					{
-						return;
-					}
 				}
 				break;
 
@@ -396,7 +440,7 @@ namespace Metasound
 				case ELiteralType::UObjectProxyArray:
 				case ELiteralType::None:
 				{
-					return;
+					return false;
 				}
 				break;
 
@@ -405,62 +449,12 @@ namespace Metasound
 				{
 					static_assert(static_cast<int32>(ELiteralType::COUNT) == 13, "Possible missing ELiteralType case coverage.");
 					ensureMsgf(false, TEXT("Failed to set input node default: Literal type not supported"));
-					return;
+					return false;
 				}
 				break;
 			}
 
-			FOutputHandle OutputHandle = InputHandle->GetCurrentlyConnectedOutput();
-			if (!OutputHandle->IsValid())
-			{
-				FMetasoundFrontendNodeStyle Style;
-				Style.Display.Visibility = EMetasoundFrontendNodeStyleDisplayVisibility::Hidden;
-
-				const FString NewInputName = GenerateUniqueInputName(InMetasound, "LiteralInput");
-				Frontend::FNodeHandle NewInputNode = AddInputNodeHandle(InMetasound, NewInputName, TypeName, Style, FText::GetEmpty(), &PinDataTypeDefaultLiteral);
-				NewInputNode->SetNodeStyle(Style);
-
-				Frontend::FGraphHandle GraphHandle = InNodeHandle->GetOwningGraph();
-
-				const FGuid VertexID = GraphHandle->GetVertexIDForInputVertex(NewInputNode->GetNodeName());
-
-				const FMetasoundFrontendLiteral* DefaultLiteral = InputHandle->GetDefaultLiteral();
-
-				if (nullptr != DefaultLiteral)
-				{
-					if (DefaultLiteral->IsValid())
-					{
-						GraphHandle->SetDefaultInput(VertexID, *DefaultLiteral);
-						InInputPin.DefaultValue = DefaultLiteral->ToString();
-					}
-				}
-
-				TArray<Metasound::Frontend::FOutputHandle> OutputHandles = NewInputNode->GetOutputs();
-				if (ensure(OutputHandles.Num() == 1))
-				{
-					OutputHandle = OutputHandles[0];
-				}
-
-				ensure(InputHandle->Connect(*OutputHandle));
-
-				InMetasound.PostEditChange();
-				InMetasound.MarkPackageDirty();
-			}
-			else
-			{
-				FNodeHandle InputNode = OutputHandle->GetOwningNode();
-				const FMetasoundFrontendNodeStyle& Style = InputNode->GetNodeStyle();
-				if (InputNode->IsValid() && Style.Display.Visibility == EMetasoundFrontendNodeStyleDisplayVisibility::Hidden)
-				{
-					Frontend::FGraphHandle GraphHandle = InputNode->GetOwningGraph();
-					const FGuid VertexID = GraphHandle->GetVertexIDForInputVertex(InputNode->GetNodeName());
-
-					GraphHandle->SetDefaultInput(VertexID, PinDataTypeDefaultLiteral);
-
-					InMetasound.PostEditChange();
-					InMetasound.MarkPackageDirty();
-				}
-			}
+			return true;
 		}
 
 		Frontend::FNodeHandle FGraphBuilder::AddInputNodeHandle(UObject& InMetasound, const FString& InName, const FName InTypeName, const FMetasoundFrontendNodeStyle& InNodeStyle, const FText& InToolTip, const FMetasoundFrontendLiteral* InDefaultValue)
@@ -469,16 +463,12 @@ namespace Metasound
 			check(MetasoundAsset);
 
 			Frontend::FGraphHandle GraphHandle = MetasoundAsset->GetRootGraphHandle();
+			const FGuid VertexID = GraphHandle->GetNewVertexID();
 
 			FMetasoundFrontendClassInput Description;
-
 			Description.Name = InName;
 			Description.TypeName = InTypeName;
-
 			Description.Metadata.Description = InToolTip;
-
-			const FGuid VertexID = GraphHandle->GetNewVertexID();
-
 			Description.VertexID = VertexID;
 
 			Frontend::FNodeHandle NodeHandle = GraphHandle->AddInputVertex(Description);
@@ -533,7 +523,7 @@ namespace Metasound
 
 		bool FGraphBuilder::ConnectNodes(UEdGraphPin& InInputPin, UEdGraphPin& InOutputPin, bool bConnectEdPins)
 		{
-			using namespace Metasound::Frontend;
+			using namespace Frontend;
 
 			// When true, will recursively call back into this function
 			// from the schema if the editor pins are successfully connected
@@ -686,8 +676,7 @@ namespace Metasound
 
 		void FGraphBuilder::DeleteLiteralInputs(UEdGraphNode& InNode)
 		{
-			using namespace Metasound::Editor;
-			using namespace Metasound::Frontend;
+			using namespace Frontend;
 
 			FNodeHandle NodeHandle = Frontend::INodeController::GetInvalidHandle();
 			if (UMetasoundEditorGraphNode* Node = Cast<UMetasoundEditorGraphNode>(&InNode))
@@ -711,11 +700,9 @@ namespace Metasound
 			}
 		}
 
-		bool FGraphBuilder::DeleteNode(UEdGraphNode& InNode, bool bInRecordTransaction)
+		bool FGraphBuilder::DeleteNode(UEdGraphNode& InNode)
 		{
-			using namespace Metasound::Frontend;
-
-			const FScopedTransaction Transaction(LOCTEXT("DeleteMetasoundGraphNode", "Delete Metasound Node"), bInRecordTransaction);
+			using namespace Frontend;
 
 			FNodeHandle NodeHandle = Frontend::INodeController::GetInvalidHandle();
 			if (UMetasoundEditorGraphNode* Node = Cast<UMetasoundEditorGraphNode>(&InNode))
@@ -791,8 +778,6 @@ namespace Metasound
 		void FGraphBuilder::RebuildNodePins(UMetasoundEditorGraphNode& InGraphNode, Frontend::FNodeHandle InNodeHandle)
 		{
 			using namespace Frontend;
-			const FScopedTransaction Transaction(LOCTEXT("RebuildMetasoundGraphNodePins", "Rebuild Metasound Pins"));
-
 			for (int32 i = InGraphNode.Pins.Num() - 1; i >= 0; i--)
 			{
 				DeleteLiteralInputs(InGraphNode);
@@ -800,21 +785,14 @@ namespace Metasound
 			}
 
 			TArray<FInputHandle> InputHandles = InNodeHandle->GetInputs();
-			Algo::SortBy(InputHandles, [](FInputHandle InputHandle)
-			{
-				return InputHandle->GetDisplayIndex();
-			});
-
+			InputHandles = InNodeHandle->GetInputStyle().SortDefaults(InputHandles);
 			for (const FInputHandle& InputHandle : InputHandles)
 			{
 				AddPinToNode(InGraphNode, InputHandle);
 			}
 
 			TArray<FOutputHandle> OutputHandles = InNodeHandle->GetOutputs();
-			Algo::SortBy(OutputHandles, [](FOutputHandle OutputHandle)
-			{
-				return OutputHandle->GetDisplayIndex();
-			});
+			OutputHandles = InNodeHandle->GetOutputStyle().SortDefaults(OutputHandles);
 			for (const FOutputHandle& OutputHandle : OutputHandles)
 			{
 				AddPinToNode(InGraphNode, OutputHandle);
@@ -903,8 +881,7 @@ namespace Metasound
 
 		UEdGraphPin* FGraphBuilder::AddPinToNode(UMetasoundEditorGraphNode& InEditorNode, Frontend::FInputHandle InInputHandle)
 		{
-			using namespace Metasound::Editor;
-			using namespace Metasound::Frontend;
+			using namespace Frontend;
 
 			IMetasoundEditorModule& EditorModule = FModuleManager::GetModuleChecked<IMetasoundEditorModule>("MetasoundEditor");
 
@@ -1002,15 +979,42 @@ namespace Metasound
 			}
 
 			// Add missing editor nodes marked as visible.
-			// TODO: Synchronize Input/Output nodes which are own types now.
-			bIsEditorGraphDirty |= (FrontendNodes.Num() > 0);
 			for (FNodeHandle Node : FrontendNodes)
 			{
 				const FMetasoundFrontendNodeStyle& Style = Node->GetNodeStyle();
 				if (Style.Display.Visibility == EMetasoundFrontendNodeStyleDisplayVisibility::Visible)
 				{
-					UEdGraphNode* NewNode = AddNode(InMetasound, Node, false /* bInSelectNewNode */);
-					PairedNodes.Add(Node->GetID(), { Cast<UMetasoundEditorGraphNode>(NewNode), Node });
+					bIsEditorGraphDirty = true;
+					UEdGraphNode* NewNode = nullptr;
+					switch (Node->GetClassType())
+					{
+						case EMetasoundFrontendClassType::Input:
+						{
+							UMetasoundEditorGraphInputNode* InputNode = AddInput(InMetasound, Node, false /* bInSelectNewNode */);
+							NewNode = CastChecked<UEdGraphNode>(InputNode);
+						}
+						break;
+
+						case EMetasoundFrontendClassType::External:
+						case EMetasoundFrontendClassType::Graph:
+						case EMetasoundFrontendClassType::Output:
+						{
+							NewNode = AddNode(InMetasound, Node, false /* bInSelectNewNode */);
+						}
+						break;
+
+						case EMetasoundFrontendClassType::Invalid:
+						default:
+						{
+							static_assert(static_cast<int32>(EMetasoundFrontendClassType::Invalid) == 4, "Possible missing EMetasoundFrontendClassType case coverage");
+						}
+						break;
+					}
+
+					if (ensure(NewNode))
+					{
+						PairedNodes.Add(Node->GetID(), { Cast<UMetasoundEditorGraphNode>(NewNode), Node });
+					}
 				}
 			}
 
@@ -1113,12 +1117,14 @@ namespace Metasound
 
 		bool FGraphBuilder::SynchronizeConnections(UObject& InMetasound)
 		{
+			using namespace Frontend;
+
 			bool bIsGraphDirty = false;
 
 			FMetasoundAssetBase* MetasoundAsset = IMetasoundUObjectRegistry::Get().GetObjectAsAssetBase(&InMetasound);
 			check(MetasoundAsset);
 
-			Frontend::FGraphHandle GraphHandle = MetasoundAsset->GetRootGraphHandle();
+			FGraphHandle GraphHandle = MetasoundAsset->GetRootGraphHandle();
 
 			UMetasoundEditorGraph* EditorGraph = CastChecked<UMetasoundEditorGraph>(MetasoundAsset->GetGraph());
 			TArray<UMetasoundEditorGraphNode*> EditorNodes;
@@ -1135,24 +1141,16 @@ namespace Metasound
 			}
 
 			// Iterate through all nodes in metasound editor graph and synchronize connections.
-			for (UEdGraphNode* EdGraphNode : EditorGraph->Nodes)
+			for (UMetasoundEditorGraphNode* MetasoundEditorNode : EditorNodes)
 			{
 				bool bIsNodeDirty = false;
 
-				// Get all editor nodes from editor graph (some nodes on graph may *NOT* be metasound ed nodes,
-				// such as comment boxes, etc, so just get nodes of class UMetasoundEditorGraph).
-				UMetasoundEditorGraphNode* MetasoundEditorNode = Cast<UMetasoundEditorGraphNode>(EdGraphNode);
-				if (!MetasoundEditorNode)
-				{
-					continue;
-				}
-
-				Frontend::FNodeHandle Node = MetasoundEditorNode->GetNodeHandle();
+				FNodeHandle Node = MetasoundEditorNode->GetNodeHandle();
 
 				TArray<UEdGraphPin*> Pins = MetasoundEditorNode->GetAllPins();
-				TArray<Frontend::FInputHandle> NodeInputs = Node->GetInputs();
+				TArray<FInputHandle> NodeInputs = Node->GetInputs();
 
-				for (Frontend::FInputHandle& NodeInput : NodeInputs)
+				for (FInputHandle& NodeInput : NodeInputs)
 				{
 					auto IsMatchingInputPin = [&](const UEdGraphPin* Pin) -> bool
 					{
@@ -1165,62 +1163,84 @@ namespace Metasound
 						MatchingPin = *DoublePointer;
 					}
 
-					if (ensure(MatchingPin))
+					if (!ensure(MatchingPin))
 					{
-						// Remove pin so it isn't used twice.
-						Pins.Remove(MatchingPin);
+						continue;
+					}
+
+					// Remove pin so it isn't used twice.
+					Pins.Remove(MatchingPin);
+
+					bool bShowConnectionInEditor = true;
+					bool bIsConnectedToLiteralInput = false;
 
-						bool bShowConnectionInEditor = true;
+					FOutputHandle OutputHandle = NodeInput->GetCurrentlyConnectedOutput();
+					if (OutputHandle->IsValid())
+					{
+						FNodeHandle InputNodeHandle = OutputHandle->GetOwningNode();
+						const FMetasoundFrontendNodeStyle& InputNodeStyle = InputNodeHandle->GetNodeStyle();
+						bShowConnectionInEditor = InputNodeStyle.Display.Visibility == EMetasoundFrontendNodeStyleDisplayVisibility::Visible;
+						bIsConnectedToLiteralInput = !bShowConnectionInEditor;
+					}
+					else
+					{
+						bShowConnectionInEditor = false;
+					}
+						
+					if (bShowConnectionInEditor)
+					{
+						bool bAddLink = false;
 
-						Frontend::FOutputHandle OutputHandle = NodeInput->GetCurrentlyConnectedOutput();
-						if (OutputHandle->IsValid())
+						if (MatchingPin->LinkedTo.IsEmpty())
 						{
-							Frontend::FNodeHandle InputNodeHandle = OutputHandle->GetOwningNode();
-							const FMetasoundFrontendNodeStyle& InputNodeStyle = InputNodeHandle->GetNodeStyle();
-							bShowConnectionInEditor = InputNodeStyle.Display.Visibility == EMetasoundFrontendNodeStyleDisplayVisibility::Visible;
+							// No link currently exists. Add the appropriate link.
+							bAddLink = true;
 						}
-						else
+						else if (!IsMatchingOutputHandleAndPin(OutputHandle, *MatchingPin->LinkedTo[0]))
 						{
-							bShowConnectionInEditor = false;
+							// The wrong link exists.
+							MatchingPin->BreakAllPinLinks();
+							bAddLink = true;
 						}
-						
-						if (bShowConnectionInEditor)
-						{
-							bool bAddLink = false;
 
-							if (MatchingPin->LinkedTo.IsEmpty())
-							{
-								// No link currently exists. Add the appropriate link.
-								bAddLink = true;
-							}
-							else if (!IsMatchingOutputHandleAndPin(OutputHandle, *MatchingPin->LinkedTo[0]))
-							{
-								// The wrong link exists.
-								MatchingPin->BreakAllPinLinks();
-								bAddLink = true;
-							}
-
-							if (bAddLink)
-							{
-								const FGuid InputNodeID = OutputHandle->GetOwningNodeID();
-								UMetasoundEditorGraphNode* OutputEditorNode = EditorNodesByID[InputNodeID];
-								UEdGraphPin* OutputPin = OutputEditorNode->FindPinChecked(OutputHandle->GetName(), EEdGraphPinDirection::EGPD_Output);
-								MatchingPin->MakeLinkTo(OutputPin);
-								bIsNodeDirty = true;
-							}
+						if (bAddLink)
+						{
+							const FGuid InputNodeID = OutputHandle->GetOwningNodeID();
+							UMetasoundEditorGraphNode* OutputEditorNode = EditorNodesByID[InputNodeID];
+							UEdGraphPin* OutputPin = OutputEditorNode->FindPinChecked(OutputHandle->GetName(), EEdGraphPinDirection::EGPD_Output);
+							MatchingPin->MakeLinkTo(OutputPin);
+							bIsNodeDirty = true;
 						}
+					}
+					else
+					{
 						// No link should exist.
-						else
+						if (!MatchingPin->LinkedTo.IsEmpty())
 						{
 							MatchingPin->BreakAllPinLinks();
 							bIsNodeDirty = true;
 						}
+
+						// Hidden literal input Link exists, so match pin default value to literal input value
+						if (bIsConnectedToLiteralInput)
+						{
+							FNodeHandle InputNode = OutputHandle->GetOwningNode();
+							const FMetasoundFrontendNodeStyle& Style = InputNode->GetNodeStyle();
+							const FGuid VertexID = GraphHandle->GetVertexIDForInputVertex(InputNode->GetNodeName());
+							const FMetasoundFrontendLiteral DefaultLiteral = GraphHandle->GetDefaultInput(VertexID);
+							const FString DefaultLiteralString = DefaultLiteral.ToString();
+							if (MatchingPin->DefaultValue != DefaultLiteralString)
+							{
+								MatchingPin->DefaultValue = DefaultLiteralString;
+								bIsNodeDirty = true;
+							}
+						}
 					}
 				}
 
 				if (bIsNodeDirty)
 				{
-					EdGraphNode->MarkPackageDirty();
+					MetasoundEditorNode->MarkPackageDirty();
 				}
 
 				bIsGraphDirty |= bIsNodeDirty;
diff --git a/Engine/Plugins/Runtime/Metasound/Source/MetasoundEditor/Private/MetasoundEditorGraphBuilder.h b/Engine/Plugins/Runtime/Metasound/Source/MetasoundEditor/Private/MetasoundEditorGraphBuilder.h
--- a/Engine/Plugins/Runtime/Metasound/Source/MetasoundEditor/Private/MetasoundEditorGraphBuilder.h
+++ b/Engine/Plugins/Runtime/Metasound/Source/MetasoundEditor/Private/MetasoundEditorGraphBuilder.h
@@ -85,6 +85,8 @@ namespace Metasound
 			// 2. the pin's DataType literal default.
 			static void AddOrUpdateLiteralInput(UObject& InMetasound, Frontend::FNodeHandle InNodeHandle, UEdGraphPin& InInputPin);
 
+			static bool GetPinDefaultLiteral(UEdGraphPin& InInputPin, Frontend::FInputHandle InputHandle, FMetasoundFrontendLiteral& OutLiteralDefault);
+
 			static bool IsRequiredInput(Frontend::FNodeHandle InNodeHandle);
 			static bool IsRequiredOutput(Frontend::FNodeHandle InNodeHandle);
 
@@ -106,7 +108,7 @@ namespace Metasound
 			static void DeleteLiteralInputs(UEdGraphNode& InNode);
 
 			// Deletes both the editor graph & frontend nodes from respective graphs
-			static bool DeleteNode(UEdGraphNode& InNode, bool bInRecordTransaction);
+			static bool DeleteNode(UEdGraphNode& InNode);
 
 			// Adds an Input UEdGraphPin to a UMetasoundEditorGraphNode
 			static UEdGraphPin* AddPinToNode(UMetasoundEditorGraphNode& InEditorNode, Frontend::FInputHandle InInputHandle);
diff --git a/Engine/Plugins/Runtime/Metasound/Source/MetasoundEditor/Private/MetasoundEditorGraphNode.cpp b/Engine/Plugins/Runtime/Metasound/Source/MetasoundEditor/Private/MetasoundEditorGraphNode.cpp
--- a/Engine/Plugins/Runtime/Metasound/Source/MetasoundEditor/Private/MetasoundEditorGraphNode.cpp
+++ b/Engine/Plugins/Runtime/Metasound/Source/MetasoundEditor/Private/MetasoundEditorGraphNode.cpp
@@ -259,7 +259,6 @@ FText UMetasoundEditorGraphNode::GetNodeTitle(ENodeTitleType::Type TitleType) co
 
 void UMetasoundEditorGraphNode::PinDefaultValueChanged(UEdGraphPin* Pin)
 {
-	using namespace Metasound;
 	using namespace Metasound::Editor;
 	using namespace Metasound::Frontend;
 
@@ -290,6 +289,36 @@ void UMetasoundEditorGraphNode::PostEditChangeProperty(FPropertyChangedEvent& In
 		Style.Display.Location = FVector2D(NodePosX, NodePosY);
 		GetNodeHandle()->SetNodeStyle(Style);
 	}
+
+	Super::PostEditChangeProperty(InEvent);
+}
+
+void UMetasoundEditorGraphNode::PostEditChangeChainProperty(FPropertyChangedChainEvent& InEvent)
+{
+	Super::PostEditChangeChainProperty(InEvent);
+}
+
+void UMetasoundEditorGraphNode::PostEditUndo()
+{
+	UEdGraphPin::ResolveAllPinReferences();
+
+	// This can trigger and the handle is no longer valid if transaction
+	// is being undone on a graph node that is orphaned.  If orphaned,
+	// bail early.
+	Metasound::Frontend::FNodeHandle NodeHandle = GetNodeHandle();
+	if (!NodeHandle->IsValid())
+	{
+		return;
+	}
+
+	for (UEdGraphPin* Pin : Pins)
+	{
+		if (Pin && Pin->Direction == EGPD_Input)
+		{
+			UObject& Metasound = GetMetasoundChecked();
+			Metasound::Editor::FGraphBuilder::AddOrUpdateLiteralInput(Metasound, NodeHandle, *Pin);
+		}
+	}
 }
 
 void UMetasoundEditorGraphNode::PostDuplicate(bool bDuplicateForPIE)
diff --git a/Engine/Plugins/Runtime/Metasound/Source/MetasoundEditor/Private/MetasoundEditorGraphSchema.cpp b/Engine/Plugins/Runtime/Metasound/Source/MetasoundEditor/Private/MetasoundEditorGraphSchema.cpp
--- a/Engine/Plugins/Runtime/Metasound/Source/MetasoundEditor/Private/MetasoundEditorGraphSchema.cpp
+++ b/Engine/Plugins/Runtime/Metasound/Source/MetasoundEditor/Private/MetasoundEditorGraphSchema.cpp
@@ -56,6 +56,7 @@ namespace Metasound
 						{
 							if (ensure(FGraphBuilder::ConnectNodes(*FromPin, *Pin, true /* bConnectEdPins */)))
 							{
+								NewGraphNode.MarkPackageDirty();
 								return true;
 							}
 						}
@@ -75,6 +76,7 @@ namespace Metasound
 						{
 							if (ensure(FGraphBuilder::ConnectNodes(*Pin, *FromPin, true /* bConnectEdPins */)))
 							{
+								NewGraphNode.MarkPackageDirty();
 								return true;
 							}
 						}
@@ -587,12 +589,17 @@ bool UMetasoundEditorGraphSchema::TryCreateConnection(UEdGraphPin* PinA, UEdGrap
 
 	// TODO: Implement YesWithConverterNode with selected conversion option
 
+	BreakPinLinks(*InputPin, false);
 	bool bModified = UEdGraphSchema::TryCreateConnection(PinA, PinB);
 	if (bModified)
 	{
 		bModified = Metasound::Editor::FGraphBuilder::ConnectNodes(*InputPin, *OutputPin, false);
 	}
 
+	if (bModified)
+	{
+		InputPin->GetOwningNode()->MarkPackageDirty();
+	}
 	return bModified;
 }
 
@@ -704,11 +711,19 @@ void UMetasoundEditorGraphSchema::BreakPinLinks(UEdGraphPin& TargetPin, bool bSe
 	IMetasoundEditorModule& EditorModule = FModuleManager::GetModuleChecked<IMetasoundEditorModule>("MetasoundEditor");
 	for (int32 i = 0; i < InputHandles.Num(); ++i)
 	{
-		FInputHandle Handle = InputHandles[i];
-		Handle->Disconnect();
+		FInputHandle InputHandle = InputHandles[i];
+		FConstOutputHandle OutputHandle = InputHandle->GetCurrentlyConnectedOutput();
+		const FMetasoundFrontendNodeStyle& Style = OutputHandle->GetOwningNode()->GetNodeStyle();
 
-		FNodeHandle NodeHandle = Handle->GetOwningNode();
-		FGraphBuilder::AddOrUpdateLiteralInput(Metasound, NodeHandle, *InputPins[i]);
+		// Hidden nodes are not "connected" from the perspective of EdGraph,
+		// and therefore should be ignored when breaking links.
+		if (Style.Display.Visibility != EMetasoundFrontendNodeStyleDisplayVisibility::Hidden)
+		{
+			InputHandle->Disconnect();
+
+			FNodeHandle NodeHandle = InputHandle->GetOwningNode();
+			FGraphBuilder::AddOrUpdateLiteralInput(Metasound, NodeHandle, *InputPins[i]);
+		}
 	}
 
 	Super::BreakPinLinks(TargetPin, bSendsNodeNotifcation);
diff --git a/Engine/Plugins/Runtime/Metasound/Source/MetasoundEditor/Private/MetasoundInputNodeDetailCustomization.cpp b/Engine/Plugins/Runtime/Metasound/Source/MetasoundEditor/Private/MetasoundInputNodeDetailCustomization.cpp
--- a/Engine/Plugins/Runtime/Metasound/Source/MetasoundEditor/Private/MetasoundInputNodeDetailCustomization.cpp
+++ b/Engine/Plugins/Runtime/Metasound/Source/MetasoundEditor/Private/MetasoundInputNodeDetailCustomization.cpp
@@ -59,26 +59,6 @@ namespace Metasound
 
 			static const FName DataTypeNameIdentifier = "DataTypeName";
 			static const FName ProxyGeneratorClassNameIdentifier = "GeneratorClass";
-
-			void UpdateInputLiteralFromNode(const TWeakObjectPtr<UMetasoundEditorGraphInputNode>& GraphNode)
-			{
-				using namespace Frontend;
-
-				if (!GraphNode.IsValid())
-				{
-					return;
-				}
-
-				FConstNodeHandle NodeHandle = GraphNode->GetNodeHandle();
-				const FString& NodeName = NodeHandle->GetNodeName();
-				FGraphHandle GraphHandle = GraphNode->GetRootGraphHandle();
-
-				FGuid VertexID = GraphHandle->GetVertexIDForInputVertex(NodeName);
-
-				const FMetasoundFrontendLiteral Literal = GraphNode->GetLiteralDefault();
-
-				GraphHandle->SetDefaultInput(VertexID, Literal);
-			}
 		}
 
 		void FMetasoundInputNodeIntDetailCustomization::CacheProxyData(TSharedPtr<IPropertyHandle> ProxyHandle)
@@ -374,7 +354,10 @@ namespace Metasound
 			{
 				for (const TWeakObjectPtr<UMetasoundEditorGraphInputNode>& InputNode : InInputs)
 				{
-					InputNodeCustomizationPrivate::UpdateInputLiteralFromNode(InputNode);
+					if (InputNode.IsValid())
+					{
+						InputNode->UpdateDocumentInput();
+					}
 				}
 			});
 			StructPropertyHandle->SetOnChildPropertyValueChanged(OnLiteralChanged);
@@ -499,7 +482,10 @@ namespace Metasound
 					SetDefaultPropertyMetaData(Handle);
 					FSimpleDelegate OnLiteralChanged = FSimpleDelegate::CreateLambda([InGraphNode = GraphNode]()
 					{
-						InputNodeCustomizationPrivate::UpdateInputLiteralFromNode(InGraphNode);
+						if (InGraphNode.IsValid())
+						{
+							InGraphNode->UpdateDocumentInput();
+						}
 					});
 
 					Handle->SetOnPropertyValueChanged(OnLiteralChanged);
diff --git a/Engine/Plugins/Runtime/Metasound/Source/MetasoundEditor/Public/MetasoundEditorGraph.h b/Engine/Plugins/Runtime/Metasound/Source/MetasoundEditor/Public/MetasoundEditorGraph.h
--- a/Engine/Plugins/Runtime/Metasound/Source/MetasoundEditor/Public/MetasoundEditorGraph.h
+++ b/Engine/Plugins/Runtime/Metasound/Source/MetasoundEditor/Public/MetasoundEditorGraph.h
@@ -2,9 +2,9 @@
 #pragma once
 
 #include "CoreMinimal.h"
-#include "EdGraph/EdGraph.h"
 #include "MetasoundEditor.h"
 #include "MetasoundFrontendLiteral.h"
+#include "MetasoundSource.h"
 #include "UObject/ObjectMacros.h"
 
 #include "MetasoundEditorGraph.generated.h"
@@ -15,7 +15,7 @@ class UMetasoundEditorGraphInputNode;
 
 
 UCLASS(MinimalAPI)
-class UMetasoundEditorGraph : public UEdGraph
+class UMetasoundEditorGraph : public UMetasoundEditorGraphBase
 {
 	GENERATED_BODY()
 
@@ -25,6 +25,8 @@ public:
 	UObject* GetMetasound() const;
 	UObject& GetMetasoundChecked() const;
 
+	virtual void Synchronize();
+
 private:
 	UPROPERTY()
 	UObject* ParentMetasound;
diff --git a/Engine/Plugins/Runtime/Metasound/Source/MetasoundEditor/Public/MetasoundEditorGraphInputNodes.h b/Engine/Plugins/Runtime/Metasound/Source/MetasoundEditor/Public/MetasoundEditorGraphInputNodes.h
--- a/Engine/Plugins/Runtime/Metasound/Source/MetasoundEditor/Public/MetasoundEditorGraphInputNodes.h
+++ b/Engine/Plugins/Runtime/Metasound/Source/MetasoundEditor/Public/MetasoundEditorGraphInputNodes.h
@@ -34,6 +34,27 @@ public:
 		return FMetasoundFrontendLiteral();
 	}
 
+	void UpdateDocumentInput() const
+	{
+		using namespace Metasound::Frontend;
+
+		FConstNodeHandle NodeHandle = GetNodeHandle();
+		const FString& NodeName = NodeHandle->GetNodeName();
+
+		FGraphHandle GraphHandle = GetRootGraphHandle();
+		const FGuid VertexID = GraphHandle->GetVertexIDForInputVertex(NodeName);
+		GraphHandle->SetDefaultInput(VertexID, GetLiteralDefault());
+	}
+
+#if WITH_EDITORONLY_DATA
+	virtual void PostEditUndo() override
+	{
+		Super::PostEditUndo();
+
+		UpdateDocumentInput();
+	}
+#endif // WITH_EDITORONLY_DATA
+
 	FName GetLiteralDefaultPropertyFName() const
 	{
 		return "Default";
diff --git a/Engine/Plugins/Runtime/Metasound/Source/MetasoundEditor/Public/MetasoundEditorGraphNode.h b/Engine/Plugins/Runtime/Metasound/Source/MetasoundEditor/Public/MetasoundEditorGraphNode.h
--- a/Engine/Plugins/Runtime/Metasound/Source/MetasoundEditor/Public/MetasoundEditorGraphNode.h
+++ b/Engine/Plugins/Runtime/Metasound/Source/MetasoundEditor/Public/MetasoundEditorGraphNode.h
@@ -54,6 +54,8 @@ public:
 
 #if WITH_EDITOR
 	virtual void PostEditChangeProperty(FPropertyChangedEvent& InEvent) override;
+	virtual void PostEditChangeChainProperty(FPropertyChangedChainEvent& InEvent) override;
+	virtual void PostEditUndo() override;
 #endif // WITH_EDITOR
 
 	virtual void PostEditImport() override;
diff --git a/Engine/Plugins/Runtime/Metasound/Source/MetasoundEngine/Private/MetasoundSource.cpp b/Engine/Plugins/Runtime/Metasound/Source/MetasoundEngine/Private/MetasoundSource.cpp
--- a/Engine/Plugins/Runtime/Metasound/Source/MetasoundEngine/Private/MetasoundSource.cpp
+++ b/Engine/Plugins/Runtime/Metasound/Source/MetasoundEngine/Private/MetasoundSource.cpp
@@ -53,6 +53,16 @@ UMetasoundSource::UMetasoundSource(const FObjectInitializer& ObjectInitializer)
 
 #if WITH_EDITOR
 
+void UMetasoundSource::PostEditUndo()
+{
+	Super::PostEditUndo();
+
+	if (Graph)
+	{
+		Graph->Synchronize();
+	}
+}
+
 void UMetasoundSource::PostEditChangeProperty(FPropertyChangedEvent& InEvent)
 {
 	Super::PostEditChangeProperty(InEvent);
@@ -79,7 +89,7 @@ void UMetasoundSource::PostEditChangeProperty(FPropertyChangedEvent& InEvent)
 	}
 }
 
-#endif
+#endif // WITHEDITOR
 
 bool UMetasoundSource::IsPlayable() const
 {
diff --git a/Engine/Plugins/Runtime/Metasound/Source/MetasoundEngine/Public/MetasoundSource.h b/Engine/Plugins/Runtime/Metasound/Source/MetasoundEngine/Public/MetasoundSource.h
--- a/Engine/Plugins/Runtime/Metasound/Source/MetasoundEngine/Public/MetasoundSource.h
+++ b/Engine/Plugins/Runtime/Metasound/Source/MetasoundEngine/Public/MetasoundSource.h
@@ -2,21 +2,15 @@
 #pragma once
 
 #include "CoreMinimal.h"
+#include "EdGraph/EdGraph.h"
 #include "MetasoundAssetBase.h"
 #include "MetasoundFrontend.h"
 #include "MetasoundFrontendDocument.h"
 #include "MetasoundInstanceTransmitter.h"
 #include "MetasoundOperatorSettings.h"
 #include "MetasoundRouter.h"
-
 #include "Sound/SoundWaveProcedural.h"
 
-
-#if WITH_EDITORONLY_DATA
-#include "EdGraph/EdGraph.h"
-#include "Misc/AssertionMacros.h"
-#endif // WITH_EDITORONLY_DATA
-
 #include "MetasoundSource.generated.h"
 
 /** Declares the output audio format of the UMetasoundSource */
@@ -27,6 +21,15 @@ enum class EMetasoundSourceAudioFormat : uint8
 	Stereo //< Stereo audio output.
 };
 
+UCLASS()
+class METASOUNDENGINE_API UMetasoundEditorGraphBase : public UEdGraph
+{
+	GENERATED_BODY()
+
+public:
+	virtual void Synchronize() { }
+};
+
 /**
  * This Metasound type can be played as an audio source.
  */
@@ -41,7 +44,7 @@ protected:
 
 #if WITH_EDITORONLY_DATA
 	UPROPERTY(EditAnywhere, Category = CustomView)
-	UEdGraph* Graph;
+	UMetasoundEditorGraphBase* Graph;
 #endif // WITH_EDITORONLY_DATA
 
 public:
@@ -51,7 +54,7 @@ public:
 	UPROPERTY(EditAnywhere, BlueprintReadOnly, Category = Audio)
 	EMetasoundSourceAudioFormat OutputFormat;
 
-#if WITH_EDITORONLY_DATA
+#if WITH_EDITOR
 	// Returns document name (for editor purposes, and avoids making document public for edit
 	// while allowing editor to reference directly)
 	static FName GetDocumentPropertyName()
@@ -59,6 +62,8 @@ public:
 		return GET_MEMBER_NAME_CHECKED(UMetasoundSource, RootMetasoundDocument);
 	}
 
+	virtual void PostEditUndo() override;
+
 	// Returns the graph associated with this Metasound. Graph is required to be referenced on
 	// Metasound UObject for editor serialization purposes.
 	// @return Editor graph associated with UMetasoundSource.
@@ -89,11 +94,8 @@ public:
 	// @param Editor graph associated with UMetasoundSource.
 	virtual void SetGraph(UEdGraph* InGraph) override
 	{
-		Graph = InGraph;
+		Graph = CastChecked<UMetasoundEditorGraphBase>(InGraph);
 	}
-#endif // WITH_EDITORONLY_DATA
-
-#if WITH_EDITOR
 
 	void PostEditChangeProperty(FPropertyChangedEvent& InEvent);
 
diff --git a/Engine/Plugins/Runtime/Metasound/Source/MetasoundFrontend/Private/MetasoundFrontend.cpp b/Engine/Plugins/Runtime/Metasound/Source/MetasoundFrontend/Private/MetasoundFrontend.cpp
--- a/Engine/Plugins/Runtime/Metasound/Source/MetasoundFrontend/Private/MetasoundFrontend.cpp
+++ b/Engine/Plugins/Runtime/Metasound/Source/MetasoundFrontend/Private/MetasoundFrontend.cpp
@@ -67,9 +67,12 @@ namespace Metasound
 				DisplayStyle
 			};
 
+			FMetasoundFrontendClassInterface& ClassInterface = ClassDescription.Interface;
+
 			// External metasounds aren't dependent on any other nodes by definition, so all we need to do
 			// is populate the Input and Output sets.
 			const FInputVertexInterface& InputInterface = InNodeMetadata.DefaultInterface.GetInputInterface();
+			TArray<int>& InputSortOrder = ClassInterface.InputStyle.DefaultSortOrder;
 			for (auto& InputTuple : InputInterface)
 			{
 				// TODO: lots to be added here. 
@@ -80,17 +83,21 @@ namespace Metasound
 				ClassInput.VertexID = FGuid::NewGuid();
 				ClassInput.Metadata.DisplayName = FText::FromString(InputTuple.Value.GetVertexName());
 				ClassInput.Metadata.Description = InputTuple.Value.GetDescription();
-				ClassInput.Metadata.DisplayIndex = InputInterface.GetOrderIndex(InputTuple.Key);
+
 				const FLiteral& DefaultLiteral = InputTuple.Value.GetDefaultValue();
 				if (DefaultLiteral.GetType() != ELiteralType::None)
 				{
 					ClassInput.DefaultLiteral.SetFromLiteral(DefaultLiteral);
 				}
 
-				ClassDescription.Interface.Inputs.Add(ClassInput);
+				const int32 OrderIndex = InputInterface.GetOrderIndex(InputTuple.Key);
+				InputSortOrder.Add(OrderIndex);
+
+				ClassInterface.Inputs.Add(ClassInput);
 			}
 
 			const FOutputVertexInterface& OutputInterface = InNodeMetadata.DefaultInterface.GetOutputInterface();
+			TArray<int32>& OutputSortOrder = ClassInterface.InputStyle.DefaultSortOrder;
 			for (auto& OutputTuple : OutputInterface)
 			{
 				// TODO: lots to be added here. 
@@ -101,7 +108,9 @@ namespace Metasound
 				ClassOutput.VertexID = FGuid::NewGuid();
 				ClassOutput.Metadata.DisplayName = FText::FromString(OutputTuple.Value.GetVertexName());
 				ClassOutput.Metadata.Description = OutputTuple.Value.GetDescription();
-				ClassOutput.Metadata.DisplayIndex = OutputInterface.GetOrderIndex(OutputTuple.Key);
+
+				const int32 OrderIndex = OutputInterface.GetOrderIndex(OutputTuple.Key);
+				ClassInterface.OutputStyle.DefaultSortOrder.Add(OrderIndex);
 
 				ClassDescription.Interface.Outputs.Add(ClassOutput);
 			}
diff --git a/Engine/Plugins/Runtime/Metasound/Source/MetasoundFrontend/Private/MetasoundFrontendInvalidController.h b/Engine/Plugins/Runtime/Metasound/Source/MetasoundFrontend/Private/MetasoundFrontendInvalidController.h
--- a/Engine/Plugins/Runtime/Metasound/Source/MetasoundFrontend/Private/MetasoundFrontendInvalidController.h
+++ b/Engine/Plugins/Runtime/Metasound/Source/MetasoundFrontend/Private/MetasoundFrontendInvalidController.h
@@ -46,7 +46,6 @@ namespace Metasound
 			FGuid GetID() const override { return Metasound::FrontendInvalidID; }
 			const FName& GetDataType() const override { return MetasoundFrontendInvalidControllerPrivate::GetInvalid<FName>(); }  
 			const FString& GetName() const override { return MetasoundFrontendInvalidControllerPrivate::GetInvalid<FString>(); }
-			int32 GetDisplayIndex() const override { return INDEX_NONE; }
 			const FText& GetDisplayName() const override { return FText::GetEmpty(); }
 			const FText& GetTooltip() const override { return FText::GetEmpty(); }
 
@@ -88,7 +87,6 @@ namespace Metasound
 			bool IsValid() const override { return false; }
 			FGuid GetID() const override { return Metasound::FrontendInvalidID; }
 			bool IsConnected() const override { return false; }
-			int32 GetDisplayIndex() const override { return INDEX_NONE; }
 			const FName& GetDataType() const override { return MetasoundFrontendInvalidControllerPrivate::GetInvalid<FName>(); }
 			const FString& GetName() const override { return MetasoundFrontendInvalidControllerPrivate::GetInvalid<FString>(); }
 			const FText& GetDisplayName() const override { return FText::GetEmpty(); }
@@ -160,6 +158,8 @@ namespace Metasound
 			FMetasoundFrontendVersionNumber GetClassVersionNumber() const override { return MetasoundFrontendInvalidControllerPrivate::GetInvalid<FMetasoundFrontendVersionNumber>(); }
 			const FText& GetClassDisplayName() const override { return MetasoundFrontendInvalidControllerPrivate::GetInvalid<FText>(); }
 			const FText& GetClassDescription() const override { return MetasoundFrontendInvalidControllerPrivate::GetInvalid<FText>(); }
+			const FMetasoundFrontendInterfaceStyle& GetInputStyle() const override { static const FMetasoundFrontendInterfaceStyle Invalid; return Invalid; }
+			const FMetasoundFrontendInterfaceStyle& GetOutputStyle() const override { static const FMetasoundFrontendInterfaceStyle Invalid; return Invalid; }
 			const FMetasoundFrontendClassStyle& GetClassStyle() const override { static const FMetasoundFrontendClassStyle Invalid; return Invalid; }
 
 			TSharedRef<IGraphController> AsGraph() override;
diff --git a/Engine/Plugins/Runtime/Metasound/Source/MetasoundFrontend/Private/MetasoundFrontendStandardController.cpp b/Engine/Plugins/Runtime/Metasound/Source/MetasoundFrontend/Private/MetasoundFrontendStandardController.cpp
--- a/Engine/Plugins/Runtime/Metasound/Source/MetasoundFrontend/Private/MetasoundFrontendStandardController.cpp
+++ b/Engine/Plugins/Runtime/Metasound/Source/MetasoundFrontend/Private/MetasoundFrontendStandardController.cpp
@@ -526,17 +526,6 @@ namespace Metasound
 			return FBaseOutputController::IsValid() && ClassOutputPtr.IsValid();
 		}
 
-		/** Returns index to display this input in parent structure. */
-		int32 FNodeOutputController::GetDisplayIndex() const
-		{
-			if (ClassOutputPtr.IsValid())
-			{
-				return ClassOutputPtr->Metadata.DisplayIndex;
-			}
-
-			return INDEX_NONE;
-		}
-
 		const FText& FNodeOutputController::GetDisplayName() const
 		{
 			if (ClassOutputPtr.IsValid())
@@ -589,16 +578,6 @@ namespace Metasound
 			return FBaseOutputController::IsValid() && OwningGraphClassInputPtr.IsValid();
 		}
 
-		int32 FInputNodeOutputController::GetDisplayIndex() const
-		{
-			if (OwningGraphClassInputPtr.IsValid())
-			{
-				return OwningGraphClassInputPtr->Metadata.DisplayIndex;
-			}
-
-			return INDEX_NONE;
-		}
-
 		const FText& FInputNodeOutputController::GetDisplayName() const
 		{
 			if (OwningGraphClassInputPtr.IsValid())
@@ -930,16 +909,6 @@ namespace Metasound
 			return FBaseInputController::IsValid() && OwningGraphClassInputPtr.IsValid();
 		}
 
-		int32 FNodeInputController::GetDisplayIndex() const
-		{
-			if (OwningGraphClassInputPtr.IsValid())
-			{
-				return OwningGraphClassInputPtr->Metadata.DisplayIndex;
-			}
-
-			return INDEX_NONE;
-		}
-
 		const FText& FNodeInputController::GetDisplayName() const
 		{
 			if (OwningGraphClassInputPtr.IsValid())
@@ -1001,16 +970,6 @@ namespace Metasound
 			return FBaseInputController::IsValid() && OwningGraphClassOutputPtr.IsValid();
 		}
 
-		int32 FOutputNodeInputController::GetDisplayIndex() const
-		{
-			if (OwningGraphClassOutputPtr.IsValid())
-			{
-				return OwningGraphClassOutputPtr->Metadata.DisplayIndex;
-			}
-
-			return INDEX_NONE;
-		}
-
 		const FText& FOutputNodeInputController::GetDisplayName() const
 		{
 			if (OwningGraphClassOutputPtr.IsValid())
@@ -1129,6 +1088,28 @@ namespace Metasound
 			return FMetasoundFrontendVersionNumber();
 		}
 
+		const FMetasoundFrontendInterfaceStyle& FBaseNodeController::GetInputStyle() const
+		{
+			if (ClassPtr.IsValid())
+			{
+				return ClassPtr->Interface.InputStyle;
+			}
+
+			static const FMetasoundFrontendInterfaceStyle Invalid;
+			return Invalid;
+		}
+
+		const FMetasoundFrontendInterfaceStyle& FBaseNodeController::GetOutputStyle() const
+		{
+			if (ClassPtr.IsValid())
+			{
+				return ClassPtr->Interface.OutputStyle;
+			}
+
+			static const FMetasoundFrontendInterfaceStyle Invalid;
+			return Invalid;
+		}
+
 		const FMetasoundFrontendClassStyle& FBaseNodeController::GetClassStyle() const
 		{
 			if (ClassPtr.IsValid())
diff --git a/Engine/Plugins/Runtime/Metasound/Source/MetasoundFrontend/Private/MetasoundFrontendStandardController.h b/Engine/Plugins/Runtime/Metasound/Source/MetasoundFrontend/Private/MetasoundFrontendStandardController.h
--- a/Engine/Plugins/Runtime/Metasound/Source/MetasoundFrontend/Private/MetasoundFrontendStandardController.h
+++ b/Engine/Plugins/Runtime/Metasound/Source/MetasoundFrontend/Private/MetasoundFrontendStandardController.h
@@ -99,7 +99,6 @@ namespace Metasound
 			bool IsValid() const override;
 
 			// Output metadata
-			int32 GetDisplayIndex() const override;
 			const FText& GetDisplayName() const override;
 			const FText& GetTooltip() const override;
 
@@ -141,7 +140,6 @@ namespace Metasound
 			bool IsValid() const override;
 
 			// Output metadata
-			int32 GetDisplayIndex() const override;
 			const FText& GetDisplayName() const override;
 			const FText& GetTooltip() const override;
 
@@ -240,7 +238,6 @@ namespace Metasound
 			const FMetasoundFrontendLiteral* GetDefaultLiteral() const override;
 
 			// Input metadata
-			int32 GetDisplayIndex() const override;
 			const FText& GetDisplayName() const override;
 			const FText& GetTooltip() const override;
 
@@ -278,7 +275,6 @@ namespace Metasound
 			bool IsValid() const override;
 
 			// Input metadata
-			int32 GetDisplayIndex() const override;
 			const FText& GetDisplayName() const override;
 			const FText& GetTooltip() const override;
 			const FMetasoundFrontendLiteral* GetDefaultLiteral() const override;
@@ -324,6 +320,8 @@ namespace Metasound
 
 			FMetasoundFrontendVersionNumber GetClassVersionNumber() const override;
 			const FText& GetClassDescription() const override;
+			const FMetasoundFrontendInterfaceStyle& GetInputStyle() const override;
+			const FMetasoundFrontendInterfaceStyle& GetOutputStyle() const override;
 			const FMetasoundFrontendClassStyle& GetClassStyle() const override;
 
 			const FMetasoundFrontendNodeStyle& GetNodeStyle() const override;
diff --git a/Engine/Plugins/Runtime/Metasound/Source/MetasoundFrontend/Public/MetasoundFrontendController.h b/Engine/Plugins/Runtime/Metasound/Source/MetasoundFrontend/Public/MetasoundFrontendController.h
--- a/Engine/Plugins/Runtime/Metasound/Source/MetasoundFrontend/Public/MetasoundFrontendController.h
+++ b/Engine/Plugins/Runtime/Metasound/Source/MetasoundFrontend/Public/MetasoundFrontendController.h
@@ -45,16 +45,16 @@ namespace Metasound
 		class IDocumentController;
 
 		// Metasound Frontend Handles are all TSharedRefs of various Metasound Frontend Controllers.
-		typedef TSharedRef<IInputController> FInputHandle;
-		typedef TSharedRef<const IInputController> FConstInputHandle;
-		typedef TSharedRef<IOutputController> FOutputHandle;
-		typedef TSharedRef<const IOutputController> FConstOutputHandle;
-		typedef TSharedRef<INodeController> FNodeHandle;
-		typedef TSharedRef<const INodeController> FConstNodeHandle;
-		typedef TSharedRef<IGraphController> FGraphHandle;
-		typedef TSharedRef<const IGraphController> FConstGraphHandle;
-		typedef TSharedRef<IDocumentController> FDocumentHandle;
-		typedef TSharedRef<const IDocumentController> FConstDocumentHandle;
+		using FInputHandle = TSharedRef<IInputController>;
+		using FConstInputHandle = TSharedRef<const IInputController>;
+		using FOutputHandle = TSharedRef<IOutputController>;
+		using FConstOutputHandle = TSharedRef<const IOutputController>;
+		using FNodeHandle = TSharedRef<INodeController>;
+		using FConstNodeHandle = TSharedRef<const INodeController>;
+		using FGraphHandle = TSharedRef<IGraphController>;
+		using FConstGraphHandle = TSharedRef<const IGraphController>;
+		using FDocumentHandle = TSharedRef<IDocumentController>;
+		using FConstDocumentHandle = TSharedRef<const IDocumentController>;
 
 		using FVertexAccessPtr = TAccessPtr<FMetasoundFrontendVertex>;
 		using FConstVertexAccessPtr = TAccessPtr<const FMetasoundFrontendVertex>;
@@ -77,19 +77,12 @@ namespace Metasound
 		struct FConstDocumentAccess
 		{
 			FConstVertexAccessPtr ConstVertex;
-
 			FConstClassInputAccessPtr ConstClassInput;
-
 			FConstClassOutputAccessPtr ConstClassOutput;
-		
 			FConstNodeAccessPtr ConstNode;
-
 			FConstClassAccessPtr ConstClass;
-
 			FConstGraphClassAccessPtr ConstGraphClass;
-
 			FConstGraphAccessPtr ConstGraph;
-
 			FConstDocumentAccessPtr ConstDocument;
 		};
 
@@ -97,19 +90,12 @@ namespace Metasound
 		struct FDocumentAccess : public FConstDocumentAccess
 		{
 			FVertexAccessPtr Vertex; 
-
 			FClassInputAccessPtr ClassInput;
-
 			FClassOutputAccessPtr ClassOutput;
-		
 			FNodeAccessPtr Node;
-
 			FClassAccessPtr Class;
-
 			FGraphClassAccessPtr GraphClass;
-
 			FGraphAccessPtr Graph;
-
 			FDocumentAccessPtr Document;
 		};
 
@@ -172,9 +158,6 @@ namespace Metasound
 			/** Returns the name associated with this output. */
 			virtual const FString& GetName() const = 0;
 			
-			/** Returns index to display this input in parent structure. */
-			virtual int32 GetDisplayIndex() const = 0;
-
 			/** Returns the human readable name associated with this output. */
 			virtual const FText& GetDisplayName() const = 0;
 			
@@ -243,9 +226,6 @@ namespace Metasound
 			/** Returns the data type name associated with this input. */
 			virtual const FString& GetName() const = 0;
 
-			/** Returns index to display this input in parent structure. */
-			virtual int32 GetDisplayIndex() const = 0;
-
 			/** Returns the data type name associated with this input. */
 			virtual const FText& GetDisplayName() const = 0;
 
@@ -376,6 +356,8 @@ namespace Metasound
 			virtual FMetasoundFrontendVersionNumber GetClassVersionNumber() const = 0;
 			virtual const FText& GetClassDisplayName() const = 0;
 			virtual const FText& GetClassDescription() const = 0;
+			virtual const FMetasoundFrontendInterfaceStyle& GetOutputStyle() const = 0;
+			virtual const FMetasoundFrontendInterfaceStyle& GetInputStyle() const = 0;
 			virtual const FMetasoundFrontendClassStyle& GetClassStyle() const = 0;
 
 			/** If the node is also a graph, this returns a graph handle.
diff --git a/Engine/Plugins/Runtime/Metasound/Source/MetasoundFrontend/Public/MetasoundFrontendDocument.h b/Engine/Plugins/Runtime/Metasound/Source/MetasoundFrontend/Public/MetasoundFrontendDocument.h
--- a/Engine/Plugins/Runtime/Metasound/Source/MetasoundFrontend/Public/MetasoundFrontendDocument.h
+++ b/Engine/Plugins/Runtime/Metasound/Source/MetasoundFrontend/Public/MetasoundFrontendDocument.h
@@ -48,11 +48,11 @@ struct FMetasoundFrontendVersionNumber
 	GENERATED_BODY()
 
 	// Major version number.
-	UPROPERTY(EditAnywhere, Category = General)
+	UPROPERTY(VisibleAnywhere, Category = General)
 	int32 Major = 1;
 
 	// Minor version number.
-	UPROPERTY(EditAnywhere, Category = General)
+	UPROPERTY(VisibleAnywhere, Category = General)
 	int32 Minor = 0;
 };
 
@@ -298,10 +298,6 @@ struct FMetasoundFrontendVertexMetadata
 {
 	GENERATED_BODY()
 
-	// Index to inform where the vertex should be displayed in parent collection.
-	UPROPERTY()
-	int32 DisplayIndex = 0;
-
 	// Display name for a vertex
 	UPROPERTY(EditAnywhere, Category = Parameters, meta = (DisplayName = "Name"))
 	FText DisplayName;
@@ -464,13 +460,29 @@ struct FMetasoundFrontendInterfaceStyle
 {
 	GENERATED_BODY()
 
-	// Interface layoud mode
+	// Interface layout mode
 	UPROPERTY()
 	EMetasoundFrontendStyleInterfaceLayoutMode LayoutMode = EMetasoundFrontendStyleInterfaceLayoutMode::Inherited;
 
-	// Default vertex order
+	// Default vertex sort order, where array index mirrors array interface index and value is display sort index.
 	UPROPERTY()
-	TArray<FString> DefaultOrder;
+	TArray<int32> DefaultSortOrder;
+
+	template <typename HandleType>
+	TArray<HandleType> SortDefaults(const TArray<HandleType>& InHandles) const
+	{
+		TArray<HandleType> SortedHandles = InHandles;
+		for (int32 i = 0; i < DefaultSortOrder.Num(); ++i)
+		{
+			const int32 SortedIndex = DefaultSortOrder[i];
+			if (SortedHandles.IsValidIndex(SortedIndex) && InHandles.IsValidIndex(i))
+			{
+				SortedHandles[SortedIndex] = InHandles[i];
+			}
+		}
+
+		return SortedHandles;
+	}
 };
 
 
