Trigger Threshold Node

diff --git a/Engine/Plugins/Runtime/Metasound/Source/MetasoundStandardNodes/Private/MetasoundTriggerOnThresholdNode.cpp b/Engine/Plugins/Runtime/Metasound/Source/MetasoundStandardNodes/Private/MetasoundTriggerOnThresholdNode.cpp
new file mode 100644
--- /dev/null
+++ b/Engine/Plugins/Runtime/Metasound/Source/MetasoundStandardNodes/Private/MetasoundTriggerOnThresholdNode.cpp
@@ -0,0 +1,253 @@
+// Copyright Epic Games, Inc. All Rights Reserved.
+#include "MetasoundTriggerOnThresholdNode.h"
+
+#include "MetasoundAudioBuffer.h"
+#include "MetasoundExecutableOperator.h"
+#include "MetasoundFacade.h"
+#include "MetasoundNodeRegistrationMacro.h"
+#include "MetasoundDataTypeRegistrationMacro.h"
+#include "MetasoundOperatorSettings.h"
+#include "MetasoundPrimitives.h"
+#include "MetasoundStandardNodesNames.h"
+#include "MetasoundTrigger.h"
+#include "MetasoundVertex.h"
+
+#define LOCTEXT_NAMESPACE "MetasoundStandardNodes"
+
+namespace Metasound
+{
+	enum class EBufferTriggerType
+	{
+		RisingEdge,
+		FallingEdge,
+		AbsThreshold,
+	};
+
+	DECLARE_METASOUND_ENUM(EBufferTriggerType, EBufferTriggerType::RisingEdge, METASOUNDSTANDARDNODES_API,
+	FEnumBufferTriggerType, FEnumBufferTriggerTypeInfo, FBufferTriggerTypeReadRef, FEnumBufferTriggerTypeWriteRef);
+
+	DEFINE_METASOUND_ENUM_BEGIN(EBufferTriggerType, FEnumBufferTriggerType, "BufferTriggerType")
+		DEFINE_METASOUND_ENUM_ENTRY(EBufferTriggerType::RisingEdge, LOCTEXT("RisingEdgeDescription", "Rising Edge"), LOCTEXT("RisingEdgeDescriptionTT", "")),
+		DEFINE_METASOUND_ENUM_ENTRY(EBufferTriggerType::FallingEdge, LOCTEXT("FallingEdgeDescription", "Falling Edge"), LOCTEXT("FallingEdgeDescriptionTT", "")),
+		DEFINE_METASOUND_ENUM_ENTRY(EBufferTriggerType::AbsThreshold, LOCTEXT("AbsThresholdDescription", "Abs Threshold"), LOCTEXT("AbsThresholdDescriptionTT", ""))
+		DEFINE_METASOUND_ENUM_END()
+
+	class FTriggerOnThresholdOperator : public IOperator
+	{
+	public:
+		static constexpr float DefaultThreshold = 0.85f;
+
+		static TUniquePtr<IOperator> CreateOperator(const FCreateOperatorParams& InParams, FBuildErrorArray& OutErrors);
+		static const FNodeClassMetadata& GetNodeInfo();
+		static FVertexInterface DeclareVertexInterface();
+
+		FTriggerOnThresholdOperator(const FOperatorSettings& InSettings, FAudioBufferReadRef InBuffer, FFloatReadRef InThreshold);
+
+		virtual FDataReferenceCollection GetInputs() const override;
+		virtual FDataReferenceCollection GetOutputs() const override;
+
+	protected:
+		FAudioBufferReadRef In;
+		FFloatReadRef Threshold;
+		FTriggerWriteRef Out;
+		bool bTriggered = false;
+		float LastSample = 0.f;
+
+		static constexpr const TCHAR* OutPinName = TEXT("Out");
+		static constexpr const TCHAR* InPinName = TEXT("In");
+		static constexpr const TCHAR* ThresholdPinName = TEXT("Threshold");
+		static constexpr const TCHAR* TriggerType = TEXT("Type");
+	};
+
+	struct FTriggerOnThresholdOperator_EdgeCommon : public FTriggerOnThresholdOperator
+	{
+		using FTriggerOnThresholdOperator::FTriggerOnThresholdOperator;
+		
+		template<typename PREDICATE>
+		void Generate(const PREDICATE& ValueTester)
+		{	
+			Out->AdvanceBlock();
+
+			const float* InputBuffer = In->GetData();
+			const int32 NumFrames = In->Num();
+			const float ThreshValue = *Threshold;
+			
+			float Previous = LastSample;
+			for (int32 i = 0; i < NumFrames; ++i)
+			{
+				const float Current = *InputBuffer;
+
+				// If Previous didn't trigger but current value does... fire!
+				if (!ValueTester(Previous, ThreshValue) && ValueTester(Current,ThreshValue))
+				{
+					Out->TriggerFrame(i);
+				}
+
+				Previous = Current;
+				++InputBuffer;
+			}
+
+			// Remember the last sample for next frame.
+			LastSample = Previous;
+		}
+	};
+
+	struct FTriggerOnThresholdOperator_RisingEdge final : public FTriggerOnThresholdOperator_EdgeCommon
+	{
+		using FTriggerOnThresholdOperator_EdgeCommon::FTriggerOnThresholdOperator_EdgeCommon;
+
+		void Execute() 
+		{
+			Generate(TGreater<float>{});
+		}
+
+		static void ExecuteFunction(IOperator* InOperator) { static_cast<FTriggerOnThresholdOperator_RisingEdge*>(InOperator)->Execute(); }
+		FExecuteFunction GetExecuteFunction() override { return &FTriggerOnThresholdOperator_RisingEdge::ExecuteFunction; }
+	};
+
+	struct FTriggerOnThresholdOperator_FallingEdge final : public FTriggerOnThresholdOperator_EdgeCommon
+	{
+		using FTriggerOnThresholdOperator_EdgeCommon::FTriggerOnThresholdOperator_EdgeCommon;
+
+		static void ExecuteFunction(IOperator* InOperator) { static_cast<FTriggerOnThresholdOperator_FallingEdge*>(InOperator)->Execute(); }
+		FExecuteFunction GetExecuteFunction() override { return &FTriggerOnThresholdOperator_FallingEdge::ExecuteFunction; }
+
+		void Execute()
+		{
+			Generate(TLess<float>{});
+		}		
+	};
+
+	struct FTriggerOnThresholdOperator_AbsThreshold final : public FTriggerOnThresholdOperator
+	{
+		using FTriggerOnThresholdOperator::FTriggerOnThresholdOperator;
+
+		static void ExecuteFunction(IOperator* InOperator) { static_cast<FTriggerOnThresholdOperator_AbsThreshold*>(InOperator)->Execute(); }
+		FExecuteFunction GetExecuteFunction() override { return &FTriggerOnThresholdOperator_AbsThreshold::ExecuteFunction; }
+
+		void Execute()
+		{
+			Out->AdvanceBlock();
+
+			const float* InputBuffer = In->GetData();
+			const int32 NumFrames = In->Num();
+			const float ThresholdSqr = *Threshold * *Threshold;
+					
+			for (int32 i = 0; i < NumFrames; ++i)
+			{
+				const float Current = *InputBuffer;
+				const float CurrentSqr = Current * Current;
+				 
+				if (CurrentSqr > ThresholdSqr && !bTriggered)
+				{
+					bTriggered = true;
+					Out->TriggerFrame(i);
+				}
+				else if (CurrentSqr < ThresholdSqr && bTriggered)
+				{
+					bTriggered = false;
+				}
+
+				++InputBuffer;
+			}
+		}
+	};
+
+	FTriggerOnThresholdOperator::FTriggerOnThresholdOperator(const FOperatorSettings& InSettings, FAudioBufferReadRef InBuffer, FFloatReadRef InThreshold)
+		: In(InBuffer)
+		, Threshold(InThreshold)
+		, Out(FTriggerWriteRef::CreateNew(InSettings))
+	{}
+
+	FDataReferenceCollection FTriggerOnThresholdOperator::GetInputs() const
+	{
+		FDataReferenceCollection InputDataReferences;
+		InputDataReferences.AddDataReadReference(ThresholdPinName, FFloatReadRef(Threshold));
+		InputDataReferences.AddDataReadReference(InPinName, FAudioBufferReadRef(In));
+		return InputDataReferences;
+	}
+
+	FDataReferenceCollection FTriggerOnThresholdOperator::GetOutputs() const
+	{
+		FDataReferenceCollection OutputDataReferences;
+		OutputDataReferences.AddDataReadReference(OutPinName, FTriggerWriteRef(Out));
+		return OutputDataReferences;
+	}
+
+	FVertexInterface FTriggerOnThresholdOperator::DeclareVertexInterface()
+	{
+		static const FVertexInterface Interface(
+			FInputVertexInterface(
+				TInputDataVertexModel<FAudioBuffer>(InPinName, LOCTEXT("BufferInDescription", "Input")),
+				TInputDataVertexModel<float>(ThresholdPinName, LOCTEXT("ThresholdDescription", "Trigger Threshold"), DefaultThreshold),
+				TInputDataVertexModel<FEnumBufferTriggerType>(TriggerType, LOCTEXT("ThresholdDescription", "Trigger Threshold"))
+			),
+			FOutputVertexInterface(
+				TOutputDataVertexModel<FTrigger>(OutPinName, LOCTEXT("TriggerOutDescription", "Output"))
+			)
+		);
+		return Interface;
+	}
+
+	const FNodeClassMetadata& FTriggerOnThresholdOperator::GetNodeInfo()
+	{
+		auto InitNodeInfo = []() -> FNodeClassMetadata
+		{
+			FNodeClassMetadata Info;
+			Info.ClassName = {Metasound::StandardNodes::Namespace, TEXT("TriggerOnThreshold"), Metasound::StandardNodes::AudioVariant};
+			Info.MajorVersion = 1;
+			Info.MinorVersion = 0;
+			Info.DisplayName = LOCTEXT("Metasound_TriggerOnThresholdNodeDisplayName", "Trigger On Threshold");
+			Info.Description = LOCTEXT("Metasound_TriggerOnThresholdNodeDescription", "Trigger based on a audio buffer input");
+			Info.Author = PluginAuthor;
+			Info.PromptIfMissing = PluginNodeMissingPrompt;
+			Info.DefaultInterface = DeclareVertexInterface(); 
+
+			return Info;
+		};
+		static const FNodeClassMetadata Info = InitNodeInfo();
+		return Info;
+	}
+
+	FTriggerOnThresholdNode::FTriggerOnThresholdNode(const FNodeInitData& InInitData)
+		: FTriggerOnThresholdNode(InInitData.InstanceName, InInitData.InstanceID, FTriggerOnThresholdOperator::DefaultThreshold)
+	{
+	}
+
+	FTriggerOnThresholdNode::FTriggerOnThresholdNode(const FString& InName, const FGuid& InInstanceID, float InDefaultThreshold)
+		: FNodeFacade(InName, InInstanceID, TFacadeOperatorClass<FTriggerOnThresholdOperator>())
+		, DefaultThreshold(InDefaultThreshold)
+	{
+	}
+
+	TUniquePtr<IOperator> FTriggerOnThresholdOperator::CreateOperator(const FCreateOperatorParams& InParams, FBuildErrorArray& OutErrors)
+	{
+		const FTriggerOnThresholdNode& Node = static_cast<const FTriggerOnThresholdNode&>(InParams.Node);
+		const FDataReferenceCollection& InputCol = InParams.InputDataReferences;
+		const FOperatorSettings& Settings = InParams.OperatorSettings;
+		const FInputVertexInterface& InputInterface = DeclareVertexInterface().GetInputInterface();
+		
+		// Static property pin, only used for factory.
+		FBufferTriggerTypeReadRef Type = InputCol.GetDataReadReferenceOrConstruct<FEnumBufferTriggerType>(TriggerType);
+
+		FAudioBufferReadRef InputBuffer = InputCol.GetDataReadReferenceOrConstruct<FAudioBuffer>(InPinName, Settings);
+		FFloatReadRef Threshold = InputCol.GetDataReadReferenceOrConstructWithVertexDefault<float>(InputInterface, ThresholdPinName);
+
+		switch (*Type)
+		{
+		default:
+		case EBufferTriggerType::RisingEdge:
+			return MakeUnique<FTriggerOnThresholdOperator_RisingEdge>(InParams.OperatorSettings, InputBuffer, Threshold);
+		case EBufferTriggerType::FallingEdge:
+			return MakeUnique<FTriggerOnThresholdOperator_FallingEdge>(InParams.OperatorSettings, InputBuffer, Threshold);
+		case EBufferTriggerType::AbsThreshold:
+			return MakeUnique<FTriggerOnThresholdOperator_AbsThreshold>(InParams.OperatorSettings, InputBuffer, Threshold);		
+		}
+		checkNoEntry();
+		return nullptr;
+	}
+	
+	METASOUND_REGISTER_NODE(FTriggerOnThresholdNode);
+}
+#undef LOCTEXT_NAMESPACE //MetasoundStandardNodes
+
diff --git a/Engine/Plugins/Runtime/Metasound/Source/MetasoundStandardNodes/Public/MetasoundTriggerOnThresholdNode.h b/Engine/Plugins/Runtime/Metasound/Source/MetasoundStandardNodes/Public/MetasoundTriggerOnThresholdNode.h
new file mode 100644
--- /dev/null
+++ b/Engine/Plugins/Runtime/Metasound/Source/MetasoundStandardNodes/Public/MetasoundTriggerOnThresholdNode.h
@@ -0,0 +1,27 @@
+// Copyright Epic Games, Inc. All Rights Reserved.
+
+#pragma once
+
+#include "CoreMinimal.h"
+
+#include "MetasoundAudioBuffer.h"
+#include "MetasoundBuilderInterface.h"
+#include "MetasoundDataReferenceCollection.h"
+#include "MetasoundExecutableOperator.h"
+#include "MetasoundFacade.h"
+#include "MetasoundNode.h"
+#include "MetasoundOperatorInterface.h"
+
+namespace Metasound
+{
+	class METASOUNDSTANDARDNODES_API FTriggerOnThresholdNode : public FNodeFacade
+	{
+	public:
+		FTriggerOnThresholdNode(const FString& InName, const FGuid& InInstanceID, float InDefaultThreshold);
+		FTriggerOnThresholdNode(const FNodeInitData& InInitData);
+
+		float GetDefaultThreshold() const { return DefaultThreshold; }
+	private:
+		float DefaultThreshold;
+	};
+}
