CIS mac fix for trigger on threshold

diff --git a/Engine/Plugins/Runtime/Metasound/Source/MetasoundStandardNodes/Private/MetasoundTriggerOnThresholdNode.cpp b/Engine/Plugins/Runtime/Metasound/Source/MetasoundStandardNodes/Private/MetasoundTriggerOnThresholdNode.cpp
--- a/Engine/Plugins/Runtime/Metasound/Source/MetasoundStandardNodes/Private/MetasoundTriggerOnThresholdNode.cpp
+++ b/Engine/Plugins/Runtime/Metasound/Source/MetasoundStandardNodes/Private/MetasoundTriggerOnThresholdNode.cpp
@@ -14,6 +14,8 @@
 
 #define LOCTEXT_NAMESPACE "MetasoundStandardNodes"
 
+#if 0
+
 namespace Metasound
 {
 	enum class EBufferTriggerType
@@ -249,5 +251,7 @@ namespace Metasound
 	
 	METASOUND_REGISTER_NODE(FTriggerOnThresholdNode);
 }
+#endif
+
 #undef LOCTEXT_NAMESPACE //MetasoundStandardNodes
 
diff --git a/Engine/Plugins/Runtime/Metasound/Source/MetasoundStandardNodes/Public/MetasoundTriggerOnThresholdNode.h b/Engine/Plugins/Runtime/Metasound/Source/MetasoundStandardNodes/Public/MetasoundTriggerOnThresholdNode.h
--- a/Engine/Plugins/Runtime/Metasound/Source/MetasoundStandardNodes/Public/MetasoundTriggerOnThresholdNode.h
+++ b/Engine/Plugins/Runtime/Metasound/Source/MetasoundStandardNodes/Public/MetasoundTriggerOnThresholdNode.h
@@ -12,6 +12,7 @@
 #include "MetasoundNode.h"
 #include "MetasoundOperatorInterface.h"
 
+#if 0
 namespace Metasound
 {
 	class METASOUNDSTANDARDNODES_API FTriggerOnThresholdNode : public FNodeFacade
@@ -25,3 +26,4 @@ namespace Metasound
 		float DefaultThreshold;
 	};
 }
+#endif
