Audio format output node specialization

diff --git a/Engine/Plugins/Runtime/Metasound/Source/MetasoundEngine/Private/MetasoundSource.cpp b/Engine/Plugins/Runtime/Metasound/Source/MetasoundEngine/Private/MetasoundSource.cpp
--- a/Engine/Plugins/Runtime/Metasound/Source/MetasoundEngine/Private/MetasoundSource.cpp
+++ b/Engine/Plugins/Runtime/Metasound/Source/MetasoundEngine/Private/MetasoundSource.cpp
@@ -516,7 +516,7 @@ const FMetasoundFrontendArchetype& UMetasoundSource::GetMonoSourceArchetype()
 		FMetasoundFrontendClassVertex GeneratedAudio;
 		GeneratedAudio.Name = UMetasoundSource::GetAudioOutputName();
 		GeneratedAudio.TypeName = Metasound::Frontend::GetDataTypeName<Metasound::FMonoAudioFormat>();
-		GeneratedAudio.Metadata.DisplayName = LOCTEXT("GeneratedMono", "Mono");
+		GeneratedAudio.Metadata.DisplayName = LOCTEXT("GeneratedMono", "Audio");
 		GeneratedAudio.Metadata.Description = LOCTEXT("GeneratedAudioToolTip", "The resulting output audio from this source.");
 		GeneratedAudio.VertexID = FGuid::NewGuid();
 
@@ -540,7 +540,7 @@ const FMetasoundFrontendArchetype& UMetasoundSource::GetStereoSourceArchetype()
 		FMetasoundFrontendClassVertex GeneratedAudio;
 		GeneratedAudio.Name = UMetasoundSource::GetAudioOutputName();
 		GeneratedAudio.TypeName = Metasound::Frontend::GetDataTypeName<Metasound::FStereoAudioFormat>();
-		GeneratedAudio.Metadata.DisplayName = LOCTEXT("GeneratedStereo", "Stereo");
+		GeneratedAudio.Metadata.DisplayName = LOCTEXT("GeneratedStereo", "Audio");
 		GeneratedAudio.Metadata.Description = LOCTEXT("GeneratedAudioToolTip", "The resulting output audio from this source.");
 		GeneratedAudio.VertexID = FGuid::NewGuid();
 
diff --git a/Engine/Plugins/Runtime/Metasound/Source/MetasoundFrontend/Private/MetasoundFrontendStandardController.cpp b/Engine/Plugins/Runtime/Metasound/Source/MetasoundFrontend/Private/MetasoundFrontendStandardController.cpp
--- a/Engine/Plugins/Runtime/Metasound/Source/MetasoundFrontend/Private/MetasoundFrontendStandardController.cpp
+++ b/Engine/Plugins/Runtime/Metasound/Source/MetasoundFrontend/Private/MetasoundFrontendStandardController.cpp
@@ -17,6 +17,7 @@
 #include "StructSerializer.h"
 #include "Serialization/MemoryReader.h"
 
+#define LOCTEXT_NAMESPACE "MetasoundFrontendStandardController"
 
 namespace Metasound
 {
@@ -349,6 +350,7 @@ namespace Metasound
 		FBaseOutputController::FBaseOutputController(const FBaseOutputController::FInitParams& InParams)
 		: ID(InParams.ID)
 		, NodeVertexPtr(InParams.NodeVertexPtr)
+		, ClassOutputPtr(InParams.ClassOutputPtr)
 		, GraphPtr(InParams.GraphPtr)
 		, OwningNode(InParams.OwningNode)
 		{
@@ -491,42 +493,7 @@ namespace Metasound
 		}
 
 
-		FDocumentAccess FBaseOutputController::ShareAccess() 
-		{
-			FDocumentAccess Access;
-
-			Access.ConstVertex = NodeVertexPtr;
-			Access.Graph = GraphPtr;
-			Access.ConstGraph = GraphPtr;
-
-			return Access;
-		}
-
-		FConstDocumentAccess FBaseOutputController::ShareAccess() const
-		{
-			FConstDocumentAccess Access;
-
-			Access.ConstVertex = NodeVertexPtr;
-			Access.ConstGraph = GraphPtr;
-
-			return Access;
-		}
-
-		//
-		// FNodeOutputController
-		//
-		FNodeOutputController::FNodeOutputController(const FNodeOutputController::FInitParams& InParams)
-		: FBaseOutputController({InParams.ID, InParams.NodeVertexPtr, InParams.GraphPtr, InParams.OwningNode})
-		, ClassOutputPtr(InParams.ClassOutputPtr)
-		{
-		}
-
-		bool FNodeOutputController::IsValid() const
-		{
-			return FBaseOutputController::IsValid() && ClassOutputPtr.IsValid();
-		}
-
-		const FText& FNodeOutputController::GetDisplayName() const
+		const FText& FBaseOutputController::GetDisplayName() const
 		{
 			if (ClassOutputPtr.IsValid())
 			{
@@ -536,7 +503,7 @@ namespace Metasound
 			return FrontendControllerIntrinsics::GetInvalidValueConstRef<FText>();
 		}
 
-		const FText& FNodeOutputController::GetTooltip() const
+		const FText& FBaseOutputController::GetTooltip() const
 		{
 			if (ClassOutputPtr.IsValid())
 			{
@@ -546,29 +513,36 @@ namespace Metasound
 			return FrontendControllerIntrinsics::GetInvalidValueConstRef<FText>();
 		}
 
-		FDocumentAccess FNodeOutputController::ShareAccess() 
+
+		FDocumentAccess FBaseOutputController::ShareAccess() 
 		{
-			FDocumentAccess Access = FBaseOutputController::ShareAccess();
+			FDocumentAccess Access;
 
+			Access.ConstVertex = NodeVertexPtr;
 			Access.ConstClassOutput = ClassOutputPtr;
+			Access.Graph = GraphPtr;
+			Access.ConstGraph = GraphPtr;
 
 			return Access;
 		}
 
-		FConstDocumentAccess FNodeOutputController::ShareAccess() const
+		FConstDocumentAccess FBaseOutputController::ShareAccess() const
 		{
-			FConstDocumentAccess Access = FBaseOutputController::ShareAccess();
+			FConstDocumentAccess Access;
 
+			Access.ConstVertex = NodeVertexPtr;
 			Access.ConstClassOutput = ClassOutputPtr;
+			Access.ConstGraph = GraphPtr;
 
 			return Access;
 		}
 
+
 		//
 		// FInputNodeOutputController
 		// 
 		FInputNodeOutputController::FInputNodeOutputController(const FInputNodeOutputController::FInitParams& InParams)
-		: FBaseOutputController({InParams.ID, InParams.NodeVertexPtr, InParams.GraphPtr, InParams.OwningNode})
+		: FBaseOutputController({InParams.ID, InParams.NodeVertexPtr, InParams.ClassOutputPtr, InParams.GraphPtr, InParams.OwningNode})
 		, OwningGraphClassInputPtr(InParams.OwningGraphClassInputPtr)
 		{
 		}
@@ -580,12 +554,21 @@ namespace Metasound
 
 		const FText& FInputNodeOutputController::GetDisplayName() const
 		{
-			if (OwningGraphClassInputPtr.IsValid())
+			if (const FMetasoundFrontendClassInput* OwningInput = OwningGraphClassInputPtr.Get())
 			{
-				return OwningGraphClassInputPtr->Metadata.DisplayName;
+				if (const FMetasoundFrontendClassOutput* ClassOutput = ClassOutputPtr.Get())
+				{
+					// If there is a valid ClassOutput, combine the names.
+					CachedDisplayName = FText::Format(LOCTEXT("InputNodeOutputControllerFormat", "{1} {0}"), OwningInput->Metadata.DisplayName, ClassOutput->Metadata.DisplayName);
+				}
+				else
+				{
+					// If there is no valid ClassOutput, use the owning value display name.
+					CachedDisplayName = OwningInput->Metadata.DisplayName;
+				}
 			}
 
-			return FrontendControllerIntrinsics::GetInvalidValueConstRef<FText>();
+			return CachedDisplayName;
 		}
 
 		const FText& FInputNodeOutputController::GetTooltip() const
@@ -622,6 +605,7 @@ namespace Metasound
 		FBaseInputController::FBaseInputController(const FBaseInputController::FInitParams& InParams)
 		: ID(InParams.ID)
 		, NodeVertexPtr(InParams.NodeVertexPtr)
+		, ClassInputPtr(InParams.ClassInputPtr)
 		, GraphPtr(InParams.GraphPtr)
 		, OwningNode(InParams.OwningNode)
 		{
@@ -657,6 +641,36 @@ namespace Metasound
 			return FrontendControllerIntrinsics::GetInvalidValueConstRef<FString>();
 		}
 
+		const FText& FBaseInputController::GetDisplayName() const
+		{
+			if (const FMetasoundFrontendClassInput* ClassInput = ClassInputPtr.Get())
+			{
+				return ClassInput->Metadata.DisplayName;
+			}
+			
+			return FrontendControllerIntrinsics::GetInvalidValueConstRef<FText>();
+		}
+
+		const FMetasoundFrontendLiteral* FBaseInputController::GetDefaultLiteral() const 
+		{
+			if (const FMetasoundFrontendClassInput* ClassInput = ClassInputPtr.Get())
+			{
+				return &(ClassInput->DefaultLiteral);
+			}
+			return nullptr;
+		}
+
+		const FText& FBaseInputController::GetTooltip() const
+		{
+			if (const FMetasoundFrontendClassInput* ClassInput = ClassInputPtr.Get())
+			{
+				return ClassInput->Metadata.Description;
+			}
+			
+			return FrontendControllerIntrinsics::GetInvalidValueConstRef<FText>();
+		}
+
+
 		bool FBaseInputController::IsConnected() const 
 		{
 			return (nullptr != FindEdge());
@@ -879,6 +893,7 @@ namespace Metasound
 			FDocumentAccess Access;
 
 			Access.ConstVertex = NodeVertexPtr;
+			Access.ConstClassInput = ClassInputPtr;
 			Access.Graph = GraphPtr;
 			Access.ConstGraph = GraphPtr;
 
@@ -890,77 +905,18 @@ namespace Metasound
 			FConstDocumentAccess Access;
 
 			Access.ConstVertex = NodeVertexPtr;
+			Access.ConstClassInput = ClassInputPtr;
 			Access.ConstGraph = GraphPtr;
 
 			return Access;
 		}
 
-		//
-		// FNodeInputController
-		//
-		FNodeInputController::FNodeInputController(const FNodeInputController::FInitParams& InParams)
-		: FBaseInputController({InParams.ID, InParams.NodeVertexPtr, InParams.GraphPtr, InParams.OwningNode})
-		, OwningGraphClassInputPtr(InParams.OwningGraphClassInputPtr)
-		{
-		}
-
-		bool FNodeInputController::IsValid() const
-		{
-			return FBaseInputController::IsValid() && OwningGraphClassInputPtr.IsValid();
-		}
-
-		const FText& FNodeInputController::GetDisplayName() const
-		{
-			if (OwningGraphClassInputPtr.IsValid())
-			{
-				return OwningGraphClassInputPtr->Metadata.DisplayName;
-			}
-			
-			return FrontendControllerIntrinsics::GetInvalidValueConstRef<FText>();
-		}
-
-		const FMetasoundFrontendLiteral* FNodeInputController::GetDefaultLiteral() const 
-		{
-			if (const FMetasoundFrontendClassInput* OwningGraphClassInput = OwningGraphClassInputPtr.Get())
-			{
-				return &(OwningGraphClassInput->DefaultLiteral);
-			}
-			return nullptr;
-		}
-
-		const FText& FNodeInputController::GetTooltip() const
-		{
-			if (OwningGraphClassInputPtr.IsValid())
-			{
-				return OwningGraphClassInputPtr->Metadata.Description;
-			}
-			
-			return FrontendControllerIntrinsics::GetInvalidValueConstRef<FText>();
-		}
-
-		FDocumentAccess FNodeInputController::ShareAccess() 
-		{
-			FDocumentAccess Access = FBaseInputController::ShareAccess();
-
-			Access.ConstClassInput = OwningGraphClassInputPtr;
-
-			return Access;
-		}
-
-		FConstDocumentAccess FNodeInputController::ShareAccess() const
-		{
-			FConstDocumentAccess Access = FBaseInputController::ShareAccess();
-
-			Access.ConstClassInput = OwningGraphClassInputPtr;
-
-			return Access;
-		}
 
 		//
 		// FOutputNodeInputController
 		//
 		FOutputNodeInputController::FOutputNodeInputController(const FOutputNodeInputController::FInitParams& InParams)
-		: FBaseInputController({InParams.ID, InParams.NodeVertexPtr, InParams.GraphPtr, InParams.OwningNode})
+		: FBaseInputController({InParams.ID, InParams.NodeVertexPtr, InParams.ClassInputPtr, InParams.GraphPtr, InParams.OwningNode})
 		, OwningGraphClassOutputPtr(InParams.OwningGraphClassOutputPtr)
 		{
 		}
@@ -972,12 +928,22 @@ namespace Metasound
 
 		const FText& FOutputNodeInputController::GetDisplayName() const
 		{
-			if (OwningGraphClassOutputPtr.IsValid())
+			if (const FMetasoundFrontendClassOutput* OwningOutput = OwningGraphClassOutputPtr.Get())
 			{
-				return OwningGraphClassOutputPtr->Metadata.DisplayName;
+				if (const FMetasoundFrontendClassInput* ClassInput = ClassInputPtr.Get())
+				{
+					// If there the ClassInput exists, combine the variable name and class input name.
+					// of the variable should be added to the names of the vertices.
+					CachedDisplayName = FText::Format(LOCTEXT("OutputNodeInputControllerFormat", "{1} {0}"), OwningOutput->Metadata.DisplayName, ClassInput->Metadata.DisplayName);
+				}
+				else
+				{
+					// If there is not ClassInput, then use the variable name.
+					CachedDisplayName = OwningOutput->Metadata.DisplayName;
+				}
 			}
-			
-			return FrontendControllerIntrinsics::GetInvalidValueConstRef<FText>();
+
+			return CachedDisplayName;
 		}
 
 		const FText& FOutputNodeInputController::GetTooltip() const
@@ -1680,12 +1646,12 @@ namespace Metasound
 
 		FInputHandle FNodeController::CreateInputController(FGuid InVertexID, FConstVertexAccessPtr InNodeVertexPtr, FConstClassInputAccessPtr InClassInputPtr, FNodeHandle InOwningNode) const
 		{
-			return MakeShared<FNodeInputController>(FNodeInputController::FInitParams{InVertexID, InNodeVertexPtr, InClassInputPtr, GraphPtr, InOwningNode});
+			return MakeShared<FBaseInputController>(FBaseInputController::FInitParams{InVertexID, InNodeVertexPtr, InClassInputPtr, GraphPtr, InOwningNode});
 		}
 
 		FOutputHandle FNodeController::CreateOutputController(FGuid InVertexID, FConstVertexAccessPtr InNodeVertexPtr, FConstClassOutputAccessPtr InClassOutputPtr, FNodeHandle InOwningNode) const
 		{
-			return MakeShared<FNodeOutputController>(FNodeOutputController::FInitParams{InVertexID, InNodeVertexPtr, InClassOutputPtr, GraphPtr, InOwningNode});
+			return MakeShared<FBaseOutputController>(FBaseOutputController::FInitParams{InVertexID, InNodeVertexPtr, InClassOutputPtr, GraphPtr, InOwningNode});
 		}
 
 		FDocumentAccess FNodeController::ShareAccess() 
@@ -1773,7 +1739,7 @@ namespace Metasound
 
 		FInputHandle FOutputNodeController::CreateInputController(FGuid InVertexID, FConstVertexAccessPtr InNodeVertexPtr, FConstClassInputAccessPtr InClassInputPtr, FNodeHandle InOwningNode) const
 		{
-			return MakeShared<FOutputNodeInputController>(FOutputNodeInputController::FInitParams{InVertexID, InNodeVertexPtr, OwningGraphClassOutputPtr, GraphPtr, InOwningNode});
+			return MakeShared<FOutputNodeInputController>(FOutputNodeInputController::FInitParams{InVertexID, InNodeVertexPtr, InClassInputPtr, OwningGraphClassOutputPtr, GraphPtr, InOwningNode});
 		}
 
 		FOutputHandle FOutputNodeController::CreateOutputController(FGuid InVertexID, FConstVertexAccessPtr InNodeVertexPtr, FConstClassOutputAccessPtr InClassOutputPtr, FNodeHandle InOwningNode) const
@@ -1873,7 +1839,7 @@ namespace Metasound
 
 		FOutputHandle FInputNodeController::CreateOutputController(FGuid InVertexID, FConstVertexAccessPtr InNodeVertexPtr, FConstClassOutputAccessPtr InClassOutputPtr, FNodeHandle InOwningNode) const
 		{
-			return MakeShared<FInputNodeOutputController>(FInputNodeOutputController::FInitParams{InVertexID, InNodeVertexPtr, OwningGraphClassInputPtr, GraphPtr, InOwningNode});
+			return MakeShared<FInputNodeOutputController>(FInputNodeOutputController::FInitParams{InVertexID, InNodeVertexPtr, InClassOutputPtr, OwningGraphClassInputPtr, GraphPtr, InOwningNode});
 		}
 
 		FDocumentAccess FInputNodeController::ShareAccess() 
@@ -2164,7 +2130,6 @@ namespace Metasound
 							Node.Name = InClassInput.Name;
 							Node.ID = NewNodeID();
 
-							// TODO: have something that checks if input node has valid interface.
 							auto IsVertexWithTypeName = [&](FMetasoundFrontendVertex& Vertex) { return Vertex.TypeName == InClassInput.TypeName; };
 							if (FMetasoundFrontendVertex* InputVertex = Node.Interface.Inputs.FindByPredicate(IsVertexWithTypeName))
 							{ 
@@ -2175,14 +2140,14 @@ namespace Metasound
 								UE_LOG(LogMetasound, Error, TEXT("Input node [TypeName:%s] does not contain input vertex with type [TypeName:%s]"), *InClassInput.TypeName.ToString(), *InClassInput.TypeName.ToString());
 							}
 
-							if (FMetasoundFrontendVertex* OutputVertex = Node.Interface.Outputs.FindByPredicate(IsVertexWithTypeName))
+							if (Node.Interface.Outputs.Num() == 1)
+							{
+								Node.Interface.Outputs[0].Name = InClassInput.Name;
+							}
+							else if (FMetasoundFrontendVertex* OutputVertex = Node.Interface.Outputs.FindByPredicate(IsVertexWithTypeName))
 							{ 
 								OutputVertex->Name = InClassInput.Name;
 							}
-							else
-							{
-								UE_LOG(LogMetasound, Error, TEXT("Input node [TypeName:%s] does not contain input vertex with type [TypeName:%s]"), *InClassInput.TypeName.ToString(), *InClassInput.TypeName.ToString());
-							}
 
 							FMetasoundFrontendClassInput& NewInput = GraphClassPtr->Interface.Inputs.Add_GetRef(InClassInput);
 							NewInput.NodeID = Node.ID;
@@ -2243,14 +2208,14 @@ namespace Metasound
 
 							// TODO: have something that checks if input node has valid interface.
 							auto IsVertexWithTypeName = [&](FMetasoundFrontendVertex& Vertex) { return Vertex.TypeName == InClassOutput.TypeName; };
-							if (FMetasoundFrontendVertex* InputVertex = Node.Interface.Inputs.FindByPredicate(IsVertexWithTypeName))
+							if (Node.Interface.Inputs.Num() == 1)
+							{
+								Node.Interface.Inputs[0].Name = InClassOutput.Name;
+							}
+							else if (FMetasoundFrontendVertex* InputVertex = Node.Interface.Inputs.FindByPredicate(IsVertexWithTypeName))
 							{ 
 								InputVertex->Name = InClassOutput.Name;
 							}
-							else
-							{
-								UE_LOG(LogMetasound, Error, TEXT("Output node [TypeName:%s] does not contain input vertex with type [TypeName:%s]"), *InClassOutput.TypeName.ToString(), *InClassOutput.TypeName.ToString());
-							}
 
 							if (FMetasoundFrontendVertex* OutputVertex = Node.Interface.Outputs.FindByPredicate(IsVertexWithTypeName))
 							{ 
@@ -2258,7 +2223,7 @@ namespace Metasound
 							}
 							else
 							{
-								UE_LOG(LogMetasound, Error, TEXT("Output node [TypeName:%s] does not contain input vertex with type [TypeName:%s]"), *InClassOutput.TypeName.ToString(), *InClassOutput.TypeName.ToString());
+								UE_LOG(LogMetasound, Error, TEXT("Output node [TypeName:%s] does not contain output vertex with type [TypeName:%s]"), *InClassOutput.TypeName.ToString(), *InClassOutput.TypeName.ToString());
 							}
 
 							FMetasoundFrontendClassOutput& NewOutput = GraphClassPtr->Interface.Outputs.Add_GetRef(InClassOutput);
@@ -3490,3 +3455,5 @@ namespace Metasound
 	}
 }
 
+#undef LOCTEXT_NAMESPACE 
+
diff --git a/Engine/Plugins/Runtime/Metasound/Source/MetasoundFrontend/Private/MetasoundFrontendStandardController.h b/Engine/Plugins/Runtime/Metasound/Source/MetasoundFrontend/Private/MetasoundFrontendStandardController.h
--- a/Engine/Plugins/Runtime/Metasound/Source/MetasoundFrontend/Private/MetasoundFrontendStandardController.h
+++ b/Engine/Plugins/Runtime/Metasound/Source/MetasoundFrontend/Private/MetasoundFrontendStandardController.h
@@ -28,6 +28,7 @@ namespace Metasound
 				FGuid ID;
 
 				FConstVertexAccessPtr NodeVertexPtr;
+				FConstClassOutputAccessPtr ClassOutputPtr;
 				FGraphAccessPtr GraphPtr; 
 
 				/* Node handle which owns this output. */
@@ -45,6 +46,10 @@ namespace Metasound
 			const FName& GetDataType() const override;
 			const FString& GetName() const override;
 
+			// Output metadata
+			const FText& GetDisplayName() const override;
+			const FText& GetTooltip() const override;
+
 			// Return info on containing node. 
 			FGuid GetOwningNodeID() const override;
 			FNodeHandle GetOwningNode() override;
@@ -67,6 +72,7 @@ namespace Metasound
 
 			FGuid ID;
 			FConstVertexAccessPtr NodeVertexPtr;	
+			FConstClassOutputAccessPtr ClassOutputPtr;
 			FGraphAccessPtr GraphPtr; 
 			FNodeHandle OwningNode;
 
@@ -75,40 +81,6 @@ namespace Metasound
 			TArray<FMetasoundFrontendEdge> FindEdges() const;
 		};
 
-		/** FNodeOutputController represents a single output of a single node. */
-		class METASOUNDFRONTEND_API FNodeOutputController : public FBaseOutputController
-		{
-		public:
-			struct FInitParams
-			{
-				FGuid ID;
-
-				FConstVertexAccessPtr NodeVertexPtr;
-			
-				FConstClassOutputAccessPtr ClassOutputPtr;
-				FGraphAccessPtr GraphPtr; 
-
-				/* Node handle which owns this output. */
-				FNodeHandle OwningNode;
-			};
-
-			FNodeOutputController(const FInitParams& InInitParams);
-
-			virtual ~FNodeOutputController() = default;
-
-			bool IsValid() const override;
-
-			// Output metadata
-			const FText& GetDisplayName() const override;
-			const FText& GetTooltip() const override;
-
-		protected:
-			FDocumentAccess ShareAccess() override;
-			FConstDocumentAccess ShareAccess() const override;
-
-		private:
-			FConstClassOutputAccessPtr ClassOutputPtr;
-		};
 
 		/** FInputNodeOutputController represents the output vertex of an input 
 		 * node. 
@@ -124,7 +96,7 @@ namespace Metasound
 				FGuid ID;
 
 				FConstVertexAccessPtr NodeVertexPtr;
-			
+				FConstClassOutputAccessPtr ClassOutputPtr;
 				FConstClassInputAccessPtr OwningGraphClassInputPtr;
 				FGraphAccessPtr GraphPtr; 
 
@@ -149,6 +121,7 @@ namespace Metasound
 			FConstDocumentAccess ShareAccess() const override;
 
 		private:
+			mutable FText CachedDisplayName;
 
 			FConstClassInputAccessPtr OwningGraphClassInputPtr;
 		};
@@ -165,6 +138,7 @@ namespace Metasound
 			{
 				FGuid ID; 
 				FConstVertexAccessPtr NodeVertexPtr;
+				FConstClassInputAccessPtr ClassInputPtr;
 				FGraphAccessPtr GraphPtr; 
 				FNodeHandle OwningNode;
 			};
@@ -180,6 +154,11 @@ namespace Metasound
 			const FName& GetDataType() const override;
 			const FString& GetName() const override;
 
+			const FMetasoundFrontendLiteral* GetDefaultLiteral() const override;
+
+			// Input metadata
+			const FText& GetDisplayName() const override;
+			const FText& GetTooltip() const override;
 
 			// Owning node info
 			FGuid GetOwningNodeID() const override;
@@ -208,47 +187,13 @@ namespace Metasound
 			const FMetasoundFrontendEdge* FindEdge() const;
 			FMetasoundFrontendEdge* FindEdge();
 
-		private:
-
 			FGuid ID;
 			FConstVertexAccessPtr NodeVertexPtr;
+			FConstClassInputAccessPtr ClassInputPtr;
 			FGraphAccessPtr GraphPtr;
 			FNodeHandle OwningNode;
-		};
-
-		/** FNodeInputController represents a single input of a single node. */
-		class METASOUNDFRONTEND_API FNodeInputController : public FBaseInputController 
-		{
-		public:
-
-			struct FInitParams
-			{
-				FGuid ID; 
-				FConstVertexAccessPtr NodeVertexPtr;
-				FConstClassInputAccessPtr OwningGraphClassInputPtr;
-				FGraphAccessPtr GraphPtr; 
-				FNodeHandle OwningNode;
-			};
 
-			/** Constructs the input controller.  */
-			FNodeInputController(const FInitParams& InParams);
 
-			bool IsValid() const override;
-
-			const FMetasoundFrontendLiteral* GetDefaultLiteral() const override;
-
-			// Input metadata
-			const FText& GetDisplayName() const override;
-			const FText& GetTooltip() const override;
-
-		protected:
-
-			FDocumentAccess ShareAccess() override;
-			FConstDocumentAccess ShareAccess() const override;
-
-		private:
-
-			FConstClassInputAccessPtr OwningGraphClassInputPtr;
 		};
 
 		/** FOutputNodeInputController represents the input vertex of an output 
@@ -264,6 +209,7 @@ namespace Metasound
 			{
 				FGuid ID; 
 				FConstVertexAccessPtr NodeVertexPtr;
+				FConstClassInputAccessPtr ClassInputPtr;
 				FConstClassOutputAccessPtr OwningGraphClassOutputPtr;
 				FGraphAccessPtr GraphPtr; 
 				FNodeHandle OwningNode;
@@ -286,6 +232,8 @@ namespace Metasound
 
 		private:
 
+			mutable FText CachedDisplayName;
+
 			FConstClassOutputAccessPtr OwningGraphClassOutputPtr;
 		};
 
diff --git a/Engine/Plugins/Runtime/Metasound/Source/MetasoundFrontend/Public/MetasoundInputNode.h b/Engine/Plugins/Runtime/Metasound/Source/MetasoundFrontend/Public/MetasoundInputNode.h
--- a/Engine/Plugins/Runtime/Metasound/Source/MetasoundFrontend/Public/MetasoundInputNode.h
+++ b/Engine/Plugins/Runtime/Metasound/Source/MetasoundFrontend/Public/MetasoundInputNode.h
@@ -18,7 +18,7 @@ namespace Metasound
 		public:
 			using FDataWriteReference = TDataWriteReference<DataType>;
 
-			TInputOperator(const FString& InDataReferenceName, FDataWriteReference InDataReference)
+			TInputOperator(const FVertexKey& InDataReferenceName, FDataWriteReference InDataReference)
 			{
 				Inputs.AddDataWriteReference<DataType>(InDataReferenceName, InDataReference);
 				Outputs.AddDataReadReference<DataType>(InDataReferenceName, InDataReference);
@@ -105,7 +105,7 @@ namespace Metasound
 			// If true, this node can be instantiated by the FrontEnd.
 			static constexpr bool bCanRegister = TInputOperatorLiteralFactory<DataType>::bCanRegister;
 
-			static FVertexInterface DeclareVertexInterface(const FString& InVertexName)
+			static FVertexInterface DeclareVertexInterface(const FVertexKey& InVertexName)
 			{
 				return FVertexInterface(
 					FInputVertexInterface(
@@ -117,7 +117,7 @@ namespace Metasound
 				);
 			}
 
-			static FNodeClassMetadata GetNodeInfo(const FString& InVertexName)
+			static FNodeClassMetadata GetNodeInfo(const FVertexKey& InVertexName)
 			{
 				FNodeClassMetadata Info;
 
@@ -137,7 +137,7 @@ namespace Metasound
 			/* Construct a TInputNode using the TInputOperatorFactory<> and forwarding 
 			 * Args to the TInputOperatorFactory constructor.*/
 			template<typename... ArgTypes>
-			TInputNode(const FString& InNodeDescription, const FGuid& InInstanceID, const FString& InVertexName, ArgTypes&&... Args)
+			TInputNode(const FString& InNodeDescription, const FGuid& InInstanceID, const FVertexKey& InVertexName, ArgTypes&&... Args)
 			:	FNode(InNodeDescription, InInstanceID, GetNodeInfo(InVertexName))
 			,	VertexName(InVertexName)
 			,	Interface(DeclareVertexInterface(InVertexName))
@@ -147,7 +147,7 @@ namespace Metasound
 
 			/* Construct a TInputNode using the TInputOperatorLiteralFactory<> and moving
 			 * InParam to the TInputOperatorLiteralFactory constructor.*/
-			explicit TInputNode(const FString& InNodeDescription, const FGuid& InInstanceID, const FString& InVertexName, FLiteral&& InParam)
+			explicit TInputNode(const FString& InNodeDescription, const FGuid& InInstanceID, const FVertexKey& InVertexName, FLiteral&& InParam)
 			:	FNode(InNodeDescription, InInstanceID, GetNodeInfo(InVertexName))
 			,	VertexName(InVertexName)
 			,	Interface(DeclareVertexInterface(InVertexName))
@@ -155,7 +155,7 @@ namespace Metasound
 			{
 			}
 
-			const FString& GetVertexName() const
+			const FVertexKey& GetVertexName() const
 			{
 				return VertexName;
 			}
@@ -181,7 +181,7 @@ namespace Metasound
 			}
 
 		private:
-			FString VertexName;
+			FVertexKey VertexName;
 
 			FVertexInterface Interface;
 			FOperatorFactorySharedRef Factory;
diff --git a/Engine/Plugins/Runtime/Metasound/Source/MetasoundFrontend/Public/MetasoundOutputNode.h b/Engine/Plugins/Runtime/Metasound/Source/MetasoundFrontend/Public/MetasoundOutputNode.h
--- a/Engine/Plugins/Runtime/Metasound/Source/MetasoundFrontend/Public/MetasoundOutputNode.h
+++ b/Engine/Plugins/Runtime/Metasound/Source/MetasoundFrontend/Public/MetasoundOutputNode.h
@@ -82,8 +82,6 @@ namespace Metasound
 
 		static FNodeClassMetadata GetNodeInfo(const FString& InVertexName)
 		{
-			static const FString ClassNameString = FString(TEXT("Output_")) + GetMetasoundDataTypeName<DataType>().ToString();
-
 			FNodeClassMetadata Info;
 
 			Info.ClassName = {TEXT("Output"), GetMetasoundDataTypeName<DataType>(), TEXT("")};
diff --git a/Engine/Plugins/Runtime/Metasound/Source/MetasoundStandardNodes/Private/MetasoundAudioFormatNodes.cpp b/Engine/Plugins/Runtime/Metasound/Source/MetasoundStandardNodes/Private/MetasoundAudioFormatNodes.cpp
deleted file mode 100644
diff --git a/Engine/Plugins/Runtime/Metasound/Source/MetasoundStandardNodes/Private/MetasoundAudioFormats.cpp b/Engine/Plugins/Runtime/Metasound/Source/MetasoundStandardNodes/Private/MetasoundAudioFormats.cpp
--- a/Engine/Plugins/Runtime/Metasound/Source/MetasoundStandardNodes/Private/MetasoundAudioFormats.cpp
+++ b/Engine/Plugins/Runtime/Metasound/Source/MetasoundStandardNodes/Private/MetasoundAudioFormats.cpp
@@ -5,12 +5,11 @@
 #include "CoreMinimal.h"
 #include "MetasoundAudioBuffer.h"
 #include "MetasoundDataReference.h"
-#include "MetasoundLiteral.h"
 #include "MetasoundDataTypeRegistrationMacro.h"
+#include "MetasoundLiteral.h"
+#include "MetasoundOutputNode.h"
 
-REGISTER_METASOUND_DATATYPE(Metasound::FMonoAudioFormat, "Audio:Mono");
-REGISTER_METASOUND_DATATYPE(Metasound::FStereoAudioFormat, "Audio:Stereo");
-REGISTER_METASOUND_DATATYPE(Metasound::FMultichannelAudioFormat, "Audio:Multichannel", ELiteralType::Integer);
+#define LOCTEXT_NAMESPACE "MetasoundStandardNodes"
 
 namespace Metasound
 {
@@ -83,4 +82,566 @@ namespace Metasound
 			ReadableBuffers = ReadableBufferStorage;
 		}
 	}
-} // namespace Metasound
+
+	// Special vertex keys for stereo input/output nodes.
+	namespace StereoAudioFormatVertexKeys
+	{
+		static const FVertexKey LeftChannelVertexKey = "Left";
+		static const FVertexKey RightChannelVertexKey = "Right";
+	}
+
+	// Specialization of TOutputNode<FStereoAudio> to support direct connection
+	// of audio buffers to left/right inputs. 
+	template<>
+	class METASOUNDSTANDARDNODES_API TOutputNode<FStereoAudioFormat> : public FNode
+	{
+		// FOutputOperator primarly used to report inputs and outputs. Has no execute function.
+		class FOutputOperator : public IOperator
+		{
+		public:
+			FOutputOperator(const FVertexKey& InOutputName, TDataReadReference<FAudioBuffer> InLeft, TDataReadReference<FAudioBuffer> InRight, TDataReadReference<FStereoAudioFormat> InStereo)
+			: OutputName(InOutputName)
+			, Left(InLeft)
+			, Right(InRight)
+			, Stereo(InStereo)
+			{
+			}
+
+			virtual ~FOutputOperator() {}
+
+			virtual FDataReferenceCollection GetInputs() const override
+			{
+				FDataReferenceCollection Inputs;
+
+				Inputs.AddDataReadReference(StereoAudioFormatVertexKeys::LeftChannelVertexKey, Left);
+				Inputs.AddDataReadReference(StereoAudioFormatVertexKeys::RightChannelVertexKey, Right);
+
+				return Inputs;
+			}
+
+			virtual FDataReferenceCollection GetOutputs() const override
+			{
+				FDataReferenceCollection Outputs;
+
+				Outputs.AddDataReadReference(OutputName, Stereo);
+
+				return Outputs;
+			}
+
+			virtual FExecuteFunction GetExecuteFunction() override
+			{
+				return nullptr;
+			}
+
+		private:
+
+			FVertexKey OutputName;
+			TDataReadReference<FAudioBuffer> Left;
+			TDataReadReference<FAudioBuffer> Right;
+			TDataReadReference<FStereoAudioFormat> Stereo;
+		};
+
+		class FOutputOperatorFactory : public IOperatorFactory
+		{
+		public:
+			FOutputOperatorFactory(const FVertexKey& InOutputName)
+			: OutputName(InOutputName)
+			{
+			}
+
+			TUniquePtr<IOperator> CreateOperator(const FCreateOperatorParams& InParams, FBuildErrorArray& OutErrors) override
+			{
+				// Construct stereo from left and right audio buffers. 
+				TDataReadReference<FAudioBuffer> Left = InParams.InputDataReferences.GetDataReadReferenceOrConstruct<FAudioBuffer>(StereoAudioFormatVertexKeys::LeftChannelVertexKey, InParams.OperatorSettings);
+				TDataReadReference<FAudioBuffer> Right = InParams.InputDataReferences.GetDataReadReferenceOrConstruct<FAudioBuffer>(StereoAudioFormatVertexKeys::RightChannelVertexKey, InParams.OperatorSettings);
+
+				TDataReadReference<FStereoAudioFormat> Stereo = TDataReadReferenceFactory<FStereoAudioFormat>::CreateExplicitArgs(InParams.OperatorSettings, WriteCast(Left), WriteCast(Right));
+
+				return MakeUnique<FOutputOperator>(OutputName, Left, Right, Stereo);
+			}
+
+		private:
+			FVertexKey OutputName;
+		};
+
+		static FVertexInterface GetVertexInterface(const FVertexKey& InVertexName)
+		{
+			return FVertexInterface(
+				FInputVertexInterface(
+					TInputDataVertexModel<FAudioBuffer>(StereoAudioFormatVertexKeys::LeftChannelVertexKey, LOCTEXT("Metasound_LeftStereoOutputVertexDescription", "Left channel audio output.")),
+					TInputDataVertexModel<FAudioBuffer>(StereoAudioFormatVertexKeys::RightChannelVertexKey, LOCTEXT("Metasound_RightStereoOutputVertexDescription", "Right channel audio output."))
+				),
+				FOutputVertexInterface(
+					TOutputDataVertexModel<FStereoAudioFormat>(InVertexName, LOCTEXT("Metasound_StereoOutputVertexDescription", "Stereo Output."))
+				)
+			);
+		}
+
+		static FNodeClassMetadata GetNodeClassMetadata(const FVertexKey& InOutputName)
+		{
+			FNodeClassMetadata Info;
+
+			Info.ClassName = {TEXT("Output"), GetMetasoundDataTypeName<FStereoAudioFormat>(), TEXT("")};
+			Info.MajorVersion = 1;
+			Info.MinorVersion = 0;
+			Info.DisplayName = FText::Format(LOCTEXT("Metasound_OutputNodeDisplayNameFormat", "Output {0}"), FText::FromString(GetMetasoundDataTypeString<FStereoAudioFormat>()));
+			Info.Description = LOCTEXT("Metasound_OutputNodeDescription", "Output from the parent Metasound graph.");
+			Info.Author = PluginAuthor;
+			Info.PromptIfMissing = PluginNodeMissingPrompt;
+			Info.DefaultInterface = GetVertexInterface(InOutputName);
+
+			return Info;
+		};
+
+
+	public:
+		TOutputNode(const FString& InInstanceName, const FGuid& InInstanceID, const FVertexKey& InVertexName)
+		:	FNode(InInstanceName, InInstanceID, GetNodeClassMetadata(InVertexName))
+		,	VertexInterface(GetVertexInterface(InVertexName))
+		,	Factory(MakeShared<FOutputOperatorFactory, ESPMode::ThreadSafe>(InVertexName))
+		{
+		}
+
+		const FVertexInterface& GetVertexInterface() const override
+		{
+			return VertexInterface;
+		}
+
+		bool SetVertexInterface(const FVertexInterface& InInterface) override
+		{
+			return VertexInterface == InInterface;
+		}
+
+		bool IsVertexInterfaceSupported(const FVertexInterface& InInterface) const 
+		{
+			return VertexInterface == InInterface;
+		}
+
+		virtual TSharedRef<IOperatorFactory, ESPMode::ThreadSafe> GetDefaultOperatorFactory() const override
+		{
+			return Factory;
+		}
+
+	private:
+		FVertexInterface VertexInterface;
+
+		TSharedRef<FOutputOperatorFactory, ESPMode::ThreadSafe> Factory;
+
+	};
+
+	// Input node specialization to expose left/right audio buffers 
+	template<>
+	class METASOUNDSTANDARDNODES_API TInputNode<FStereoAudioFormat> : public FNode
+	{
+		// Noop operator. Used to return inputs / outputs for debugging.
+		class FInputOperator : public IOperator
+		{
+		public:
+
+			FInputOperator(const FVertexKey& InInputName, TDataReadReference<FAudioBuffer> InLeft, TDataReadReference<FAudioBuffer> InRight, TDataReadReference<FStereoAudioFormat> InStereo)
+			: InputName(InInputName)
+			, Left(InLeft)
+			, Right(InRight)
+			, Stereo(InStereo)
+			{
+			}
+
+			virtual ~FInputOperator() {}
+
+
+			virtual FDataReferenceCollection GetInputs() const override
+			{
+				FDataReferenceCollection Inputs;
+
+				Inputs.AddDataReadReference(InputName, Stereo);
+
+				return Inputs;
+			}
+
+			virtual FDataReferenceCollection GetOutputs() const override
+			{
+				FDataReferenceCollection Outputs;
+
+				Outputs.AddDataReadReference(StereoAudioFormatVertexKeys::LeftChannelVertexKey, Left);
+				Outputs.AddDataReadReference(StereoAudioFormatVertexKeys::RightChannelVertexKey, Right);
+
+				return Outputs;
+			}
+
+			virtual FExecuteFunction GetExecuteFunction() override
+			{
+				return nullptr;
+			}
+
+		private:
+
+			FVertexKey InputName;
+			TDataReadReference<FAudioBuffer> Left;
+			TDataReadReference<FAudioBuffer> Right;
+			TDataReadReference<FStereoAudioFormat> Stereo;
+		};
+
+		class FInputOperatorFactory : public IOperatorFactory
+		{
+		public:
+			FInputOperatorFactory(const FVertexKey& InInputName)
+			: InputName(InInputName)
+			{
+			}
+
+			TUniquePtr<IOperator> CreateOperator(const FCreateOperatorParams& InParams, FBuildErrorArray& OutErrors) override
+			{
+				// Split a stereo signal into left/right.
+				TDataReadReference<FStereoAudioFormat> Stereo = InParams.InputDataReferences.GetDataReadReferenceOrConstruct<FStereoAudioFormat>(InputName, InParams.OperatorSettings);
+				TDataReadReference<FAudioBuffer> Left = Stereo->GetLeft();
+				TDataReadReference<FAudioBuffer> Right = Stereo->GetRight();
+
+				return MakeUnique<FInputOperator>(InputName, Left, Right, Stereo);
+			}
+
+		private:
+			FVertexKey InputName;
+		};
+
+		static FVertexInterface GetVertexInterface(const FVertexKey& InVertexName)
+		{
+			return FVertexInterface(
+				FInputVertexInterface(
+					TInputDataVertexModel<FStereoAudioFormat>(InVertexName, LOCTEXT("Metasound_StereoInputVertexDescription", "Stereo Input."))
+				),
+				FOutputVertexInterface(
+					TOutputDataVertexModel<FAudioBuffer>(StereoAudioFormatVertexKeys::LeftChannelVertexKey, LOCTEXT("Metasound_LeftStereoInputVertexDescription", "Left channel audio output.")),
+					TOutputDataVertexModel<FAudioBuffer>(StereoAudioFormatVertexKeys::RightChannelVertexKey, LOCTEXT("Metasound_RightStereoInputVertexDescription", "Right channel audio output."))
+				)
+			);
+		}
+
+		static FNodeClassMetadata GetNodeClassMetadata(const FVertexKey& InInputName)
+		{
+			FNodeClassMetadata Info;
+
+			Info.ClassName = {TEXT("Input"), GetMetasoundDataTypeName<FStereoAudioFormat>(), TEXT("")};
+			Info.MajorVersion = 1;
+			Info.MinorVersion = 0;
+			Info.DisplayName = FText::Format(LOCTEXT("Metasound_InputNodeDisplayNameFormat", "Input {0}"), FText::FromString(GetMetasoundDataTypeString<FStereoAudioFormat>()));
+			Info.Description = LOCTEXT("Metasound_InputNodeDescription", "Input from the parent Metasound graph.");
+			Info.Author = PluginAuthor;
+			Info.PromptIfMissing = PluginNodeMissingPrompt;
+			Info.DefaultInterface = GetVertexInterface(InInputName);
+
+			return Info;
+		};
+
+
+	public:
+		static constexpr bool bCanRegister = true;
+
+		TInputNode(const FString& InInstanceName, const FGuid& InInstanceID, const FVertexKey& InVertexName, FLiteral&& InLiteral)
+		:	FNode(InInstanceName, InInstanceID, GetNodeClassMetadata(InVertexName))
+		,	VertexInterface(GetVertexInterface(InVertexName))
+		,	Factory(MakeShared<FInputOperatorFactory, ESPMode::ThreadSafe>(InVertexName))
+		{
+		}
+
+		const FVertexInterface& GetVertexInterface() const override
+		{
+			return VertexInterface;
+		}
+
+		bool SetVertexInterface(const FVertexInterface& InInterface) override
+		{
+			return VertexInterface == InInterface;
+		}
+
+		bool IsVertexInterfaceSupported(const FVertexInterface& InInterface) const 
+		{
+			return VertexInterface == InInterface;
+		}
+
+		TSharedRef<IOperatorFactory, ESPMode::ThreadSafe> GetDefaultOperatorFactory() const override
+		{
+			return Factory;
+		}
+
+	private:
+		FVertexInterface VertexInterface;
+
+		TSharedRef<FInputOperatorFactory, ESPMode::ThreadSafe> Factory;
+
+	};
+
+	// TOutputNode<> specialization for FMonoAudioFormat. Allows an audio buffer
+	// to be directly connected to an mono audio output.
+	template<>
+	class METASOUNDSTANDARDNODES_API TOutputNode<FMonoAudioFormat> : public FNode
+	{
+		class FOutputOperator : public IOperator
+		{
+		public:
+			FOutputOperator(const FVertexKey& InOutputName, TDataReadReference<FAudioBuffer> InCenter, TDataReadReference<FMonoAudioFormat> InMono)
+			: OutputName(InOutputName)
+			, Center(InCenter)
+			, Mono(InMono)
+			{
+			}
+
+			virtual ~FOutputOperator() {}
+
+			virtual FDataReferenceCollection GetInputs() const override
+			{
+				FDataReferenceCollection Inputs;
+
+				Inputs.AddDataReadReference(OutputName, Center);
+
+				return Inputs;
+			}
+
+			virtual FDataReferenceCollection GetOutputs() const override
+			{
+				FDataReferenceCollection Outputs;
+
+				Outputs.AddDataReadReference(OutputName, Mono);
+
+				return Outputs;
+			}
+
+			virtual FExecuteFunction GetExecuteFunction() override
+			{
+				return nullptr;
+			}
+
+		private:
+
+			FVertexKey OutputName;
+			TDataReadReference<FAudioBuffer> Center;
+			TDataReadReference<FMonoAudioFormat> Mono;
+		};
+
+		class FOutputOperatorFactory : public IOperatorFactory
+		{
+		public:
+			FOutputOperatorFactory(const FVertexKey& InOutputName)
+			: OutputName(InOutputName)
+			{
+			}
+
+			TUniquePtr<IOperator> CreateOperator(const FCreateOperatorParams& InParams, FBuildErrorArray& OutErrors) override
+			{
+				TDataReadReference<FAudioBuffer> Center = InParams.InputDataReferences.GetDataReadReferenceOrConstruct<FAudioBuffer>(OutputName, InParams.OperatorSettings);
+
+				TDataReadReference<FMonoAudioFormat> Mono = TDataReadReferenceFactory<FMonoAudioFormat>::CreateExplicitArgs(InParams.OperatorSettings, WriteCast(Center));
+
+				return MakeUnique<FOutputOperator>(OutputName, Center, Mono);
+			}
+
+		private:
+			FVertexKey OutputName;
+		};
+
+		static FVertexInterface GetVertexInterface(const FVertexKey& InVertexName)
+		{
+			return FVertexInterface(
+				FInputVertexInterface(
+					TInputDataVertexModel<FAudioBuffer>(InVertexName, LOCTEXT("Metasound_CenterMonoOutputVertexDescription", "Center channel audio output."))
+				),
+				FOutputVertexInterface(
+					TOutputDataVertexModel<FMonoAudioFormat>(InVertexName, LOCTEXT("Metasound_MonoOutputVertexDescription", "Mono Output."))
+				)
+			);
+		}
+
+		static FNodeClassMetadata GetNodeClassMetadata(const FVertexKey& InOutputName)
+		{
+			FNodeClassMetadata Info;
+
+			Info.ClassName = {TEXT("Output"), GetMetasoundDataTypeName<FMonoAudioFormat>(), TEXT("")};
+			Info.MajorVersion = 1;
+			Info.MinorVersion = 0;
+			Info.DisplayName = FText::Format(LOCTEXT("Metasound_OutputNodeDisplayNameFormat", "Output {0}"), FText::FromString(GetMetasoundDataTypeString<FMonoAudioFormat>()));
+			Info.Description = LOCTEXT("Metasound_OutputNodeDescription", "Output from the parent Metasound graph.");
+			Info.Author = PluginAuthor;
+			Info.PromptIfMissing = PluginNodeMissingPrompt;
+			Info.DefaultInterface = GetVertexInterface(InOutputName);
+
+			return Info;
+		};
+
+
+	public:
+		TOutputNode(const FString& InInstanceName, const FGuid& InInstanceID, const FVertexKey& InVertexName)
+		:	FNode(InInstanceName, InInstanceID, GetNodeClassMetadata(InVertexName))
+		,	VertexInterface(GetVertexInterface(InVertexName))
+		,	Factory(MakeShared<FOutputOperatorFactory, ESPMode::ThreadSafe>(InVertexName))
+		{
+		}
+
+		const FVertexInterface& GetVertexInterface() const override
+		{
+			return VertexInterface;
+		}
+
+		bool SetVertexInterface(const FVertexInterface& InInterface) override
+		{
+			return VertexInterface == InInterface;
+		}
+
+		bool IsVertexInterfaceSupported(const FVertexInterface& InInterface) const 
+		{
+			return VertexInterface == InInterface;
+		}
+
+		TSharedRef<IOperatorFactory, ESPMode::ThreadSafe> GetDefaultOperatorFactory() const override
+		{
+			return Factory;
+		}
+
+	private:
+		FVertexInterface VertexInterface;
+
+		TSharedRef<FOutputOperatorFactory, ESPMode::ThreadSafe> Factory;
+
+	};
+
+	// TInputNode<> specializastion for FMonoAudioFormat. Allows an input mono audio
+	// format to be exposed as a single buffer.
+	template<>
+	class METASOUNDSTANDARDNODES_API TInputNode<FMonoAudioFormat> : public FNode
+	{
+		class FInputOperator : public IOperator
+		{
+		public:
+
+			FInputOperator(const FVertexKey& InInputName, TDataReadReference<FAudioBuffer> InCenter, TDataReadReference<FMonoAudioFormat> InMono)
+			: InputName(InInputName)
+			, Center(InCenter)
+			, Mono(InMono)
+			{
+			}
+
+			virtual ~FInputOperator() {}
+
+
+			virtual FDataReferenceCollection GetInputs() const override
+			{
+				FDataReferenceCollection Inputs;
+
+				Inputs.AddDataReadReference(InputName, Mono);
+
+				return Inputs;
+			}
+
+			virtual FDataReferenceCollection GetOutputs() const override
+			{
+				FDataReferenceCollection Outputs;
+
+				Outputs.AddDataReadReference(InputName, Center);
+
+				return Outputs;
+			}
+
+			virtual FExecuteFunction GetExecuteFunction() override
+			{
+				return nullptr;
+			}
+
+		private:
+
+			FVertexKey InputName;
+			TDataReadReference<FAudioBuffer> Center;
+			TDataReadReference<FMonoAudioFormat> Mono;
+		};
+
+		class FInputOperatorFactory : public IOperatorFactory
+		{
+		public:
+			FInputOperatorFactory(const FVertexKey& InInputName)
+			: InputName(InInputName)
+			{
+			}
+
+			TUniquePtr<IOperator> CreateOperator(const FCreateOperatorParams& InParams, FBuildErrorArray& OutErrors) override
+			{
+				TDataReadReference<FMonoAudioFormat> Mono = InParams.InputDataReferences.GetDataReadReferenceOrConstruct<FMonoAudioFormat>(InputName, InParams.OperatorSettings);
+				TDataReadReference<FAudioBuffer> Center = Mono->GetCenter();
+
+				return MakeUnique<FInputOperator>(InputName, Center, Mono);
+			}
+
+		private:
+			FVertexKey InputName;
+		};
+
+		static FVertexInterface GetVertexInterface(const FVertexKey& InVertexName)
+		{
+			return FVertexInterface(
+				FInputVertexInterface(
+					TInputDataVertexModel<FMonoAudioFormat>(InVertexName, LOCTEXT("Metasound_MonoInputVertexDescription", "Mono Input."))
+				),
+				FOutputVertexInterface(
+					TOutputDataVertexModel<FAudioBuffer>(InVertexName, LOCTEXT("Metasound_CenterMonoInputVertexDescription", "Center channel audio output."))
+				)
+			);
+		}
+
+		static FNodeClassMetadata GetNodeClassMetadata(const FVertexKey& InInputName)
+		{
+			FNodeClassMetadata Info;
+
+			Info.ClassName = {TEXT("Input"), GetMetasoundDataTypeName<FMonoAudioFormat>(), TEXT("")};
+			Info.MajorVersion = 1;
+			Info.MinorVersion = 0;
+			Info.DisplayName = FText::Format(LOCTEXT("Metasound_InputNodeDisplayNameFormat", "Input {0}"), FText::FromString(GetMetasoundDataTypeString<FMonoAudioFormat>()));
+			Info.Description = LOCTEXT("Metasound_InputNodeDescription", "Input from the parent Metasound graph.");
+			Info.Author = PluginAuthor;
+			Info.PromptIfMissing = PluginNodeMissingPrompt;
+			Info.DefaultInterface = GetVertexInterface(InInputName);
+
+			return Info;
+		};
+
+
+	public:
+		static constexpr bool bCanRegister = true;
+
+		TInputNode(const FString& InInstanceName, const FGuid& InInstanceID, const FVertexKey& InVertexName, FLiteral&& InLiteral)
+		:	FNode(InInstanceName, InInstanceID, GetNodeClassMetadata(InVertexName))
+		,	VertexInterface(GetVertexInterface(InVertexName))
+		,	Factory(MakeShared<FInputOperatorFactory, ESPMode::ThreadSafe>(InVertexName))
+		{
+		}
+
+		const FVertexInterface& GetVertexInterface() const override
+		{
+			return VertexInterface;
+		}
+
+		bool SetVertexInterface(const FVertexInterface& InInterface) override
+		{
+			return VertexInterface == InInterface;
+		}
+
+		bool IsVertexInterfaceSupported(const FVertexInterface& InInterface) const 
+		{
+			return VertexInterface == InInterface;
+		}
+
+		TSharedRef<IOperatorFactory, ESPMode::ThreadSafe> GetDefaultOperatorFactory() const override
+		{
+			return Factory;
+		}
+
+	private:
+		FVertexInterface VertexInterface;
+
+		TSharedRef<FInputOperatorFactory, ESPMode::ThreadSafe> Factory;
+
+	};
+}
+
+// Data type registration has to happen after TInputNode<> and TOutputNode<> specializations
+// so that the registration macro has access to the specializations.
+REGISTER_METASOUND_DATATYPE(Metasound::FMonoAudioFormat, "Audio:Mono");
+REGISTER_METASOUND_DATATYPE(Metasound::FStereoAudioFormat, "Audio:Stereo");
+REGISTER_METASOUND_DATATYPE(Metasound::FMultichannelAudioFormat, "Audio:Multichannel", ELiteralType::Integer);
+
+#undef LOCTEXT_NAMESPACE // MetasoundStandardNodes
diff --git a/Engine/Plugins/Runtime/Metasound/Source/MetasoundStandardNodes/Public/MetasoundAudioFormatNodes.h b/Engine/Plugins/Runtime/Metasound/Source/MetasoundStandardNodes/Public/MetasoundAudioFormatNodes.h
deleted file mode 100644
