Fix bug manifesting in PeriodicTrigger caused by inadvertent implicit casting in FTime ops
- Make addition and subtraction in FSampleCounter require integral types to
    1. Avoid invalid baseless time against sample operations, and
    2. Require client code to be explicit when truncation is ultimately desired

diff --git a/Engine/Plugins/Runtime/Metasound/Source/MetasoundFrontend/Public/MetasoundTime.h b/Engine/Plugins/Runtime/Metasound/Source/MetasoundFrontend/Public/MetasoundTime.h
--- a/Engine/Plugins/Runtime/Metasound/Source/MetasoundFrontend/Public/MetasoundTime.h
+++ b/Engine/Plugins/Runtime/Metasound/Source/MetasoundFrontend/Public/MetasoundTime.h
@@ -21,7 +21,7 @@ namespace Metasound
 			{
 			}
 
-			FTime(TimeType InTime)
+			explicit FTime(TimeType InTime)
 				: Time(InTime)
 			{
 			}
@@ -127,7 +127,7 @@ namespace Metasound
 
 			FTime operator-(const FTime& InRHS) const
 			{
-				return Time - InRHS.Time;
+				return FTime(Time - InRHS.Time);
 			}
 
 			template<typename ArithmeticType>
@@ -146,7 +146,7 @@ namespace Metasound
 
 			FTime operator+(const FTime& InRHS) const
 			{
-				return Time + InRHS.Time;
+				return FTime(Time + InRHS.Time);
 			}
 
 			template<typename ArithmeticType>
@@ -160,7 +160,7 @@ namespace Metasound
 			friend FTime operator+(const FTime& InLHS, const ArithmeticType& InRHS)
 			{
 				static_assert(TIsArithmetic<ArithmeticType>::Value, "Must be arithmetic type.");
-				return InLHS.Time + InRHS;
+				return FTime(InLHS.Time + InRHS);
 			}
 
 			template<typename ArithmeticType>
@@ -181,7 +181,7 @@ namespace Metasound
 			friend FTime operator*(const FTime& InLHS, const ArithmeticType& InRHS)
 			{
 				static_assert(TIsArithmetic<ArithmeticType>::Value, "Must be arithmetic type.");
-				return InLHS.Time * InRHS;
+				return FTime(InLHS.Time * InRHS);
 			}
 
 			template<typename ArithmeticType>
@@ -202,7 +202,7 @@ namespace Metasound
 			friend FTime operator/(const FTime& InLHS, const ArithmeticType& InRHS)
 			{
 				static_assert(TIsArithmetic<ArithmeticType>::Value, "Must be arithmetic type.");
-				return InLHS.Time / InRHS;
+				return FTime(InLHS.Time / InRHS);
 			}
 
 			bool operator<(const FTime& InRHS) const
diff --git a/Engine/Plugins/Runtime/Metasound/Source/MetasoundFrontend/Public/MetasoundTrigger.h b/Engine/Plugins/Runtime/Metasound/Source/MetasoundFrontend/Public/MetasoundTrigger.h
--- a/Engine/Plugins/Runtime/Metasound/Source/MetasoundFrontend/Public/MetasoundTrigger.h
+++ b/Engine/Plugins/Runtime/Metasound/Source/MetasoundFrontend/Public/MetasoundTrigger.h
@@ -46,16 +46,6 @@ namespace Metasound
 			 */
 			void TriggerFrame(int32 InFrameToTrigger);
 
-			/** Trigger a specific time in the future.
-			 *
-			 * @param InTime - Time to set trigger.
-			 */
-			template<typename TimeType>
-			void TriggerTime(const TimeType& InTime)
-			{
-				TriggerFrame(InTime.GetNumSamples(SampleRate));
-			}
-
 			/** Advance internal frame counters by block size. */
 			void AdvanceBlock();
 
diff --git a/Engine/Plugins/Runtime/Metasound/Source/MetasoundStandardNodes/Private/MetasoundPeriodicTriggerNode.cpp b/Engine/Plugins/Runtime/Metasound/Source/MetasoundStandardNodes/Private/MetasoundPeriodicTriggerNode.cpp
--- a/Engine/Plugins/Runtime/Metasound/Source/MetasoundStandardNodes/Private/MetasoundPeriodicTriggerNode.cpp
+++ b/Engine/Plugins/Runtime/Metasound/Source/MetasoundStandardNodes/Private/MetasoundPeriodicTriggerNode.cpp
@@ -85,6 +85,7 @@ namespace Metasound
 			[this](int32 StartFrame, int32 EndFrame)
 			{
 				bEnabled = true;
+				SampleCounter.SetNumSamples(0);
 			}
 		);
 
@@ -105,7 +106,7 @@ namespace Metasound
 			// Time must march on, can't stay in the now forever.
 			PeriodInSamples = FMath::Max(static_cast<FSampleCount>(1), PeriodInSamples);
 
-			while ((SampleCounter - BlockSize) <= 0)
+			while ((SampleCounter - BlockSize).GetNumSamples() <= 0)
 			{
 				const int32 StartOffset = static_cast<int32>(SampleCounter.GetNumSamples());
 				TriggerOut->TriggerFrame(StartOffset);
diff --git a/Engine/Plugins/Runtime/Metasound/Source/MetasoundStandardNodes/Private/MetasoundTriggerDelayNode.cpp b/Engine/Plugins/Runtime/Metasound/Source/MetasoundStandardNodes/Private/MetasoundTriggerDelayNode.cpp
--- a/Engine/Plugins/Runtime/Metasound/Source/MetasoundStandardNodes/Private/MetasoundTriggerDelayNode.cpp
+++ b/Engine/Plugins/Runtime/Metasound/Source/MetasoundStandardNodes/Private/MetasoundTriggerDelayNode.cpp
@@ -111,7 +111,7 @@ namespace Metasound
 			}
 		);
 
-		if (NextTriggerCounter >= 0)
+		if (NextTriggerCounter.GetNumSamples() >= 0)
 		{
 			FSampleCount SamplesRemaining = NextTriggerCounter.GetNumSamples() - FramesPerBlock;
 			if (SamplesRemaining > 0.0f)
diff --git a/Engine/Plugins/Runtime/Metasound/Source/MetasoundStandardNodes/Public/MetasoundSampleCounter.h b/Engine/Plugins/Runtime/Metasound/Source/MetasoundStandardNodes/Public/MetasoundSampleCounter.h
--- a/Engine/Plugins/Runtime/Metasound/Source/MetasoundStandardNodes/Public/MetasoundSampleCounter.h
+++ b/Engine/Plugins/Runtime/Metasound/Source/MetasoundStandardNodes/Public/MetasoundSampleCounter.h
@@ -161,6 +161,20 @@ namespace Metasound
 				return { FSampleCount(InLHS.NumSamples - (InRHS * InLHS.SampleRate)), InLHS.SampleRate };
 			}
 
+			template<typename IntegralType>
+			friend FSampleCounter operator+(const FSampleCounter& InLHS, const IntegralType& InRHS)
+			{
+				static_assert(TIsIntegral<IntegralType>::Value, "Right-hand addend type must be integral.");
+				return FSampleCounter(InLHS.NumSamples + InRHS, InLHS.SampleRate);
+			}
+
+			template<typename IntegralType>
+			friend FSampleCounter operator-(const FSampleCounter& InLHS, const IntegralType& InRHS)
+			{
+				static_assert(TIsIntegral<IntegralType>::Value, "Subtrahend type must be integral.");
+				return FSampleCounter(InLHS.NumSamples - InRHS, InLHS.SampleRate);
+			}
+
 			template<typename ArithmeticType>
 			friend FSampleCounter operator*(const FSampleCounter& InLHS, const ArithmeticType& InRHS)
 			{
@@ -213,18 +227,18 @@ namespace Metasound
 				return InLHS;
 			}
 
-			template<typename ArithmeticType>
-			friend FSampleCounter operator+=(FSampleCounter& InLHS, const ArithmeticType& InRHS)
+			template<typename IntegralType>
+			friend FSampleCounter operator+=(FSampleCounter& InLHS, const IntegralType& InRHS)
 			{
-				static_assert(TIsArithmetic<ArithmeticType>::Value, "Addend type must be arithmetic.");
+				static_assert(TIsIntegral<IntegralType>::Value, "Addend type must be integral.");
 				InLHS.NumSamples += InRHS;
 				return InLHS;
 			}
 
-			template<typename ArithmeticType>
-			friend FSampleCounter& operator-=(FSampleCounter& InLHS, const ArithmeticType& InRHS)
+			template<typename IntegralType>
+			friend FSampleCounter& operator-=(FSampleCounter& InLHS, const IntegralType& InRHS)
 			{
-				static_assert(TIsArithmetic<ArithmeticType>::Value, "Subtrahend type must be arithmetic.");
+				static_assert(TIsIntegral<IntegralType>::Value, "Subtrahend type must be integral.");
 				InLHS.NumSamples -= InRHS;
 				return InLHS;
 			}
diff --git a/Engine/Plugins/Runtime/Metasound/Source/MetasoundStandardNodes/Public/MetasoundTriggerDelayNode.h b/Engine/Plugins/Runtime/Metasound/Source/MetasoundStandardNodes/Public/MetasoundTriggerDelayNode.h
--- a/Engine/Plugins/Runtime/Metasound/Source/MetasoundStandardNodes/Public/MetasoundTriggerDelayNode.h
+++ b/Engine/Plugins/Runtime/Metasound/Source/MetasoundStandardNodes/Public/MetasoundTriggerDelayNode.h
@@ -20,7 +20,5 @@ namespace Metasound
 	{
 		public:
 			FTriggerDelayNode(const FNodeInitData& InInitData);
-
-			virtual ~FTriggerDelayNode() = default;
 	};
 }
