New modulo metasound math node

* removes cruft of random node
* adds icon for modulo and modifies division icon to avoid ambiguity with modulo

diff --git a/Engine/Plugins/Runtime/Metasound/Source/MetasoundEditor/Private/MetasoundEditorModule.cpp b/Engine/Plugins/Runtime/Metasound/Source/MetasoundEditor/Private/MetasoundEditorModule.cpp
--- a/Engine/Plugins/Runtime/Metasound/Source/MetasoundEditor/Private/MetasoundEditorModule.cpp
+++ b/Engine/Plugins/Runtime/Metasound/Source/MetasoundEditor/Private/MetasoundEditorModule.cpp
@@ -83,8 +83,8 @@ namespace Metasound
 					Set("MetasoundEditor.Graph.Node.Math.Add", new FSlateImageBrush(RootToContentDir(TEXT("/Graph/node_math_add_40x.png")), Icon40x40));
 					Set("MetasoundEditor.Graph.Node.Math.Divide", new FSlateImageBrush(RootToContentDir(TEXT("/Graph/node_math_divide_40x.png")), Icon40x40));
 					Set("MetasoundEditor.Graph.Node.Math.Multiply", new FSlateImageBrush(RootToContentDir(TEXT("/Graph/node_math_multiply_40x.png")), Icon40x40));
-					Set("MetasoundEditor.Graph.Node.Math.RandRange", new FSlateImageBrush(RootToContentDir(TEXT("/Graph/node_math_random_40x.png")), Icon40x40));
 					Set("MetasoundEditor.Graph.Node.Math.Subtract", new FSlateImageBrush(RootToContentDir(TEXT("/Graph/node_math_subtract_40x.png")), Icon40x40));
+					Set("MetasoundEditor.Graph.Node.Math.Modulo", new FSlateImageBrush(RootToContentDir(TEXT("/Graph/node_math_modulo_40x.png")), Icon40x40));
 
 					// Misc
 					Set("MetasoundEditor.Speaker", new FSlateImageBrush(RootToContentDir(TEXT("/Icons/speaker_144x.png")), FVector2D(144.0f, 144.0f)));
diff --git a/Engine/Plugins/Runtime/Metasound/Source/MetasoundStandardNodes/Private/MetasoundMathNodes.cpp b/Engine/Plugins/Runtime/Metasound/Source/MetasoundStandardNodes/Private/MetasoundMathNodes.cpp
--- a/Engine/Plugins/Runtime/Metasound/Source/MetasoundStandardNodes/Private/MetasoundMathNodes.cpp
+++ b/Engine/Plugins/Runtime/Metasound/Source/MetasoundStandardNodes/Private/MetasoundMathNodes.cpp
@@ -29,6 +29,9 @@ namespace Metasound
 	METASOUND_REGISTER_NODE(FDivideInt32Node)
 	METASOUND_REGISTER_NODE(FDivideInt64Node)
 	METASOUND_REGISTER_NODE(FDivideTimeFloatNode)
+
+	METASOUND_REGISTER_NODE(FModuloInt32Node)
+	METASOUND_REGISTER_NODE(FModuloInt64Node)
 }
 
 #undef LOCTEXT_NAMESPACE // MetasoundMathOpNode
diff --git a/Engine/Plugins/Runtime/Metasound/Source/MetasoundStandardNodes/Public/MetasoundMathNodes.h b/Engine/Plugins/Runtime/Metasound/Source/MetasoundStandardNodes/Public/MetasoundMathNodes.h
--- a/Engine/Plugins/Runtime/Metasound/Source/MetasoundStandardNodes/Public/MetasoundMathNodes.h
+++ b/Engine/Plugins/Runtime/Metasound/Source/MetasoundStandardNodes/Public/MetasoundMathNodes.h
@@ -120,6 +120,24 @@ namespace Metasound
 				FText::FromString(GetMetasoundDataTypeString<OperandDataType>()));
 			return DisplayName;
 		}
+
+		template<typename DataType>
+		const FText GetModuloDisplayName()
+		{
+			static const FText DisplayName = FText::Format(LOCTEXT("ModuloNodeDisplayNamePattern", "Modulo ({0})"), FText::FromString(GetMetasoundDataTypeString<DataType>()));
+			return DisplayName;
+		}
+
+		template<typename DataType, typename OperandDataType>
+		const FText GetModuloDisplayName()
+		{
+			static const FText DisplayName = FText::Format(
+				LOCTEXT("ModuloNodeOperandTypedDisplayNamePattern", "Modulo ({0} by {1})"),
+				FText::FromString(GetMetasoundDataTypeString<DataType>()),
+				FText::FromString(GetMetasoundDataTypeString<OperandDataType>()));
+			return DisplayName;
+		}
+
 	}
 
 	template <typename TDerivedClass, typename TMathOpClass, typename TDataClass, typename TOperandDataClass>
@@ -245,8 +263,8 @@ namespace Metasound
 		{
 			static const FVertexInterface Interface(
 				FInputVertexInterface(
-					TInputDataVertexModel<TDataClass>(MathOpNames::PrimaryOperandName, LOCTEXT("MathOpAddendInitialTooltip", "Initial attend."), static_cast<TDataClass>(0)),
-					TInputDataVertexModel<TOperandDataClass>(MathOpNames::AdditionalOperandsName, LOCTEXT("MathOpAddendsTooltip", "Additional attend(s)."), static_cast<TOperandDataClass>(0))
+					TInputDataVertexModel<TDataClass>(MathOpNames::PrimaryOperandName, LOCTEXT("MathOpAddendInitialTooltip", "Initial addend."), static_cast<TDataClass>(0)),
+					TInputDataVertexModel<TOperandDataClass>(MathOpNames::AdditionalOperandsName, LOCTEXT("MathOpAddendsTooltip", "Additional addend(s)."), static_cast<TOperandDataClass>(0))
 				),
 				FOutputVertexInterface(
 					TOutputDataVertexModel<TDataClass>(TEXT("Out"), LOCTEXT("MathOpOutTooltip", "Math operation result"))
@@ -392,8 +410,8 @@ namespace Metasound
 		{
 			static const FVertexInterface Interface(
 				FInputVertexInterface(
-					TInputDataVertexModel<TDataClass>(MathOpNames::PrimaryOperandName, LOCTEXT("MathOpDividendTooltip", "Dividend."), static_cast<TDataClass>(0)),
-					TInputDataVertexModel<TOperandDataClass>(MathOpNames::AdditionalOperandsName, LOCTEXT("MathOpDivisorsTooltip", "Divisor(s)."), static_cast<TOperandDataClass>(0))
+					TInputDataVertexModel<TDataClass>(MathOpNames::PrimaryOperandName, LOCTEXT("MathOpDividendTooltip", "Dividend."), static_cast<TDataClass>(1)),
+					TInputDataVertexModel<TOperandDataClass>(MathOpNames::AdditionalOperandsName, LOCTEXT("MathOpDivisorsTooltip", "Divisor(s)."), static_cast<TOperandDataClass>(1))
 				),
 				FOutputVertexInterface(
 					TOutputDataVertexModel<TDataClass>(TEXT("Out"), LOCTEXT("MathOpOutTooltip", "Math operation result"))
@@ -437,6 +455,66 @@ namespace Metasound
 		}
 	};
 
+	template <typename TDataClass, typename TOperandDataClass = TDataClass>
+	class TMathOpModulo
+	{
+	public:
+		using TDataClassReadRef = TDataReadReference<TDataClass>;
+		using TDataClassWriteRef = TDataWriteReference<TDataClass>;
+
+		using TOperandDataClassReadRef = TDataReadReference<TOperandDataClass>;
+
+		static const FVertexInterface& GetVertexInterface()
+		{
+			static const FVertexInterface Interface(
+				FInputVertexInterface(
+					TInputDataVertexModel<TDataClass>(MathOpNames::PrimaryOperandName, LOCTEXT("MathOpModuloDividendTooltip", "Dividend."), static_cast<TDataClass>(1)),
+					TInputDataVertexModel<TOperandDataClass>(MathOpNames::AdditionalOperandsName, LOCTEXT("MathOpModuloDivisorsTooltip", "Divisor(s)."), static_cast<TOperandDataClass>(1))
+				),
+				FOutputVertexInterface(
+					TOutputDataVertexModel<TDataClass>(TEXT("Out"), LOCTEXT("MathOpOutTooltip", "Math operation result"))
+				)
+			);
+
+			return Interface;
+		}
+
+		static TDataClass GetDefault(const FOperatorSettings& InSettings, const FLiteral& InVertexDefault)
+		{
+			return static_cast<TDataClass>(1);
+		}
+
+		static TOperandDataClass GetDefaultOp(const FOperatorSettings& InSettings, const FLiteral& InVertexDefault)
+		{
+			return static_cast<TOperandDataClass>(1);
+		}
+
+		static void Calculate(TMathOpModulo& InInstanceData, const TDataClassReadRef& InPrimaryOperand, const TArray<TOperandDataClassReadRef>& InAdditionalOperands, TDataClassWriteRef& OutResult)
+		{
+			*OutResult = *InPrimaryOperand;
+
+			if (!ensure(!InAdditionalOperands.IsEmpty()))
+			{
+				return;
+			}
+
+			*OutResult = *InPrimaryOperand;
+
+			// TODO: chaining modulo operations doesn't make too much sense... how do we forbid additional operands in some math ops?
+			for (int32 i = 0; i < InAdditionalOperands.Num(); ++i)
+			{
+				const TDataClass& OperandValue = *InAdditionalOperands[i];
+				if (OperandValue == static_cast<TDataClass>(0))
+				{
+					// TODO: Error here
+					return;
+				}
+
+				*OutResult %= OperandValue;
+			}
+		}
+	};
+
 // Specialized Math Operations
 	template <>
 	class TMathOpAdd<FAudioBuffer>
@@ -1116,11 +1194,11 @@ namespace Metasound
 		{
 			static const FVertexInterface Interface(
 				FInputVertexInterface(
-					TInputDataVertexModel<int64>(MathOpNames::PrimaryOperandName, LOCTEXT("MathOpRandomMinTooltip", "Random value minimum."), 1),
-					TInputDataVertexModel<int64>(MathOpNames::AdditionalOperandsName, LOCTEXT("MathOpRandomMaxTooltip", "Random value maximum."), 1)
+					TInputDataVertexModel<int64>(MathOpNames::PrimaryOperandName, LOCTEXT("MathOpDivideDividendTooltip", "Dividend of operation."), 1),
+					TInputDataVertexModel<int64>(MathOpNames::AdditionalOperandsName, LOCTEXT("MathOpDivideDivisorTooltip", "Divisor of operation."), 1)
 				),
 				FOutputVertexInterface(
-					TOutputDataVertexModel<int64>(TEXT("Out"), LOCTEXT("MathOpRandomOutTooltip", "Resulting random value"))
+					TOutputDataVertexModel<int64>(TEXT("Out"), LOCTEXT("MathOpDivideOutTooltip", "Resulting value"))
 				)
 			);
 
@@ -1129,11 +1207,13 @@ namespace Metasound
 
 		static int64 GetDefault(const FOperatorSettings& InSettings, const FLiteral& InVertexDefault)
 		{
+			// TODO: we are truncating to int32, why do we need int64?
 			return static_cast<int64>(InVertexDefault.Value.Get<int32>());
 		}
 
 		static int64 GetDefaultOp(const FOperatorSettings& InSettings, const FLiteral& InVertexDefault)
 		{
+			// TODO: we are truncating to int32, why do we need in64?
 			return static_cast<int64>(InVertexDefault.Value.Get<int32>());
 		}
 
@@ -1161,6 +1241,59 @@ namespace Metasound
 	};
 
 
+	// int64 has to be explicitly implemented as they are PODs that are not implemented as a literal base type.
+	template <>
+	class TMathOpModulo<int64>
+	{
+	public:
+		static const FVertexInterface& GetVertexInterface()
+		{
+			static const FVertexInterface Interface(
+				FInputVertexInterface(
+					TInputDataVertexModel<int64>(MathOpNames::PrimaryOperandName, LOCTEXT("MathOpModuloDividendTooltip", "Dividend of modulo operation."), 1),
+					TInputDataVertexModel<int64>(MathOpNames::AdditionalOperandsName, LOCTEXT("MathOpModuluDivisorTooltip", "Divisor of modulo operation."), 1)
+				),
+				FOutputVertexInterface(
+					TOutputDataVertexModel<int64>(TEXT("Out"), LOCTEXT("MathOpModuloOutTooltip", "Resulting value"))
+				)
+			);
+
+			return Interface;
+		}
+
+		static int64 GetDefault(const FOperatorSettings& InSettings, const FLiteral& InVertexDefault)
+		{
+			return static_cast<int64>(InVertexDefault.Value.Get<int32>());
+		}
+
+		static int64 GetDefaultOp(const FOperatorSettings& InSettings, const FLiteral& InVertexDefault)
+		{
+			return static_cast<int64>(InVertexDefault.Value.Get<int32>());
+		}
+
+		static void Calculate(TMathOpModulo& InInstanceData, const FInt64ReadRef& InPrimaryOperand, const TArray<FInt64ReadRef>& InAdditionalOperands, FInt64WriteRef& OutResult)
+		{
+			*OutResult = *InPrimaryOperand;
+
+			if (!ensure(!InAdditionalOperands.IsEmpty()))
+			{
+				return;
+			}
+
+			for (int32 i = 0; i < InAdditionalOperands.Num(); ++i)
+			{
+				const int64 OperandValue = *InAdditionalOperands[i];
+				if (OperandValue == 0)
+				{
+					// TODO: Error here
+					continue;
+				}
+
+				*OutResult %= OperandValue;
+			}
+		}
+	};
+
 	// Definitions
 	DEFINE_METASOUND_MATHOP(Add, Float, float, LOCTEXT("Metasound_MathAddFloatNodeDescription", "Adds floats."))
 	DEFINE_METASOUND_MATHOP(Add, Int32, int32, LOCTEXT("Metasound_MathAddInt32NodeDescription", "Adds int32s."))
@@ -1189,6 +1322,10 @@ namespace Metasound
  	DEFINE_METASOUND_MATHOP(Divide, Int64, int64, LOCTEXT("Metasound_MathDivideInt64NodeDescription", "Divide int64 by another int64."))
 
 	DEFINE_METASOUND_OPERAND_TYPED_MATHOP(Divide, Time, FTime, Float, float, LOCTEXT("Metasound_MathDivideTimeNodeDescription", "Divides time by floats."))
+
+	DEFINE_METASOUND_MATHOP(Modulo, Int32, int32, LOCTEXT("Metasound_MathModulusInt32NodeDescription", "Modulo int32 by another int32."))
+	DEFINE_METASOUND_MATHOP(Modulo, Int64, int64, LOCTEXT("Metasound_MathModulusInt64NodeDescription", "Modulo int64 by another int64."))
 }
 
+
 #undef LOCTEXT_NAMESPACE // MetasoundMathOpNode
