Fix for static init issu w/ FNames in metasound node registration

diff --git a/Engine/Plugins/Runtime/Metasound/Source/MetasoundEngine/Private/MetasoundEngineModule.cpp b/Engine/Plugins/Runtime/Metasound/Source/MetasoundEngine/Private/MetasoundEngineModule.cpp
--- a/Engine/Plugins/Runtime/Metasound/Source/MetasoundEngine/Private/MetasoundEngineModule.cpp
+++ b/Engine/Plugins/Runtime/Metasound/Source/MetasoundEngine/Private/MetasoundEngineModule.cpp
@@ -23,6 +23,9 @@ class FMetasoundEngineModule : public IMetasoundEngineModule
 		// Register preferred archetypes
 		Metasound::IMetasoundUObjectRegistry::RegisterUClassPreferredArchetypes<UMetasoundSource>();
 
+		// flush node regsitration queue
+		FMetasoundFrontendRegistryContainer::Get()->RegisterPendingNodes();
+
 		FModuleManager::Get().LoadModuleChecked(TEXT("AudioCodecEngine"));
 
 		UE_LOG(LogMetasoundEngine, Log, TEXT("Metasound Engine Initialized"));
diff --git a/Engine/Plugins/Runtime/Metasound/Source/MetasoundFrontend/Private/MetasoundFrontend.cpp b/Engine/Plugins/Runtime/Metasound/Source/MetasoundFrontend/Private/MetasoundFrontend.cpp
--- a/Engine/Plugins/Runtime/Metasound/Source/MetasoundFrontend/Private/MetasoundFrontend.cpp
+++ b/Engine/Plugins/Runtime/Metasound/Source/MetasoundFrontend/Private/MetasoundFrontend.cpp
@@ -190,7 +190,7 @@ class FMetasoundFrontendModule : public IModuleInterface
 {
 	virtual void StartupModule() override
 	{
-		Metasound::Frontend::InitializeFrontend();
+		Metasound::Frontend::RegisterPendingNodes();
 	}
 };
 
diff --git a/Engine/Plugins/Runtime/Metasound/Source/MetasoundFrontend/Private/MetasoundFrontendBaseClasses.cpp b/Engine/Plugins/Runtime/Metasound/Source/MetasoundFrontend/Private/MetasoundFrontendBaseClasses.cpp
--- a/Engine/Plugins/Runtime/Metasound/Source/MetasoundFrontend/Private/MetasoundFrontendBaseClasses.cpp
+++ b/Engine/Plugins/Runtime/Metasound/Source/MetasoundFrontend/Private/MetasoundFrontendBaseClasses.cpp
@@ -700,9 +700,9 @@ namespace Metasound
 
 		
 
-		void InitializeFrontend()
+		void RegisterPendingNodes()
 		{
-			FMetasoundFrontendRegistryContainer::Get()->InitializeFrontend();
+			FMetasoundFrontendRegistryContainer::Get()->RegisterPendingNodes();
 		}
 	}
 }
diff --git a/Engine/Plugins/Runtime/Metasound/Source/MetasoundFrontend/Private/MetasoundFrontendRegistries.cpp b/Engine/Plugins/Runtime/Metasound/Source/MetasoundFrontend/Private/MetasoundFrontendRegistries.cpp
--- a/Engine/Plugins/Runtime/Metasound/Source/MetasoundFrontend/Private/MetasoundFrontendRegistries.cpp
+++ b/Engine/Plugins/Runtime/Metasound/Source/MetasoundFrontend/Private/MetasoundFrontendRegistries.cpp
@@ -130,7 +130,6 @@ namespace Metasound
 }
 
 FMetasoundFrontendRegistryContainer::FMetasoundFrontendRegistryContainer()
-	: bHasModuleBeenInitialized(false)
 {
 }
 
@@ -155,17 +154,11 @@ void FMetasoundFrontendRegistryContainer::ShutdownMetasoundFrontend()
 	}
 }
 
-void FMetasoundFrontendRegistryContainer::InitializeFrontend()
+void FMetasoundFrontendRegistryContainer::RegisterPendingNodes()
 {
 	FScopeLock ScopeLock(&LazyInitCommandCritSection);
 
-	if (bHasModuleBeenInitialized)
-	{
-		// this function should only be called once.
-		checkNoEntry();
-	}
-
-	UE_LOG(LogTemp, Display, TEXT("Initializing Metasounds Frontend."));
+	UE_LOG(LogTemp, Display, TEXT("Processing %i Metasounds Frontend Registration Requests."), LazyInitCommands.Num());
 	uint64 CurrentTime = FPlatformTime::Cycles64();
 
 	for (TUniqueFunction<void()>& Command : LazyInitCommands)
@@ -174,7 +167,6 @@ void FMetasoundFrontendRegistryContainer::InitializeFrontend()
 	}
 
 	LazyInitCommands.Empty();
-	bHasModuleBeenInitialized = true;
 
 	uint64 CyclesUsed = FPlatformTime::Cycles64() - CurrentTime;
 	UE_LOG(LogTemp, Display, TEXT("Initializing Metasounds Frontend took %f seconds."), FPlatformTime::ToSeconds64(CyclesUsed));
@@ -182,14 +174,7 @@ void FMetasoundFrontendRegistryContainer::InitializeFrontend()
 
 bool FMetasoundFrontendRegistryContainer::EnqueueInitCommand(TUniqueFunction<void()>&& InFunc)
 {
-	// if the module has been initalized already, we can safely call this function now.
-	if (bHasModuleBeenInitialized)
-	{
-		InFunc();
-	}
-
-	// otherwise, we enqueue the function to be executed after the frontend module has been initialized.
-
+	FScopeLock ScopeLock(&LazyInitCommandCritSection);
 	if (LazyInitCommands.Num() >= MaxNumNodesAndDatatypesToInitialize)
 	{
 		UE_LOG(LogTemp, Warning, TEXT("Registering more that %d nodes and datatypes for metasounds! Consider increasing MetasoundFrontendRegistryContainer::MaxNumNodesAndDatatypesToInitialize."));
diff --git a/Engine/Plugins/Runtime/Metasound/Source/MetasoundFrontend/Public/MetasoundFrontendBaseClasses.h b/Engine/Plugins/Runtime/Metasound/Source/MetasoundFrontend/Public/MetasoundFrontendBaseClasses.h
--- a/Engine/Plugins/Runtime/Metasound/Source/MetasoundFrontend/Public/MetasoundFrontendBaseClasses.h
+++ b/Engine/Plugins/Runtime/Metasound/Source/MetasoundFrontend/Public/MetasoundFrontendBaseClasses.h
@@ -13,7 +13,7 @@ namespace Metasound
 	namespace Frontend
 	{
 		// This is called by FMetasoundFrontendModule, and flushes any node or datatype registration that was done prior to boot.
-		void InitializeFrontend();
+		void RegisterPendingNodes();
 
 		/** Use this to get a reference to the static registry of all nodes implemented in C++. */
 		METASOUNDFRONTEND_API TMap<FNodeRegistryKey, FNodeRegistryElement>& GetExternalNodeRegistry();
diff --git a/Engine/Plugins/Runtime/Metasound/Source/MetasoundFrontend/Public/MetasoundFrontendRegistries.h b/Engine/Plugins/Runtime/Metasound/Source/MetasoundFrontend/Public/MetasoundFrontendRegistries.h
--- a/Engine/Plugins/Runtime/Metasound/Source/MetasoundFrontend/Public/MetasoundFrontendRegistries.h
+++ b/Engine/Plugins/Runtime/Metasound/Source/MetasoundFrontend/Public/MetasoundFrontendRegistries.h
@@ -266,7 +266,7 @@ public:
 	FMetasoundFrontendRegistryContainer(const FMetasoundFrontendRegistryContainer&) = delete;
 
 	// This is called on module startup. This invokes any registration commands enqueued by our registration macros.
-	void InitializeFrontend();
+	void RegisterPendingNodes();
 
 
 	bool EnqueueInitCommand(TUniqueFunction<void()>&& InFunc);
@@ -349,7 +349,6 @@ private:
 	TArray<TUniqueFunction<void()>, TInlineAllocator<MaxNumNodesAndDatatypesToInitialize>> LazyInitCommands;
 	
 	FCriticalSection LazyInitCommandCritSection;
-	bool bHasModuleBeenInitialized;
 
 	// Registry in which we keep all information about nodes implemented in C++.
 	TMap<FNodeRegistryKey, FNodeRegistryElement> ExternalNodeRegistry;
diff --git a/Engine/Plugins/Runtime/Metasound/Source/MetasoundStandardNodes/Private/MetasoundStandardNodesModule.cpp b/Engine/Plugins/Runtime/Metasound/Source/MetasoundStandardNodes/Private/MetasoundStandardNodesModule.cpp
--- a/Engine/Plugins/Runtime/Metasound/Source/MetasoundStandardNodes/Private/MetasoundStandardNodesModule.cpp
+++ b/Engine/Plugins/Runtime/Metasound/Source/MetasoundStandardNodes/Private/MetasoundStandardNodesModule.cpp
@@ -14,6 +14,10 @@ namespace Metasound
 {
 	class FMetasoundStandardNodesModule : public IMetasoundGraphCoreModule
 	{
+		virtual void StartupModule() override
+		{
+			FMetasoundFrontendRegistryContainer::Get()->RegisterPendingNodes();
+		}
 	};
 }
 
