Value node

* allows setting a metasound type on trigger.
* Effectively a parameter block-rate sample and hold

diff --git a/Engine/Plugins/Runtime/Metasound/Source/MetasoundStandardNodes/Private/MetasoundValueNode.cpp b/Engine/Plugins/Runtime/Metasound/Source/MetasoundStandardNodes/Private/MetasoundValueNode.cpp
new file mode 100644
--- /dev/null
+++ b/Engine/Plugins/Runtime/Metasound/Source/MetasoundStandardNodes/Private/MetasoundValueNode.cpp
@@ -0,0 +1,24 @@
+// Copyright Epic Games, Inc. All Rights Reserved.
+
+#include "MetasoundValueNode.h"
+
+#include "MetasoundNodeRegistrationMacro.h"
+
+#define LOCTEXT_NAMESPACE "MetasoundStandardNodes"
+
+namespace Metasound
+{
+	using FValueNodeInt32 = TValueNode<int32>;
+ 	METASOUND_REGISTER_NODE(FValueNodeInt32)
+
+	using FValueNodeFloat = TValueNode<float>;
+	METASOUND_REGISTER_NODE(FValueNodeFloat)
+
+	using FValueNodeBool = TValueNode<bool>;
+	METASOUND_REGISTER_NODE(FValueNodeBool)
+
+	using FValueNodeString = TValueNode<FString>;
+	METASOUND_REGISTER_NODE(FValueNodeString)
+}
+
+#undef LOCTEXT_NAMESPACE
diff --git a/Engine/Plugins/Runtime/Metasound/Source/MetasoundStandardNodes/Public/MetasoundValueNode.h b/Engine/Plugins/Runtime/Metasound/Source/MetasoundStandardNodes/Public/MetasoundValueNode.h
new file mode 100644
--- /dev/null
+++ b/Engine/Plugins/Runtime/Metasound/Source/MetasoundStandardNodes/Public/MetasoundValueNode.h
@@ -0,0 +1,230 @@
+// Copyright Epic Games, Inc. All Rights Reserved.
+
+#pragma once
+
+#include "CoreMinimal.h"
+#include "MetasoundFacade.h"
+#include "MetasoundExecutableOperator.h"
+#include "MetasoundPrimitives.h"
+#include "MetasoundStandardNodesNames.h"
+#include "MetasoundTrigger.h"
+#include "Internationalization/Text.h"
+
+#define LOCTEXT_NAMESPACE "MetasoundStandardNodes"
+
+namespace Metasound
+{
+	template<typename ValueType>
+	struct TValueDisplayName
+	{
+		bool bIsSupported = false;
+	};
+
+	template<>
+	struct TValueDisplayName<int32>
+	{
+		static FText GetDisplayName()
+		{
+			return LOCTEXT("IntValueOpDisplayName", "Int Value");
+		}
+
+		static FText GetDescription()
+		{
+			return LOCTEXT("ValueOpDescription", "Allows setting an int value to the output on trigger.");
+		}
+	};
+
+	template<>
+	struct TValueDisplayName<float>
+	{
+		static FText GetDisplayName()
+		{
+			return LOCTEXT("FloatValueOpDisplayName", "Float Value");
+		}
+
+		static FText GetDescription()
+		{
+			return LOCTEXT("FloatValueOpDisplayName", "Allows setting a float value to the output on trigger.");
+		}
+	};
+
+	template<>
+	struct TValueDisplayName<bool>
+	{
+		static FText GetDisplayName()
+		{
+			return LOCTEXT("BoolValueOpDisplayName", "Bool Value");
+		}
+
+		static FText GetDescription()
+		{
+			return LOCTEXT("BoolValueOpDisplayName", "Allows setting a bool value to the output on trigger.");
+		}
+	};
+
+	template<>
+	struct TValueDisplayName<FString>
+	{
+		static FText GetDisplayName()
+		{
+			return LOCTEXT("StringValueOpDisplayName", "String Value");
+		}
+
+		static FText GetDescription()
+		{
+			return LOCTEXT("StringValueOpDisplayName", "Allows setting a string value to the output on trigger.");
+		}
+	};
+
+	template<typename ValueType>
+	class TValueOperator : public TExecutableOperator<TValueOperator<ValueType>>
+	{
+	public:
+		using FArrayDataReadReference = TDataReadReference<ValueType>;
+
+		static const FString& GetInitValueName()
+		{
+			static const FString Name = TEXT("Init");
+			return Name;
+		}
+
+		static const FString& GetSetValueName()
+		{
+			static const FString Name = TEXT("Set");
+			return Name;
+		}
+
+		static const FString& GetInputTriggerName()
+		{
+			static const FString Name = TEXT("Trigger");
+			return Name;
+		}
+
+		static const FString& GetOutputValueName()
+		{
+			static const FString Name = TEXT("Output");
+			return Name;
+		}
+
+		static const FVertexInterface& GetDefaultInterface()
+		{
+			static const FVertexInterface DefaultInterface(
+				FInputVertexInterface(
+					TInputDataVertexModel<ValueType>(GetInitValueName(), LOCTEXT("InitValue", "Value to init the output to.")),
+					TInputDataVertexModel<ValueType>(GetSetValueName(), LOCTEXT("ValueInput", "Value to set the output to when triggered.")),
+					TInputDataVertexModel<FTrigger>(GetInputTriggerName(), LOCTEXT("ValueTrigger", "Trigger to write the set value to the output."))
+				),
+				FOutputVertexInterface(
+					TOutputDataVertexModel<ValueType>(GetOutputValueName(), LOCTEXT("ValueOutput", "The output value cached in the node."))
+				)
+			);
+
+			return DefaultInterface;
+		}
+
+		static const FNodeClassMetadata& GetNodeInfo()
+		{
+			auto CreateNodeClassMetadata = []() -> FNodeClassMetadata
+			{
+				FNodeClassMetadata Metadata
+				{
+					FNodeClassName{FName("Value"), GetMetasoundDataTypeName<ValueType>(), TEXT("") },
+					1, // Major Version
+					0, // Minor Version
+					TValueDisplayName<ValueType>::GetDisplayName(),
+					LOCTEXT("ValueOpDescription", "Allows setting a value to output on trigger."),
+					PluginAuthor,
+					PluginNodeMissingPrompt,
+					GetDefaultInterface(),
+					{LOCTEXT("ValueCategory", "Value")},
+					{TEXT("Value")},
+					FNodeDisplayStyle{}
+				};
+
+				return Metadata;
+			};
+
+			static const FNodeClassMetadata Metadata = CreateNodeClassMetadata();
+			return Metadata;
+		}
+
+		static TUniquePtr<IOperator> CreateOperator(const FCreateOperatorParams& InParams, TArray<TUniquePtr<IOperatorBuildError>>& OutErrors)
+		{
+			const FInputVertexInterface& InputInterface = InParams.Node.GetVertexInterface().GetInputInterface();
+			const FDataReferenceCollection& InputCollection = InParams.InputDataReferences;
+
+			FTriggerReadRef Trigger = InputCollection.GetDataReadReferenceOrConstruct<FTrigger>(GetInputTriggerName(), InParams.OperatorSettings);
+			TDataReadReference<ValueType> InitValue = InputCollection.GetDataReadReferenceOrConstructWithVertexDefault<ValueType>(InputInterface, GetInitValueName());
+			TDataReadReference<ValueType> SetValue = InputCollection.GetDataReadReferenceOrConstructWithVertexDefault<ValueType>(InputInterface, GetSetValueName());
+
+			return MakeUnique<TValueOperator<ValueType>>(InParams.OperatorSettings, Trigger, InitValue, SetValue);
+		}
+
+
+		TValueOperator(const FOperatorSettings& InSettings, TDataReadReference<FTrigger> InTrigger, TDataReadReference<ValueType> InInitValue, TDataReadReference<ValueType> InSetValue)
+			: Trigger(InTrigger)
+			, InitValue(InInitValue)
+			, SetValue(InSetValue)
+			, OutputValue(TDataWriteReferenceFactory<ValueType>::CreateAny(InSettings))
+		{
+			*OutputValue = *InitValue;
+		}
+
+		virtual ~TValueOperator() = default;
+
+
+		virtual FDataReferenceCollection GetInputs() const override
+		{
+			FDataReferenceCollection Inputs;
+
+			Inputs.AddDataReadReference(GetInputTriggerName(), Trigger);
+			Inputs.AddDataReadReference(GetInitValueName(), InitValue);
+			Inputs.AddDataReadReference(GetSetValueName(), SetValue);
+
+			return Inputs;
+		}
+
+		virtual FDataReferenceCollection GetOutputs() const override
+		{
+			FDataReferenceCollection Outputs;
+
+			Outputs.AddDataReadReference(GetOutputValueName(), OutputValue);
+
+			return Outputs;
+		}
+
+		void Execute()
+		{
+			if (*Trigger)
+			{
+				*OutputValue = *SetValue;
+			}
+		}
+
+	private:
+
+		TDataReadReference<FTrigger> Trigger;
+		TDataReadReference<ValueType> InitValue;
+		TDataReadReference<ValueType> SetValue;
+		TDataWriteReference<ValueType> OutputValue;
+	};
+
+	/** TValueNode
+	 *
+	 *  Generates a random float value when triggered.
+	 */
+	template<typename ValueType>
+	class METASOUNDSTANDARDNODES_API TValueNode : public FNodeFacade
+	{
+	public:
+		/**
+		 * Constructor used by the Metasound Frontend.
+		 */
+		TValueNode(const FNodeInitData& InInitData)
+			: FNodeFacade(InInitData.InstanceName, InInitData.InstanceID, TFacadeOperatorClass<TValueOperator<ValueType>>())
+		{}
+
+		virtual ~TValueNode() = default;
+	};
+} // namespace Metasound
+
