Metasound Noise Generators (pink, white)

diff --git a/Engine/Plugins/Runtime/Metasound/Source/MetasoundStandardNodes/Private/MetasoundNoiseGenerator.cpp b/Engine/Plugins/Runtime/Metasound/Source/MetasoundStandardNodes/Private/MetasoundNoiseGenerator.cpp
new file mode 100644
--- /dev/null
+++ b/Engine/Plugins/Runtime/Metasound/Source/MetasoundStandardNodes/Private/MetasoundNoiseGenerator.cpp
@@ -0,0 +1,228 @@
+// Copyright Epic Games, Inc. All Rights Reserved.
+#include "MetasoundNoiseGenerator.h"
+
+#include "MetasoundAudioBuffer.h"
+#include "MetasoundExecutableOperator.h"
+#include "MetasoundFacade.h"
+#include "MetasoundNodeRegistrationMacro.h"
+#include "MetasoundDataTypeRegistrationMacro.h"
+#include "MetasoundOperatorSettings.h"
+#include "MetasoundPrimitives.h"
+#include "MetasoundStandardNodesNames.h"
+#include "MetasoundTrigger.h"
+#include "MetasoundVertex.h"
+#include "DSP/Noise.h"
+
+#define LOCTEXT_NAMESPACE "MetasoundStandardNodes"
+
+namespace Metasound
+{
+	enum class ENoiseType
+	{
+		Pink,
+		White
+	};
+
+	DECLARE_METASOUND_ENUM(ENoiseType, ENoiseType::Pink, METASOUNDSTANDARDNODES_API,
+		FEnumNoiseType, FEnumNoiseTypeInfo, FEnumNoiseTypeReadRef, FEnumNoiseTypeWriteRef);
+
+	DEFINE_METASOUND_ENUM_BEGIN(ENoiseType, FEnumNoiseType, "NoiseType")
+		DEFINE_METASOUND_ENUM_ENTRY(ENoiseType::Pink, LOCTEXT("PinkDescription", "Pink Noise"), LOCTEXT("PinkDescriptionTT", "Pink noise or ​1⁄f noise, spectral density is inversely proportional to the frequency of the signal")),
+		DEFINE_METASOUND_ENUM_ENTRY(ENoiseType::White, LOCTEXT("WhiteDescription", "White Noise"), LOCTEXT("WhiteDescriptionTT", "A random signal having equal intensity at different frequencies")),
+	DEFINE_METASOUND_ENUM_END()
+
+	class FNoiseOperator : public IOperator
+	{
+	public:
+		static constexpr int32 DefaultSeed = INDEX_NONE;
+
+		static TUniquePtr<IOperator> CreateOperator(const FCreateOperatorParams& InParams, FBuildErrorArray& OutErrors);
+		static const FNodeClassMetadata& GetNodeInfo();
+		static FVertexInterface DeclareVertexInterface();
+
+		FNoiseOperator(const FOperatorSettings& InSettings, FInt32ReadRef&& InReadRef);
+		virtual FDataReferenceCollection GetInputs() const override;
+		virtual FDataReferenceCollection GetOutputs() const override;
+
+	protected:
+		FInt32ReadRef Seed;
+		FAudioBufferWriteRef Out;
+		int32 OldSeed;
+
+		template<typename T>
+		static T MakeGenerator(int32 InSeed)
+		{
+			// If we the default seed, use the default noise generator constructor.
+			if (InSeed == DefaultSeed)
+			{
+				return T{};
+			}
+			else
+			{
+				return T{ InSeed };
+			}
+		}
+
+		template<typename T>
+		FORCEINLINE void CheckAndReseed(T& InOutGenerator)
+		{
+			// Reseed?
+			int32 NewSeed = *Seed;
+			if (OldSeed != NewSeed)
+			{
+				InOutGenerator = MakeGenerator<T>(NewSeed);
+				OldSeed = NewSeed;
+			}
+		}
+
+		template<typename T>
+		FORCEINLINE void Generate(T& InGenerator)
+		{
+			float* WritePtr = Out->GetData();
+			for (int32 i = Out->Num(); i > 0; --i)
+			{
+				*WritePtr++ = InGenerator.Generate();
+			}
+		}
+
+		static constexpr const TCHAR* OutPinName = TEXT("Audio");
+		static constexpr const TCHAR* TriggerTypePinName = TEXT("Type");
+		static constexpr const TCHAR* SeedPinName = TEXT("Seed");
+	};
+
+	struct FNoiseOperator_White final : public FNoiseOperator
+	{
+		Audio::FWhiteNoise Generator;
+
+		FNoiseOperator_White(const FOperatorSettings& InSettings, FInt32ReadRef&& InSeed)
+			: FNoiseOperator{ InSettings, MoveTemp(InSeed) }
+			, Generator{ MakeGenerator<Audio::FWhiteNoise>(*Seed) }
+		{}
+
+		void Execute()
+		{
+			// Reseed if necessary.
+			CheckAndReseed(Generator);
+
+			// Generate a block.
+			Generate(Generator);
+		}
+		static void ExecuteFunction(IOperator* InOperator) { static_cast<FNoiseOperator_White*>(InOperator)->Execute(); }
+		FExecuteFunction GetExecuteFunction() override { return &FNoiseOperator_White::ExecuteFunction; }
+	};
+
+	struct FNoiseOperator_Pink final : public FNoiseOperator
+	{
+		using FNoiseOperator::FNoiseOperator;
+		Audio::FPinkNoise Generator;
+
+		FNoiseOperator_Pink(const FOperatorSettings& InSettings, FInt32ReadRef&& InSeed)
+			: FNoiseOperator{ InSettings, MoveTemp(InSeed) }
+			, Generator{ MakeGenerator<Audio::FPinkNoise>(*Seed) }
+		{}
+
+		static void ExecuteFunction(IOperator* InOperator) { static_cast<FNoiseOperator_Pink*>(InOperator)->Execute(); }
+		FExecuteFunction GetExecuteFunction() override { return &FNoiseOperator_Pink::ExecuteFunction; }
+
+		void Execute()
+		{
+			// Reseed if necessary.
+			CheckAndReseed(Generator);
+
+			// Generate a block.
+			Generate(Generator);
+		}
+	};
+	
+	FNoiseOperator::FNoiseOperator(const FOperatorSettings& InSettings, FInt32ReadRef&& InReadRef)
+		: Seed{ MoveTemp(InReadRef) }
+		, Out{ FAudioBufferWriteRef::CreateNew(InSettings) }
+		, OldSeed(*Seed)
+	{}
+
+	FDataReferenceCollection FNoiseOperator::GetInputs() const
+	{
+		FDataReferenceCollection InputDataReferences;
+		InputDataReferences.AddDataReadReference(SeedPinName, FInt32ReadRef(Seed));
+		return InputDataReferences;
+	}
+
+	FDataReferenceCollection FNoiseOperator::GetOutputs() const
+	{
+		FDataReferenceCollection OutputDataReferences;
+		OutputDataReferences.AddDataReadReference(OutPinName, FAudioBufferReadRef(Out));
+		return OutputDataReferences;
+	}
+
+	FVertexInterface FNoiseOperator::DeclareVertexInterface()
+	{
+		static const FVertexInterface Interface(
+			FInputVertexInterface(
+				TInputDataVertexModel<int32>(SeedPinName, LOCTEXT("SeedPinDescription", "Seed for seeding the Random Number Generator, -1 (default) will use current time"), FNoiseOperator::DefaultSeed),
+				TInputDataVertexModel<FEnumNoiseType>(TriggerTypePinName, LOCTEXT("NoiseTypeDescription", "Type of Noise to Generate"))
+			),
+			FOutputVertexInterface(
+				TOutputDataVertexModel<FAudioBuffer>(OutPinName, LOCTEXT("AudioOutDescription", "Audio output"))
+			)
+		);
+		return Interface;
+	}
+
+	const FNodeClassMetadata& FNoiseOperator::GetNodeInfo()
+	{
+		auto InitNodeInfo = []() -> FNodeClassMetadata
+		{
+			FNodeClassMetadata Info;
+			Info.ClassName = { Metasound::StandardNodes::Namespace, TEXT("Noise"), Metasound::StandardNodes::AudioVariant };
+			Info.MajorVersion = 1;
+			Info.MinorVersion = 0;
+			Info.DisplayName = LOCTEXT("Metasound_NoiseNodeDisplayName", "Noise Generator");
+			Info.Description = LOCTEXT("Metasound_NoiseNodeDescription", "Noise Generated that produces differnt types of noise");
+			Info.Author = PluginAuthor;
+			Info.PromptIfMissing = PluginNodeMissingPrompt;
+			Info.DefaultInterface = DeclareVertexInterface();
+
+			return Info;
+		};
+		static const FNodeClassMetadata Info = InitNodeInfo();
+		return Info;
+	}
+
+	FNoiseNode::FNoiseNode(const FString& InName, const FGuid& InInstanceID, int32 InDefaultSeed)
+		: FNodeFacade(InName, InInstanceID, TFacadeOperatorClass<FNoiseOperator>())
+		, DefaultSeed(InDefaultSeed)
+	{}
+
+	FNoiseNode::FNoiseNode(const FNodeInitData& InInitData)
+		: FNoiseNode(InInitData.InstanceName, InInitData.InstanceID, FNoiseOperator::DefaultSeed)
+	{}
+
+	TUniquePtr<IOperator> FNoiseOperator::CreateOperator(const FCreateOperatorParams& InParams, FBuildErrorArray& OutErrors)
+	{
+		const FNoiseNode& Node = static_cast<const FNoiseNode&>(InParams.Node);
+		const FDataReferenceCollection& InputCol = InParams.InputDataReferences;
+		const FOperatorSettings& Settings = InParams.OperatorSettings;
+		const FInputVertexInterface& InputInterface = DeclareVertexInterface().GetInputInterface();
+
+		// Static property pin, only used for factory.
+		FEnumNoiseTypeReadRef Type = InputCol.GetDataReadReferenceOrConstruct<FEnumNoiseType>(TriggerTypePinName);
+
+		// Seed.
+		FInt32ReadRef Seed = InputCol.GetDataReadReferenceOrConstruct<int32>(TriggerTypePinName, Node.GetDefaultSeed());
+
+		switch (*Type)
+		{
+		default:
+		case ENoiseType::White:
+			return MakeUnique<FNoiseOperator_White>(InParams.OperatorSettings, MoveTemp(Seed));
+		case ENoiseType::Pink:
+			return MakeUnique<FNoiseOperator_Pink>(InParams.OperatorSettings, MoveTemp(Seed));
+		}
+		checkNoEntry();
+		return nullptr;
+	}
+
+	METASOUND_REGISTER_NODE(FNoiseNode);
+}
+#undef LOCTEXT_NAMESPACE //MetasoundStandardNodes
+
diff --git a/Engine/Plugins/Runtime/Metasound/Source/MetasoundStandardNodes/Public/MetasoundNoiseGenerator.h b/Engine/Plugins/Runtime/Metasound/Source/MetasoundStandardNodes/Public/MetasoundNoiseGenerator.h
new file mode 100644
--- /dev/null
+++ b/Engine/Plugins/Runtime/Metasound/Source/MetasoundStandardNodes/Public/MetasoundNoiseGenerator.h
@@ -0,0 +1,20 @@
+// Copyright Epic Games, Inc. All Rights Reserved.
+
+#pragma once
+
+#include "CoreMinimal.h"
+#include "MetasoundFacade.h"
+
+namespace Metasound
+{
+	class METASOUNDSTANDARDNODES_API FNoiseNode : public FNodeFacade
+	{
+	public:
+		FNoiseNode(const FString& InName, const FGuid& InInstanceID, int32 InDefaultSeed);
+		FNoiseNode(const FNodeInitData& InInitData);
+
+		int32 GetDefaultSeed() const { return DefaultSeed; }
+	private:
+		int32 DefaultSeed;
+	};
+}
diff --git a/Engine/Source/Runtime/AudioMixer/Private/AudioMixerDevice.cpp b/Engine/Source/Runtime/AudioMixer/Private/AudioMixerDevice.cpp
--- a/Engine/Source/Runtime/AudioMixer/Private/AudioMixerDevice.cpp
+++ b/Engine/Source/Runtime/AudioMixer/Private/AudioMixerDevice.cpp
@@ -2405,17 +2405,14 @@ namespace Audio
 		const int32 NumFrames = OpenStreamParams.NumFrames;
 		const int32 NumChannels = PlatformInfo.NumChannels;
 
-		static FWhiteNoise WhiteNoise(DebugGeneratorAmpCVar);
-
-		WhiteNoise.SetScaleAdd(DebugGeneratorAmpCVar, 0.0f);
-
+		static FWhiteNoise WhiteNoise;
 
 		for (int32 FrameIndex = 0; FrameIndex < NumFrames; ++FrameIndex)
 		{
 			for (int32 ChannelIndex = 0; ChannelIndex < NumChannels; ++ChannelIndex)
 			{
 				int32 Index = FrameIndex * NumChannels + ChannelIndex;
-				Output[Index] += WhiteNoise.Generate();
+				Output[Index] += WhiteNoise.Generate(DebugGeneratorAmpCVar, 0.f);
 			}
 		}
 	}
diff --git a/Engine/Source/Runtime/SignalProcessing/Private/Noise.cpp b/Engine/Source/Runtime/SignalProcessing/Private/Noise.cpp
--- a/Engine/Source/Runtime/SignalProcessing/Private/Noise.cpp
+++ b/Engine/Source/Runtime/SignalProcessing/Private/Noise.cpp
@@ -1,82 +1,69 @@
 // Copyright Epic Games, Inc. All Rights Reserved.
 
 #include "DSP/Noise.h"
+#include "DSP/Dsp.h"
+#include "HAL/PlatformTime.h"
 
 namespace Audio
 {
-	FWhiteNoise::FWhiteNoise(const float InScale, const float InAdd)
-	{
-		SetScaleAdd(InScale, InAdd);
-	}
-
-	void FWhiteNoise::SetScaleAdd(const float InScale, const float InAdd)
-	{
-		if (Scale == InScale && Add == InAdd)
-		{
-			return;
-		}
-
-		Scale = InScale;
-		Add = InAdd;
+	FWhiteNoise::FWhiteNoise(int32 InRandomSeed)
+		: RandomStream{ InRandomSeed }
+	{}
 
-		RandomStream.Initialize(HashCombine(GetTypeHash(Scale), GetTypeHash(Add)));
-	}
+	FWhiteNoise::FWhiteNoise()
+		: FWhiteNoise{ static_cast<int32>(FPlatformTime::Cycles()) }
+	{}
 
-	float FWhiteNoise::Generate()
-	{
-		return Add + Scale * RandomStream.FRand() * 2 - 1.0f;
-	}
+	static const float DefaultFilterGain = Audio::ConvertToLinear(-3.f);
 
-	FPinkNoise::FPinkNoise(const float InScale, const float InAdd)
+	FPinkNoise::FPinkNoise(int32 InRandomSeed)
+		: Noise{ InRandomSeed }
+		, X_Z{ 0,0,0,0 }
+		, Y_Z{ 0,0.0,0 }
+		, A0{ DefaultFilterGain }
 	{
-		InitFilter();
-		Noise.SetScaleAdd(InScale, InAdd);
+		static_assert(UE_ARRAY_COUNT(X_Z) == 4, "sizeof(X_Z)==4");
+		static_assert(UE_ARRAY_COUNT(Y_Z) == 4, "sizeof(Y_Z)==4");
 	}
 
-	void FPinkNoise::InitFilter()
-	{
-		for (int32 i = 0; i < 4; ++i)
-		{
-			X[i] = 0.0f;
-			Y[i] = 0.0f;
-		}
-
-		A[0] = 0.0f;
-		A[1] = -2.495f;
-		A[2] = 2.017;
-		A[3] = -0.522f;
-
-		B[0] = 0.041f;
-		B[1] = -0.96f;
-		B[2] = 0.051f;
-		B[3] = -0.004f;
-	}
-
-	/** Sets the output scale and add parameter. */
-	void FPinkNoise::SetScaleAdd(const float InScale, const float InAdd)
-	{
-		Noise.SetScaleAdd(InScale, InAdd);
-	}
+	FPinkNoise::FPinkNoise()
+		: FPinkNoise{static_cast<int32>(FPlatformTime::Cycles())}
+	{}
 
 	float FPinkNoise::Generate()
 	{
-		X[0] = Noise.Generate();
-
-		float Output = 0.0f;
-		for (int32 i = 0; i < 4; ++i)
-		{
-			Output += (B[i] * X[i] - A[i] * Y[i]);
-		}
-
-		X[3] = X[2];
-		X[2] = X[1];
-		X[1] = X[0];
-
-		Y[3] = Y[2];
-		Y[2] = Y[1];
-		Y[1] = Y[0];
-		Y[0] = Output;
-
-		return Output;
+		// Filter Coefficients based on:
+		// https://ccrma.stanford.edu/~jos/sasp/Example_Synthesis_1_F_Noise.html
+		static constexpr float A[4] { 1.0f, -2.494956002f,2.017265875f, -0.522189400f };
+		static constexpr float B[4] { 0.049922035f,-0.095993537f,0.050612699f,-0.004408786f};	
+
+		static_assert(UE_ARRAY_COUNT(A) == UE_ARRAY_COUNT(X_Z), "A Coefficients and X need to be the same size");
+		static_assert(UE_ARRAY_COUNT(B) == UE_ARRAY_COUNT(Y_Z), "B Coefficients and Y need to be the same size");
+
+		X_Z[0] = Noise.Generate(); // Xn
+
+		float Yn =
+				A0 * X_Z[0]		// a0 * (x) (inject non-const filter-gain A0 in here).
+			+ A[1] * X_Z[1]		// a1 * (x-1)
+			+ A[2] * X_Z[2]		// a2 * (x-2)
+			+ A[3] * X_Z[3]		// a3 * (x-3) 
+
+			- B[0] * Y_Z[0]		// b1 * (y-1)
+			- B[1] * Y_Z[1]		// b2 * (y-2)
+			- B[2] * Y_Z[2]		// b3 * (y-3)
+			- B[3] * Y_Z[3];	// b4 * (y-4)
+
+		// Shuffle feed-forward state by one.
+		X_Z[3] = X_Z[2];
+		X_Z[2] = X_Z[1];
+		X_Z[1] = X_Z[0];
+
+		// Shuffle feed-back state by one.
+		Y_Z[3] = Y_Z[2]; 
+		Y_Z[2] = Y_Z[1];
+		Y_Z[1] = Y_Z[0];
+		Y_Z[0] = Yn;
+
+		return Yn;
 	}
 }
diff --git a/Engine/Source/Runtime/SignalProcessing/Public/DSP/Noise.h b/Engine/Source/Runtime/SignalProcessing/Public/DSP/Noise.h
--- a/Engine/Source/Runtime/SignalProcessing/Public/DSP/Noise.h
+++ b/Engine/Source/Runtime/SignalProcessing/Public/DSP/Noise.h
@@ -14,18 +14,21 @@ namespace Audio
 	class SIGNALPROCESSING_API FWhiteNoise
 	{
 	public:
-		/** Constructor with a default scale add parameter */
-		FWhiteNoise(const float InScale = 1.0f, const float InAdd = 0.0f);
-
-		void SetScaleAdd(const float InScale, const float InAdd);
+		FWhiteNoise();
+		FWhiteNoise(int32 InRandomSeed);
 
 		/** Generate next sample of white noise */
-		float Generate();
-
-	private:
-		float Scale;
-		float Add;
-
+		FORCEINLINE float Generate()
+		{
+			return (RandomStream.FRand() * 2.f) - 1.0f;
+		}
+		
+		/** Generate next sample of white noise (with optional Scale and Add params) */
+		FORCEINLINE float Generate(float InScale, float InAdd)
+		{
+			return Generate() * InScale + InAdd;
+		}
+	private:	
 		FRandomStream RandomStream;
 	};
 
@@ -36,26 +39,25 @@ namespace Audio
 	class SIGNALPROCESSING_API FPinkNoise
 	{
 	public:
-		/** Constructor. */
-		FPinkNoise(const float InScale = 1.0f, const float InAdd = 0.0f);
+		/** Constructor. Without seed argument, uses Cpu cycles to chose one at "random" */
+		FPinkNoise();
 
-		/** Sets the output scale and add parameter. */
-		void SetScaleAdd(const float InScale, const float InAdd);
+		/** Constructor with seed input */
+		FPinkNoise(int32 InRandomSeed);
 
 		/** Generate next sample of pink noise. */
 		float Generate();
 
-	private:
-
-		void InitFilter();
+		/** Set Pink Noise Filter Gain (default -3db) */
+		void SetFilterGain(float InFilterGain) 
+		{
+			A0 = InFilterGain;
+		}
 
+	private:
 		FWhiteNoise Noise;
-
-		float A[4];
-		float B[4];
-		float X[4];
-		float Y[4];
+		float X_Z[4];
+		float Y_Z[4];
+		float A0 = 1.f;
 	};
-
-
 }
