Fix for Mac specific CIS break

diff --git a/Engine/Plugins/Runtime/Metasound/Source/MetasoundStandardNodes/Private/MetasoundNoiseGenerator.cpp b/Engine/Plugins/Runtime/Metasound/Source/MetasoundStandardNodes/Private/MetasoundNoiseGenerator.cpp
--- a/Engine/Plugins/Runtime/Metasound/Source/MetasoundStandardNodes/Private/MetasoundNoiseGenerator.cpp
+++ b/Engine/Plugins/Runtime/Metasound/Source/MetasoundStandardNodes/Private/MetasoundNoiseGenerator.cpp
@@ -90,6 +90,8 @@ namespace Metasound
 		static constexpr const TCHAR* SeedPinName = TEXT("Seed");
 	};
 
+	constexpr int32 FNoiseOperator::DefaultSeed;
+
 	struct FNoiseOperator_White final : public FNoiseOperator
 	{
 		Audio::FWhiteNoise Generator;
