Metasound trigger counter node

Counts triggers and writes the count to the output integer

diff --git a/Engine/Plugins/Runtime/Metasound/Source/MetasoundStandardNodes/Private/MetasoundTriggerCounterNode.cpp b/Engine/Plugins/Runtime/Metasound/Source/MetasoundStandardNodes/Private/MetasoundTriggerCounterNode.cpp
new file mode 100644
--- /dev/null
+++ b/Engine/Plugins/Runtime/Metasound/Source/MetasoundStandardNodes/Private/MetasoundTriggerCounterNode.cpp
@@ -0,0 +1,151 @@
+// Copyright Epic Games, Inc. All Rights Reserved.
+#include "MetasoundTriggerCounterNode.h"
+
+#include "CoreMinimal.h"
+
+#include "MetasoundBuilderInterface.h"
+#include "MetasoundDataReferenceCollection.h"
+#include "MetasoundExecutableOperator.h"
+#include "MetasoundFacade.h"
+#include "MetasoundNode.h"
+#include "MetasoundNodeInterface.h"
+#include "MetasoundNodeRegistrationMacro.h"
+#include "MetasoundOperatorInterface.h"
+#include "MetasoundPrimitives.h"
+#include "MetasoundStandardNodesNames.h"
+
+#define LOCTEXT_NAMESPACE "MetasoundStandardNodes"
+
+namespace Metasound
+{
+	namespace TriggerCounter
+	{
+		static const TCHAR* InParamNameInput = TEXT("In");
+		static const TCHAR* InParamNameReset = TEXT("Reset");
+		static const TCHAR* OutParamNameOutput = TEXT("Out");
+	}
+
+	class FTriggerCounterOperator : public TExecutableOperator<FTriggerCounterOperator>
+	{
+		public:
+			static const FNodeClassMetadata& GetNodeInfo();
+			static const FVertexInterface& GetVertexInterface();
+			static TUniquePtr<IOperator> CreateOperator(const FCreateOperatorParams& InParams, FBuildErrorArray& OutErrors);
+
+			FTriggerCounterOperator(const FOperatorSettings& InSettings, const FTriggerReadRef& InTriggerIn, const FTriggerReadRef& InTriggerReset);
+
+			virtual FDataReferenceCollection GetInputs() const override;
+			virtual FDataReferenceCollection GetOutputs() const override;
+
+			void Execute();
+
+		private:
+			FTriggerReadRef TriggerIn;
+			FTriggerReadRef TriggerReset;
+			FInt32WriteRef TriggerCount;
+	};
+
+	FTriggerCounterOperator::FTriggerCounterOperator(const FOperatorSettings& InSettings, const FTriggerReadRef& InTriggerIn, const FTriggerReadRef& InTriggerReset)
+	: TriggerIn(InTriggerIn)
+	, TriggerReset(InTriggerReset)
+	, TriggerCount(FInt32WriteRef::CreateNew(0))
+	{
+	}
+
+	FDataReferenceCollection FTriggerCounterOperator::GetInputs() const
+	{
+		FDataReferenceCollection InputDataReferences;
+		InputDataReferences.AddDataReadReference(TriggerCounter::InParamNameInput, TriggerIn);
+		InputDataReferences.AddDataReadReference(TriggerCounter::InParamNameReset, TriggerReset);
+		return InputDataReferences;
+	}
+
+	FDataReferenceCollection FTriggerCounterOperator::GetOutputs() const
+	{
+		FDataReferenceCollection OutputDataReferences;
+		OutputDataReferences.AddDataReadReference(TriggerCounter::OutParamNameOutput, FInt32WriteRef(TriggerCount));
+
+		return OutputDataReferences;
+	}
+
+	void FTriggerCounterOperator::Execute()
+	{
+		// TODO: fix this up when we implement merged trigger execution to avoid edge cases of ordering problems
+		// for when these triggers are executed at differnet times within the same block
+		TriggerReset->ExecuteBlock(
+			[&](int32 StartFrame, int32 EndFrame)
+			{
+			},
+			[this](int32 StartFrame, int32 EndFrame)
+			{
+				*TriggerCount = 0;
+			}
+		);
+
+		TriggerIn->ExecuteBlock(
+			[&](int32 StartFrame, int32 EndFrame)
+			{
+			},
+			[this](int32 StartFrame, int32 EndFrame)
+			{
+				UE_LOG(LogTemp, Log, TEXT("TriggerCount: %d"), *TriggerCount)
+
+				++(*TriggerCount);
+			}
+		);
+	}
+
+	TUniquePtr<IOperator> FTriggerCounterOperator::CreateOperator(const FCreateOperatorParams& InParams, FBuildErrorArray& OutErrors)
+	{
+		const FInputVertexInterface& InputInterface = GetVertexInterface().GetInputInterface();
+		FTriggerReadRef TriggerIn = InParams.InputDataReferences.GetDataReadReferenceOrConstruct<FTrigger>(TriggerCounter::InParamNameInput, InParams.OperatorSettings);
+		FTriggerReadRef TriggerReset = InParams.InputDataReferences.GetDataReadReferenceOrConstruct<FTrigger>(TriggerCounter::InParamNameReset, InParams.OperatorSettings);
+
+		return MakeUnique<FTriggerCounterOperator>(InParams.OperatorSettings, TriggerIn, TriggerReset);
+	}
+
+	const FVertexInterface& FTriggerCounterOperator::GetVertexInterface()
+	{
+		static const FVertexInterface Interface(
+			FInputVertexInterface(
+				TInputDataVertexModel<FTrigger>(TriggerCounter::InParamNameInput, LOCTEXT("TriggerCounterInTooltip", "Trigger input to count.")),
+				TInputDataVertexModel<FTrigger>(TriggerCounter::InParamNameReset, LOCTEXT("TriggerCounterResetInTooltip", "Resets the counter back to 0."))
+			),
+			FOutputVertexInterface(
+				TOutputDataVertexModel<int32>(TriggerCounter::OutParamNameOutput, LOCTEXT("TriggerCounterOutTooltip", "The current count of triggers."))
+			)
+		);
+
+		return Interface;
+	}
+
+	const FNodeClassMetadata& FTriggerCounterOperator::GetNodeInfo()
+	{
+		auto InitNodeInfo = []() -> FNodeClassMetadata
+		{
+			FNodeClassMetadata Info;
+			Info.ClassName = {Metasound::StandardNodes::Namespace, TEXT("Trigger Counter"), TEXT("")};
+			Info.MajorVersion = 1;
+			Info.MinorVersion = 0;
+			Info.DisplayName = LOCTEXT("Metasound_TriggerCounterNodeDisplayName", "Trigger Counter");
+			Info.Description = LOCTEXT("Metasound_TriggerCounterNodeDescription", "Counts the trigger inputs.");
+			Info.Author = PluginAuthor;
+			Info.PromptIfMissing = PluginNodeMissingPrompt;
+			Info.DefaultInterface = GetVertexInterface();
+
+			return Info;
+		};
+
+		static const FNodeClassMetadata Info = InitNodeInfo();
+		return Info;
+	}
+
+	FTriggerCounterNode::FTriggerCounterNode(const FNodeInitData& InInitData)
+	:	FNodeFacade(InInitData.InstanceName, InInitData.InstanceID, TFacadeOperatorClass<FTriggerCounterOperator>())
+	{
+	}
+
+	METASOUND_REGISTER_NODE(FTriggerCounterNode)
+}
+
+#undef LOCTEXT_NAMESPACE // MetasoundTriggerDelayNode
diff --git a/Engine/Plugins/Runtime/Metasound/Source/MetasoundStandardNodes/Public/MetasoundTriggerCounterNode.h b/Engine/Plugins/Runtime/Metasound/Source/MetasoundStandardNodes/Public/MetasoundTriggerCounterNode.h
new file mode 100644
--- /dev/null
+++ b/Engine/Plugins/Runtime/Metasound/Source/MetasoundStandardNodes/Public/MetasoundTriggerCounterNode.h
@@ -0,0 +1,26 @@
+// Copyright Epic Games, Inc. All Rights Reserved.
+
+#pragma once
+
+#include "CoreMinimal.h"
+#include "MetasoundBuilderInterface.h"
+#include "MetasoundDataReferenceCollection.h"
+#include "MetasoundExecutableOperator.h"
+#include "MetasoundFacade.h"
+#include "MetasoundNode.h"
+#include "MetasoundNodeInterface.h"
+#include "MetasoundOperatorInterface.h"
+#include "MetasoundTime.h"
+#include "MetasoundTrigger.h"
+
+
+namespace Metasound
+{
+	class METASOUNDSTANDARDNODES_API FTriggerCounterNode : public FNodeFacade
+	{
+		public:
+			FTriggerCounterNode(const FNodeInitData& InInitData);
+
+			virtual ~FTriggerCounterNode() = default;
+	};
+}
