Fix Mac only CIS

diff --git a/Engine/Plugins/Runtime/Metasound/Source/MetasoundStandardNodes/Private/MetasoundTriggerOnThresholdNode.cpp b/Engine/Plugins/Runtime/Metasound/Source/MetasoundStandardNodes/Private/MetasoundTriggerOnThresholdNode.cpp
--- a/Engine/Plugins/Runtime/Metasound/Source/MetasoundStandardNodes/Private/MetasoundTriggerOnThresholdNode.cpp
+++ b/Engine/Plugins/Runtime/Metasound/Source/MetasoundStandardNodes/Private/MetasoundTriggerOnThresholdNode.cpp
@@ -14,8 +14,6 @@
 
 #define LOCTEXT_NAMESPACE "MetasoundStandardNodes"
 
-#if 0
-
 namespace Metasound
 {
 	enum class EBufferTriggerType
@@ -61,6 +59,9 @@ namespace Metasound
 		static constexpr const TCHAR* TriggerType = TEXT("Type");
 	};
 
+	// Mac Clang requires linkage on constexpr
+	constexpr float FTriggerOnThresholdOperator::DefaultThreshold;
+
 	struct FTriggerOnThresholdOperator_EdgeCommon : public FTriggerOnThresholdOperator
 	{
 		using FTriggerOnThresholdOperator::FTriggerOnThresholdOperator;
@@ -251,7 +252,5 @@ namespace Metasound
 	
 	METASOUND_REGISTER_NODE(FTriggerOnThresholdNode);
 }
-#endif
-
 #undef LOCTEXT_NAMESPACE //MetasoundStandardNodes
 
diff --git a/Engine/Plugins/Runtime/Metasound/Source/MetasoundStandardNodes/Public/MetasoundTriggerOnThresholdNode.h b/Engine/Plugins/Runtime/Metasound/Source/MetasoundStandardNodes/Public/MetasoundTriggerOnThresholdNode.h
--- a/Engine/Plugins/Runtime/Metasound/Source/MetasoundStandardNodes/Public/MetasoundTriggerOnThresholdNode.h
+++ b/Engine/Plugins/Runtime/Metasound/Source/MetasoundStandardNodes/Public/MetasoundTriggerOnThresholdNode.h
@@ -12,7 +12,6 @@
 #include "MetasoundNode.h"
 #include "MetasoundOperatorInterface.h"
 
-#if 0
 namespace Metasound
 {
 	class METASOUNDSTANDARDNODES_API FTriggerOnThresholdNode : public FNodeFacade
@@ -26,4 +25,3 @@ namespace Metasound
 		float DefaultThreshold;
 	};
 }
-#endif
