Fix for Oscillator phase reset on "2D Rotation".
Add "Generator" Category to all noise/sine/triange/saw/square.
Disable Genertion enum options that are defaulting to others (for Early Access)

diff --git a/Engine/Plugins/Runtime/Metasound/Source/MetasoundStandardNodes/Private/MetasoundNoiseGenerator.cpp b/Engine/Plugins/Runtime/Metasound/Source/MetasoundStandardNodes/Private/MetasoundNoiseGenerator.cpp
--- a/Engine/Plugins/Runtime/Metasound/Source/MetasoundStandardNodes/Private/MetasoundNoiseGenerator.cpp
+++ b/Engine/Plugins/Runtime/Metasound/Source/MetasoundStandardNodes/Private/MetasoundNoiseGenerator.cpp
@@ -178,11 +178,12 @@ namespace Metasound
 			Info.ClassName = { Metasound::StandardNodes::Namespace, TEXT("Noise"), Metasound::StandardNodes::AudioVariant };
 			Info.MajorVersion = 1;
 			Info.MinorVersion = 0;
-			Info.DisplayName = LOCTEXT("Metasound_NoiseNodeDisplayName", "Noise Generator");
-			Info.Description = LOCTEXT("Metasound_NoiseNodeDescription", "Noise Generated that produces differnt types of noise");
+			Info.DisplayName = LOCTEXT("Metasound_NoiseNodeDisplayNameX", "Noise");
+			Info.Description = LOCTEXT("Metasound_NoiseNodeDescription", "Noise Generator that produces differnt types of noise");
 			Info.Author = PluginAuthor;
 			Info.PromptIfMissing = PluginNodeMissingPrompt;
 			Info.DefaultInterface = DeclareVertexInterface();
+			Info.CategoryHierarchy = { LOCTEXT("Metasound_GeneratorCategory", "Generators") };
 
 			return Info;
 		};
diff --git a/Engine/Plugins/Runtime/Metasound/Source/MetasoundStandardNodes/Private/MetasoundOscillatorNodes.cpp b/Engine/Plugins/Runtime/Metasound/Source/MetasoundStandardNodes/Private/MetasoundOscillatorNodes.cpp
--- a/Engine/Plugins/Runtime/Metasound/Source/MetasoundStandardNodes/Private/MetasoundOscillatorNodes.cpp
+++ b/Engine/Plugins/Runtime/Metasound/Source/MetasoundStandardNodes/Private/MetasoundOscillatorNodes.cpp
@@ -18,6 +18,8 @@ namespace Metasound
 {			
 #pragma region Common
 
+	static const FText GeneratorCategory = { LOCTEXT("Metasound_GeneratorCategory", "Generators") };
+
 	// Contained here are a small selection of block generators and sample generators.
 	// They are templatized to allow quick construction of new oscillators, but aren't considered
 	// optimal (as they contain many branches) and they should be only as a default 
@@ -164,7 +166,9 @@ namespace Metasound
 		struct F2DRotatorGenerateBlock
 		{
 			Audio::FSinOsc2DRotation Rotator;
-			F2DRotatorGenerateBlock(float InStartingPhase) { Rotator.LastPhasePerSample = InStartingPhase; }
+			F2DRotatorGenerateBlock(float InStartingLinearPhase)
+				: Rotator{ 2.f * PI * InStartingLinearPhase }
+			{}
 			void operator()(const FGeneratorArgs& Args)
 			{
 				Rotator.GenerateBuffer(Args.SampleRate, Args.FrequencyHz, Args.AlignedBuffer.GetData(), Args.AlignedBuffer.Num());
@@ -501,8 +505,8 @@ namespace Metasound
 		DEFINE_METASOUND_ENUM_ENTRY(ESineGenerationType::Sinf, LOCTEXT("SinfDescription", "Pure Math"), LOCTEXT("SinfDescriptionTT", "Uses Pure Math (Sinf) to generate the Sine Wave (most expensive)")),
 		DEFINE_METASOUND_ENUM_ENTRY(ESineGenerationType::Rotation, LOCTEXT("RotationDescription", "2D Rotation"), LOCTEXT("RotationDescriptionTT", "Rotates around the unit circle generate the sinewave")),
 		DEFINE_METASOUND_ENUM_ENTRY(ESineGenerationType::Bhaskara, LOCTEXT("BhaskaraDescription", "Bhaskara"), LOCTEXT("BhaskaraDescriptionTT", "Sine approximation using Bhaskara technique discovered in 7th century")),
-		DEFINE_METASOUND_ENUM_ENTRY(ESineGenerationType::Wavetable, LOCTEXT("WavetableDescription", "Wavetable"), LOCTEXT("WavetableDescriptionTT", "Uses wavetable to generate the sinewave")),
-		DEFINE_METASOUND_ENUM_ENTRY(ESineGenerationType::FilterResonanceFeedback, LOCTEXT("FRFDescription", "Filter resonance"), LOCTEXT("FRFDescriptionTT", "Filter resonance/feedback")),
+		//DEFINE_METASOUND_ENUM_ENTRY(ESineGenerationType::Wavetable, LOCTEXT("WavetableDescription", "Wavetable"), LOCTEXT("WavetableDescriptionTT", "Uses wavetable to generate the sinewave")),
+		//DEFINE_METASOUND_ENUM_ENTRY(ESineGenerationType::FilterResonanceFeedback, LOCTEXT("FRFDescription", "Filter resonance"), LOCTEXT("FRFDescriptionTT", "Filter resonance/feedback")),
 	DEFINE_METASOUND_ENUM_END()
 
 	class FSineOscilatorNode::FFactory : public FOscilatorFactoryBase
@@ -525,6 +529,7 @@ namespace Metasound
 				Info.Author = PluginAuthor;
 				Info.PromptIfMissing = PluginNodeMissingPrompt;
 				Info.DefaultInterface = GetVertexInterface();
+				Info.CategoryHierarchy.Emplace(GeneratorCategory);
 				return Info;
 			};
 			static const FNodeClassMetadata Info = InitNodeInfo();
@@ -619,7 +624,7 @@ namespace Metasound
 	DEFINE_METASOUND_ENUM_BEGIN(ESawGenerationType, FEnumSawGenerationType, "SawGenerationType")
 		DEFINE_METASOUND_ENUM_ENTRY(ESawGenerationType::PolySmooth, LOCTEXT("SawPolySmoothDescription", "Poly Smooth"), LOCTEXT("PolySmoothDescriptionTT", "PolySmooth (i.e. BLEP)")),
 		DEFINE_METASOUND_ENUM_ENTRY(ESawGenerationType::Trivial, LOCTEXT("SawTrivialDescription", "Trivial"), LOCTEXT("TrivialDescriptionTT", "The most basic raw implementation")),
-		DEFINE_METASOUND_ENUM_ENTRY(ESawGenerationType::Wavetable, LOCTEXT("SawWavetableDescription", "Wavetable"), LOCTEXT("SawWavetableDescriptionTT", "Use a Wavetable iterpolation to generate the Waveform"))
+		//DEFINE_METASOUND_ENUM_ENTRY(ESawGenerationType::Wavetable, LOCTEXT("SawWavetableDescription", "Wavetable"), LOCTEXT("SawWavetableDescriptionTT", "Use a Wavetable iterpolation to generate the Waveform"))
 	DEFINE_METASOUND_ENUM_END()
 
 	class FSawOscilatorNode::FFactory : public FOscilatorFactoryBase
@@ -644,6 +649,7 @@ namespace Metasound
 				Info.Author = PluginAuthor;
 				Info.PromptIfMissing = PluginNodeMissingPrompt;
 				Info.DefaultInterface = GetVertexInterface();
+				Info.CategoryHierarchy.Emplace(GeneratorCategory);
 				return Info;
 			};
 			static const FNodeClassMetadata Info = InitNodeInfo();
@@ -736,7 +742,7 @@ namespace Metasound
 	DEFINE_METASOUND_ENUM_BEGIN(ESquareGenerationType, FEnumSquareGenerationType, "SquareGenerationType")
 		DEFINE_METASOUND_ENUM_ENTRY(ESquareGenerationType::PolySmooth, LOCTEXT("SquarePolySmoothDescription", "Poly Smooth"), LOCTEXT("PolySmoothDescriptionTT", "PolySmooth (i.e. BLEP)")),
 		DEFINE_METASOUND_ENUM_ENTRY(ESquareGenerationType::Trivial, LOCTEXT("SquareTrivialDescription", "Trivial"), LOCTEXT("TrivialDescriptionTT", "The most basic raw implementation")),
-		DEFINE_METASOUND_ENUM_ENTRY(ESquareGenerationType::Wavetable, LOCTEXT("SquareWavetableDescription", "Wavetable"), LOCTEXT("SquareWavetableDescriptionTT", "Use a Wavetable interpolation to generate the Waveform"))
+		//DEFINE_METASOUND_ENUM_ENTRY(ESquareGenerationType::Wavetable, LOCTEXT("SquareWavetableDescription", "Wavetable"), LOCTEXT("SquareWavetableDescriptionTT", "Use a Wavetable interpolation to generate the Waveform"))
 	DEFINE_METASOUND_ENUM_END()
 	
 	template<typename GeneratorPolicy>
@@ -868,6 +874,7 @@ namespace Metasound
 				Info.Author = PluginAuthor;
 				Info.PromptIfMissing = PluginNodeMissingPrompt;
 				Info.DefaultInterface = GetVertexInterface();
+				Info.CategoryHierarchy.Emplace(GeneratorCategory);
 				return Info;
 			};
 			static const FNodeClassMetadata Info = InitNodeInfo();
@@ -914,7 +921,7 @@ namespace Metasound
 	DEFINE_METASOUND_ENUM_BEGIN(ETriangleGenerationType, FEnumTriangleGenerationType, "TriangleGenerationType")
 		DEFINE_METASOUND_ENUM_ENTRY(ETriangleGenerationType::PolySmooth, LOCTEXT("TrianglePolySmoothDescription", "Poly Smooth"), LOCTEXT("PolySmoothDescriptionTT", "PolySmooth (i.e. BLEP)")),
 		DEFINE_METASOUND_ENUM_ENTRY(ETriangleGenerationType::Trivial, LOCTEXT("TriangleTrivialDescription", "Trivial"), LOCTEXT("TrivialDescriptionTT", "The most basic raw implementation")),
-		DEFINE_METASOUND_ENUM_ENTRY(ETriangleGenerationType::Wavetable, LOCTEXT("TriangleWavetableDescription", "Wavetable"), LOCTEXT("TriangleWavetableDescriptionTT", "Use a Wavetable iterpolation to generate the Waveform"))
+		//DEFINE_METASOUND_ENUM_ENTRY(ETriangleGenerationType::Wavetable, LOCTEXT("TriangleWavetableDescription", "Wavetable"), LOCTEXT("TriangleWavetableDescriptionTT", "Use a Wavetable iterpolation to generate the Waveform"))
 	DEFINE_METASOUND_ENUM_END()
 
 	class FTriangleOscilatorNode::FFactory : public FOscilatorFactoryBase
@@ -978,6 +985,7 @@ namespace Metasound
 				Info.Author = PluginAuthor;
 				Info.PromptIfMissing = PluginNodeMissingPrompt;
 				Info.DefaultInterface = GetVertexInterface();
+				Info.CategoryHierarchy.Emplace(GeneratorCategory);
 				return Info;
 			};
 			static const FNodeClassMetadata Info = InitNodeInfo();
diff --git a/Engine/Source/Runtime/SignalProcessing/Public/DSP/Dsp.h b/Engine/Source/Runtime/SignalProcessing/Public/DSP/Dsp.h
--- a/Engine/Source/Runtime/SignalProcessing/Public/DSP/Dsp.h
+++ b/Engine/Source/Runtime/SignalProcessing/Public/DSP/Dsp.h
@@ -115,17 +115,13 @@ namespace Audio
 	* continue the sine tone where it left off.
 	*
 	*/
-	struct FSinOsc2DRotation
+	class FSinOsc2DRotation
 	{
-		VectorRegister QuadDxVec, QuadDyVec;
-		float LastPhasePerSample;
-		float LastPhase;
-		float Dx, Dy;
-
-		FSinOsc2DRotation()
+	public:
+		FSinOsc2DRotation(float InStartingPhaseRadians = 0.f)
 		{
 			LastPhasePerSample = -1;
-			LastPhase = 0;
+			LastPhase = InStartingPhaseRadians;
 		}
 
 		FSinOsc2DRotation(FSinOsc2DRotation const& other)
@@ -144,14 +140,12 @@ namespace Audio
 		* 
 		* @param SampleRate the sample rate the buffer will be played at
 		* @param ClampedFrequency the frequency of the tone to emit, clamped to nyquist (SampleRate/2)
-		* @param AlignedBuffer the output buffer, aligned for vector operations
+		* @param Buffer the output buffer
 		* @param BufferSampleCount the number of samples in the buffer. Generally, this should be %4 granularity.
 		* 
 		*/
-		void GenerateBuffer(float SampleRate, float ClampedFrequency, float* AlignedBuffer, int32 BufferSampleCount)
+		void GenerateBuffer(float SampleRate, float ClampedFrequency, float* Buffer, int32 BufferSampleCount)
 		{
-			check(IsAligned(AlignedBuffer, sizeof(VectorRegister)));
-
 			// Regenerate our vector rotation components if our changes.
 			const float PhasePerSample = (ClampedFrequency * (2 * PI)) / (SampleRate);
 			if (LastPhasePerSample != PhasePerSample)
@@ -165,7 +159,7 @@ namespace Audio
 				LastPhasePerSample = PhasePerSample;
 			}
 
-			float* Write = AlignedBuffer;
+			float* Write = Buffer;
 
 			// The rotation matrix drifts, so we resync every so often
 			while (BufferSampleCount)
@@ -222,7 +216,8 @@ namespace Audio
 				{
 					// We've actually already calculated the next quad - it's in YVector
 					alignas(16) float YFloats[4];
-					VectorStoreAligned(YVector, YFloats);
+					//VectorStoreAligned(YVector, YFloats);
+					VectorStore(YVector, YFloats);
 
 					int32 Remn = BlockSampleCount & SIMD_MASK;
 					for (int32 i = 0; i < Remn; i++)
@@ -239,6 +234,12 @@ namespace Audio
 				BufferSampleCount -= BlockSampleCount;
 			} // end while BufferSampleCount
 		} // end GenerateBuffer
+
+	private:
+		VectorRegister QuadDxVec, QuadDyVec;
+		float LastPhasePerSample;
+		float LastPhase;
+		float Dx, Dy;
 	}; // end FSinOsc2DRotation
 
 	// Fast tanh based on pade approximation
