Support setting value by literal in MetasoundInstanceTransmitter
- Support for Literal in ISender
- Support IDataChannel creation registration
- Make IDataChannel shared ptrs thread safe.
- Add LexToString for FLiteral.
#okforgihub public

diff --git a/Engine/Plugins/Runtime/Metasound/Source/MetasoundFrontend/Private/MetasoundFrontendRegistries.cpp b/Engine/Plugins/Runtime/Metasound/Source/MetasoundFrontend/Private/MetasoundFrontendRegistries.cpp
--- a/Engine/Plugins/Runtime/Metasound/Source/MetasoundFrontend/Private/MetasoundFrontendRegistries.cpp
+++ b/Engine/Plugins/Runtime/Metasound/Source/MetasoundFrontend/Private/MetasoundFrontendRegistries.cpp
@@ -4,6 +4,7 @@
 
 #include "CoreMinimal.h"
 #include "MetasoundLog.h"
+#include "MetasoundRouter.h"
 #include "Misc/ScopeLock.h"
 #include "HAL/PlatformTime.h"
 
@@ -595,3 +596,16 @@ FMetasoundFrontendRegistryContainer::GetEnumInterfaceForDataType(FName InDataTyp
 	return nullptr;
 }
 
+
+TSharedPtr<Metasound::IDataChannel, ESPMode::ThreadSafe> FMetasoundFrontendRegistryContainer::CreateDataChannelForDataType(const FName& InDataType, const Metasound::FOperatorSettings& InSettings) const
+{
+	if (const FDataTypeRegistryElement* Element = DataTypeRegistry.Find(InDataType))
+	{
+		if (Element->Callbacks.CreateDataChannel)
+		{
+			return Element->Callbacks.CreateDataChannel(InSettings);
+		}
+	}
+	return nullptr;
+}
+
diff --git a/Engine/Plugins/Runtime/Metasound/Source/MetasoundFrontend/Private/MetasoundInstanceTransmitter.cpp b/Engine/Plugins/Runtime/Metasound/Source/MetasoundFrontend/Private/MetasoundInstanceTransmitter.cpp
--- a/Engine/Plugins/Runtime/Metasound/Source/MetasoundFrontend/Private/MetasoundInstanceTransmitter.cpp
+++ b/Engine/Plugins/Runtime/Metasound/Source/MetasoundFrontend/Private/MetasoundInstanceTransmitter.cpp
@@ -18,17 +18,39 @@ namespace Metasound
 
 	bool FMetasoundInstanceTransmitter::SetFloatParameter(const FName& InParameterName, float InValue)
 	{
-		return SendValue<float>(InParameterName, InValue);
+		return SetParameterWithLiteral(InParameterName, FLiteral(InValue));
 	}
 
 	bool FMetasoundInstanceTransmitter::SetIntParameter(const FName& InParameterName, int32 InValue) 
 	{
-		return SendValue<int32>(InParameterName, InValue);
+		return SetParameterWithLiteral(InParameterName, FLiteral(InValue));
 	}
 
 	bool FMetasoundInstanceTransmitter::SetBoolParameter(const FName& InParameterName, bool InValue) 
 	{
-		return SendValue<bool>(InParameterName, InValue);
+		return SetParameterWithLiteral(InParameterName, FLiteral(InValue));
+	}
+
+	bool FMetasoundInstanceTransmitter::SetParameterWithLiteral(const FName& InParameterName, const FLiteral& InLiteral) 
+	{
+		ISender* Sender = FindSender(InParameterName);
+
+		if (nullptr == Sender)
+		{
+			// If not sender exists for parameter name, attempt to add one.
+			if (const FSendInfo* SendInfo = FindSendInfo(InParameterName))
+			{
+				Sender = AddSender(*SendInfo);
+			}
+		}
+
+		if (nullptr != Sender)
+		{
+			// Push data using found sender.
+			return Sender->PushLiteral(InLiteral);
+		}
+
+		return false;
 	}
 
 	TUniquePtr<IAudioInstanceTransmitter> FMetasoundInstanceTransmitter::Clone() const
@@ -50,4 +72,27 @@ namespace Metasound
 		}
 		return nullptr;
 	}
+
+	TUniquePtr<ISender> FMetasoundInstanceTransmitter::CreateSender(const FSendInfo& InInfo) const
+	{
+		// TODO: likely want to remove this and opt for different protocols having different behaviors.
+		const float DelayTimeInSeconds = 0.1f; //< This not used for non-audio routing.
+		FSenderInitParams InitParams = {OperatorSettings, DelayTimeInSeconds};
+
+		return FDataTransmissionCenter::Get().RegisterNewSender(InInfo.TypeName, InInfo.Address, InitParams);
+	}
+
+	ISender* FMetasoundInstanceTransmitter::AddSender(const FSendInfo& InInfo)
+	{
+		TUniquePtr<ISender> Sender = CreateSender(InInfo);
+
+		if (ensureMsgf(Sender.IsValid(), TEXT("Failed to create sender [DataType:%s, Address:%s]"), *InInfo.TypeName.ToString(), *InInfo.Address.ChannelName.ToString()))
+		{
+			ISender* Ptr = Sender.Get();
+			InputSends.Add(InInfo.ParameterName, MoveTemp(Sender));
+			return Ptr;
+		}
+		
+		return nullptr;
+	}
 }
diff --git a/Engine/Plugins/Runtime/Metasound/Source/MetasoundFrontend/Private/MetasoundRouter.cpp b/Engine/Plugins/Runtime/Metasound/Source/MetasoundFrontend/Private/MetasoundRouter.cpp
--- a/Engine/Plugins/Runtime/Metasound/Source/MetasoundFrontend/Private/MetasoundRouter.cpp
+++ b/Engine/Plugins/Runtime/Metasound/Source/MetasoundFrontend/Private/MetasoundRouter.cpp
@@ -1,6 +1,8 @@
 // Copyright Epic Games, Inc. All Rights Reserved.
 
 #include "MetasoundRouter.h"
+
+#include "MetasoundFrontendRegistries.h"
 #include "MetasoundOperatorInterface.h"
 #include "CoreMinimal.h"
 #include "HAL/IConsoleManager.h"
@@ -128,14 +130,104 @@ namespace Metasound
 		FDataTransmissionCenter::Get().UnregisterSubsystem(SubsystemName);
 	}
 
-	Metasound::FDataTransmissionCenter& FDataTransmissionCenter::Get()
+	FDataTransmissionCenter& FDataTransmissionCenter::Get()
 	{
 		static FDataTransmissionCenter Singleton;
 		return Singleton;
 	}
 
-	FDataTransmissionCenter::FDataTransmissionCenter()
+	TUniquePtr<ISender> FDataTransmissionCenter::RegisterNewSender(const FName& InDataTypeName, const FSendAddress& InAddress, const FSenderInitParams& InitParams)
 	{
+		TUniquePtr<ISender> Sender;
+
+		if (InAddress.Subsystem == GetSubsystemNameForSendScope(ETransmissionScope::ThisInstanceOnly))
+		{
+			Sender = InstanceRouter.RegisterNewSender(InAddress.MetasoundInstanceID, InDataTypeName, InAddress.ChannelName, InitParams);
+		}
+		else if (InAddress.Subsystem == GetSubsystemNameForSendScope(ETransmissionScope::Global))
+		{
+			Sender = GlobalRouter.RegisterNewSender(InDataTypeName, InAddress.ChannelName, InitParams);
+		}
+		else if (FSubsystemData* FoundSubsystem = SubsystemRouters.Find(InAddress.Subsystem))
+		{
+			ITransmissionSubsystem* SubsystemInterface = FoundSubsystem->SubsystemPtr;
+
+			// Ensure that this subsystem can support this send type.
+			if (ensureAlways(SubsystemInterface && SubsystemInterface->CanSupportDataType(InDataTypeName)))
+			{
+
+				FScopeLock ScopeLock(&SubsystemRoutersMutationLock);
+				Sender = FoundSubsystem->AddressRouter.RegisterNewSender(InDataTypeName, InAddress.ChannelName, InitParams);
+
+				FoundSubsystem->SubsystemPtr->OnNewSendRegistered(InAddress, InDataTypeName);
+			}
+		}
+		else
+		{
+			// Otherwise, the subsystem FName was invalid.
+			UE_LOG(LogMetasound, Error, TEXT("Cannot create Sender. Did not find transmission subsystem [Name:%s]"), *InAddress.Subsystem.ToString());
+		}
+
+		return MoveTemp(Sender);
+	}
+
+	TUniquePtr<IReceiver> FDataTransmissionCenter::RegisterNewReceiver(const FName& InDataTypeName, const FSendAddress& InAddress, const FReceiverInitParams& InitParams)
+	{
+		TUniquePtr<IReceiver> Receiver;
+
+		if (InAddress.Subsystem == GetSubsystemNameForSendScope(ETransmissionScope::ThisInstanceOnly))
+		{
+			Receiver = InstanceRouter.RegisterNewReceiver(InAddress.MetasoundInstanceID, InDataTypeName, InAddress.ChannelName, InitParams);
+		}
+		else if (InAddress.Subsystem == GetSubsystemNameForSendScope(ETransmissionScope::Global))
+		{
+			Receiver = GlobalRouter.RegisterNewReceiver(InDataTypeName, InAddress.ChannelName, InitParams);
+		}
+		else if (FSubsystemData* FoundSubsystem = SubsystemRouters.Find(InAddress.Subsystem))
+		{
+			ITransmissionSubsystem* SubsystemInterface = FoundSubsystem->SubsystemPtr;
+
+			// Ensure that this subsystem can support this send type.
+			if (ensureAlways(SubsystemInterface && SubsystemInterface->CanSupportDataType(InDataTypeName)))
+			{
+				FScopeLock ScopeLock(&SubsystemRoutersMutationLock);
+
+				Receiver = FoundSubsystem->AddressRouter.RegisterNewReceiver(InDataTypeName, InAddress.ChannelName, InitParams);
+				FoundSubsystem->SubsystemPtr->OnNewReceiverRegistered(InAddress, InDataTypeName);
+			}
+		}
+		else
+		{
+			// Otherwise, the subsystem FName was invalid.
+			UE_LOG(LogMetasound, Error, TEXT("Cannot create Receiver. Did not find transmission subsystem [Name:%s]"), *InAddress.Subsystem.ToString());
+		}
+
+		return MoveTemp(Receiver);
+	}
+
+	bool FDataTransmissionCenter::PushLiteral(FName GlobalChannelName, const FLiteral& InParam)
+	{
+		return GlobalRouter.PushLiteral(GlobalChannelName, InParam);
+	}
+
+	void FDataTransmissionCenter::RegisterSubsystem(ITransmissionSubsystem* InSystem, FName InSubsytemName)
+	{
+		FScopeLock ScopeLock(&SubsystemRoutersMutationLock);
+
+		//check to make sure we're not adding a subsystem twice.
+		check(!SubsystemRouters.Contains(InSubsytemName));
+
+		SubsystemRouters.Emplace(InSubsytemName, InSystem);
+	}
+
+	void FDataTransmissionCenter::UnregisterSubsystem(FName InSubsystemName)
+	{
+		FScopeLock ScopeLock(&SubsystemRoutersMutationLock);
+
+		//check to make sure we're not adding a subsystem twice.
+		check(SubsystemRouters.Contains(InSubsystemName));
+
+		SubsystemRouters.Remove(InSubsystemName);
 	}
 
 	TArray<FName> FAddressRouter::GetAvailableChannels()
@@ -156,7 +248,7 @@ namespace Metasound
 	{
 		FScopeLock ScopeLock(&DataChannelMapMutationLock);
 
-		if (TSharedRef<IDataChannel>* FoundChannel = DataChannelMap.Find(InChannelName))
+		if (TSharedRef<IDataChannel, ESPMode::ThreadSafe>* FoundChannel = DataChannelMap.Find(InChannelName))
 		{
 			return (*FoundChannel)->GetDataType();
 		}
@@ -166,6 +258,66 @@ namespace Metasound
 		}
 	}
 
+	TSharedPtr<IDataChannel, ESPMode::ThreadSafe> FAddressRouter::GetDataChannel(const FName& InDataTypeName, const FName& InChannelName, const FOperatorSettings& InOperatorSettings)
+	{
+		TSharedPtr<IDataChannel, ESPMode::ThreadSafe> DataChannel; 
+
+		{
+			FScopeLock ScopeLock(&DataChannelMapMutationLock);
+
+			if (TSharedRef<IDataChannel, ESPMode::ThreadSafe>* ExistingChannelPtr = DataChannelMap.Find(InChannelName))
+			{
+				DataChannel = *ExistingChannelPtr;
+			}
+			else
+			{
+				// This is the first time we're seeing this, add it to the map.
+				FMetasoundFrontendRegistryContainer* Registry = FMetasoundFrontendRegistryContainer::Get();
+
+				if (ensure(nullptr != Registry))
+				{
+					DataChannel = Registry->CreateDataChannelForDataType(InDataTypeName, InOperatorSettings);
+					if (DataChannel.IsValid())
+					{
+						DataChannelMap.Add(InChannelName, DataChannel.ToSharedRef());
+					}
+				}
+			}
+		}
+
+		return DataChannel;
+	}
+
+	TUniquePtr<ISender> FAddressRouter::RegisterNewSender(const FName& InDataTypeName, const FName& InChannelName, const FSenderInitParams& InitParams)
+	{
+
+		TSharedPtr<IDataChannel, ESPMode::ThreadSafe> DataChannel = GetDataChannel(InDataTypeName, InChannelName, InitParams.OperatorSettings);
+
+		if (DataChannel.IsValid())
+		{
+			return DataChannel->NewSender(InitParams);
+		}
+		else
+		{
+			return TUniquePtr<ISender>(nullptr);
+		}
+	}
+
+	TUniquePtr<IReceiver> FAddressRouter::RegisterNewReceiver(const FName& InDataTypeName, const FName& InChannelName, const FReceiverInitParams& InitParams)
+	{
+
+		TSharedPtr<IDataChannel, ESPMode::ThreadSafe> DataChannel = GetDataChannel(InDataTypeName, InChannelName, InitParams.OperatorSettings);
+
+		if (DataChannel.IsValid())
+		{
+			return DataChannel->NewReceiver(InitParams);
+		}
+		else
+		{
+			return TUniquePtr<IReceiver>(nullptr);
+		}
+	}
+
 	TArray<FName> FInstanceLocalRouter::GetAvailableChannels(uint64 InInstanceID)
 	{
 		FScopeLock ScopeLock(&InstanceRouterMapMutationLock);
@@ -194,4 +346,36 @@ namespace Metasound
 		}
 	}
 
+	TUniquePtr<IReceiver> FInstanceLocalRouter::RegisterNewReceiver(uint64 InInstanceID, const FName& InDataTypeName, const FName& InChannelName, const FReceiverInitParams& InitParams)
+	{
+		FScopeLock ScopeLock(&InstanceRouterMapMutationLock);
+
+		if (FAddressRouter* AddressRouter = InstanceRouterMap.Find(InInstanceID))
+		{
+			return AddressRouter->RegisterNewReceiver(InDataTypeName, InChannelName, InitParams);
+		}
+		else
+		{
+			// This is the first time we're seeing this, add it to the map.
+			FAddressRouter& NewRouter = InstanceRouterMap.Add(InInstanceID, FAddressRouter());
+			return NewRouter.RegisterNewReceiver(InDataTypeName, InChannelName, InitParams);
+		}
+	}
+
+	TUniquePtr<ISender> FInstanceLocalRouter::RegisterNewSender(uint64 InInstanceID, const FName& InDataTypeName, const FName& InChannelName, const FSenderInitParams& InitParams)
+	{
+		FScopeLock ScopeLock(&InstanceRouterMapMutationLock);
+
+		if (FAddressRouter* AddressRouter = InstanceRouterMap.Find(InInstanceID))
+		{
+			return AddressRouter->RegisterNewSender(InDataTypeName, InChannelName, InitParams);
+		}
+		else
+		{
+			// This is the first time we're seeing this, add it to the map.
+			FAddressRouter& NewRouter = InstanceRouterMap.Add(InInstanceID, FAddressRouter());
+			return NewRouter.RegisterNewSender(InDataTypeName, InChannelName, InitParams);
+		}
+	}
+
 }
diff --git a/Engine/Plugins/Runtime/Metasound/Source/MetasoundFrontend/Public/MetasoundDataTypeRegistrationMacro.h b/Engine/Plugins/Runtime/Metasound/Source/MetasoundFrontend/Public/MetasoundDataTypeRegistrationMacro.h
--- a/Engine/Plugins/Runtime/Metasound/Source/MetasoundFrontend/Public/MetasoundDataTypeRegistrationMacro.h
+++ b/Engine/Plugins/Runtime/Metasound/Source/MetasoundFrontend/Public/MetasoundDataTypeRegistrationMacro.h
@@ -21,6 +21,7 @@
 #include "MetasoundReceiveNode.h"
 #include "MetasoundRouter.h"
 #include "MetasoundSendNode.h"
+#include "MetasoundTransmissionRegistration.h"
 
 #include <type_traits>
 
@@ -319,7 +320,7 @@ namespace Metasound
 				};
 			}
 
-			static const FName DataTypeName = GetMetasoundDataTypeName<TDataType>();
+			FCreateDataChannelFunction CreateDataChannelFunc(&FTransmissionDataChannelFactory::CreateDataChannel<TDataType>);
 
 			// Pack all of our various constructor lambdas to a single struct.
 			FDataTypeConstructorCallbacks Callbacks = 
@@ -328,13 +329,14 @@ namespace Metasound
 				MoveTemp(CreateFrontendInputClass),
 				MoveTemp(OutputNodeConstructor), 
 				MoveTemp(CreateFrontendOutputClass),
-				MoveTemp(ProxyGenerator) 
+				MoveTemp(ProxyGenerator),
+				MoveTemp(CreateDataChannelFunc)
 			};
 
 
 			FDataTypeRegistryInfo RegistryInfo;
 
-			RegistryInfo.DataTypeName = DataTypeName;
+			RegistryInfo.DataTypeName = GetMetasoundDataTypeName<TDataType>();
 			RegistryInfo.PreferredLiteralType = PreferredArgType;
 
 			RegistryInfo.bIsDefaultParsable = TIsParsable<TDataType, FLiteral::FNone>::Value;		
diff --git a/Engine/Plugins/Runtime/Metasound/Source/MetasoundFrontend/Public/MetasoundFrontendRegistries.h b/Engine/Plugins/Runtime/Metasound/Source/MetasoundFrontend/Public/MetasoundFrontendRegistries.h
--- a/Engine/Plugins/Runtime/Metasound/Source/MetasoundFrontend/Public/MetasoundFrontendRegistries.h
+++ b/Engine/Plugins/Runtime/Metasound/Source/MetasoundFrontend/Public/MetasoundFrontendRegistries.h
@@ -5,11 +5,12 @@
 #include "CoreMinimal.h"
 #include "IAudioProxyInitializer.h"
 #include "MetasoundDataReference.h"
+#include "MetasoundEnum.h"
 #include "MetasoundFrontendDocument.h"
 #include "MetasoundLiteral.h"
 #include "MetasoundNodeInterface.h"
 #include "MetasoundOperatorInterface.h"
-#include "MetasoundEnum.h"
+#include "MetasoundRouter.h"
 
 namespace Metasound
 {
@@ -62,6 +63,9 @@ namespace Metasound
 	// This function is used to create a proxy from a datatype's base uclass.
 	typedef TFunction<Audio::IProxyDataPtr(UObject*)> FCreateAudioProxyFunction;
 
+	// Creates a data channel for senders and receivers of this data type.
+	typedef TFunction<TSharedPtr<IDataChannel, ESPMode::ThreadSafe>(const FOperatorSettings&)> FCreateDataChannelFunction;
+
 	typedef TFunction<TUniquePtr<Metasound::INode>(const Metasound::FNodeInitData&)> FCreateMetasoundNodeFunction;
 	typedef TFunction<FMetasoundFrontendClass()> FCreateMetasoundFrontendClassFunction;
 
@@ -239,6 +243,8 @@ namespace Metasound
 
 		// For datatypes that use a UObject literal or a UObject literal array, this lambda generates a literal from the corresponding UObject.
 		FCreateAudioProxyFunction CreateAudioProxy;
+
+		FCreateDataChannelFunction CreateDataChannel;
 	};
 }
 
@@ -336,6 +342,8 @@ public:
 
 	TSharedPtr<const Metasound::Frontend::IEnumDataTypeInterface> GetEnumInterfaceForDataType(FName InDataType) const;
 
+	TSharedPtr<Metasound::IDataChannel, ESPMode::ThreadSafe> CreateDataChannelForDataType(const FName& InDataType, const Metasound::FOperatorSettings& InOperatorSettings) const;
+
 private:
 	FMetasoundFrontendRegistryContainer();
 
diff --git a/Engine/Plugins/Runtime/Metasound/Source/MetasoundFrontend/Public/MetasoundInstanceTransmitter.h b/Engine/Plugins/Runtime/Metasound/Source/MetasoundFrontend/Public/MetasoundInstanceTransmitter.h
--- a/Engine/Plugins/Runtime/Metasound/Source/MetasoundFrontend/Public/MetasoundInstanceTransmitter.h
+++ b/Engine/Plugins/Runtime/Metasound/Source/MetasoundFrontend/Public/MetasoundInstanceTransmitter.h
@@ -11,22 +11,44 @@
 
 namespace Metasound
 {
+	/** FMetasoundInstanceTransmitter provides a communication interface for 
+	 * sending values to a MetaSound instance. It relies on the send/receive transmission
+	 * system to ferry data from the transmitter to the MetaSound instance. Data will
+	 * be safely ushered across thread boundaries in scenarios where the instance
+	 * transmitter and metasound instance live on different threads. 
+	 */
 	class METASOUNDFRONTEND_API FMetasoundInstanceTransmitter : public IAudioInstanceTransmitter
 	{
 		FMetasoundInstanceTransmitter(const FMetasoundInstanceTransmitter&) = delete;
 		FMetasoundInstanceTransmitter& operator=(const FMetasoundInstanceTransmitter&) = delete;
 	public:
+		
+		/** FSendInfo describes the MetaSounds input parameters as well as the 
+		 * necessary information to route data to the instances inputs. 
+		 */
 		struct FSendInfo
 		{
+			/** Global address of instance input. */
 			FSendAddress Address;
+
+			/** Name of parameter on MetaSound instance. */
 			FName ParameterName;
+
+			/** Type name of parameter on MetaSound instance. */
 			FName TypeName;
 		};
 
+		/** Initialization parameters for a FMetasoundInstanceTransmitter. */
 		struct FInitParams
 		{
+			/** FOperatorSettings must match the operator settings of the MetaSound 
+			 * instance to ensure proper operation. */
 			FOperatorSettings OperatorSettings;
+
+			/** ID of the MetaSound instance.  */
 			uint64 InstanceID;
+
+			/** Available input parameters on MetaSound instance. */
 			TArray<FSendInfo> Infos;
 
 			FInitParams(const FOperatorSettings& InSettings, uint64 InInstanceID, const TArray<FSendInfo>& InInfos=TArray<FSendInfo>())
@@ -41,80 +63,61 @@ namespace Metasound
 		FMetasoundInstanceTransmitter(const FMetasoundInstanceTransmitter::FInitParams& InInitParams);
 		virtual ~FMetasoundInstanceTransmitter() = default;
 
+		/** Returns ID of the MetaSound instance associated with this transmitter. */
 		uint64 GetInstanceID() const override;
 
+		/** Set a float parameter on the MetaSound instance by name. 
+		 *
+		 * @param InParameterName - Name of MetaSound instance parameter.
+		 * @param InValue - Value to set. 
+		 *
+		 * @return true on success, false on failure.
+		 */
 		bool SetFloatParameter(const FName& InParameterName, float InValue) override;
+
+		/** Set a int parameter on the MetaSound instance by name. 
+		 *
+		 * @param InParameterName - Name of MetaSound instance parameter.
+		 * @param InValue - Value to set. 
+		 *
+		 * @return true on success, false on failure.
+		 */
 		bool SetIntParameter(const FName& InParameterName, int32 InValue) override;
-		bool SetBoolParameter(const FName& InParameterName, bool InValue) override;
 
+		/** Set a bool parameter on the MetaSound instance by name. 
+		 *
+		 * @param InParameterName - Name of MetaSound instance parameter.
+		 * @param InValue - Value to set. 
+		 *
+		 * @return true on success, false on failure.
+		 */
+		bool SetBoolParameter(const FName& InParameterName, bool InValue) override;
+		
+		/** Set a parameter using a literal.
+		 *
+		 * @param InParameterName - Name of MetaSound instance parameter.
+		 * @param InValue - Literal value used to construct paramter value. 
+		 *
+		 * @return true on success, false on failure. 
+		 */
+		bool SetParameterWithLiteral(const FName& InParameterName, const FLiteral& InValue);
+
+		/** Duplicate this transmitter interface. The transmitters association with
+		 * the MetaSound instance will be maintained. */
 		TUniquePtr<IAudioInstanceTransmitter> Clone() const override;
 
-		template<typename DataType>
-		bool SendValue(const FName& InParameterName, const DataType& InValue)
-		{
-			if (TSender<DataType>* Sender = FindSenderOfDataType<DataType>(InParameterName))
-			{
-				return Sender->Push(InValue);
-			}
-			return false;
-		}
-
 	private:
-
+		// Find FSendInfo by parameter name. 
 		const FSendInfo* FindSendInfo(const FName& InParameterName) const;
-		ISender* FindSender(const FName& InParameterName);
-
-		template<typename DataType> 
-		TSender<DataType>* FindSenderOfDataType(const FName& InParameterName)
-		{
-			using FSenderType = TSender<DataType>;
 
-			if (ISender* ExistingSender = FindSender(InParameterName))
-			{
-				if (ExistingSender->CheckType<FSenderType>())
-				{
-					FSenderType& DerivedSender = ExistingSender->GetAs<FSenderType>();
-					return &DerivedSender;
-				}
-				else
-				{
-					UE_LOG(LogMetasound, Warning, TEXT("Could not send to Metasound input. Mismatched data type [Recieved:%s, Expected: %s]"), *GetMetasoundDataTypeString<DataType>(), *ExistingSender->GetDataType().ToString());
-				}
-			}
-			else if (const FSendInfo* SendInfo = FindSendInfo(InParameterName))
-			{
-				TUniquePtr<TSender<DataType>> NewSender = CreateSender<DataType>(*SendInfo);
-				if (NewSender.IsValid())
-				{
-					FSenderType* DerivedSender = NewSender.Get();
-
-					InputSends.Add(InParameterName, MoveTemp(NewSender));
-
-					return DerivedSender;
-				}
-			}
-
-			return nullptr;
-		}
-
-		template<typename DataType>
-		TUniquePtr<TSender<DataType>> CreateSender(const FSendInfo& InInfo)
-		{
-			const float DelayTimeInSeconds = 10.0f; // TODO: likely want to remove this and opt for different protocols having different behaviors.
-			FSenderInitParams InitParams = {OperatorSettings, DelayTimeInSeconds};
-
-			if (InInfo.TypeName == GetMetasoundDataTypeName<DataType>())
-			{
-				return FDataTransmissionCenter::Get().RegisterNewSend<DataType>(InInfo.Address, InitParams);
-			}
-			else
-			{
-				UE_LOG(LogMetasound, Error, TEXT("Cannot create sender. Requested data type [TypeName:%s] does not match receive data type [TypeName:%s]"), *GetMetasoundDataTypeName<DataType>().ToString(), *InInfo.TypeName.ToString())
-			}
+		// Find ISender by parameter name. 
+		ISender* FindSender(const FName& InParameterName);
 
-			return TUniquePtr<TSender<DataType>>(nullptr);
-		}
+		// Create and store a new ISender for the given FSendInfo.
+		ISender* AddSender(const FSendInfo& InInfo);
 
+		// Create a new ISender from FSendInfo.
+		TUniquePtr<ISender> CreateSender(const FSendInfo& InInfo) const;
 
 		TArray<FSendInfo> SendInfos;
 		FOperatorSettings OperatorSettings;
diff --git a/Engine/Plugins/Runtime/Metasound/Source/MetasoundFrontend/Public/MetasoundRouter.h b/Engine/Plugins/Runtime/Metasound/Source/MetasoundFrontend/Public/MetasoundRouter.h
--- a/Engine/Plugins/Runtime/Metasound/Source/MetasoundFrontend/Public/MetasoundRouter.h
+++ b/Engine/Plugins/Runtime/Metasound/Source/MetasoundFrontend/Public/MetasoundRouter.h
@@ -22,7 +22,7 @@
  * This system is used by Send and Receive nodes in the Metasound graph, but can also be used in C++ for arbitrary message passing.
  *
  * Typical use case:
- * TSenderPtr<float> FloatSender FDataTransmissionCenter::Get().RegisterNewSend<float>(ExampleAddress, InitSettings);
+ * TSenderPtr<float> FloatSender FDataTransmissionCenter::Get().RegisterNewSender<float>(ExampleAddress, InitSettings);
  * TReceiverPtr<float> FloatReceiver FDataTransmissionCenter::Get().RegisterNewReceiver<float>(ExampleAddress, InitSettings);
  *
  * //...
@@ -164,7 +164,7 @@ namespace Metasound
 	class METASOUNDFRONTEND_API IReceiver : public IDataTransmissionBase
 	{
 	public:
-		IReceiver(TSharedPtr<IDataChannel> InDataChannel, FName DataType)
+		IReceiver(TSharedPtr<IDataChannel, ESPMode::ThreadSafe> InDataChannel, FName DataType)
 			: IDataTransmissionBase(DataType)
 			, DataChannel(InDataChannel)
 		{}
@@ -173,7 +173,7 @@ namespace Metasound
 		virtual ~IReceiver();
 
 	protected:
-		TSharedPtr<IDataChannel> DataChannel;
+		TSharedPtr<IDataChannel, ESPMode::ThreadSafe> DataChannel;
 	};
 
 	// This is a handle to something that will retrieve a datatype.
@@ -181,15 +181,18 @@ namespace Metasound
 	class METASOUNDFRONTEND_API ISender : public IDataTransmissionBase
 	{
 	public:
-		ISender(TSharedPtr<IDataChannel> InDataChannel, FName DataType)
+		ISender(TSharedPtr<IDataChannel, ESPMode::ThreadSafe> InDataChannel, FName DataType)
 			: IDataTransmissionBase(DataType)
 			, DataChannel(InDataChannel)
 		{}
 
 		virtual ~ISender();
 
+		// Constructs the value with a literal and pushes to data channel. 
+		virtual bool PushLiteral(const FLiteral& InLiteral) = 0;
+
 	protected:
-		TSharedPtr<IDataChannel> DataChannel;
+		TSharedPtr<IDataChannel, ESPMode::ThreadSafe> DataChannel;
 	};
 
 	struct FSenderInitParams
@@ -205,7 +208,7 @@ namespace Metasound
 
 	// This contains an intermediary buffer for use between the send and receive nodes.
 	// Typically only used by ISender and IReceiver implementations.
-	class METASOUNDFRONTEND_API IDataChannel : public IDataTransmissionBase, public TSharedFromThis<IDataChannel>
+	class METASOUNDFRONTEND_API IDataChannel : public IDataTransmissionBase, public TSharedFromThis<IDataChannel, ESPMode::ThreadSafe>
 	{
 	public:
 		virtual ~IDataChannel() = default;
@@ -315,8 +318,7 @@ namespace Metasound
 
 		static FName GetDataTypeName()
 		{
-			static FName TypeName = TDataReferenceTypeInfo<TDataType>::TypeName;
-			return TypeName;
+			return GetMetasoundDataTypeName<TDataType>();
 		}
 
 	protected:
@@ -356,26 +358,46 @@ namespace Metasound
 			}
 		}
 
+		bool PushLiteral(const FLiteral& InLiteral) override
+		{
+			if (!TLiteralTraits<TDataType>::IsParsable(InLiteral))
+			{
+				// The literal sent in is not compatible with the data type.
+				ensureMsgf(false, TEXT("Failed to send data. Data type [TypeName:%s] cannot be parsed from literal [Literal:%s]."), *GetMetasoundDataTypeString<TDataType>(), *LexToString(InLiteral));
+
+				return false;
+			}
+			else
+			{
+				return Push(TDataTypeLiteralFactory<TDataType>::CreateExplicitArgs(Params.OperatorSettings, InLiteral));
+			}
+		}
+
 		// Resets the delay for this specific sender. if this goes beyond the threshold set by BufferResizeThreshold, we reallocate.
 		bool SetDelay(float InSeconds)
 		{
 			Params.DelayTimeInSeconds = InSeconds;
-			NumElementsToDelayBy = FMath::Min(InSeconds * Params.OperatorSettings.GetActualBlockRate(), 1.f);
+			// TODO: LOL. I was wondering why the having a delay didn't change the behavior.
+			// Probably for the best since it does not make much sense to delay
+			// sent values. The FMath::Min() results in it being either 1 or 0. 
+			//NumElementsToDelayBy = FMath::Min(InSeconds * Params.OperatorSettings.GetActualBlockRate(), 1.f);
 
+			/*
 			if (NumElementsToDelayBy >= SenderBuffer.GetCapacity())
 			{
 				SenderBuffer.SetCapacity(NumElementsToDelayBy * BufferResizeThreshold);
 			}
+			*/
+			SenderBuffer.SetCapacity(2);
 			return true;
 		}
 
 		static FName GetDataTypeName()
 		{
-			static FName TypeName = TDataReferenceTypeInfo<TDataType>::TypeName;
-			return TypeName;
+			return GetMetasoundDataTypeName<TDataType>();
 		}
 
-		TSender(const FSenderInitParams& InitParams, TSharedPtr<IDataChannel> InDataChannel)
+		TSender(const FSenderInitParams& InitParams, TSharedPtr<IDataChannel, ESPMode::ThreadSafe> InDataChannel)
 			: ISender(InDataChannel, GetDataTypeName())
 			, Params(InitParams)
 		{
@@ -440,11 +462,10 @@ namespace Metasound
 
 		static FName GetDataTypeName()
 		{
-			static FName TypeName = TDataReferenceTypeInfo<TDataType>::TypeName;
-			return TypeName;
+			return GetMetasoundDataTypeName<TDataType>();
 		}
 
-		TReceiver(const FReceiverInitParams& InitParams, TSharedPtr<IDataChannel> InDataChannel)
+		TReceiver(const FReceiverInitParams& InitParams, TSharedPtr<IDataChannel, ESPMode::ThreadSafe> InDataChannel)
 			: IReceiver(InDataChannel, GetDataTypeName())
 			, OperatorSettings(InitParams.OperatorSettings)
 		{
@@ -515,8 +536,7 @@ namespace Metasound
 
 		static FName GetDataTypeName()
 		{
-			static FName TypeName = TDataReferenceTypeInfo<TDataType>::TypeName;
-			return TypeName;
+			return GetMetasoundDataTypeName<TDataType>();
 		}
 
 		virtual FName GetDataType() override
@@ -602,6 +622,12 @@ namespace Metasound
 			return true;
 		}
 
+		bool PushLiteral(const FLiteral& InLiteral) override
+		{
+			ensureMsgf(false, TEXT("Cannot push literal to audio format"));
+			return false;
+		}
+
 		// Resets the delay for this sender.
 		bool SetDelay(float InSeconds)
 		{
@@ -629,11 +655,10 @@ namespace Metasound
 
 		static FName GetDataTypeName()
 		{
-			static FName TypeName = TDataReferenceTypeInfo<TDataType>::TypeName;
-			return TypeName;
+			return GetMetasoundDataTypeName<TDataType>();
 		}
 
-		TAudioSender(const FSenderInitParams& InitParams, TSharedPtr<IDataChannel> InDataChannel, TArray<Audio::FPatchInput>&& Inputs)
+		TAudioSender(const FSenderInitParams& InitParams, TSharedPtr<IDataChannel, ESPMode::ThreadSafe> InDataChannel, TArray<Audio::FPatchInput>&& Inputs)
 			: ISender(InDataChannel, GetDataTypeName())
 			, DataChannelInputs(MoveTemp(Inputs))
 			, Params(InitParams)
@@ -688,11 +713,10 @@ namespace Metasound
 
 		static FName GetDataTypeName()
 		{
-			static FName TypeName = TDataReferenceTypeInfo<TDataType>::TypeName;
-			return TypeName;
+			return GetMetasoundDataTypeName<TDataType>();
 		}
 
-		TAudioReceiver(const FReceiverInitParams& InitParams, TSharedPtr<IDataChannel> InDataChannel, TArray<Audio::FPatchOutputStrongPtr>&& Outputs)
+		TAudioReceiver(const FReceiverInitParams& InitParams, TSharedPtr<IDataChannel, ESPMode::ThreadSafe> InDataChannel, TArray<Audio::FPatchOutputStrongPtr>&& Outputs)
 			: IReceiver(InDataChannel, GetDataTypeName())
 			, DataChannelOutputs(MoveTemp(Outputs))
 			, Params(InitParams)
@@ -738,8 +762,7 @@ namespace Metasound
 
 		static FName GetDataTypeName()
 		{
-			static FName TypeName = TDataReferenceTypeInfo<TDataType>::TypeName;
-			return TypeName;
+			return GetMetasoundDataTypeName<TDataType>();
 		}
 
 		virtual FName GetDataType() override
@@ -786,19 +809,19 @@ namespace Metasound
 	// SFINAE for creating the correct data channel type for the given datatype:
 
 	template<typename TDataType, typename TEnableIf<std::is_copy_constructible<TDataType>::value && !TIsDerivedFrom<TDataType, IAudioDataType>::Value, bool>::Type = true>
-	TSharedRef<IDataChannel> MakeDataChannel(const FOperatorSettings& InSettings)
+	TSharedRef<IDataChannel, ESPMode::ThreadSafe> MakeDataChannel(const FOperatorSettings& InSettings)
 	{
 		return MakeShareable(new TCopyableDataChannel<TDataType>(InSettings));
 	}
 
 	template<typename TDataType, typename TEnableIf<TIsDerivedFrom<TDataType, IAudioDataType>::Value, bool>::Type = true>
-	TSharedRef<IDataChannel> MakeDataChannel(const FOperatorSettings& InSettings)
+	TSharedRef<IDataChannel, ESPMode::ThreadSafe> MakeDataChannel(const FOperatorSettings& InSettings)
 	{
 		return MakeShareable(new TAudioDataChannel<TDataType>(InSettings));
 	}
 
 	template<typename TDataType, typename TEnableIf<!std::is_copy_constructible<TDataType>::value && !TIsDerivedFrom<TDataType, IAudioDataType>::Value, bool>::Type = true>
-	TSharedRef<IDataChannel> MakeDataChannel(const FOperatorSettings& InSettings)
+	TSharedRef<IDataChannel, ESPMode::ThreadSafe> MakeDataChannel(const FOperatorSettings& InSettings)
 	{
 		return MakeShareable(new TNonOperationalDataChannel<TDataType>(InSettings));
 	}
@@ -859,53 +882,47 @@ namespace Metasound
 
 		FName GetDatatypeForChannel(FName InChannelName);
 
+		TSharedPtr<IDataChannel, ESPMode::ThreadSafe> GetDataChannel(const FName& InDataTypeName, const FName& InChannelName, const FOperatorSettings& InOperatorSettings);
+
+		TUniquePtr<ISender> RegisterNewSender(const FName& InDataTypeName, const FName& InChannelName, const FSenderInitParams& InitParams);
+		TUniquePtr<IReceiver> RegisterNewReceiver(const FName& InDataTypeName, const FName& InChannelName, const FReceiverInitParams& InitParams);
+
 		template<typename TDataType>
-		TSenderPtr<TDataType> RegisterNewSend(FName InChannelName, const FSenderInitParams& InitParams)
+		TReceiverPtr<TDataType> RegisterNewReceiver(const FName& InChannelName, const FReceiverInitParams& InitParams)
 		{
-			FScopeLock ScopeLock(&DataChannelMapMutationLock);
+			TReceiverPtr<TDataType> Receiver;
 
-			if (TSharedRef<IDataChannel>* ExistingChannelPtr = DataChannelMap.Find(InChannelName))
-			{
-				TSharedRef<IDataChannel>& ExistingChannel = *ExistingChannelPtr;
+			TUniquePtr<IReceiver> ReceiverBase = RegisterNewReceiver(GetMetasoundDataTypeName<TDataType>(), InChannelName, InitParams);
 
-				// ensure that the existing data channel is of the same datatype as the send we're trying to register.
-				return Downcast<TDataType>(ExistingChannel->NewSender(InitParams));
-			}
-			else
+			if (ReceiverBase.IsValid())
 			{
-				// This is the first time we're seeing this, add it to the map.
-				TSharedRef<IDataChannel>& ExistingChannel = DataChannelMap.Add(InChannelName, MakeDataChannel<TDataType>(InitParams.OperatorSettings));
-				return Downcast<TDataType>(ExistingChannel->NewSender(InitParams));
+				Receiver = Downcast<TDataType>(MoveTemp(ReceiverBase));
 			}
+
+			return MoveTemp(Receiver);
 		}
 
 		template<typename TDataType>
-		TReceiverPtr<TDataType> RegisterNewReceiver(FName InChannelName, const FReceiverInitParams& InitParams)
+		TSenderPtr<TDataType> RegisterNewSender(const FName& InChannelName, const FSenderInitParams& InitParams)
 		{
-			FScopeLock ScopeLock(&DataChannelMapMutationLock);
+			TSenderPtr<TDataType> Sender;
 
-			if (TSharedRef<IDataChannel>* ExistingChannelPtr = DataChannelMap.Find(InChannelName))
-			{
-				TSharedRef<IDataChannel>& ExistingChannel = *ExistingChannelPtr;
+			TUniquePtr<ISender> SenderBase = RegisterNewSender(GetMetasoundDataTypeName<TDataType>(), InChannelName, InitParams);
 
-				TUniquePtr<IReceiver> BaseReceiver = ExistingChannel->NewReceiver(InitParams);
-				return Downcast<TDataType>(MoveTemp(BaseReceiver));
-			}
-			else
+			if (SenderBase.IsValid())
 			{
-				// This is the first time we're seeing this, add it to the map.
-				TSharedRef<IDataChannel>& ExistingChannel = DataChannelMap.Add(InChannelName, MakeDataChannel<TDataType>(InitParams.OperatorSettings));
-
-				return Downcast<TDataType>(ExistingChannel->NewReceiver(InitParams));
+				Sender = Downcast<TDataType>(MoveTemp(SenderBase));
 			}
+
+			return MoveTemp(Sender);
 		}
 
-		bool PushLiteral(FName InChannelName, const FLiteral& InParam)
+		bool PushLiteral(const FName& InChannelName, const FLiteral& InParam)
 		{
 			FScopeLock ScopeLock(&DataChannelMapMutationLock);
-			if (TSharedRef<IDataChannel>* ExistingChannelPtr = DataChannelMap.Find(InChannelName))
+			if (TSharedRef<IDataChannel, ESPMode::ThreadSafe>* ExistingChannelPtr = DataChannelMap.Find(InChannelName))
 			{
-				TSharedRef<IDataChannel>& ExistingChannel = *ExistingChannelPtr;
+				TSharedRef<IDataChannel, ESPMode::ThreadSafe>& ExistingChannel = *ExistingChannelPtr;
 				return ExistingChannel->PushLiteral(InParam);
 			}
 			else
@@ -923,7 +940,7 @@ namespace Metasound
 		{}
 
 	private:
-		TMap<FName, TSharedRef<IDataChannel>> DataChannelMap;
+		TMap<FName, TSharedRef<IDataChannel, ESPMode::ThreadSafe>> DataChannelMap;
 		FCriticalSection DataChannelMapMutationLock;
 	};
 
@@ -936,38 +953,38 @@ namespace Metasound
 
 		FName GetDatatypeForChannel(uint64 InInstanceID, FName InChannelName);
 
+		TUniquePtr<ISender> RegisterNewSender(uint64 InInstanceID, const FName& InDataTypeName, const FName& InChannelName, const FSenderInitParams& InitParams);
+
 		template<typename TDataType>
-		TSenderPtr<TDataType> RegisterNewSend(uint64 InInstanceID, FName InChannelName, const FSenderInitParams& InitParams)
+		TSenderPtr<TDataType> RegisterNewSender(uint64 InInstanceID, FName InChannelName, const FSenderInitParams& InitParams)
 		{
-			FScopeLock ScopeLock(&InstanceRouterMapMutationLock);
+			TSenderPtr<TDataType> Sender;
 
-			if (FAddressRouter* AddressRouter = InstanceRouterMap.Find(InInstanceID))
-			{
-				return AddressRouter->RegisterNewSend<TDataType>(InChannelName, InitParams);
-			}
-			else
+			TUniquePtr<ISender> SenderBase = RegisterNewSender(InInstanceID, GetMetasoundDataTypeName<TDataType>(), InChannelName, InitParams);
+
+			if (SenderBase.IsValid())
 			{
-				// This is the first time we're seeing this, add it to the map.
-				FAddressRouter& NewRouter = InstanceRouterMap.Add(InInstanceID, FAddressRouter());
-				return NewRouter.RegisterNewSend<TDataType>(InChannelName, InitParams);
+				Sender = Downcast<TDataType>(SenderBase);
 			}
+
+			return MoveTemp(Sender);
 		}
 
+		TUniquePtr<IReceiver> RegisterNewReceiver(uint64 InInstanceID, const FName& InDataTypeName, const FName& InChannelName, const FReceiverInitParams& InitParams);
+
 		template<typename TDataType>
 		TReceiverPtr<TDataType> RegisterNewReceiver(uint64 InInstanceID, FName InChannelName, const FReceiverInitParams& InitParams)
 		{
-			FScopeLock ScopeLock(&InstanceRouterMapMutationLock);
+			TReceiverPtr<TDataType> Receiver;
 
-			if (FAddressRouter* AddressRouter = InstanceRouterMap.Find(InInstanceID))
-			{
-				return AddressRouter->RegisterNewReceiver<TDataType>(InChannelName, InitParams);
-			}
-			else
+			TUniquePtr<IReceiver> ReceiverBase = RegisterNewReceiver(InInstanceID, GetMetasoundDataTypeName<TDataType>(), InChannelName, InitParams);
+
+			if (ReceiverBase.IsValid())
 			{
-				// This is the first time we're seeing this, add it to the map.
-				FAddressRouter& NewRouter = InstanceRouterMap.Add(InInstanceID, FAddressRouter());
-				return NewRouter.RegisterNewReceiver<TDataType>(InChannelName, InitParams);
+				Receiver = Downcast<TDataType>(ReceiverBase);
 			}
+
+			return MoveTemp(ReceiverBase);
 		}
 
 	private:
@@ -1031,7 +1048,7 @@ namespace Metasound
 				check(SubsystemInterface);
 
 				// Ensure that this subsystem can support this send type.
-				return SubsystemInterface->CanSupportDataType(FName(TDataReferenceTypeInfo<TDataType>::TypeName));
+				return SubsystemInterface->CanSupportDataType(GetMetasoundDataTypeName<TDataType>());
 			}
 			else
 			{
@@ -1041,84 +1058,51 @@ namespace Metasound
 			}
 		}
 
+		// Creates a new object to push data to an address.
+		// Returns a new sender, or nullptr if registration failed.
+		TUniquePtr<ISender> RegisterNewSender(const FName& InDataTypeName, const FSendAddress& InAddress, const FSenderInitParams& InitParams);
+
 		// Creates a new object to push data to an address.
 		// Returns a new sender, or nullptr if registration failed.
 		template<typename TDataType>
-		TSenderPtr<TDataType> RegisterNewSend(const FSendAddress& InAddress, const FSenderInitParams& InitParams)
+		TSenderPtr<TDataType> RegisterNewSender(const FSendAddress& InAddress, const FSenderInitParams& InitParams)
 		{
-			if (InAddress.Subsystem == GetSubsystemNameForSendScope(ETransmissionScope::ThisInstanceOnly))
-			{
-				return InstanceRouter.RegisterNewSend<TDataType>(InAddress.MetasoundInstanceID, InAddress.ChannelName, InitParams);
-			}
-			else if (InAddress.Subsystem == GetSubsystemNameForSendScope(ETransmissionScope::Global))
-			{
-				return GlobalRouter.RegisterNewSend<TDataType>(InAddress.ChannelName, InitParams);
-			}
-			else if (FSubsystemData* FoundSubsystem = SubsystemRouters.Find(InAddress.Subsystem))
-			{
-				ITransmissionSubsystem* SubsystemInterface = FoundSubsystem->SubsystemPtr;
-
-				// Ensure that this subsystem can support this send type.
-				if (!ensureAlways(SubsystemInterface && SubsystemInterface->CanSupportDataType(FName(TDataReferenceTypeInfo<TDataType>::TypeName))))
-				{
-					return TSenderPtr<TDataType>();
-				}
+			TSenderPtr<TDataType> Sender;
+			
+			TUniquePtr<ISender> SenderBase = RegisterNewSender(GetMetasoundDataTypeName<TDataType>(), InAddress, InitParams);
 
-				FScopeLock ScopeLock(&SubsystemRoutersMutationLock);
-				TSenderPtr<TDataType> Sender = FoundSubsystem->AddressRouter.RegisterNewSend<TDataType>(InAddress.ChannelName, InitParams);
-				FoundSubsystem->SubsystemPtr->OnNewSendRegistered(InAddress, FName(TDataReferenceTypeInfo<TDataType>::TypeName));
-				return MoveTemp(Sender);
-			}
-			else
+			if (SenderBase.IsValid())
 			{
-				// Otherwise, the subsystem FName was invalid.
-				ensureAlways(false);
-				return TSenderPtr<TDataType>();
+				Sender = Downcast<TDataType>(MoveTemp(SenderBase));
 			}
+
+			return Sender;
 		}
 
+		// Registers a new object to poll data from an address.
+		// Returns a new receiver, or nullptr if registration failed.
+		TUniquePtr<IReceiver> RegisterNewReceiver(const FName& InDataTypeName, const FSendAddress& InAddress, const FReceiverInitParams& InitParams);
+
 		// Registers a new object to poll data from an address.
 		// Returns a new receiver, or nullptr if registration failed.
 		template<typename TDataType>
 		TReceiverPtr<TDataType> RegisterNewReceiver(const FSendAddress& InAddress, const FReceiverInitParams& InitParams)
 		{
-			if (InAddress.Subsystem == GetSubsystemNameForSendScope(ETransmissionScope::ThisInstanceOnly))
-			{
-				return InstanceRouter.RegisterNewReceiver<TDataType>(InAddress.MetasoundInstanceID, InAddress.ChannelName, InitParams);
-			}
-			else if (InAddress.Subsystem == GetSubsystemNameForSendScope(ETransmissionScope::Global))
-			{
-				return GlobalRouter.RegisterNewReceiver<TDataType>(InAddress.ChannelName, InitParams);
-			}
-			else if (FSubsystemData* FoundSubsystem = SubsystemRouters.Find(InAddress.Subsystem))
-			{
-				ITransmissionSubsystem* SubsystemInterface = FoundSubsystem->SubsystemPtr;
+			TReceiverPtr<TDataType> Receiver;
 
-				// Ensure that this subsystem can support this send type.
-				if (!ensureAlways(SubsystemInterface && SubsystemInterface->CanSupportDataType(FName(TDataReferenceTypeInfo<TDataType>::TypeName))))
-				{
-					return TReceiverPtr<TDataType>();
-				}
+			TUniquePtr<IReceiver> ReceiverBase = RegisterNewReceiver(GetMetasoundDataTypeName<TDataType>(), InAddress, InitParams);
 
-				FScopeLock ScopeLock(&SubsystemRoutersMutationLock);
-				TReceiverPtr<TDataType> Receiver = FoundSubsystem->AddressRouter.RegisterNewReceiver<TDataType>(InAddress.ChannelName, InitParams);
-				FoundSubsystem->SubsystemPtr->OnNewReceiverRegistered(InAddress, FName(TDataReferenceTypeInfo<TDataType>::TypeName));
-				return MoveTemp(Receiver);
-			}
-			else
+			if (ReceiverBase.IsValid())
 			{
-				// Otherwise, the subsystem FName was invalid.
-				ensureAlways(false);
-				return TReceiverPtr<TDataType>();
+				Receiver = Downcast<TDataType>(MoveTemp(ReceiverBase));
 			}
+
+			return Receiver;
 		}
 
 		// Pushes a literal parameter to a specific data channel in the global router.
 		// returns false if the literal type isn't supported.
-		bool PushLiteral(FName GlobalChannelName, const FLiteral& InParam)
-		{
-			return GlobalRouter.PushLiteral(GlobalChannelName, InParam);
-		}
+		bool PushLiteral(FName GlobalChannelName, const FLiteral& InParam);
 
 	private:
 		// Single map of FNames to IDataChannels
@@ -1132,27 +1116,10 @@ namespace Metasound
 		FCriticalSection SubsystemRoutersMutationLock;
 
 		// These are used by the constructor and destructor of implementations of ITransmissionSubsystem.
-		void RegisterSubsystem(ITransmissionSubsystem* InSystem, FName InSubsytemName)
-		{
-			FScopeLock ScopeLock(&SubsystemRoutersMutationLock);
-
-			//check to make sure we're not adding a subsystem twice.
-			check(!SubsystemRouters.Contains(InSubsytemName));
-
-			SubsystemRouters.Emplace(InSubsytemName, InSystem);
-		}
-
-		void UnregisterSubsystem(FName InSubsystemName)
-		{
-			FScopeLock ScopeLock(&SubsystemRoutersMutationLock);
-
-			//check to make sure we're not adding a subsystem twice.
-			check(SubsystemRouters.Contains(InSubsystemName));
-
-			SubsystemRouters.Remove(InSubsystemName);
-		}
+		void RegisterSubsystem(ITransmissionSubsystem* InSystem, FName InSubsytemName);
+		void UnregisterSubsystem(FName InSubsystemName);
 
-		FDataTransmissionCenter();
+		FDataTransmissionCenter() = default;
 
 		friend class ITransmissionSubsystem;
 	};
diff --git a/Engine/Plugins/Runtime/Metasound/Source/MetasoundFrontend/Public/MetasoundSendNode.h b/Engine/Plugins/Runtime/Metasound/Source/MetasoundFrontend/Public/MetasoundSendNode.h
--- a/Engine/Plugins/Runtime/Metasound/Source/MetasoundFrontend/Public/MetasoundSendNode.h
+++ b/Engine/Plugins/Runtime/Metasound/Source/MetasoundFrontend/Public/MetasoundSendNode.h
@@ -82,7 +82,7 @@ namespace Metasound
 					, SendAddress(InSendAddress)
 					, CachedSendAddress(*InSendAddress)
 					, CachedSenderParams({InOperatorSettings, 0.0f})
-					, Sender(FDataTransmissionCenter::Get().RegisterNewSend<TDataType>(CachedSendAddress, CachedSenderParams))
+					, Sender(FDataTransmissionCenter::Get().RegisterNewSender<TDataType>(CachedSendAddress, CachedSenderParams))
 				{
 				}
 
@@ -106,7 +106,7 @@ namespace Metasound
 					if (SendAddress->ChannelName != CachedSendAddress.ChannelName)
 					{
 						CachedSendAddress = *SendAddress;
-						Sender = FDataTransmissionCenter::Get().RegisterNewSend<TDataType>(CachedSendAddress, CachedSenderParams);
+						Sender = FDataTransmissionCenter::Get().RegisterNewSender<TDataType>(CachedSendAddress, CachedSenderParams);
 						check(Sender.IsValid());
 					}
 
diff --git a/Engine/Plugins/Runtime/Metasound/Source/MetasoundFrontend/Public/MetasoundTransmissionRegistration.h b/Engine/Plugins/Runtime/Metasound/Source/MetasoundFrontend/Public/MetasoundTransmissionRegistration.h
new file mode 100644
--- /dev/null
+++ b/Engine/Plugins/Runtime/Metasound/Source/MetasoundFrontend/Public/MetasoundTransmissionRegistration.h
@@ -0,0 +1,87 @@
+// Copyright Epic Games, Inc. All Rights Reserved.
+
+#pragma once
+
+#include "MetasoundRouter.h"
+#include "MetasoundFrontendRegistries.h"
+
+#include <type_traits>
+
+namespace Metasound
+{
+	namespace MetasoundTransmissionPrivate
+	{
+		// TTransmissionSupport determines whether the send/receive system is 
+		// supported for a given data type. It is used to add send/receive
+		// support when possible, and avoid errors when registering data types
+		// that do not support the transmission system. 
+		template<typename DataType>
+		struct TTransmissionSupport
+		{
+		private:
+
+			// All types that support copy assignment and copy construction can be
+			// used in the transmission system. 
+			static constexpr bool bIsCopyable = std::is_copy_assignable<DataType>::value && std::is_copy_constructible<DataType>::value;
+
+			// IAudioDataType derived classes have specialized templates for the
+			// transmission system.
+			static constexpr bool bIsAudioDataType = std::is_base_of<IAudioDataType, DataType>::value;
+
+		public:
+			
+			static constexpr bool bIsTransmissionSupported = bIsCopyable || bIsAudioDataType; 
+		};
+
+		// At the time of writing this code, TArray incorrectly defines a copy constructor
+		// even when the underlying elements or not copyable. Normally this results
+		// in compilation errors when trying to utilize the erroneously defined 
+		// TArray copy constructor. This specialization checks the underlying TArray
+		// element type to determine if the array is copyable.
+		template<typename ElementType>
+		struct TTransmissionSupport<TArray<ElementType>>
+		{
+			static constexpr bool bIsCopyable = std::is_copy_assignable<ElementType>::value && std::is_copy_constructible<ElementType>::value;
+
+		public:
+			
+			static constexpr bool bIsTransmissionSupported = bIsCopyable;
+		};
+	}
+
+	struct FTransmissionDataChannelFactory
+	{
+		/**  Create a transmission IDataChannel given a data type.
+		 *
+		 * This function is defined if a data type is supported by the transmission system. 
+		 */
+		template<
+			typename DataType,
+			typename std::enable_if<
+				MetasoundTransmissionPrivate::TTransmissionSupport<DataType>::bIsTransmissionSupported, 
+				bool
+			>::type = true
+		>
+		static TSharedPtr<IDataChannel, ESPMode::ThreadSafe> CreateDataChannel(const FOperatorSettings& InOperatorSettings)
+		{
+			return MakeDataChannel<DataType>(InOperatorSettings);
+		}
+
+		/**  Create a transmission IDataChannel given a data type.
+		 *
+		 * This function is defined if a data type is not supported by the transmission system. 
+		 * It returns null data channel. 
+		 */
+		template<
+			typename DataType,
+			typename std::enable_if<
+				!MetasoundTransmissionPrivate::TTransmissionSupport<DataType>::bIsTransmissionSupported, 
+				bool
+			>::type = true
+		>
+		static TSharedPtr<IDataChannel, ESPMode::ThreadSafe> CreateDataChannel(const FOperatorSettings& InOperatorSettings)
+		{
+			return TSharedPtr<IDataChannel, ESPMode::ThreadSafe>(nullptr);
+		}
+	};
+}
diff --git a/Engine/Plugins/Runtime/Metasound/Source/MetasoundGraphCore/Private/MetasoundLiteral.cpp b/Engine/Plugins/Runtime/Metasound/Source/MetasoundGraphCore/Private/MetasoundLiteral.cpp
--- a/Engine/Plugins/Runtime/Metasound/Source/MetasoundGraphCore/Private/MetasoundLiteral.cpp
+++ b/Engine/Plugins/Runtime/Metasound/Source/MetasoundGraphCore/Private/MetasoundLiteral.cpp
@@ -10,73 +10,9 @@
 namespace Metasound
 {
 #if METASOUND_DEBUG_LITERALS
-	void FLiteral::InitDebugString()
+	void FLiteral::InitDebugString() const
 	{
-		switch (GetType())
-		{
-			case ELiteralType::None:
-				DebugString = TEXT("NONE");
-			break;
-
-			case ELiteralType::Boolean:
-				DebugString = FString::Printf(TEXT("Boolean: %s"), Value.Get<bool>() ? TEXT("true") : TEXT("false"));
-			break;
-
-			case ELiteralType::Integer:
-				DebugString = FString::Printf(TEXT("Int32: %d"), Value.Get<int32>());
-			break;
-
-			case ELiteralType::Float:
-				DebugString = FString::Printf(TEXT("Float: %f"), Value.Get<float>());
-			break;
-
-			case ELiteralType::String:
-				DebugString = FString::Printf(TEXT("String: %s"), *Value.Get<FString>());
-			break;
-
-			case ELiteralType::UObjectProxy:
-			{
-				FString ProxyType = TEXT("nullptr");
-				if (Value.Get<Audio::IProxyDataPtr>().IsValid())
-				{
-					ProxyType = Value.Get<Audio::IProxyDataPtr>()->GetProxyTypeName().ToString();
-				}
-				DebugString = FString::Printf(TEXT("Audio::IProxyDataPtr: %s"), *ProxyType);
-			}
-			break;
-
-			case ELiteralType::NoneArray:
-				DebugString = TEXT("TArray<NONE>");
-			break;
-
-			case ELiteralType::BooleanArray:
-				DebugString = FString::Printf(TEXT("TArray<Boolean>"));
-			break;
-
-			case ELiteralType::IntegerArray:
-				DebugString = FString::Printf(TEXT("TArray<int32>"));
-			break;
-
-			case ELiteralType::FloatArray:
-				DebugString = FString::Printf(TEXT("TArray<float>"));
-			break;
-
-			case ELiteralType::StringArray:
-				DebugString = FString::Printf(TEXT("TArray<FString>"));
-			break;
-
-			case ELiteralType::UObjectProxyArray:
-				DebugString = FString::Printf(TEXT("TArray<Audio::IProxyDataPtr>"));
-			break;
-
-			case ELiteralType::Invalid:
-				DebugString = TEXT("INVALID");
-			break;
-
-			default:
-				static_assert(static_cast<int32>(ELiteralType::Invalid) == 12, "Possible missing ELiteralType case coverage");
-				checkNoEntry();
-		}
+		DebugString = LexToString(*this);
 	}
 #endif // METASOUND_DEBUG_LITERALS
 
@@ -257,4 +193,63 @@ namespace Metasound
 
 }
 
+FString LexToString(const Metasound::FLiteral& InLiteral)
+{
+	using namespace Metasound;
+
+	switch (InLiteral.GetType())
+	{
+		case ELiteralType::None:
+			return TEXT("NONE");
+
+		case ELiteralType::Boolean:
+			return FString::Printf(TEXT("Boolean: %s"), InLiteral.Value.Get<bool>() ? TEXT("true") : TEXT("false"));
+
+		case ELiteralType::Integer:
+			return FString::Printf(TEXT("Int32: %d"), InLiteral.Value.Get<int32>());
+
+		case ELiteralType::Float:
+			return FString::Printf(TEXT("Float: %f"), InLiteral.Value.Get<float>());
+
+		case ELiteralType::String:
+			return FString::Printf(TEXT("String: %s"), *InLiteral.Value.Get<FString>());
+
+		case ELiteralType::UObjectProxy:
+		{
+			FString ProxyType = TEXT("nullptr");
+			if (InLiteral.Value.Get<Audio::IProxyDataPtr>().IsValid())
+			{
+				ProxyType = InLiteral.Value.Get<Audio::IProxyDataPtr>()->GetProxyTypeName().ToString();
+			}
+			return FString::Printf(TEXT("Audio::IProxyDataPtr: %s"), *ProxyType);
+		}
+		break;
+
+		case ELiteralType::NoneArray:
+			return TEXT("TArray<NONE>");
+
+		case ELiteralType::BooleanArray:
+			return TEXT("TArray<Boolean>");
+
+		case ELiteralType::IntegerArray:
+			return TEXT("TArray<int32>");
+
+		case ELiteralType::FloatArray:
+			return TEXT("TArray<float>");
+
+		case ELiteralType::StringArray:
+			return TEXT("TArray<FString>");
+
+		case ELiteralType::UObjectProxyArray:
+			return TEXT("TArray<Audio::IProxyDataPtr>");
+
+		case ELiteralType::Invalid:
+			return TEXT("INVALID");
+
+		default:
+			static_assert(static_cast<int32>(ELiteralType::Invalid) == 12, "Possible missing ELiteralType case coverage");
+			checkNoEntry();
+			return TEXT("INVALID");
+	}
+}
 
diff --git a/Engine/Plugins/Runtime/Metasound/Source/MetasoundGraphCore/Public/MetasoundLiteral.h b/Engine/Plugins/Runtime/Metasound/Source/MetasoundGraphCore/Public/MetasoundLiteral.h
--- a/Engine/Plugins/Runtime/Metasound/Source/MetasoundGraphCore/Public/MetasoundLiteral.h
+++ b/Engine/Plugins/Runtime/Metasound/Source/MetasoundGraphCore/Public/MetasoundLiteral.h
@@ -139,7 +139,7 @@ namespace Metasound
 
 		FString DebugString;
 
-		void InitDebugString();
+		void InitDebugString() const;
 
 #endif // METAOUND_DEBUG_LITERALS
 	};
@@ -331,6 +331,6 @@ namespace Metasound
 			return FLiteral::FNone{};
 		}
 	};
-
-
 }
+
+METASOUNDGRAPHCORE_API FString LexToString(const Metasound::FLiteral& InLiteral);
