More StandardNode Categories (filter,generators,trigger utils).

diff --git a/Engine/Plugins/Runtime/Metasound/Source/MetasoundStandardNodes/Private/MetasoundBasicFilters.cpp b/Engine/Plugins/Runtime/Metasound/Source/MetasoundStandardNodes/Private/MetasoundBasicFilters.cpp
--- a/Engine/Plugins/Runtime/Metasound/Source/MetasoundStandardNodes/Private/MetasoundBasicFilters.cpp
+++ b/Engine/Plugins/Runtime/Metasound/Source/MetasoundStandardNodes/Private/MetasoundBasicFilters.cpp
@@ -9,6 +9,7 @@
 #include "MetasoundNodeRegistrationMacro.h"
 #include "MetasoundOperatorSettings.h"
 #include "MetasoundStandardNodesNames.h"
+#include "MetasoundStandardNodesCategories.h"
 
 #include "DSP/Filter.h"
 #include "DSP/InterpolatedOnePole.h"
@@ -121,6 +122,7 @@ namespace Metasound
 			Info.Author = PluginAuthor;
 			Info.PromptIfMissing = PluginNodeMissingPrompt;
 			Info.DefaultInterface = DeclareVertexInterface();
+			Info.CategoryHierarchy.Emplace(StandardNodes::Filters);
 
 			return Info;
 		};
@@ -312,6 +314,7 @@ namespace Metasound
 			Info.Author = PluginAuthor;
 			Info.PromptIfMissing = PluginNodeMissingPrompt;
 			Info.DefaultInterface = DeclareVertexInterface();
+			Info.CategoryHierarchy.Emplace(StandardNodes::Filters);
 
 			return Info;
 		};
@@ -488,6 +491,7 @@ namespace Metasound
 			Info.Author = PluginAuthor;
 			Info.PromptIfMissing = PluginNodeMissingPrompt;
 			Info.DefaultInterface = DeclareVertexInterface();
+			Info.CategoryHierarchy.Emplace(StandardNodes::Filters);
 
 			return Info;
 		};
@@ -632,6 +636,7 @@ namespace Metasound
 			Info.Author = PluginAuthor;
 			Info.PromptIfMissing = PluginNodeMissingPrompt;
 			Info.DefaultInterface = DeclareVertexInterface();
+			Info.CategoryHierarchy.Emplace(StandardNodes::Filters);
 
 			return Info;
 		};
@@ -805,6 +810,7 @@ namespace Metasound
 			Info.Author = PluginAuthor;
 			Info.PromptIfMissing = PluginNodeMissingPrompt;
 			Info.DefaultInterface = DeclareVertexInterface();
+			Info.CategoryHierarchy.Emplace(StandardNodes::Filters);
 
 			return Info;
 		};
diff --git a/Engine/Plugins/Runtime/Metasound/Source/MetasoundStandardNodes/Private/MetasoundNoiseGenerator.cpp b/Engine/Plugins/Runtime/Metasound/Source/MetasoundStandardNodes/Private/MetasoundNoiseGenerator.cpp
--- a/Engine/Plugins/Runtime/Metasound/Source/MetasoundStandardNodes/Private/MetasoundNoiseGenerator.cpp
+++ b/Engine/Plugins/Runtime/Metasound/Source/MetasoundStandardNodes/Private/MetasoundNoiseGenerator.cpp
@@ -12,6 +12,7 @@
 #include "MetasoundTrigger.h"
 #include "MetasoundVertex.h"
 #include "DSP/Noise.h"
+#include "MetasoundStandardNodesCategories.h"
 
 #define LOCTEXT_NAMESPACE "MetasoundStandardNodes"
 
@@ -183,8 +184,7 @@ namespace Metasound
 			Info.Author = PluginAuthor;
 			Info.PromptIfMissing = PluginNodeMissingPrompt;
 			Info.DefaultInterface = DeclareVertexInterface();
-			Info.CategoryHierarchy = { LOCTEXT("Metasound_GeneratorCategory", "Generators") };
-
+			Info.CategoryHierarchy.Emplace(StandardNodes::Generators);
 			return Info;
 		};
 		static const FNodeClassMetadata Info = InitNodeInfo();
diff --git a/Engine/Plugins/Runtime/Metasound/Source/MetasoundStandardNodes/Private/MetasoundOscillatorNodes.cpp b/Engine/Plugins/Runtime/Metasound/Source/MetasoundStandardNodes/Private/MetasoundOscillatorNodes.cpp
--- a/Engine/Plugins/Runtime/Metasound/Source/MetasoundStandardNodes/Private/MetasoundOscillatorNodes.cpp
+++ b/Engine/Plugins/Runtime/Metasound/Source/MetasoundStandardNodes/Private/MetasoundOscillatorNodes.cpp
@@ -11,6 +11,7 @@
 #include "MetasoundStandardNodesNames.h"
 #include "MetasoundTrigger.h"
 #include "MetasoundVertex.h"
+#include "MetasoundStandardNodesCategories.h"
 
 #define LOCTEXT_NAMESPACE "MetasoundStandardNodes"
 
@@ -18,8 +19,6 @@ namespace Metasound
 {			
 #pragma region Common
 
-	static const FText GeneratorCategory = { LOCTEXT("Metasound_GeneratorCategory", "Generators") };
-
 	// Contained here are a small selection of block generators and sample generators.
 	// They are templatized to allow quick construction of new oscillators, but aren't considered
 	// optimal (as they contain many branches) and they should be only as a default 
@@ -499,7 +498,7 @@ namespace Metasound
 	};
 
 	DECLARE_METASOUND_ENUM(ESineGenerationType, ESineGenerationType::Sinf, METASOUNDSTANDARDNODES_API,
-	FEnumSineGenerationType, FEnumSineGenerationTypeInfo, FEnumSineGenerationTypeReadRef, FEnumSineGenerationTypeWriteRef);
+		FEnumSineGenerationType, FEnumSineGenerationTypeInfo, FEnumSineGenerationTypeReadRef, FEnumSineGenerationTypeWriteRef);
 
 	DEFINE_METASOUND_ENUM_BEGIN(ESineGenerationType, FEnumSineGenerationType, "SineGenerationType")
 		DEFINE_METASOUND_ENUM_ENTRY(ESineGenerationType::Sinf, LOCTEXT("SinfDescription", "Pure Math"), LOCTEXT("SinfDescriptionTT", "Uses Pure Math (Sinf) to generate the Sine Wave (most expensive)")),
@@ -529,7 +528,7 @@ namespace Metasound
 				Info.Author = PluginAuthor;
 				Info.PromptIfMissing = PluginNodeMissingPrompt;
 				Info.DefaultInterface = GetVertexInterface();
-				Info.CategoryHierarchy.Emplace(GeneratorCategory);
+				Info.CategoryHierarchy.Emplace(StandardNodes::Generators);
 				return Info;
 			};
 			static const FNodeClassMetadata Info = InitNodeInfo();
@@ -649,7 +648,7 @@ namespace Metasound
 				Info.Author = PluginAuthor;
 				Info.PromptIfMissing = PluginNodeMissingPrompt;
 				Info.DefaultInterface = GetVertexInterface();
-				Info.CategoryHierarchy.Emplace(GeneratorCategory);
+				Info.CategoryHierarchy.Emplace(StandardNodes::Generators);
 				return Info;
 			};
 			static const FNodeClassMetadata Info = InitNodeInfo();
@@ -874,7 +873,7 @@ namespace Metasound
 				Info.Author = PluginAuthor;
 				Info.PromptIfMissing = PluginNodeMissingPrompt;
 				Info.DefaultInterface = GetVertexInterface();
-				Info.CategoryHierarchy.Emplace(GeneratorCategory);
+				Info.CategoryHierarchy.Emplace(StandardNodes::Generators);
 				return Info;
 			};
 			static const FNodeClassMetadata Info = InitNodeInfo();
@@ -985,7 +984,7 @@ namespace Metasound
 				Info.Author = PluginAuthor;
 				Info.PromptIfMissing = PluginNodeMissingPrompt;
 				Info.DefaultInterface = GetVertexInterface();
-				Info.CategoryHierarchy.Emplace(GeneratorCategory);
+				Info.CategoryHierarchy.Emplace(StandardNodes::Generators);
 				return Info;
 			};
 			static const FNodeClassMetadata Info = InitNodeInfo();
diff --git a/Engine/Plugins/Runtime/Metasound/Source/MetasoundStandardNodes/Private/MetasoundStandardNodesCategories.cpp b/Engine/Plugins/Runtime/Metasound/Source/MetasoundStandardNodes/Private/MetasoundStandardNodesCategories.cpp
new file mode 100644
--- /dev/null
+++ b/Engine/Plugins/Runtime/Metasound/Source/MetasoundStandardNodes/Private/MetasoundStandardNodesCategories.cpp
@@ -0,0 +1,17 @@
+// Copyright Epic Games, Inc. All Rights Reserved.
+
+#include "MetasoundStandardNodesCategories.h"
+
+#define LOCTEXT_NAMESPACE "MetasoundStandardNodes"
+
+namespace Metasound
+{
+	namespace StandardNodes
+	{
+		const FText Generators = { LOCTEXT("Metasound_GeneratorsCategory", "Generators") };
+		const FText Filters = { LOCTEXT("Metasound_FiltersCategory", "Filters") };
+		const FText TriggerUtils = { LOCTEXT("Metasound_TriggerUtilsCategory", "Trigger Utility") };
+	}
+}
+
+#undef LOCTEXT_NAMESPACE
\ No newline at end of file
diff --git a/Engine/Plugins/Runtime/Metasound/Source/MetasoundStandardNodes/Private/MetasoundTriggerAccumulatorNode.cpp b/Engine/Plugins/Runtime/Metasound/Source/MetasoundStandardNodes/Private/MetasoundTriggerAccumulatorNode.cpp
--- a/Engine/Plugins/Runtime/Metasound/Source/MetasoundStandardNodes/Private/MetasoundTriggerAccumulatorNode.cpp
+++ b/Engine/Plugins/Runtime/Metasound/Source/MetasoundStandardNodes/Private/MetasoundTriggerAccumulatorNode.cpp
@@ -8,7 +8,7 @@
 #include "MetasoundTrigger.h"
 #include "MetasoundStandardNodesNames.h"
 #include "MetasoundVertex.h"
-
+#include "MetasoundStandardNodesCategories.h"
 
 #define LOCTEXT_NAMESPACE "MetasoundStandardNodes"
 
@@ -197,6 +197,7 @@ namespace Metasound
 			Info.Author = PluginAuthor;
 			Info.PromptIfMissing = PluginNodeMissingPrompt;
 			Info.DefaultInterface = DeclareVertexInterface();
+			Info.CategoryHierarchy.Emplace(StandardNodes::TriggerUtils);
 
 			return Info;
 		};
diff --git a/Engine/Plugins/Runtime/Metasound/Source/MetasoundStandardNodes/Private/MetasoundTriggerCounterNode.cpp b/Engine/Plugins/Runtime/Metasound/Source/MetasoundStandardNodes/Private/MetasoundTriggerCounterNode.cpp
--- a/Engine/Plugins/Runtime/Metasound/Source/MetasoundStandardNodes/Private/MetasoundTriggerCounterNode.cpp
+++ b/Engine/Plugins/Runtime/Metasound/Source/MetasoundStandardNodes/Private/MetasoundTriggerCounterNode.cpp
@@ -13,6 +13,7 @@
 #include "MetasoundOperatorInterface.h"
 #include "MetasoundPrimitives.h"
 #include "MetasoundStandardNodesNames.h"
+#include "MetasoundStandardNodesCategories.h"
 
 #define LOCTEXT_NAMESPACE "MetasoundStandardNodes"
 
@@ -130,6 +131,7 @@ namespace Metasound
 			Info.Author = PluginAuthor;
 			Info.PromptIfMissing = PluginNodeMissingPrompt;
 			Info.DefaultInterface = GetVertexInterface();
+			Info.CategoryHierarchy.Emplace(StandardNodes::TriggerUtils);
 
 			return Info;
 		};
diff --git a/Engine/Plugins/Runtime/Metasound/Source/MetasoundStandardNodes/Private/MetasoundTriggerDelayNode.cpp b/Engine/Plugins/Runtime/Metasound/Source/MetasoundStandardNodes/Private/MetasoundTriggerDelayNode.cpp
--- a/Engine/Plugins/Runtime/Metasound/Source/MetasoundStandardNodes/Private/MetasoundTriggerDelayNode.cpp
+++ b/Engine/Plugins/Runtime/Metasound/Source/MetasoundStandardNodes/Private/MetasoundTriggerDelayNode.cpp
@@ -14,6 +14,7 @@
 #include "MetasoundPrimitives.h"
 #include "MetasoundSampleCounter.h"
 #include "MetasoundStandardNodesNames.h"
+#include "MetasoundStandardNodesCategories.h"
 
 #define LOCTEXT_NAMESPACE "MetasoundStandardNodes"
 
@@ -168,6 +169,7 @@ namespace Metasound
 			Info.Author = PluginAuthor;
 			Info.PromptIfMissing = PluginNodeMissingPrompt;
 			Info.DefaultInterface = GetVertexInterface();
+			Info.CategoryHierarchy.Emplace(StandardNodes::TriggerUtils);
 
 			return Info;
 		};
diff --git a/Engine/Plugins/Runtime/Metasound/Source/MetasoundStandardNodes/Private/MetasoundTriggerOnThresholdNode.cpp b/Engine/Plugins/Runtime/Metasound/Source/MetasoundStandardNodes/Private/MetasoundTriggerOnThresholdNode.cpp
--- a/Engine/Plugins/Runtime/Metasound/Source/MetasoundStandardNodes/Private/MetasoundTriggerOnThresholdNode.cpp
+++ b/Engine/Plugins/Runtime/Metasound/Source/MetasoundStandardNodes/Private/MetasoundTriggerOnThresholdNode.cpp
@@ -11,6 +11,7 @@
 #include "MetasoundStandardNodesNames.h"
 #include "MetasoundTrigger.h"
 #include "MetasoundVertex.h"
+#include "MetasoundStandardNodesCategories.h"
 
 #define LOCTEXT_NAMESPACE "MetasoundStandardNodes"
 
@@ -205,6 +206,7 @@ namespace Metasound
 			Info.Author = PluginAuthor;
 			Info.PromptIfMissing = PluginNodeMissingPrompt;
 			Info.DefaultInterface = DeclareVertexInterface(); 
+			Info.CategoryHierarchy.Emplace(StandardNodes::TriggerUtils);
 
 			return Info;
 		};
diff --git a/Engine/Plugins/Runtime/Metasound/Source/MetasoundStandardNodes/Private/MetasoundTriggerPipeNode.cpp b/Engine/Plugins/Runtime/Metasound/Source/MetasoundStandardNodes/Private/MetasoundTriggerPipeNode.cpp
--- a/Engine/Plugins/Runtime/Metasound/Source/MetasoundStandardNodes/Private/MetasoundTriggerPipeNode.cpp
+++ b/Engine/Plugins/Runtime/Metasound/Source/MetasoundStandardNodes/Private/MetasoundTriggerPipeNode.cpp
@@ -14,7 +14,8 @@
 #include "MetasoundStandardNodesNames.h"
 #include "MetasoundTime.h"
 #include "MetasoundTrigger.h"
-#include "../Public/MetasoundSampleCounter.h"
+#include "MetasoundSampleCounter.h"
+#include "MetasoundStandardNodesCategories.h"
 
 #define LOCTEXT_NAMESPACE "MetasoundStandardNodes"
 
@@ -167,6 +168,7 @@ namespace Metasound
 			Info.Author = PluginAuthor;
 			Info.PromptIfMissing = PluginNodeMissingPrompt;
 			Info.DefaultInterface = GetVertexInterface();
+			Info.CategoryHierarchy.Emplace(StandardNodes::TriggerUtils);
 
 			return Info;
 		};
diff --git a/Engine/Plugins/Runtime/Metasound/Source/MetasoundStandardNodes/Public/MetasoundStandardNodesCategories.h b/Engine/Plugins/Runtime/Metasound/Source/MetasoundStandardNodes/Public/MetasoundStandardNodesCategories.h
new file mode 100644
--- /dev/null
+++ b/Engine/Plugins/Runtime/Metasound/Source/MetasoundStandardNodes/Public/MetasoundStandardNodesCategories.h
@@ -0,0 +1,15 @@
+// Copyright Epic Games, Inc. All Rights Reserved.
+
+#pragma once
+
+#include "CoreMinimal.h"
+
+namespace Metasound
+{
+	namespace StandardNodes
+	{
+		extern const FText Generators;
+		extern const FText Filters;
+		extern const FText TriggerUtils;
+	}
+}
