Envelope follower node

diff --git a/Engine/Plugins/Runtime/Metasound/Source/MetasoundStandardNodes/Private/MetasoundEnvelopeFollowerNode.cpp b/Engine/Plugins/Runtime/Metasound/Source/MetasoundStandardNodes/Private/MetasoundEnvelopeFollowerNode.cpp
new file mode 100644
--- /dev/null
+++ b/Engine/Plugins/Runtime/Metasound/Source/MetasoundStandardNodes/Private/MetasoundEnvelopeFollowerNode.cpp
@@ -0,0 +1,211 @@
+// Copyright Epic Games, Inc. All Rights Reserved.
+
+#include "MetasoundEnvelopeFollowerNode.h"
+
+#include "MetasoundAudioBuffer.h"
+#include "MetasoundExecutableOperator.h"
+#include "MetasoundFacade.h"
+#include "MetasoundNodeRegistrationMacro.h"
+#include "MetasoundDataTypeRegistrationMacro.h"
+#include "MetasoundOperatorSettings.h"
+#include "MetasoundPrimitives.h"
+#include "MetasoundStandardNodesNames.h"
+#include "MetasoundTrigger.h"
+#include "MetasoundTime.h"
+#include "MetasoundVertex.h"
+#include "DSP/EnvelopeFollower.h"
+
+#define LOCTEXT_NAMESPACE "MetasoundStandardNodes"
+
+namespace Metasound
+{
+	DECLARE_METASOUND_ENUM(Audio::EPeakMode::Type, Audio::EPeakMode::Peak, METASOUNDSTANDARDNODES_API,
+		FEnumPeakModeType, FEnumPeakModeTypeInfo, FEnumPeakModeTypeReadRef, FEnumPeakModeTypeWriteRef);
+
+	DEFINE_METASOUND_ENUM_BEGIN(Audio::EPeakMode::Type, FEnumPeakModeType, "PeakModeType")
+		DEFINE_METASOUND_ENUM_ENTRY(Audio::EPeakMode::MeanSquared, LOCTEXT("MeanSquaredDescription", "Mean Squared"), LOCTEXT("MeanSquaredDescriptionTT", "Sets the envelope to follow the mean-squared of the audio signal")),
+		DEFINE_METASOUND_ENUM_ENTRY(Audio::EPeakMode::RootMeanSquared, LOCTEXT("RootMeanSquaredDescription", "Root Mean Squared"), LOCTEXT("RootMeanSquaredTT", "Sets the envelope to follow the root-mean-squared of the audio signal")),
+		DEFINE_METASOUND_ENUM_ENTRY(Audio::EPeakMode::Peak, LOCTEXT("PeakDescription", "Peak"), LOCTEXT("PeakDescriptionTT", "Sets the envelope to follow the peaks of the audio signal")),
+	DEFINE_METASOUND_ENUM_END()
+
+	namespace EnvelopeFollower
+	{
+		static const TCHAR* InParamNameAudioInput = TEXT("In");
+		static const TCHAR* InParamNameAttackTime = TEXT("Attack Time");
+		static const TCHAR* InParamNameReleaseTime = TEXT("Release Time");
+		static const TCHAR* InParamNamePeakModeType = TEXT("Peak Mode");
+		static const TCHAR* OutParamNameEnvelope = TEXT("Envelope");
+	}
+
+	class FEnvelopeFollowerOperator : public TExecutableOperator<FEnvelopeFollowerOperator>
+	{
+	public:
+
+		static const FNodeClassMetadata& GetNodeInfo();
+		static const FVertexInterface& GetVertexInterface();
+		static TUniquePtr<IOperator> CreateOperator(const FCreateOperatorParams& InParams, FBuildErrorArray& OutErrors);
+
+		FEnvelopeFollowerOperator(const FOperatorSettings& InSettings, 
+			const FAudioBufferReadRef& InAudioInput, 
+			const FTimeReadRef& InAttackTime, 
+			const FTimeReadRef& InReleaseTime,
+			const FEnumPeakModeTypeReadRef& InPeakModeType);
+
+		virtual FDataReferenceCollection GetInputs() const override;
+		virtual FDataReferenceCollection GetOutputs() const override;
+		void Execute();
+
+	private:
+		// The input audio buffer
+		FAudioBufferReadRef AudioInput;
+
+		// The amount of attack time
+		FTimeReadRef AttackTime;
+
+		// The amount of release time
+		FTimeReadRef ReleaseTime;
+
+		FEnumPeakModeTypeReadRef PeakModeType;
+
+		// The audio output
+		FFloatWriteRef EnvelopeOutput;
+
+		// The envelope follower DSP object
+		Audio::FEnvelopeFollower EnvelopeFollower;
+
+		double PrevAttackTime = 0.0;
+		double PrevReleaseTime = 0.0;
+		Audio::EPeakMode::Type PrevPeakModeType = Audio::EPeakMode::MeanSquared;
+	};
+
+	FEnvelopeFollowerOperator::FEnvelopeFollowerOperator(const FOperatorSettings& InSettings, 
+		const FAudioBufferReadRef& InAudioInput, 
+		const FTimeReadRef& InAttackTime,
+		const FTimeReadRef& InReleaseTime,
+		const FEnumPeakModeTypeReadRef& InPeakModeType)
+		: AudioInput(InAudioInput)
+		, AttackTime(InAttackTime)
+		, ReleaseTime(InReleaseTime)
+		, PeakModeType(InPeakModeType)
+		, EnvelopeOutput(FFloatWriteRef::CreateNew(0.0f))
+	{
+		PrevAttackTime = FTime::ToMilliseconds(*AttackTime);
+		PrevReleaseTime = FTime::ToMilliseconds(*ReleaseTime);
+		PrevPeakModeType = *PeakModeType;
+
+		EnvelopeFollower.Init(InSettings.GetSampleRate(), PrevAttackTime, PrevReleaseTime, PrevPeakModeType);
+	}
+
+	FDataReferenceCollection FEnvelopeFollowerOperator::GetInputs() const
+	{
+		FDataReferenceCollection InputDataReferences;
+		InputDataReferences.AddDataReadReference(EnvelopeFollower::InParamNameAudioInput, AudioInput);
+		InputDataReferences.AddDataReadReference(EnvelopeFollower::InParamNameAttackTime, AttackTime);
+		InputDataReferences.AddDataReadReference(EnvelopeFollower::InParamNameReleaseTime, ReleaseTime);
+		InputDataReferences.AddDataReadReference(EnvelopeFollower::InParamNamePeakModeType, PeakModeType);
+
+		return InputDataReferences;
+	}
+
+	FDataReferenceCollection FEnvelopeFollowerOperator::GetOutputs() const
+	{
+		FDataReferenceCollection OutputDataReferences;
+		OutputDataReferences.AddDataReadReference(EnvelopeFollower::OutParamNameEnvelope, EnvelopeOutput);
+		return OutputDataReferences;
+	}
+
+	void FEnvelopeFollowerOperator::Execute()
+	{
+		// Check for any input changes
+		double CurrentAttackTime = FTime::ToMilliseconds(*AttackTime);
+		if (!FMath::IsNearlyEqual(CurrentAttackTime, PrevAttackTime))
+		{
+			PrevAttackTime = CurrentAttackTime;
+			EnvelopeFollower.SetAttackTime(CurrentAttackTime);
+		}
+
+		double CurrentReleaseTime = FTime::ToMilliseconds(*ReleaseTime);
+		if (!FMath::IsNearlyEqual(CurrentReleaseTime, PrevReleaseTime))
+		{
+			PrevReleaseTime = CurrentReleaseTime;
+			EnvelopeFollower.SetReleaseTime(CurrentReleaseTime);
+		}
+
+		Audio::EPeakMode::Type CurrentPeakModeType = *PeakModeType;
+		if (CurrentPeakModeType != PrevPeakModeType)
+		{
+			PrevPeakModeType = CurrentPeakModeType;
+			EnvelopeFollower.SetMode(CurrentPeakModeType);
+		}
+
+		// Process the audio through the envelope follower
+		const float* InputAudio = AudioInput->GetData();
+		int32 NumFrames = AudioInput->Num();
+		EnvelopeFollower.ProcessAudio(InputAudio, NumFrames);
+
+		// Write the current envelope follower value to the output
+		*EnvelopeOutput = EnvelopeFollower.GetCurrentValue();
+	}
+
+	const FVertexInterface& FEnvelopeFollowerOperator::GetVertexInterface()
+	{
+		static const FVertexInterface Interface(
+			FInputVertexInterface(
+				TInputDataVertexModel<FAudioBuffer>(EnvelopeFollower::InParamNameAudioInput, LOCTEXT("AudioInputToolTT", "Audio input.")),
+				TInputDataVertexModel<FTime>(EnvelopeFollower::InParamNameAttackTime, LOCTEXT("AttackTimeTT", "The attack time of the envelope follower."), 0.01f),
+				TInputDataVertexModel<FTime>(EnvelopeFollower::InParamNameReleaseTime, LOCTEXT("ReleaseTimeTT", "The release time of the envelope follower."), 0.1f),
+				TInputDataVertexModel<FEnumPeakModeType>(EnvelopeFollower::InParamNamePeakModeType, LOCTEXT("PeakModeTypeTT", "The peak mode of the envelope follower."))
+			),
+			FOutputVertexInterface(
+				TOutputDataVertexModel<float>(EnvelopeFollower::OutParamNameEnvelope, LOCTEXT("EnvelopeFollowerOutputTT", "The output envelope value of the audio signal."))
+			)
+		);
+
+		return Interface;
+	}
+
+	const FNodeClassMetadata& FEnvelopeFollowerOperator::GetNodeInfo()
+	{
+		auto InitNodeInfo = []() -> FNodeClassMetadata
+		{
+			FNodeClassMetadata Info;
+			Info.ClassName = { Metasound::StandardNodes::Namespace, TEXT("Envelope Follower"), TEXT("") };
+			Info.MajorVersion = 1;
+			Info.MinorVersion = 0;
+			Info.DisplayName = LOCTEXT("Metasound_EnvelopeFollowerDisplayName", "Envelope Follower");
+			Info.Description = LOCTEXT("Metasound_EnvelopeFollowerDescription", "Outputs an envelope from an input audio signal.");
+			Info.Author = PluginAuthor;
+			Info.PromptIfMissing = PluginNodeMissingPrompt;
+			Info.DefaultInterface = GetVertexInterface();
+
+			return Info;
+		};
+
+		static const FNodeClassMetadata Info = InitNodeInfo();
+
+		return Info;
+	}
+
+	TUniquePtr<IOperator> FEnvelopeFollowerOperator::CreateOperator(const FCreateOperatorParams& InParams, FBuildErrorArray& OutErrors)
+	{
+		const FEnvelopeFollowerNode& EnvelopeFollowerNode = static_cast<const FEnvelopeFollowerNode&>(InParams.Node);
+		const FDataReferenceCollection& InputCollection = InParams.InputDataReferences;
+		const FInputVertexInterface& InputInterface = GetVertexInterface().GetInputInterface();
+
+		FAudioBufferReadRef AudioIn = InputCollection.GetDataReadReferenceOrConstruct<FAudioBuffer>(EnvelopeFollower::InParamNameAudioInput, InParams.OperatorSettings);
+		FTimeReadRef AttackTime = InputCollection.GetDataReadReferenceOrConstructWithVertexDefault<FTime, float>(InputInterface, EnvelopeFollower::InParamNameAttackTime);
+		FTimeReadRef ReleaseTime = InputCollection.GetDataReadReferenceOrConstructWithVertexDefault<FTime, float>(InputInterface, EnvelopeFollower::InParamNameReleaseTime);
+		FEnumPeakModeTypeReadRef EnvelopePeakModeType = InputCollection.GetDataReadReferenceOrConstruct<FEnumPeakModeType>(EnvelopeFollower::InParamNamePeakModeType);
+
+		return MakeUnique<FEnvelopeFollowerOperator>(InParams.OperatorSettings, AudioIn, AttackTime, ReleaseTime, EnvelopePeakModeType);
+	}
+
+	FEnvelopeFollowerNode::FEnvelopeFollowerNode(const FNodeInitData& InitData)
+		: FNodeFacade(InitData.InstanceName, InitData.InstanceID, TFacadeOperatorClass<FEnvelopeFollowerOperator>())
+	{
+	}
+
+	METASOUND_REGISTER_NODE(FEnvelopeFollowerNode)
+}
+
+#undef LOCTEXT_NAMESPACE
diff --git a/Engine/Plugins/Runtime/Metasound/Source/MetasoundStandardNodes/Public/MetasoundEnvelopeFollowerNode.h b/Engine/Plugins/Runtime/Metasound/Source/MetasoundStandardNodes/Public/MetasoundEnvelopeFollowerNode.h
new file mode 100644
--- /dev/null
+++ b/Engine/Plugins/Runtime/Metasound/Source/MetasoundStandardNodes/Public/MetasoundEnvelopeFollowerNode.h
@@ -0,0 +1,23 @@
+// Copyright Epic Games, Inc. All Rights Reserved.
+
+#pragma once
+
+#include "CoreMinimal.h"
+#include "MetasoundFacade.h"
+
+namespace Metasound
+{
+	/** FEnvelopeFollowerNode
+	 *
+	 *  Delays an audio buffer by a specified amount.
+	 */
+	class METASOUNDSTANDARDNODES_API FEnvelopeFollowerNode : public FNodeFacade
+	{
+	public:
+		/**
+		 * Constructor used by the Metasound Frontend.
+		 */
+		FEnvelopeFollowerNode(const FNodeInitData& InitData);
+	};
+} // namespace Metasound
+
