Fix array values not triggering literal update in Metasound Editor

diff --git a/Engine/Plugins/Runtime/Metasound/Source/MetasoundEditor/Private/MetasoundInputNodeDetailCustomization.cpp b/Engine/Plugins/Runtime/Metasound/Source/MetasoundEditor/Private/MetasoundInputNodeDetailCustomization.cpp
--- a/Engine/Plugins/Runtime/Metasound/Source/MetasoundEditor/Private/MetasoundInputNodeDetailCustomization.cpp
+++ b/Engine/Plugins/Runtime/Metasound/Source/MetasoundEditor/Private/MetasoundInputNodeDetailCustomization.cpp
@@ -489,6 +489,7 @@ namespace Metasound
 					});
 
 					Handle->SetOnPropertyValueChanged(OnLiteralChanged);
+					Handle->SetOnChildPropertyValueChanged(OnLiteralChanged);
 
 					TSharedPtr<IPropertyHandleArray> PropertyArray = Handle->AsArray();
 					if (PropertyArray.IsValid())
