Clamping input values on envelope follower node

diff --git a/Engine/Plugins/Runtime/Metasound/Source/MetasoundStandardNodes/Private/MetasoundEnvelopeFollowerNode.cpp b/Engine/Plugins/Runtime/Metasound/Source/MetasoundStandardNodes/Private/MetasoundEnvelopeFollowerNode.cpp
--- a/Engine/Plugins/Runtime/Metasound/Source/MetasoundStandardNodes/Private/MetasoundEnvelopeFollowerNode.cpp
+++ b/Engine/Plugins/Runtime/Metasound/Source/MetasoundStandardNodes/Private/MetasoundEnvelopeFollowerNode.cpp
@@ -89,8 +89,8 @@ namespace Metasound
 		, PeakModeType(InPeakModeType)
 		, EnvelopeOutput(FFloatWriteRef::CreateNew(0.0f))
 	{
-		PrevAttackTime = FTime::ToMilliseconds(*AttackTime);
-		PrevReleaseTime = FTime::ToMilliseconds(*ReleaseTime);
+		PrevAttackTime = FMath::Max(FTime::ToMilliseconds(*AttackTime), 0.0);
+		PrevReleaseTime = FMath::Max(FTime::ToMilliseconds(*ReleaseTime), 0.0f);
 		PrevPeakModeType = *PeakModeType;
 
 		EnvelopeFollower.Init(InSettings.GetSampleRate(), PrevAttackTime, PrevReleaseTime, PrevPeakModeType);
@@ -117,14 +117,14 @@ namespace Metasound
 	void FEnvelopeFollowerOperator::Execute()
 	{
 		// Check for any input changes
-		double CurrentAttackTime = FTime::ToMilliseconds(*AttackTime);
+		double CurrentAttackTime = FMath::Max(FTime::ToMilliseconds(*AttackTime), 0.0);
 		if (!FMath::IsNearlyEqual(CurrentAttackTime, PrevAttackTime))
 		{
 			PrevAttackTime = CurrentAttackTime;
 			EnvelopeFollower.SetAttackTime(CurrentAttackTime);
 		}
 
-		double CurrentReleaseTime = FTime::ToMilliseconds(*ReleaseTime);
+		double CurrentReleaseTime = FMath::Max(FTime::ToMilliseconds(*ReleaseTime), 0.0);
 		if (!FMath::IsNearlyEqual(CurrentReleaseTime, PrevReleaseTime))
 		{
 			PrevReleaseTime = CurrentReleaseTime;
