Fixing typo

Goes to show submitting a minor change...

diff --git a/Engine/Plugins/Runtime/Metasound/Source/MetasoundStandardNodes/Private/MetasoundEnvelopeFollowerNode.cpp b/Engine/Plugins/Runtime/Metasound/Source/MetasoundStandardNodes/Private/MetasoundEnvelopeFollowerNode.cpp
--- a/Engine/Plugins/Runtime/Metasound/Source/MetasoundStandardNodes/Private/MetasoundEnvelopeFollowerNode.cpp
+++ b/Engine/Plugins/Runtime/Metasound/Source/MetasoundStandardNodes/Private/MetasoundEnvelopeFollowerNode.cpp
@@ -90,7 +90,7 @@ namespace Metasound
 		, EnvelopeOutput(FFloatWriteRef::CreateNew(0.0f))
 	{
 		PrevAttackTime = FMath::Max(FTime::ToMilliseconds(*AttackTime), 0.0);
-		PrevReleaseTime = FMath::Max(FTime::ToMilliseconds(*ReleaseTime), 0.0f);
+		PrevReleaseTime = FMath::Max(FTime::ToMilliseconds(*ReleaseTime), 0.0);
 		PrevPeakModeType = *PeakModeType;
 
 		EnvelopeFollower.Init(InSettings.GetSampleRate(), PrevAttackTime, PrevReleaseTime, PrevPeakModeType);
