Add Copy/Paste to Metasounds Editor
- Clean up FGraphBuilder's interface to clarify when node handles vs edgraph uobjects are created
- Split up FGraphBuilder's node creation interface for convenience and clarity (for inputs, outputs, & external nodes)

diff --git a/Engine/Plugins/Runtime/Metasound/Source/MetasoundEditor/Private/MetasoundEditor.cpp b/Engine/Plugins/Runtime/Metasound/Source/MetasoundEditor/Private/MetasoundEditor.cpp
--- a/Engine/Plugins/Runtime/Metasound/Source/MetasoundEditor/Private/MetasoundEditor.cpp
+++ b/Engine/Plugins/Runtime/Metasound/Source/MetasoundEditor/Private/MetasoundEditor.cpp
@@ -805,15 +805,14 @@ namespace Metasound
 				}
 			}
 
-			// Delete the deleteable nodes
+			// Delete the deletable nodes
 			DeleteSelectedNodes();
 
-			// Reselect whatever's left from the original selection after the deletion
+			// Clear deleted, and reselect remaining nodes from original selection
 			MetasoundGraphEditor->ClearSelectionSet();
-
-			for (FGraphPanelSelectionSet::TConstIterator SelectedIter(RemainingNodes); SelectedIter; ++SelectedIter)
+			for (UObject* RemainingNode : RemainingNodes)
 			{
-				if (UEdGraphNode* Node = Cast<UEdGraphNode>(*SelectedIter))
+				if (UEdGraphNode* Node = Cast<UEdGraphNode>(RemainingNode))
 				{
 					MetasoundGraphEditor->SetNodeSelection(Node, true);
 				}
@@ -821,55 +820,24 @@ namespace Metasound
 		}
 
 		void FEditor::CopySelectedNodes()
-		{
+ 		{
 			using namespace Metasound::Frontend;
 
 			// Export the selected nodes and place the text on the clipboard
-			/* TODO: rethink how this is performed. Current implementation does not
-			 * handle case where nodes are not in the root graph. 
-			 *
-			 * GraphHandle should support exporting selection of nodes to a FMetasoundFrontendDocument as nodes on the root class. 
-			 */
-			 
-			//const FGraphPanelSelectionSet SelectedNodes = MetasoundGraphEditor->GetSelectedNodes();
-			//
-			//TMap<int32, FNodeHandle> HandlesToCopy;
-			//for (FGraphPanelSelectionSet::TConstIterator SelectedIter(SelectedNodes); SelectedIter; ++SelectedIter)
-			//{
-				//if (UMetasoundEditorGraphNode* Node = Cast<UMetasoundEditorGraphNode>(*SelectedIter))
-				//{
-					//HandlesToCopy.Add(Node->GetID(), Node->GetNodeHandle());
-					//Node->PrepareForCopying();
-				//}
-			//}
-//
-			//FMetasoundAssetBase* MetasoundAsset = Metasound::IMetasoundUObjectRegistry::Get().GetObjectAsAssetBase(Metasound);
-			//check(MetasoundAsset);
-			//FGraphHandle GraphHandle = MetasoundAsset->GetRootGraphHandle();
-//
-			//UMetasound* MetasoundCopy = NewObject<UMetasound>();
-			//FGraphHandle CopyGraphHandle = InitMetasound(*MetasoundCopy);
-//
-			//GraphHandle->CopyGraph(CopyGraphHandle);
-//
-			//TArray<FNodeHandle> AllCopiedNodes = CopyGraphHandle->GetAllNodes();
-			//for (FNodeHandle& Node : AllCopiedNodes)
-			//{
-				//if (!HandlesToCopy.Contains(Node->GetID()))
-				//{
-					//CopyGraphHandle->RemoveNode(Node);
-				//}
-			//}
-//
-			//FString ExportedEdText;
-			//FEdGraphUtilities::ExportNodesToText(SelectedNodes, ExportedEdText);
-//
-			//FMetasoundEditorData EditorData;
-			//EditorData.GraphData = ExportedEdText;
-			//CopyGraphHandle.SetEditorData(EditorData);
-//
-			//const FString JSONToCopy = CopyGraphHandle.ExportToJSON();
-			//FPlatformApplicationMisc::ClipboardCopy(*JSONToCopy);
+			const FGraphPanelSelectionSet SelectedNodes = MetasoundGraphEditor->GetSelectedNodes();
+			TMap<int32, FNodeHandle> HandlesToCopy;
+			for (UObject* SelectedNode : SelectedNodes)
+			{
+				if (UMetasoundEditorGraphNode* Node = Cast<UMetasoundEditorGraphNode>(SelectedNode))
+				{
+					HandlesToCopy.Add(Node->GetUniqueID(), Node->GetNodeHandle());
+					Node->PrepareForCopying();
+				}
+			}
+
+			FString ExportedEdText;
+			FEdGraphUtilities::ExportNodesToText(SelectedNodes, ExportedEdText);
+			FPlatformApplicationMisc::ClipboardCopy(*ExportedEdText);
 		}
 
 		bool FEditor::CanCopyNodes() const
@@ -908,95 +876,163 @@ namespace Metasound
 
 		void FEditor::PasteNodes(const FVector2D* InLocation)
 		{
-// 			FVector2D Location;
-// 			if (InLocation)
-// 			{
-// 				Location = *InLocation;
-// 			}
-// 			else
-// 			{
-// 				check(MetasoundGraphEditor);
-// 				Location = MetasoundGraphEditor->GetPasteLocation();
-// 			}
-// 
-// 			FMetasoundAssetBase* MetasoundAsset = IMetasoundUObjectRegistry::Get().GetObjectAsAssetBase(Metasound);
-// 			check(MetasoundAsset);
-// 
-// 			UEdGraph* Graph = MetasoundAsset->GetGraph();
-// 			if (!Graph)
-// 			{
-// 				return;
-// 			}
-// 
-// 			// Undo/Redo support
-// 			const FScopedTransaction Transaction(NSLOCTEXT("UnrealEd", "MetasoundEditorPaste", "Paste Metasound Node"));
-// 			Graph->Modify();
-// 			Metasound->Modify();
-// 
-// 			// Clear the selection set (newly pasted stuff will be selected)
-// 			MetasoundGraphEditor->ClearSelectionSet();
-// 
-//			static_assert(false, "TODO: Delete this and assert & insert nodes/inputs/outputs from 'PastedDocument' into 'Metasound' doc here. Run fix-up on new nodes to generate ids ensuring no duplicates with nodes that pre-existed.");
-// 
-// 			// Import the nodes
-// 			TSet<UEdGraphNode*> PastedNodes;
-// 			FEdGraphUtilities::ImportNodesFromText(Graph, PastedDocument.EditorData.GraphData, PastedNodes);
-// 
-// 			// Clear out, no reason to hold a bunch of references to data no longer needed.
-// 			PastedDocument = FMetasoundFrontendDocument();
-// 
-// 			FVector2D AvgNodePosition = FVector2D::ZeroVector;
-// 			for (UEdGraphNode* Node : PastedNodes)
-// 			{
-// 				AvgNodePosition.X += Node->NodePosX;
-// 				AvgNodePosition.Y += Node->NodePosY;
-// 			}
-// 
-// 			if (PastedNodes.Num() > 0)
-// 			{
-// 				float InvNumNodes = 1.0f / PastedNodes.Num();
-// 				AvgNodePosition.X *= InvNumNodes;
-// 				AvgNodePosition.Y *= InvNumNodes;
-// 			}
-// 
-// 			for (UEdGraphNode* Node : PastedNodes)
-// 			{
-// 				// Select the newly pasted stuff
-// 				MetasoundGraphEditor->SetNodeSelection(Node, true);
-// 
-// 				Node->NodePosX = (Node->NodePosX - AvgNodePosition.X) + Location.X;
-// 				Node->NodePosY = (Node->NodePosY - AvgNodePosition.Y) + Location.Y;
-// 
-// 				Node->SnapToGrid(SNodePanel::GetSnapGridSize());
-// 				Node->CreateNewGuid();
-// 			}
-// 
-// 			MetasoundGraphEditor->NotifyGraphChanged();
-// 
-// 			Metasound->PostEditChange();
-// 			Metasound->MarkPackageDirty();
+			using namespace Frontend;
+
+			FVector2D Location;
+			if (InLocation)
+			{
+				Location = *InLocation;
+			}
+			else
+			{
+				check(MetasoundGraphEditor);
+				Location = MetasoundGraphEditor->GetPasteLocation();
+			}
+
+			FMetasoundAssetBase* MetasoundAsset = IMetasoundUObjectRegistry::Get().GetObjectAsAssetBase(Metasound);
+			check(MetasoundAsset);
+
+			UEdGraph* Graph = MetasoundAsset->GetGraph();
+			if (!Graph)
+			{
+				return;
+			}
+
+			const FScopedTransaction Transaction(NSLOCTEXT("UnrealEd", "MetasoundEditorPaste", "Paste Metasound Node(s)"));
+			Graph->Modify();
+			Metasound->Modify();
+
+			// Clear the selection set (newly pasted stuff will be selected)
+			MetasoundGraphEditor->ClearSelectionSet();
+
+			TSet<UEdGraphNode*> PastedGraphNodes;
+			FEdGraphUtilities::ImportNodesFromText(Graph, ClipboardContent, PastedGraphNodes);
+
+			ClipboardContent.Empty();
+
+			TArray<UEdGraphNode*> NodesToRemove;
+			for (UEdGraphNode* GraphNode : PastedGraphNodes)
+			{
+				GraphNode->CreateNewGuid();
+				if (UMetasoundEditorGraphNode* MetasoundGraphNode = Cast<UMetasoundEditorGraphNode>(GraphNode))
+				{
+					// TODO: Support copying inputs after adding replacing InputGraphNodes with InputReferenceGraphNodes
+					FNodeHandle NewHandle = INodeController::GetInvalidHandle();
+					if (MetasoundGraphNode->IsA(UMetasoundEditorGraphExternalNode::StaticClass()))
+					{
+						NewHandle = FGraphBuilder::AddNodeHandle(*Metasound, *MetasoundGraphNode);
+					}
+
+					if (!NewHandle->IsValid())
+					{
+						NodesToRemove.Add(GraphNode);
+					}
+				}
+			}
+
+			// Remove nodes failed to import before attempting to connect/place
+			// in frontend graph.
+			for (UEdGraphNode* Node : NodesToRemove)
+			{
+				Graph->RemoveNode(Node);
+				PastedGraphNodes.Remove(Node);
+			}
+
+			// Find average midpoint of nodes and offset subgraph accordingly
+			FVector2D AvgNodePosition = FVector2D::ZeroVector;
+			for (UEdGraphNode* Node : PastedGraphNodes)
+			{
+				AvgNodePosition.X += Node->NodePosX;
+				AvgNodePosition.Y += Node->NodePosY;
+			}
+
+			if (!PastedGraphNodes.IsEmpty())
+			{
+				float InvNumNodes = 1.0f / PastedGraphNodes.Num();
+				AvgNodePosition.X *= InvNumNodes;
+				AvgNodePosition.Y *= InvNumNodes;
+			}
+
+			for (UEdGraphNode* GraphNode : PastedGraphNodes)
+			{
+				GraphNode->NodePosX = (GraphNode->NodePosX - AvgNodePosition.X) + Location.X;
+				GraphNode->NodePosY = (GraphNode->NodePosY - AvgNodePosition.Y) + Location.Y;
+
+				GraphNode->SnapToGrid(SNodePanel::GetSnapGridSize());
+				if (UMetasoundEditorGraphNode* MetasoundGraphNode = Cast<UMetasoundEditorGraphNode>(GraphNode))
+				{
+					FNodeHandle NodeHandle = MetasoundGraphNode->GetNodeHandle();
+					if (ensure(NodeHandle->IsValid()))
+					{
+						const FVector2D NewNodeLocation = FVector2D(GraphNode->NodePosX, GraphNode->NodePosY);
+						FMetasoundFrontendNodeStyle NodeStyle = NodeHandle->GetNodeStyle();
+						NodeStyle.Display.Location = NewNodeLocation;
+						NodeHandle->SetNodeStyle(NodeStyle);
+					}
+				}
+			}
+
+			for (UEdGraphNode* GraphNode : PastedGraphNodes)
+			{
+				UMetasoundEditorGraphNode* MetasoundNode = Cast<UMetasoundEditorGraphNode>(GraphNode);
+				if (!MetasoundNode)
+				{
+					continue;
+				}
+
+				for (UEdGraphPin* Pin : GraphNode->Pins)
+				{
+					if (Pin->Direction != EGPD_Input)
+					{
+						continue;
+					}
+
+					if (Pin->LinkedTo.IsEmpty())
+					{
+						FGraphBuilder::AddOrUpdateLiteralInput(*Metasound, MetasoundNode->GetNodeHandle(), *Pin, true /* bForcePinValueAsDefault */);
+					}
+					else
+					{
+						for (UEdGraphPin* LinkedPin : Pin->LinkedTo)
+						{
+							if (UMetasoundEditorGraphNode* Node = Cast<UMetasoundEditorGraphNode>(LinkedPin->GetOwningNode()))
+							{
+								FGraphBuilder::ConnectNodes(*Pin, *LinkedPin, false /* bConnectEdPins */);
+							}
+						}
+					}
+				}
+			}
+
+			// Select the newly pasted stuff
+			for (UEdGraphNode* GraphNode : PastedGraphNodes)
+			{
+				MetasoundGraphEditor->SetNodeSelection(GraphNode, true);
+			}
+
+			FGraphBuilder::SynchronizeGraph(*Metasound);
+
+			MetasoundGraphEditor->NotifyGraphChanged();
+			Metasound->MarkPackageDirty();
 		}
 
 		bool FEditor::CanPasteNodes()
 		{
-			return false;
+			FMetasoundAssetBase* MetasoundAsset = IMetasoundUObjectRegistry::Get().GetObjectAsAssetBase(Metasound);
+			check(MetasoundAsset);
 
-// 			static_assert(false, "Delete this assert once 'PasteNodes' includes doc fix-ups as described above.");
-// 			FString ClipboardContent;
-// 			FPlatformApplicationMisc::ClipboardPaste(ClipboardContent);
-// 
-// 			FMetasoundAssetBase* AssetBase = IMetasoundUObjectRegistry::Get().GetObjectAsAssetBase(Metasound);
-// 			check(AssetBase);
-// 
-// 			PastedDocument = FMetasoundFrontendDocument();
-// 			const bool bCanImportDocument = Metasound::Frontend::ImportJSONToMetasound(ClipboardContent, PastedDocument);
-// 			if (!bCanImportDocument)
-// 			{
-// 				return false;
-// 			}
-// 
-// 		 	const bool bCanPasteNodes = FEdGraphUtilities::CanImportNodesFromText(AssetBase->GetGraph(), PastedDocument.EditorData.GraphData);
-// 		 	return bCanPasteNodes;
+			if (UEdGraph* Graph = MetasoundAsset->GetGraph())
+			{
+				FPlatformApplicationMisc::ClipboardPaste(ClipboardContent);
+				if (FEdGraphUtilities::CanImportNodesFromText(Graph, ClipboardContent))
+				{
+					return true;
+				}
+
+				ClipboardContent.Empty();
+			}
+
+			return false;
 		}
 
 		void FEditor::UndoGraphAction()
diff --git a/Engine/Plugins/Runtime/Metasound/Source/MetasoundEditor/Private/MetasoundEditor.h b/Engine/Plugins/Runtime/Metasound/Source/MetasoundEditor/Private/MetasoundEditor.h
--- a/Engine/Plugins/Runtime/Metasound/Source/MetasoundEditor/Private/MetasoundEditor.h
+++ b/Engine/Plugins/Runtime/Metasound/Source/MetasoundEditor/Private/MetasoundEditor.h
@@ -1,6 +1,8 @@
 // Copyright Epic Games, Inc. All Rights Reserved.
 #pragma once
 
+#include "AudioMeterStyle.h"
+#include "AudioSynesthesia/Classes/Meter.h"
 #include "CoreMinimal.h"
 #include "EditorUndoClient.h"
 #include "Framework/Commands/UICommandList.h"
@@ -10,25 +12,25 @@
 #include "MetasoundFrontend.h"
 #include "MetasoundFrontendController.h"
 #include "Misc/NotifyHook.h"
+#include "SAudioMeter.h"
+#include "SGraphActionMenu.h"
 #include "SMetasoundPalette.h"
+#include "Sound/AudioBus.h"
 #include "Textures/SlateIcon.h"
 #include "Toolkits/IToolkitHost.h"
 #include "UObject/GCObject.h"
-#include "Widgets/SPanel.h"
-#include "Sound/AudioBus.h"
-#include "SAudioMeter.h"
-#include "AudioMeterStyle.h"
-#include "AudioSynesthesia/Classes/Meter.h"
+#include "Toolkits/AssetEditorToolkit.h"
 #include "UObject/StrongObjectPtr.h"
+#include "Widgets/SPanel.h"
 
 // Forward Declarations
-class FSlateRect;
-class FTabManager;
 class IDetailsView;
 class IToolkitHost;
 class SDockableTab;
 class SGraphEditor;
 class SMetasoundPalette;
+class FSlateRect;
+class FTabManager;
 class SVerticalBox;
 class UEdGraphNode;
 class UMetasound;
@@ -182,7 +184,6 @@ namespace Metasound
 			TSharedRef<SGraphEditor> CreateGraphEditorWidget();
 
 		private:
-
 			/** List of open tool panels; used to ensure only one exists at any one time */
 			TMap<FName, TWeakPtr<SDockableTab>> SpawnedToolPanels;
 
@@ -224,8 +225,8 @@ namespace Metasound
 			/** Whether or not metasound being edited is valid */
 			bool bPassedValidation = true;
 
-			/** Document used when pasting from clipboard to avoid deserializing twice */
-			FMetasoundFrontendDocument PastedDocument;
+			/** Clipboard content used when pasting from clipboard to avoid deserializing twice */
+			FString ClipboardContent;
 		};
 	} // namespace Editor
 } // namespace Metasound
diff --git a/Engine/Plugins/Runtime/Metasound/Source/MetasoundEditor/Private/MetasoundEditorGraphBuilder.cpp b/Engine/Plugins/Runtime/Metasound/Source/MetasoundEditor/Private/MetasoundEditorGraphBuilder.cpp
--- a/Engine/Plugins/Runtime/Metasound/Source/MetasoundEditor/Private/MetasoundEditorGraphBuilder.cpp
+++ b/Engine/Plugins/Runtime/Metasound/Source/MetasoundEditor/Private/MetasoundEditorGraphBuilder.cpp
@@ -13,6 +13,8 @@
 #include "MetasoundEditorGraphNode.h"
 #include "MetasoundEditorGraphInputNodes.h"
 #include "MetasoundEditorModule.h"
+#include "MetasoundFrontendQuery.h"
+#include "MetasoundFrontendQuerySteps.h"
 #include "MetasoundFrontendRegistries.h"
 #include "MetasoundLiteral.h"
 #include "MetasoundUObjectRegistry.h"
@@ -60,70 +62,81 @@ namespace Metasound
 			}
 		}
 
-		UEdGraphNode* FGraphBuilder::AddNode(UObject& InMetasound, Frontend::FNodeHandle& InNodeHandle, bool bInSelectNewNode)
+		UMetasoundEditorGraphExternalNode* FGraphBuilder::AddExternalNode(UObject& InMetasound, Frontend::FNodeHandle& InNodeHandle, bool bInSelectNewNode)
 		{
+			UMetasoundEditorGraphExternalNode* NewGraphNode = nullptr;
+			if (!ensure(InNodeHandle->GetClassType() == EMetasoundFrontendClassType::External))
+			{
+				return nullptr;
+			}
+
 			FMetasoundAssetBase* MetasoundAsset = IMetasoundUObjectRegistry::Get().GetObjectAsAssetBase(&InMetasound);
 			check(MetasoundAsset);
-
 			UEdGraph& Graph = MetasoundAsset->GetGraphChecked();
+			FGraphNodeCreator<UMetasoundEditorGraphExternalNode> NodeCreator(Graph);
 
-			UMetasoundEditorGraphNode* NewGraphNode = nullptr;
-			switch (InNodeHandle->GetClassType())
-			{
-				case EMetasoundFrontendClassType::Input:
-				{
-					// Should use AddInput for this case to ensure child class is used for input.
-					checkNoEntry();
-				}
-				break;
+			NewGraphNode = NodeCreator.CreateNode(bInSelectNewNode);
+			NewGraphNode->ClassName = InNodeHandle->GetClassName();
 
-				case EMetasoundFrontendClassType::Output:
-				{
-					FGraphNodeCreator<UMetasoundEditorGraphOutputNode> NodeCreator(Graph);
-					NewGraphNode = NodeCreator.CreateNode(bInSelectNewNode);
-					NodeCreator.Finalize();
-				}
-				break;
+			NodeCreator.Finalize();
+			InitGraphNode(InNodeHandle, NewGraphNode, InMetasound);
 
-				case EMetasoundFrontendClassType::External:
-				{
-					FGraphNodeCreator<UMetasoundEditorGraphExternalNode> NodeCreator(Graph);
-					NewGraphNode = NodeCreator.CreateNode(bInSelectNewNode);
-					NodeCreator.Finalize();
-				}
-				break;
+			return NewGraphNode;
+		}
 
-				case EMetasoundFrontendClassType::Graph: // TODO: Implement
-				default:
-				{
-					static_assert(static_cast<int32>(EMetasoundFrontendClassType::Invalid) == 4, "Possible missing case coverage");
-					checkNoEntry();
-				}
-				break;
+		UMetasoundEditorGraphOutputNode* FGraphBuilder::AddOutputNode(UObject& InMetasound, Frontend::FNodeHandle& InNodeHandle, bool bInSelectNewNode)
+		{
+			using namespace Frontend;
+
+			UMetasoundEditorGraphOutputNode* NewGraphNode = nullptr;
+			if (!ensure(InNodeHandle->GetClassType() == EMetasoundFrontendClassType::Output))
+			{
+				return nullptr;
 			}
 
+			TArray<FConstInputHandle> Inputs = InNodeHandle->GetConstInputs();
+			if (!ensure(!Inputs.IsEmpty()))
+			{
+				return nullptr;
+			}
+
+			FMetasoundAssetBase* MetasoundAsset = IMetasoundUObjectRegistry::Get().GetObjectAsAssetBase(&InMetasound);
+			check(MetasoundAsset);
+			UEdGraph& Graph = MetasoundAsset->GetGraphChecked();
+			FGraphNodeCreator<UMetasoundEditorGraphOutputNode> NodeCreator(Graph);
+			NewGraphNode = NodeCreator.CreateNode(bInSelectNewNode);
+
+			NodeCreator.Finalize();
+			InitGraphNode(InNodeHandle, NewGraphNode, InMetasound);
+
+			FConstInputHandle InputHandle = Inputs[0];
+			NewGraphNode->OutputTypeName = InputHandle->GetDataType();
+			return NewGraphNode;
+		}
+
+		void FGraphBuilder::InitGraphNode(Frontend::FNodeHandle& InNodeHandle, UMetasoundEditorGraphNode* NewGraphNode, UObject& InMetasound)
+		{
 			const FMetasoundFrontendNodeStyle& Style = InNodeHandle->GetNodeStyle();
 
-			NewGraphNode->SetNodeID(InNodeHandle->GetID());
 			NewGraphNode->CreateNewGuid();
+
+			NewGraphNode->NodeID = InNodeHandle->GetID();
 			NewGraphNode->NodePosX = Style.Display.Location.X;
 			NewGraphNode->NodePosY = Style.Display.Location.Y;
 
-			RebuildNodePins(*NewGraphNode, InNodeHandle);
+			RebuildNodePins(*NewGraphNode);
 
 			InMetasound.PostEditChange();
 			InMetasound.MarkPackageDirty();
-
-			return NewGraphNode;
 		}
 
-		UEdGraphNode* FGraphBuilder::AddNode(UObject& InMetasound, const Frontend::FNodeClassInfo& InClassInfo, const FMetasoundFrontendNodeStyle& InNodeStyle, bool bInSelectNewNode)
+		UMetasoundEditorGraphExternalNode* FGraphBuilder::AddExternalNode(UObject& InMetasound, const Frontend::FNodeClassInfo& InClassInfo, const FMetasoundFrontendNodeStyle& InNodeStyle, bool bInSelectNewNode)
 		{
-			Frontend::FNodeHandle NodeHandle = AddNodeHandle(InMetasound, InClassInfo, InNodeStyle);
-			return AddNode(InMetasound, NodeHandle, bInSelectNewNode);
+			Frontend::FNodeHandle NodeHandle = AddExternalNodeHandle(InMetasound, InClassInfo, InNodeStyle);
+			return AddExternalNode(InMetasound, NodeHandle, bInSelectNewNode);
 		}
 
-		Frontend::FNodeHandle FGraphBuilder::AddNodeHandle(UObject& InMetasound, const Frontend::FNodeClassInfo& InClassInfo, const FMetasoundFrontendNodeStyle& InNodeStyle)
+		Frontend::FNodeHandle FGraphBuilder::AddExternalNodeHandle(UObject& InMetasound, const Frontend::FNodeClassInfo& InClassInfo, const FMetasoundFrontendNodeStyle& InNodeStyle)
 		{
 			FMetasoundAssetBase* MetasoundAsset = IMetasoundUObjectRegistry::Get().GetObjectAsAssetBase(&InMetasound);
 			check(MetasoundAsset);
@@ -250,61 +263,67 @@ namespace Metasound
 			return GetOutputHandleFromPin(InPin);
 		}
 
-		UMetasoundEditorGraphInputNode* FGraphBuilder::AddInput(UObject& InMetasound, Frontend::FNodeHandle InNodeHandle, bool bInSelectNewNode)
+		UMetasoundEditorGraphInputNode* FGraphBuilder::AddInputNode(UObject& InMetasound, Frontend::FNodeHandle InNodeHandle, bool bInSelectNewNode)
 		{
 			using namespace Frontend;
 
 			FMetasoundAssetBase* MetasoundAsset = IMetasoundUObjectRegistry::Get().GetObjectAsAssetBase(&InMetasound);
 			check(MetasoundAsset);
 
-			const TArray<FOutputHandle>& Outputs = InNodeHandle->GetOutputs();
-			if (!ensure(!Outputs.IsEmpty()))
+			UMetasoundEditorGraph* MetasoundGraph = Cast<UMetasoundEditorGraph>(MetasoundAsset->GetGraph());
+			if (!ensure(MetasoundGraph))
 			{
 				return nullptr;
 			}
 
-			FOutputHandle OutputHandle = InNodeHandle->GetOutputs()[0];
-			FGraphHandle Graph = InNodeHandle->GetOwningGraph();
-			FGuid VertexID = Graph->GetVertexIDForInputVertex(OutputHandle->GetName());
-			FMetasoundFrontendLiteral DefaultLiteral = Graph->GetDefaultInput(VertexID);
-
-			EMetasoundFrontendLiteralType LiteralType = DefaultLiteral.GetType();
 
-			UMetasoundEditorGraph* MetasoundGraph = Cast<UMetasoundEditorGraph>(MetasoundAsset->GetGraph());
-			if (!ensure(MetasoundGraph))
+			TArray<FConstOutputHandle> Outputs = InNodeHandle->GetConstOutputs();
+			if (!ensure(!Outputs.IsEmpty()))
 			{
 				return nullptr;
 			}
 
+			FConstOutputHandle OutputHandle = Outputs[0];
+			FGraphHandle Graph = InNodeHandle->GetOwningGraph();
+			FGuid VertexID = Graph->GetVertexIDForInputVertex(OutputHandle->GetName());
+			FMetasoundFrontendLiteral DefaultLiteral = Graph->GetDefaultInput(VertexID);
+			EMetasoundFrontendLiteralType LiteralType = DefaultLiteral.GetType();
 			UMetasoundEditorGraphInputNode* NewGraphNode = MetasoundGraph->CreateInputNode(LiteralType, nullptr, bInSelectNewNode);
 			if (ensure(NewGraphNode))
 			{
-				RebuildNodePins(*NewGraphNode, InNodeHandle);
+				NewGraphNode->ClassName = InNodeHandle->GetClassName();
+				NewGraphNode->NodeID = InNodeHandle->GetID();
+				NewGraphNode->InputTypeName = OutputHandle->GetDataType();
 
-				NewGraphNode->SetNodeID(InNodeHandle->GetID());
+				RebuildNodePins(*NewGraphNode);
 				return NewGraphNode;
 			}
 
 			return nullptr;
 		}
 
-		UMetasoundEditorGraphInputNode* FGraphBuilder::AddInput(UObject& InMetasound, const FString& InName, const FName InTypeName, const FMetasoundFrontendNodeStyle& InNodeStyle, const FText& InToolTip, bool bInSelectNewNode)
+		UMetasoundEditorGraphInputNode* FGraphBuilder::AddInputNode(UObject& InMetasound, const FString& InName, const FName InTypeName, const FMetasoundFrontendNodeStyle& InNodeStyle, const FText& InToolTip, bool bInSelectNewNode)
 		{
 			using namespace Frontend;
 
 			FNodeHandle NodeHandle = AddInputNodeHandle(InMetasound, InName, InTypeName, InNodeStyle, InToolTip);
-			if (UMetasoundEditorGraphInputNode* NewGraphNode = AddInput(InMetasound, NodeHandle))
+			UMetasoundEditorGraphInputNode* NewGraphNode = AddInputNode(InMetasound, NodeHandle, bInSelectNewNode);
+			if (ensure(NewGraphNode))
 			{
 				const FMetasoundFrontendNodeStyle& Style = NodeHandle->GetNodeStyle();
 				NewGraphNode->NodePosX = Style.Display.Location.X;
 				NewGraphNode->NodePosY = Style.Display.Location.Y;
 				return NewGraphNode;
 			}
+			else
+			{
+				GraphBuilderPrivate::DeleteNode(InMetasound, NodeHandle);
+			}
 
 			return nullptr;
 		}
 
-		void FGraphBuilder::AddOrUpdateLiteralInput(UObject& InMetasound, Frontend::FNodeHandle InNodeHandle, UEdGraphPin& InInputPin)
+		void FGraphBuilder::AddOrUpdateLiteralInput(UObject& InMetasound, Frontend::FNodeHandle InNodeHandle, UEdGraphPin& InInputPin, bool bForcePinValueAsDefault)
 		{
 			using namespace Frontend;
 
@@ -334,10 +353,14 @@ namespace Metasound
 				NewInputNode->SetNodeStyle(Style);
 
 				const FMetasoundFrontendLiteral* DefaultLiteral = InputHandle->GetDefaultLiteral();
-				if (DefaultLiteral && DefaultLiteral->IsValid())
+				Frontend::FGraphHandle GraphHandle = InNodeHandle->GetOwningGraph();
+				const FGuid VertexID = GraphHandle->GetVertexIDForInputVertex(NewInputNode->GetNodeName());
+				if (bForcePinValueAsDefault)
+				{
+					GraphHandle->SetDefaultInput(VertexID, PinDataTypeDefaultLiteral);
+				}
+				else if (DefaultLiteral && DefaultLiteral->IsValid())
 				{
-					Frontend::FGraphHandle GraphHandle = InNodeHandle->GetOwningGraph();
-					const FGuid VertexID = GraphHandle->GetVertexIDForInputVertex(NewInputNode->GetNodeName());
 					GraphHandle->SetDefaultInput(VertexID, *DefaultLiteral);
 					InInputPin.DefaultValue = DefaultLiteral->ToString();
 				}
@@ -457,6 +480,62 @@ namespace Metasound
 			return true;
 		}
 
+		Frontend::FNodeHandle FGraphBuilder::AddNodeHandle(UObject& InMetasound, UMetasoundEditorGraphNode& InGraphNode)
+		{
+			FMetasoundFrontendNodeStyle Style;
+			Style.Display.Location = FVector2D(InGraphNode.NodePosX, InGraphNode.NodePosY);
+
+ 			if (UMetasoundEditorGraphInputNode* InputNode = Cast<UMetasoundEditorGraphInputNode>(&InGraphNode))
+ 			{
+ 				const TArray<UEdGraphPin*>& Pins = InGraphNode.GetAllPins();
+ 				const UEdGraphPin* Pin = Pins.IsEmpty() ? nullptr : Pins[0];
+ 				if (ensure(Pin) && ensure(Pin->Direction == EGPD_Output))
+ 				{
+ 					return FGraphBuilder::AddInputNodeHandle(InMetasound, Pin->GetName(), InputNode->InputTypeName, Style, InGraphNode.GetTooltipText());
+ 				}
+ 			}
+
+ 			if (UMetasoundEditorGraphOutputNode* OutputNode = Cast<UMetasoundEditorGraphOutputNode>(&InGraphNode))
+ 			{
+ 				const TArray<UEdGraphPin*>& Pins = InGraphNode.GetAllPins();
+ 				const UEdGraphPin* Pin = Pins.IsEmpty() ? nullptr : Pins[0];
+ 				if (ensure(Pin) && ensure(Pin->Direction == EGPD_Input))
+ 				{
+ 					return FGraphBuilder::AddOutputNodeHandle(InMetasound, Pin->GetName(), OutputNode->OutputTypeName, Style, InGraphNode.GetTooltipText());
+ 				}
+ 			}
+
+			if (UMetasoundEditorGraphExternalNode* ExternalNode = Cast<UMetasoundEditorGraphExternalNode>(&InGraphNode))
+			{
+				FFrontendQuery Query;
+				const FFrontendQuerySelectionView Result = Query
+					.AddStep<FGenerateAllAvailableNodeClasses>()
+					.AddStep<FFilterClassesByClassName>(ExternalNode->ClassName)
+					.ExecuteQuery();
+
+				TArrayView<const FFrontendQueryEntry* const> Selection = Result.GetSelection();
+				if (ensure(!Selection.IsEmpty()))
+				{
+					const FMetasoundFrontendClass* FrontendClass = Selection[0]->Value.TryGet<FMetasoundFrontendClass>();
+					if (ensure(FrontendClass))
+					{
+						FMetasoundAssetBase* MetasoundAsset = IMetasoundUObjectRegistry::Get().GetObjectAsAssetBase(&InMetasound);
+						check(MetasoundAsset);
+
+						Frontend::FNodeHandle NewNode = MetasoundAsset->GetRootGraphHandle()->AddNode(FrontendClass->Metadata);
+						NewNode->SetNodeStyle(Style);
+
+						InGraphNode.NodeID = NewNode->GetID();
+
+						return NewNode;
+					}
+				}
+			}
+
+			// Cast failed so no handle created
+			return Frontend::INodeController::GetInvalidHandle();
+		}
+
 		Frontend::FNodeHandle FGraphBuilder::AddInputNodeHandle(UObject& InMetasound, const FString& InName, const FName InTypeName, const FMetasoundFrontendNodeStyle& InNodeStyle, const FText& InToolTip, const FMetasoundFrontendLiteral* InDefaultValue)
 		{
 			FMetasoundAssetBase* MetasoundAsset = IMetasoundUObjectRegistry::Get().GetObjectAsAssetBase(&InMetasound);
@@ -494,10 +573,10 @@ namespace Metasound
 			return NodeHandle;
 		}
 
-		UEdGraphNode* FGraphBuilder::AddOutput(UObject& InMetasound, const FString& InName, const FName InTypeName, const FMetasoundFrontendNodeStyle& InNodeStyle, const FText& InToolTip, bool bInSelectNewNode)
+		UMetasoundEditorGraphOutputNode* FGraphBuilder::AddOutputNode(UObject& InMetasound, const FString& InName, const FName InTypeName, const FMetasoundFrontendNodeStyle& InNodeStyle, const FText& InToolTip, bool bInSelectNewNode)
 		{
 			Frontend::FNodeHandle NodeHandle = AddOutputNodeHandle(InMetasound, InName, InTypeName, InNodeStyle, InToolTip);
-			return AddNode(InMetasound, NodeHandle, bInSelectNewNode);
+			return AddOutputNode(InMetasound, NodeHandle, bInSelectNewNode);
 		}
 
 		Frontend::FNodeHandle FGraphBuilder::AddOutputNodeHandle(UObject& InMetasound, const FString& InName, const FName InTypeName, const FMetasoundFrontendNodeStyle& InNodeStyle, const FText& InToolTip)
@@ -521,13 +600,48 @@ namespace Metasound
 			return NodeHandle;
 		}
 
-		bool FGraphBuilder::ConnectNodes(UEdGraphPin& InInputPin, UEdGraphPin& InOutputPin, bool bConnectEdPins)
+		UEdGraphNode* FGraphBuilder::AddNode(UObject& InMetasound, Frontend::FNodeHandle InNodeHandle, bool bInSelectNewNode)
+		{
+			switch (InNodeHandle->GetClassType())
+			{
+				case EMetasoundFrontendClassType::Input:
+				{
+					return CastChecked<UEdGraphNode>(AddInputNode(InMetasound, InNodeHandle, bInSelectNewNode));
+				}
+				break;
+
+				case EMetasoundFrontendClassType::External:
+				{
+					return CastChecked<UEdGraphNode>(AddExternalNode(InMetasound, InNodeHandle, bInSelectNewNode));
+				}
+				break;
+
+				case EMetasoundFrontendClassType::Output:
+				{
+					return CastChecked<UEdGraphNode>(AddOutputNode(InMetasound, InNodeHandle, bInSelectNewNode));
+				}
+				break;
+
+				case EMetasoundFrontendClassType::Invalid:
+				case EMetasoundFrontendClassType::Graph:
+				default:
+				{
+					checkNoEntry();
+					static_assert(static_cast<int32>(EMetasoundFrontendClassType::Invalid) == 4, "Possible missing FMetasoundFrontendClassType case coverage");
+				}
+				break;
+			}
+
+			return nullptr;
+		}
+
+		bool FGraphBuilder::ConnectNodes(UEdGraphPin& InInputPin, UEdGraphPin& InOutputPin, bool bInConnectEdPins)
 		{
 			using namespace Frontend;
 
 			// When true, will recursively call back into this function
 			// from the schema if the editor pins are successfully connected
-			if (bConnectEdPins)
+			if (bInConnectEdPins)
 			{
 				const UEdGraphSchema* Schema = InInputPin.GetSchema();
 				if (ensure(Schema))
@@ -617,16 +731,7 @@ namespace Metasound
 				}
 				NodeHandle->SetNodeStyle(Style);
 
-				UEdGraphNode* NewNode = nullptr;
-				if (NodeHandle->GetClassType() == EMetasoundFrontendClassType::Input)
-				{
-					NewNode = Cast<UEdGraphNode>(AddInput(InMetasound, NodeHandle, false /* bInSelectNewNode */));
-				}
-				else
-				{
-					NewNode = AddNode(InMetasound, NodeHandle, false /* bInSelectNewNode */);
-				}
-
+				UEdGraphNode* NewNode = AddNode(InMetasound, NodeHandle, false /* bInSelectNewNode */);
 				NewIdNodeMap.Add(NodeHandle->GetID(), FNodePair { NodeHandle, NewNode });
 			}
 
@@ -775,24 +880,26 @@ namespace Metasound
 			return true;
 		}
 
-		void FGraphBuilder::RebuildNodePins(UMetasoundEditorGraphNode& InGraphNode, Frontend::FNodeHandle InNodeHandle)
+		void FGraphBuilder::RebuildNodePins(UMetasoundEditorGraphNode& InGraphNode)
 		{
 			using namespace Frontend;
+		
+			FNodeHandle NodeHandle = InGraphNode.GetNodeHandle();
 			for (int32 i = InGraphNode.Pins.Num() - 1; i >= 0; i--)
 			{
 				DeleteLiteralInputs(InGraphNode);
 				InGraphNode.RemovePin(InGraphNode.Pins[i]);
 			}
 
-			TArray<FInputHandle> InputHandles = InNodeHandle->GetInputs();
-			InputHandles = InNodeHandle->GetInputStyle().SortDefaults(InputHandles);
+			TArray<FInputHandle> InputHandles = NodeHandle->GetInputs();
+			InputHandles = NodeHandle->GetInputStyle().SortDefaults(InputHandles);
 			for (const FInputHandle& InputHandle : InputHandles)
 			{
 				AddPinToNode(InGraphNode, InputHandle);
 			}
 
-			TArray<FOutputHandle> OutputHandles = InNodeHandle->GetOutputs();
-			OutputHandles = InNodeHandle->GetOutputStyle().SortDefaults(OutputHandles);
+			TArray<FOutputHandle> OutputHandles = NodeHandle->GetOutputs();
+			OutputHandles = NodeHandle->GetOutputStyle().SortDefaults(OutputHandles);
 			for (const FOutputHandle& OutputHandle : OutputHandles)
 			{
 				AddPinToNode(InGraphNode, OutputHandle);
@@ -884,7 +991,6 @@ namespace Metasound
 			using namespace Frontend;
 
 			IMetasoundEditorModule& EditorModule = FModuleManager::GetModuleChecked<IMetasoundEditorModule>("MetasoundEditor");
-
 			FEdGraphPinType PinType("MetasoundNode", NAME_None, nullptr, EPinContainerType::None, false, FEdGraphTerminalType());
 			UEdGraphPin* NewPin = InEditorNode.CreatePin(EGPD_Input, PinType, FName(*InInputHandle->GetName()));
 
@@ -903,8 +1009,6 @@ namespace Metasound
 		UEdGraphPin* FGraphBuilder::AddPinToNode(UMetasoundEditorGraphNode& InEditorNode, Frontend::FOutputHandle InOutputHandle)
 		{
 			IMetasoundEditorModule& EditorModule = FModuleManager::GetModuleChecked<IMetasoundEditorModule>("MetasoundEditor");
-			FText DisplayName;
-
 			FEdGraphPinType PinType("MetasoundNode", NAME_None, nullptr, EPinContainerType::None, false, FEdGraphTerminalType());
 			UEdGraphPin* NewPin = InEditorNode.CreatePin(EGPD_Output, PinType, FName(*InOutputHandle->GetName()));
 
@@ -985,32 +1089,7 @@ namespace Metasound
 				if (Style.Display.Visibility == EMetasoundFrontendNodeStyleDisplayVisibility::Visible)
 				{
 					bIsEditorGraphDirty = true;
-					UEdGraphNode* NewNode = nullptr;
-					switch (Node->GetClassType())
-					{
-						case EMetasoundFrontendClassType::Input:
-						{
-							UMetasoundEditorGraphInputNode* InputNode = AddInput(InMetasound, Node, false /* bInSelectNewNode */);
-							NewNode = CastChecked<UEdGraphNode>(InputNode);
-						}
-						break;
-
-						case EMetasoundFrontendClassType::External:
-						case EMetasoundFrontendClassType::Graph:
-						case EMetasoundFrontendClassType::Output:
-						{
-							NewNode = AddNode(InMetasound, Node, false /* bInSelectNewNode */);
-						}
-						break;
-
-						case EMetasoundFrontendClassType::Invalid:
-						default:
-						{
-							static_assert(static_cast<int32>(EMetasoundFrontendClassType::Invalid) == 4, "Possible missing EMetasoundFrontendClassType case coverage");
-						}
-						break;
-					}
-
+					UEdGraphNode* NewNode = AddNode(InMetasound, Node, false /* bInSelectNewNode */);
 					if (ensure(NewNode))
 					{
 						PairedNodes.Add(Node->GetID(), { Cast<UMetasoundEditorGraphNode>(NewNode), Node });
@@ -1037,7 +1116,7 @@ namespace Metasound
 			return bIsEditorGraphDirty;
 		}
 
-		bool FGraphBuilder::SynchronizeNodePins(UMetasoundEditorGraphNode& InEditorNode, Frontend::FNodeHandle InNode)
+		bool FGraphBuilder::SynchronizeNodePins(UMetasoundEditorGraphNode& InEditorNode, Frontend::FNodeHandle InNode, bool bRemoveUnusedPins)
 		{
 			bool bIsNodeDirty = false;
 
@@ -1086,22 +1165,27 @@ namespace Metasound
 				}
 			}
 
-			bIsNodeDirty |= (InputHandles.Num() > 0);
-			bIsNodeDirty |= (OutputHandles.Num() > 0);
-			bIsNodeDirty |= (EditorPins.Num() > 0);
+			bIsNodeDirty |= !InputHandles.IsEmpty();
+			bIsNodeDirty |= !OutputHandles.IsEmpty();
+			bIsNodeDirty |= !EditorPins.IsEmpty();
 
 			// Remove any unused editor pins.
-			for (UEdGraphPin* Pin : EditorPins)
+			if (bRemoveUnusedPins)
 			{
-				InEditorNode.RemovePin(Pin);
+				for (UEdGraphPin* Pin : EditorPins)
+				{
+					InEditorNode.RemovePin(Pin);
+				}
 			}
 
 			// Add unmatched input pins
+			InputHandles = InNode->GetOutputStyle().SortDefaults(InputHandles);
 			for (Frontend::FInputHandle& InputHandle : InputHandles)
 			{
 				AddPinToNode(InEditorNode, InputHandle);
 			}
 
+			OutputHandles = InNode->GetOutputStyle().SortDefaults(OutputHandles);
 			for (Frontend::FOutputHandle& OutputHandle : OutputHandles)
 			{
 				AddPinToNode(InEditorNode, OutputHandle);
diff --git a/Engine/Plugins/Runtime/Metasound/Source/MetasoundEditor/Private/MetasoundEditorGraphBuilder.h b/Engine/Plugins/Runtime/Metasound/Source/MetasoundEditor/Private/MetasoundEditorGraphBuilder.h
--- a/Engine/Plugins/Runtime/Metasound/Source/MetasoundEditor/Private/MetasoundEditorGraphBuilder.h
+++ b/Engine/Plugins/Runtime/Metasound/Source/MetasoundEditor/Private/MetasoundEditorGraphBuilder.h
@@ -29,6 +29,8 @@ namespace Metasound
 	{
 		class FGraphBuilder
 		{
+			static void InitGraphNode(Frontend::FNodeHandle& InNodeHandle, UMetasoundEditorGraphNode* NewGraphNode, UObject& InMetasound);
+
 		public:
 			static const FName PinCategoryAudio;
 			static const FName PinCategoryBoolean;
@@ -46,18 +48,47 @@ namespace Metasound
 			static const FText ConvertMenuName;
 			static const FText FunctionMenuName;
 
-			// Adds a node to the editor graph that corresponds to the provided node handle.
-			static UEdGraphNode* AddNode(UObject& InMetasound, Frontend::FNodeHandle& InNodeHandle, bool bInSelectNewNode = true);
+			// Adds an EdGraph node to mirror the provided FNodeHandle.
+			static UEdGraphNode* AddNode(UObject& InMetasound, Frontend::FNodeHandle InNodeHandle, bool bInSelectNewNode = true);
+
+			// Adds a node handle to mirror the provided graph node and binds to it.  Does *NOT* mirror existing EdGraph connections
+			// nor does it remove existing bound Frontend Node (if set) from associated Frontend Graph.
+			static Frontend::FNodeHandle AddNodeHandle(UObject& InMetasound, UMetasoundEditorGraphNode& InGraphNode);
+
+			// Adds an output node to the editor graph that corresponds to the provided node handle.
+			// Generates analogous FNodeHandle.
+			static UMetasoundEditorGraphInputNode* AddInputNode(UObject& InMetasound, const FString& InName, const FName InTypeName, const FMetasoundFrontendNodeStyle& InNodeStyle, const FText& InToolTip, bool bInSelectNewNode = true);
 
-			// Adds a node with the given class info to both the editor and document graphs
-			static UEdGraphNode* AddNode(UObject& InMetasound, const Frontend::FNodeClassInfo& InClassInfo, const FMetasoundFrontendNodeStyle& InNodeStyle, bool bInSelectNewNode = true);
+			// Adds a corresponding UMetasoundEditorGraphInputNode for the provided node handle.
+			static UMetasoundEditorGraphInputNode* AddInputNode(UObject& InMetasound, Frontend::FNodeHandle InNodeHandle, bool bInSelectNewNode = true);
+
+			// Generates FNodeHandle for the given external node data. Does not bind or create EdGraph representation of given node.
+			static Frontend::FNodeHandle AddInputNodeHandle(UObject& InMetasound, const FString& InName, const FName InTypeName, const FMetasoundFrontendNodeStyle& InNodeStyle, const FText& InToolTip, const FMetasoundFrontendLiteral* InDefaultValue = nullptr);
 
-			// Adds a node handle with the given class and style info
-			static Frontend::FNodeHandle AddNodeHandle(UObject& InMetasound, const Frontend::FNodeClassInfo& InClassInfo, const FMetasoundFrontendNodeStyle& InNodeStyle);
+			// Adds a corresponding UMetasoundEditorGraphExternalNode for the provided node handle.
+			static UMetasoundEditorGraphExternalNode* AddExternalNode(UObject& InMetasound, Frontend::FNodeHandle& InNodeHandle, bool bInSelectNewNode = true);
 
-			// Attempts to connect graph nodes together.  Returns true if succeeded, breaks pin link and returns false if failed.
-			// If bConnectEdPins is set, will attempt to connect the Editor Graph representation of the pins.
-			static bool ConnectNodes(UEdGraphPin& InInputPin, UEdGraphPin& InOutputPin, bool bConnectEdPins);
+			// Adds an externally-defined node with the given class info to both the editor and document graphs.
+			// Generates analogous FNodeHandle.
+			static UMetasoundEditorGraphExternalNode* AddExternalNode(UObject& InMetasound, const Frontend::FNodeClassInfo& InClassInfo, const FMetasoundFrontendNodeStyle& InNodeStyle, bool bInSelectNewNode = true);
+
+			// Generates FNodeHandle for the given external node data. Does not bind or create EdGraph representation of given node.
+			static Frontend::FNodeHandle AddExternalNodeHandle(UObject& InMetasound, const Frontend::FNodeClassInfo& InClassInfo, const FMetasoundFrontendNodeStyle& InNodeStyle);
+
+			// Adds an output node to the editor graph that corresponds to the provided node handle.
+			static UMetasoundEditorGraphOutputNode* AddOutputNode(UObject& InMetasound, Frontend::FNodeHandle& InNodeHandle, bool bInSelectNewNode = true);
+
+			// Adds an output node to the editor graph that corresponds to the provided node handle.
+			// Generates analogous FNodeHandle.
+			static UMetasoundEditorGraphOutputNode* AddOutputNode(UObject& InMetasound, const FString& InName, const FName InTypeName, const FMetasoundFrontendNodeStyle& InNodeStyle, const FText& InToolTip, bool bInSelectNewNode = true);
+
+			// Generates analogous FNodeHandle for the given internal node data. Does not bind nor create EdGraph representation of given node.
+			static Frontend::FNodeHandle AddOutputNodeHandle(UObject& InMetasound, const FString& InName, const FName InTypeName, const FMetasoundFrontendNodeStyle& InNodeStyle, const FText& InToolTip);
+
+			// Attempts to connect frontend node counterparts together for provided pines.  Returns true if succeeded,
+			// and breaks pin link and returns false if failed.  If bConnectEdPins is set, will attempt to connect
+			// the Editor Graph representation of the pins.
+			static bool ConnectNodes(UEdGraphPin& InInputPin, UEdGraphPin& InOutputPin, bool bInConnectEdPins);
 
 			static FString GenerateUniqueInputName(const UObject& InMetasound, const FName InBaseName);
 
@@ -77,32 +108,21 @@ namespace Metasound
 			static Frontend::FOutputHandle GetOutputHandleFromPin(const UEdGraphPin* InPin);
 			static Frontend::FConstOutputHandle GetConstOutputHandleFromPin(const UEdGraphPin* InPin);
 
-			static UMetasoundEditorGraphInputNode* AddInput(UObject& InMetasound, const FString& InName, const FName InTypeName, const FMetasoundFrontendNodeStyle& InNodeStyle, const FText& InToolTip, bool bInSelectNewNode = true);
-			static UMetasoundEditorGraphInputNode* AddInput(UObject& InMetasound, Frontend::FNodeHandle InNodeHandle, bool bInSelectNewNode = true);
-
 			// Adds or updates a hidden input that is set to the default literal value.  If the literal input is added, value defaults to
 			// 1. what is set on the input vertex of the node's inputer vertex interface, and if not defined there,
 			// 2. the pin's DataType literal default.
-			static void AddOrUpdateLiteralInput(UObject& InMetasound, Frontend::FNodeHandle InNodeHandle, UEdGraphPin& InInputPin);
+			static void AddOrUpdateLiteralInput(UObject& InMetasound, Frontend::FNodeHandle InNodeHandle, UEdGraphPin& InInputPin, bool bForcePinValueAsDefault = false);
 
 			static bool GetPinDefaultLiteral(UEdGraphPin& InInputPin, Frontend::FInputHandle InputHandle, FMetasoundFrontendLiteral& OutLiteralDefault);
 
 			static bool IsRequiredInput(Frontend::FNodeHandle InNodeHandle);
 			static bool IsRequiredOutput(Frontend::FNodeHandle InNodeHandle);
 
-			// Adds an input node handle with the given class info
-			static Frontend::FNodeHandle AddInputNodeHandle(UObject& InMetasound, const FString& InName, const FName InTypeName, const FMetasoundFrontendNodeStyle& InNodeStyle, const FText& InToolTip, const FMetasoundFrontendLiteral* InDefaultValue = nullptr);
-
-			static UEdGraphNode* AddOutput(UObject& InMetasound, const FString& InName, const FName InTypeName, const FMetasoundFrontendNodeStyle& InNodeStyle, const FText& InToolTip, bool bInSelectNewNode = true);
-
-			// Adds an output node handle with the given class info
-			static Frontend::FNodeHandle AddOutputNodeHandle(UObject& InMetasound, const FString& InName, const FName InTypeName, const FMetasoundFrontendNodeStyle& InNodeStyle, const FText& InToolTip);
-
 			// Constructs graph with default inputs & outputs.
 			static void ConstructGraph(UObject& InMetasound);
 
 			// Rebuilds all editor node pins based on the provided node handle's class definition.
-			static void RebuildNodePins(UMetasoundEditorGraphNode& InGraphNode, Frontend::FNodeHandle InNodeHandle);
+			static void RebuildNodePins(UMetasoundEditorGraphNode& InGraphNode);
 
 			// Removes all literal inputs connected to the given node
 			static void DeleteLiteralInputs(UEdGraphNode& InNode);
@@ -124,7 +144,7 @@ namespace Metasound
 			// Adds and removes pins so that the UMetasoundEditorGraphNode matches the InNode.
 			//
 			// @return True if the UMetasoundEditorGraphNode was altered. False otherwise.
-			static bool SynchronizeNodePins(UMetasoundEditorGraphNode& InEditorNode, Frontend::FNodeHandle InNode);
+			static bool SynchronizeNodePins(UMetasoundEditorGraphNode& InEditorNode, Frontend::FNodeHandle InNode, bool bRemoveUnusedPins = true);
 
 			// Adds and removes connections so that the UEdGraph of the metasound has the same
 			// connections as the FMetasoundFrontendDocument graph.
diff --git a/Engine/Plugins/Runtime/Metasound/Source/MetasoundEditor/Private/MetasoundEditorGraphNode.cpp b/Engine/Plugins/Runtime/Metasound/Source/MetasoundEditor/Private/MetasoundEditorGraphNode.cpp
--- a/Engine/Plugins/Runtime/Metasound/Source/MetasoundEditor/Private/MetasoundEditorGraphNode.cpp
+++ b/Engine/Plugins/Runtime/Metasound/Source/MetasoundEditor/Private/MetasoundEditorGraphNode.cpp
@@ -93,6 +93,15 @@ const UObject& UMetasoundEditorGraphNode::GetMetasoundChecked() const
 	return EdGraph->GetMetasoundChecked();
 }
 
+Metasound::Frontend::FConstGraphHandle UMetasoundEditorGraphNode::GetConstRootGraphHandle() const
+{
+	const FMetasoundAssetBase* ConstMetasoundAsset = Metasound::IMetasoundUObjectRegistry::Get().GetObjectAsAssetBase(&GetMetasoundChecked());
+	FMetasoundAssetBase* MetasoundAsset = const_cast<FMetasoundAssetBase*>(ConstMetasoundAsset);
+	check(MetasoundAsset);
+
+	return MetasoundAsset->GetRootGraphHandle();
+}
+
 Metasound::Frontend::FGraphHandle UMetasoundEditorGraphNode::GetRootGraphHandle() const
 {
 	const FMetasoundAssetBase* ConstMetasoundAsset = Metasound::IMetasoundUObjectRegistry::Get().GetObjectAsAssetBase(&GetMetasoundChecked());
@@ -102,6 +111,11 @@ Metasound::Frontend::FGraphHandle UMetasoundEditorGraphNode::GetRootGraphHandle(
 	return MetasoundAsset->GetRootGraphHandle();
 }
 
+Metasound::Frontend::FConstNodeHandle UMetasoundEditorGraphNode::GetConstNodeHandle() const
+{
+	return GetConstRootGraphHandle()->GetNodeWithID(NodeID);
+}
+
 Metasound::Frontend::FNodeHandle UMetasoundEditorGraphNode::GetNodeHandle() const
 {
 	return GetRootGraphHandle()->GetNodeWithID(NodeID);
@@ -123,14 +137,20 @@ void UMetasoundEditorGraphNode::AllocateDefaultPins()
 	using namespace Metasound;
 
 	ensureAlways(Pins.IsEmpty());
-	Editor::FGraphBuilder::RebuildNodePins(*this, GetNodeHandle());
+	Editor::FGraphBuilder::RebuildNodePins(*this);
 }
 
 void UMetasoundEditorGraphNode::ReconstructNode()
 {
 	using namespace Metasound;
 
-	Editor::FGraphBuilder::RebuildNodePins(*this, GetNodeHandle());
+	// Don't remove unused pins here. Reconstruction can occur while duplicating or pasting nodes,
+	// and subsequent steps clean-up unused pins.  This can be called mid-copy, which means the node
+	// handle may be invalid.  Setting to remove unused causes premature removal and then default values
+	// are lost.
+	// TODO: User will want to see dead pins as well for node definition changes. Label and color-code dead
+	// pins (ex. red), and leave dead connections for reference like BP.
+	Editor::FGraphBuilder::SynchronizeNodePins(*this, GetNodeHandle(), false /* bRemoveUnusedPins */);
 }
 
 void UMetasoundEditorGraphNode::AutowireNewNode(UEdGraphPin* FromPin)
@@ -207,11 +227,6 @@ FString UMetasoundEditorGraphNode::GetDocumentationLink() const
 	return TEXT("Shared/GraphNodes/Metasound");
 }
 
-void UMetasoundEditorGraphNode::SetNodeID(FGuid InNodeID)
-{
-	NodeID = InNodeID;
-}
-
 FGuid UMetasoundEditorGraphNode::GetNodeID() const
 {
 	return NodeID;
diff --git a/Engine/Plugins/Runtime/Metasound/Source/MetasoundEditor/Private/MetasoundEditorGraphSchema.cpp b/Engine/Plugins/Runtime/Metasound/Source/MetasoundEditor/Private/MetasoundEditorGraphSchema.cpp
--- a/Engine/Plugins/Runtime/Metasound/Source/MetasoundEditor/Private/MetasoundEditorGraphSchema.cpp
+++ b/Engine/Plugins/Runtime/Metasound/Source/MetasoundEditor/Private/MetasoundEditorGraphSchema.cpp
@@ -274,7 +274,7 @@ UEdGraphNode* FMetasoundGraphSchemaAction_NewNode::PerformAction(UEdGraph* Paren
 
 	FMetasoundFrontendNodeStyle Style;
 	Style.Display.Location = Location;
-	if (UEdGraphNode* NewGraphNode = FGraphBuilder::AddNode(ParentMetasound, NodeClassInfo, Style))
+	if (UEdGraphNode* NewGraphNode = FGraphBuilder::AddExternalNode(ParentMetasound, NodeClassInfo, Style, bSelectNewNode))
 	{
 		TryConnectNewNodeToPin(*NewGraphNode, FromPin);
 		return NewGraphNode;
@@ -305,7 +305,7 @@ UEdGraphNode* FMetasoundGraphSchemaAction_NewInput::PerformAction(UEdGraph* Pare
 	FMetasoundFrontendNodeStyle Style;
 	Style.Display.Location = Location;
 
-	if (UMetasoundEditorGraphInputNode* NewGraphNode = FGraphBuilder::AddInput(ParentMetasound, NewNodeName, NodeTypeName, Style, FText::GetEmpty()))
+	if (UMetasoundEditorGraphInputNode* NewGraphNode = FGraphBuilder::AddInputNode(ParentMetasound, NewNodeName, NodeTypeName, Style, FText::GetEmpty()))
 	{
 		UEdGraphNode* EdGraphNode = CastChecked<UEdGraphNode>(NewGraphNode);
 		TryConnectNewNodeToPin(*EdGraphNode, FromPin);
@@ -337,7 +337,7 @@ UEdGraphNode* FMetasoundGraphSchemaAction_NewOutput::PerformAction(UEdGraph* Par
 	FMetasoundFrontendNodeStyle Style;
 	Style.Display.Location = Location;
 
-	if (UEdGraphNode* NewGraphNode = FGraphBuilder::AddOutput(ParentMetasound, NewNodeName, NodeTypeName, Style, FText::GetEmpty()))
+	if (UEdGraphNode* NewGraphNode = FGraphBuilder::AddOutputNode(ParentMetasound, NewNodeName, NodeTypeName, Style, FText::GetEmpty()))
 	{
 		TryConnectNewNodeToPin(*NewGraphNode, FromPin);
 		return NewGraphNode;
diff --git a/Engine/Plugins/Runtime/Metasound/Source/MetasoundEditor/Public/MetasoundEditorGraphInputNodes.h b/Engine/Plugins/Runtime/Metasound/Source/MetasoundEditor/Public/MetasoundEditorGraphInputNodes.h
--- a/Engine/Plugins/Runtime/Metasound/Source/MetasoundEditor/Public/MetasoundEditorGraphInputNodes.h
+++ b/Engine/Plugins/Runtime/Metasound/Source/MetasoundEditor/Public/MetasoundEditorGraphInputNodes.h
@@ -20,6 +20,14 @@
 class UEdGraphPin;
 class UMetasound;
 
+namespace Metasound
+{
+	namespace Editor
+	{
+		class FGraphBuilder;
+	}
+}
+
 
 UCLASS()
 class METASOUNDEDITOR_API UMetasoundEditorGraphInputNode : public UMetasoundEditorGraphNode
@@ -64,6 +72,12 @@ public:
 	{
 		return EMetasoundFrontendLiteralType::None;
 	}
+
+protected:
+	UPROPERTY()
+	FName InputTypeName;
+
+	friend class Metasound::Editor::FGraphBuilder;
 };
 
 UCLASS(MinimalAPI)
diff --git a/Engine/Plugins/Runtime/Metasound/Source/MetasoundEditor/Public/MetasoundEditorGraphNode.h b/Engine/Plugins/Runtime/Metasound/Source/MetasoundEditor/Public/MetasoundEditorGraphNode.h
--- a/Engine/Plugins/Runtime/Metasound/Source/MetasoundEditor/Public/MetasoundEditorGraphNode.h
+++ b/Engine/Plugins/Runtime/Metasound/Source/MetasoundEditor/Public/MetasoundEditorGraphNode.h
@@ -19,6 +19,14 @@
 class UEdGraphPin;
 class UMetasound;
 
+namespace Metasound
+{
+	namespace Editor
+	{
+		class FGraphBuilder;
+	}
+}
+
 
 UCLASS(MinimalAPI)
 class UMetasoundEditorGraphNode : public UEdGraphNode
@@ -71,20 +79,33 @@ public:
 	const UObject& GetMetasoundChecked() const;
 
 	Metasound::Frontend::FGraphHandle GetRootGraphHandle() const;
+	Metasound::Frontend::FConstGraphHandle GetConstRootGraphHandle() const;
+
 	Metasound::Frontend::FNodeHandle GetNodeHandle() const;
+	Metasound::Frontend::FConstNodeHandle GetConstNodeHandle() const;
 
-	void SetNodeID(FGuid InNodeID);
 	FGuid GetNodeID() const;
 
 protected:
+	UPROPERTY()
+	FMetasoundFrontendClassName ClassName;
+
 	UPROPERTY()
 	FGuid NodeID;
+
+	friend class Metasound::Editor::FGraphBuilder;
 };
 
 UCLASS(MinimalAPI)
 class UMetasoundEditorGraphOutputNode : public UMetasoundEditorGraphNode
 {
 	GENERATED_BODY()
+
+protected:
+	UPROPERTY()
+	FName OutputTypeName;
+
+	friend class Metasound::Editor::FGraphBuilder;
 };
 
 UCLASS(MinimalAPI)
diff --git a/Engine/Plugins/Runtime/Metasound/Source/MetasoundFrontend/Private/MetasoundFrontendQuerySteps.cpp b/Engine/Plugins/Runtime/Metasound/Source/MetasoundFrontend/Private/MetasoundFrontendQuerySteps.cpp
--- a/Engine/Plugins/Runtime/Metasound/Source/MetasoundFrontend/Private/MetasoundFrontendQuerySteps.cpp
+++ b/Engine/Plugins/Runtime/Metasound/Source/MetasoundFrontend/Private/MetasoundFrontendQuerySteps.cpp
@@ -59,5 +59,15 @@ namespace Metasound
 	{
 		return InEntry.Value.Get<FMetasoundFrontendClass>().Metadata.ClassName == ClassName;
 	}
+
+	FFilterClassesByClassID::FFilterClassesByClassID(const FGuid InClassID)
+		: ClassID(InClassID)
+	{
+	}
+
+	bool FFilterClassesByClassID::Filter(const FFrontendQueryEntry& InEntry) const
+	{
+		return InEntry.Value.Get<FMetasoundFrontendClass>().ID == ClassID;
+	}
 }
 
diff --git a/Engine/Plugins/Runtime/Metasound/Source/MetasoundFrontend/Public/MetasoundFrontendQuerySteps.h b/Engine/Plugins/Runtime/Metasound/Source/MetasoundFrontend/Public/MetasoundFrontendQuerySteps.h
--- a/Engine/Plugins/Runtime/Metasound/Source/MetasoundFrontend/Public/MetasoundFrontendQuerySteps.h
+++ b/Engine/Plugins/Runtime/Metasound/Source/MetasoundFrontend/Public/MetasoundFrontendQuerySteps.h
@@ -60,5 +60,16 @@ namespace Metasound
 	private:
 		FMetasoundFrontendClassName ClassName;
 	};
+
+	class METASOUNDFRONTEND_API FFilterClassesByClassID : public IFrontendQueryFilterStep
+	{
+	public:
+		FFilterClassesByClassID(const FGuid InClassID);
+
+		bool Filter(const FFrontendQueryEntry& InEntry) const override;
+
+	private:
+		FGuid ClassID;
+	};
 }
 
