Metasound filter fixup (param names, default values, clamping)

diff --git a/Engine/Plugins/Runtime/Metasound/Source/MetasoundStandardNodes/Private/MetasoundBasicFilters.cpp b/Engine/Plugins/Runtime/Metasound/Source/MetasoundStandardNodes/Private/MetasoundBasicFilters.cpp
--- a/Engine/Plugins/Runtime/Metasound/Source/MetasoundStandardNodes/Private/MetasoundBasicFilters.cpp
+++ b/Engine/Plugins/Runtime/Metasound/Source/MetasoundStandardNodes/Private/MetasoundBasicFilters.cpp
@@ -13,6 +13,7 @@
 
 #include "DSP/Filter.h"
 #include "DSP/InterpolatedOnePole.h"
+#include "Math/NumericLimits.h"
 
 #define LOCTEXT_NAMESPACE "MetasoundBasicFilterNodes"
 
@@ -22,6 +23,9 @@ namespace Metasound
 
 	class FLadderFilterOperator : public TExecutableOperator<FLadderFilterOperator>
 	{
+	private:
+		static constexpr float InvalidValue = -1.f;
+
 	public:
 		static const FNodeClassMetadata& GetNodeInfo();
 
@@ -35,7 +39,6 @@ namespace Metasound
 			InputDataReferences.AddDataReadReference(TEXT("In"), FAudioBufferReadRef(AudioInput));
 			InputDataReferences.AddDataReadReference(TEXT("Freq."), FFloatReadRef(Frequency));
 			InputDataReferences.AddDataReadReference(TEXT("Res."), FFloatReadRef(Resonance));
-			InputDataReferences.AddDataReadReference(TEXT("Pass-Band Gain Compensation"), FFloatReadRef(PassBandGainCompensation));
 
 			return InputDataReferences;
 		}
@@ -44,7 +47,7 @@ namespace Metasound
 		{
 			// expose read access to our output buffer for other processors in the graph
 			FDataReferenceCollection OutputDataReferences;
-			OutputDataReferences.AddDataReadReference(TEXT("Out"), FAudioBufferReadRef(AudioOutput)); 
+			OutputDataReferences.AddDataReadReference(TEXT("Out"), FAudioBufferReadRef(AudioOutput));
 
 			return OutputDataReferences;
 		}
@@ -57,19 +60,18 @@ namespace Metasound
 		FAudioBufferReadRef AudioInput;
 		FFloatReadRef Frequency;
 		FFloatReadRef Resonance;
-		FFloatReadRef PassBandGainCompensation;
 
 		// cached data
-		float LastFrequency;
-		float LastResonance;
-		float LastPassBandGainCompensation;
+		float PreviousFrequency{ InvalidValue };
+		float PreviousResonance{ InvalidValue };
 
 		// output pins
 		FAudioBufferWriteRef AudioOutput;
 
 		// data
 		const int32 BlockSize;
-		FSampleRate SampleRate;
+		const float SampleRate;
+		const float MaxCutoffFrequency;
 
 		// dsp
 		Audio::FLadderFilter LadderFilter;
@@ -81,16 +83,15 @@ namespace Metasound
 			const FOperatorSettings& InSettings,
 			const FAudioBufferReadRef& InAudioInput,
 			const FFloatReadRef& InFrequency,
-			const FFloatReadRef& InResonance,
-			const FFloatReadRef& InPassBandGainCompensation
+			const FFloatReadRef& InResonance
 		)
 			: AudioInput(InAudioInput)
 			, Frequency(InFrequency)
 			, Resonance(InResonance)
-			, PassBandGainCompensation(InPassBandGainCompensation)
 			, AudioOutput(FAudioBufferWriteRef::CreateNew(InSettings))
 			, BlockSize(InSettings.GetNumFramesPerBlock())
 			, SampleRate(InSettings.GetSampleRate())
+			, MaxCutoffFrequency(0.5f * SampleRate)
 		{
 			// verify our buffer sizes:
 			check(AudioOutput->Num() == BlockSize);
@@ -137,9 +138,8 @@ namespace Metasound
 		static const FVertexInterface Interface(
 			FInputVertexInterface(
 				TInputDataVertexModel<FAudioBuffer>(TEXT("In"), LOCTEXT("AudioInputTooltip", "Audio Input")),
-				TInputDataVertexModel<float>(TEXT("Freq."), LOCTEXT("FrequencyTooltip", "Controls cutoff frequency")),
-				TInputDataVertexModel<float>(TEXT("Res."), LOCTEXT("ResonanceTooltip", "Controls filter resonance")),
-				TInputDataVertexModel<float>(TEXT("Pass-Band Gain Comp."), LOCTEXT("PassBandGainCompensationTooltip", "Controls pass-band gain compenation"))
+				TInputDataVertexModel<float>(TEXT("Freq."), LOCTEXT("FrequencyTooltip", "Controls cutoff frequency"), 20000.f),
+				TInputDataVertexModel<float>(TEXT("Res."), LOCTEXT("ResonanceTooltip", "Controls filter resonance"), 6.f)
 			),
 			FOutputVertexInterface(
 				TOutputDataVertexModel<FAudioBuffer>(TEXT("Out"), LOCTEXT("FilterOutputToolTip", "Audio Out"))
@@ -159,38 +159,36 @@ namespace Metasound
 		FAudioBufferReadRef AudioIn = InputDataRefs.GetDataReadReferenceOrConstruct<FAudioBuffer>(TEXT("In"), InParams.OperatorSettings);
 		FFloatReadRef FrequencyIn = InputDataRefs.GetDataReadReferenceOrConstruct<float>(TEXT("Freq."));
 		FFloatReadRef ResonanceIn = InputDataRefs.GetDataReadReferenceOrConstruct<float>(TEXT("Res."));
-		FFloatReadRef PassBandGainCompensationIn = InputDataRefs.GetDataReadReferenceOrConstruct<float>(TEXT("Pass-Band Gain Compensation"));
 
 		return MakeUnique<FLadderFilterOperator>(
 			InParams.OperatorSettings
 			, AudioIn
 			, FrequencyIn
 			, ResonanceIn
-			, PassBandGainCompensationIn
 			);
 	}
 
 	void FLadderFilterOperator::Execute()
 	{
-		bool bNeedsUpdate = false;
-		bNeedsUpdate |= (!FMath::IsNearlyEqual(LastFrequency, *Frequency));
-		bNeedsUpdate |= (!FMath::IsNearlyEqual(LastResonance, *Resonance));
-		bNeedsUpdate |= (!FMath::IsNearlyEqual(LastPassBandGainCompensation, *PassBandGainCompensation));
+		const float CurrentFrequency = FMath::Clamp(*Frequency, 0.f, MaxCutoffFrequency);
+		const float CurrentResonance = FMath::Clamp(*Resonance, 1.0f, 10.0f);
+
+		const bool bNeedsUpdate =
+			   (!FMath::IsNearlyEqual(PreviousFrequency, CurrentFrequency))
+			|| (!FMath::IsNearlyEqual(PreviousResonance, CurrentResonance));
 
 		if (bNeedsUpdate)
 		{
 			LadderFilter.SetQ(*Resonance);
 			LadderFilter.SetFrequency(*Frequency);
-			LadderFilter.SetPassBandGainCompensation(*PassBandGainCompensation);
 
 			LadderFilter.Update();
+
+			PreviousFrequency = *Frequency;
+			PreviousResonance = *Resonance;
 		}
 
 		LadderFilter.ProcessAudio(AudioInput->GetData(), AudioInput->Num(), AudioOutput->GetData());
-
-		LastFrequency = *Frequency;
-		LastResonance = *Resonance;
-		LastPassBandGainCompensation = *PassBandGainCompensation;
 	}
 
 	METASOUND_REGISTER_NODE(FLadderFilterNode);
@@ -202,6 +200,8 @@ namespace Metasound
 
 	class FStateVariableFilterOperator : public TExecutableOperator<FStateVariableFilterOperator>
 	{
+	private:
+		static constexpr float InvalidValue = -1.f;
 	public:
 		static const FNodeClassMetadata& GetNodeInfo();
 
@@ -215,7 +215,7 @@ namespace Metasound
 			InputDataReferences.AddDataReadReference(TEXT("In"), FAudioBufferReadRef(AudioInput));
 			InputDataReferences.AddDataReadReference(TEXT("Freq."), FFloatReadRef(Frequency));
 			InputDataReferences.AddDataReadReference(TEXT("Res."), FFloatReadRef(Resonance));
-			InputDataReferences.AddDataReadReference(TEXT("Pass-Band Gain Compensation"), FFloatReadRef(BandStopControl));
+			InputDataReferences.AddDataReadReference(TEXT("Band-Stop Control"), FFloatReadRef(BandStopControl));
 
 			return InputDataReferences;
 		}
@@ -243,9 +243,9 @@ namespace Metasound
 		FFloatReadRef BandStopControl;
 
 		// cached data
-		float LastFrequency;
-		float LastResonance;
-		float LastBandStopControl;
+		float PreviousFrequency{ InvalidValue };
+		float PreviousResonance{ InvalidValue };
+		float PreviousBandStopControl{ InvalidValue };
 
 		// output pins
 		FAudioBufferWriteRef LowPassOutput;
@@ -255,7 +255,8 @@ namespace Metasound
 
 		// data
 		const int32 BlockSize;
-		FSampleRate SampleRate;
+		float SampleRate;
+		float MaxCutoffFrequency;
 
 		// dsp
 		Audio::FStateVariableFilter StateVariableFilter;
@@ -280,6 +281,7 @@ namespace Metasound
 			, BandStopOutput(FAudioBufferWriteRef::CreateNew(InSettings))
 			, BlockSize(InSettings.GetNumFramesPerBlock())
 			, SampleRate(InSettings.GetSampleRate())
+			, MaxCutoffFrequency(0.5f * SampleRate)
 		{
 			// verify our buffer sizes:
 			check(LowPassOutput->Num() == BlockSize);
@@ -329,9 +331,9 @@ namespace Metasound
 		static const FVertexInterface Interface(
 			FInputVertexInterface(
 				TInputDataVertexModel<FAudioBuffer>(TEXT("In"), LOCTEXT("AudioInputTooltip", "Audio Input")),
-				TInputDataVertexModel<float>(TEXT("Freq."), LOCTEXT("FrequencyTooltip", "Controls cutoff frequency")),
-				TInputDataVertexModel<float>(TEXT("Res."), LOCTEXT("ResonanceTooltip", "Controls filter resonance")),
-				TInputDataVertexModel<float>(TEXT("Band-Stop Control"), LOCTEXT("PassBandGainCompensationTooltip", "Band Stop Control"))
+				TInputDataVertexModel<float>(TEXT("Freq."), LOCTEXT("FrequencyTooltip", "Controls cutoff frequency"), 20000.f),
+				TInputDataVertexModel<float>(TEXT("Res."), LOCTEXT("ResonanceTooltip", "Controls filter resonance"), 0.f),
+				TInputDataVertexModel<float>(TEXT("Band-Stop Control"), LOCTEXT("BandStopControlTooltip", "Band Stop Control"), 0.f)
 			),
 			FOutputVertexInterface(
 				TOutputDataVertexModel<FAudioBuffer>(TEXT("LP"), LOCTEXT("FilterOutputToolTip", "Low Pass Output")),
@@ -367,32 +369,36 @@ namespace Metasound
 
 	void FStateVariableFilterOperator::Execute()
 	{
-		bool bNeedsUpdate = false;
-		bNeedsUpdate |= (!FMath::IsNearlyEqual(LastFrequency, *Frequency));
-		bNeedsUpdate |= (!FMath::IsNearlyEqual(LastResonance, *Resonance));
-		bNeedsUpdate |= (!FMath::IsNearlyEqual(LastBandStopControl, *BandStopControl));
+		const float CurrentFrequency = FMath::Clamp(*Frequency, 0.f, MaxCutoffFrequency);
+		const float CurrentResonance = FMath::Clamp(*Resonance, 0.f, 10.f);
+		const float CurrentBandStopControl = FMath::Clamp(*BandStopControl, 0.f, 1.f);
+
+		bool bNeedsUpdate =
+			   (!FMath::IsNearlyEqual(PreviousFrequency, CurrentFrequency))
+			|| (!FMath::IsNearlyEqual(PreviousResonance, CurrentResonance))
+			|| (!FMath::IsNearlyEqual(PreviousBandStopControl, CurrentBandStopControl));
 
 		if (bNeedsUpdate)
 		{
-			StateVariableFilter.SetQ(*Resonance);
-			StateVariableFilter.SetFrequency(*Frequency);
-			StateVariableFilter.SetPassBandGainCompensation(*BandStopControl);
+			StateVariableFilter.SetQ(CurrentResonance);
+			StateVariableFilter.SetFrequency(CurrentFrequency);
+			StateVariableFilter.SetBandStopControl(CurrentBandStopControl);
 
 			StateVariableFilter.Update();
+
+			PreviousFrequency = CurrentFrequency;
+			PreviousResonance = CurrentResonance;
+			PreviousBandStopControl = CurrentBandStopControl;
 		}
 
 		StateVariableFilter.ProcessAudio(
-			  AudioInput->GetData()
+			AudioInput->GetData()
 			, AudioInput->Num()
 			, LowPassOutput->GetData()
 			, HighPassOutput->GetData()
 			, BandPassOutput->GetData()
 			, BandStopOutput->GetData()
 		);
-
-		LastFrequency = *Frequency;
-		LastResonance = *Resonance;
-		LastBandStopControl = *BandStopControl;
 	}
 
 	METASOUND_REGISTER_NODE(FStateVariableFilterNode);
@@ -404,6 +410,8 @@ namespace Metasound
 
 	class FOnePoleLowPassFilterOperator : public TExecutableOperator<FOnePoleLowPassFilterOperator>
 	{
+	private:
+		static constexpr float InvalidValue = -1.f;
 	public:
 		static const FNodeClassMetadata& GetNodeInfo();
 
@@ -442,7 +450,7 @@ namespace Metasound
 
 		// data
 		const int32 BlockSize;
-		FSampleRate SampleRate;
+		float SampleRate;
 
 		// dsp
 		Audio::FInterpolatedLPF OnePoleLowPassFilter;
@@ -506,7 +514,7 @@ namespace Metasound
 		static const FVertexInterface Interface(
 			FInputVertexInterface(
 				TInputDataVertexModel<FAudioBuffer>(TEXT("In"), LOCTEXT("AudioInputTooltip", "Audio Input")),
-				TInputDataVertexModel<float>(TEXT("Freq."), LOCTEXT("FrequencyTooltip", "Controls filter cutoff"))
+				TInputDataVertexModel<float>(TEXT("Freq."), LOCTEXT("FrequencyTooltip", "Controls filter cutoff"), 20000.f)
 			),
 			FOutputVertexInterface(
 				TOutputDataVertexModel<FAudioBuffer>(TEXT("Out"), LOCTEXT("FilterOutputToolTip", "Audio Output"))
@@ -549,6 +557,8 @@ namespace Metasound
 
 	class FOnePoleHighPassFilterOperator : public TExecutableOperator<FOnePoleHighPassFilterOperator>
 	{
+	private:
+		static constexpr float InvalidValue = -1.f;
 	public:
 		static const FNodeClassMetadata& GetNodeInfo();
 
@@ -587,7 +597,7 @@ namespace Metasound
 
 		// data
 		const int32 BlockSize;
-		FSampleRate SampleRate;
+		float SampleRate;
 
 		// dsp
 		Audio::FInterpolatedHPF OnePoleHighPassFilter;
@@ -651,8 +661,8 @@ namespace Metasound
 		static const FVertexInterface Interface(
 			FInputVertexInterface(
 				TInputDataVertexModel<FAudioBuffer>(TEXT("In"), LOCTEXT("AudioInputTooltip", "Audio In")),
-				TInputDataVertexModel<float>(TEXT("Freq."), LOCTEXT("FrequencyTooltip", "Freq In"))
-				),
+				TInputDataVertexModel<float>(TEXT("Freq."), LOCTEXT("FrequencyTooltip", "Freq In"), 10.f)
+			),
 			FOutputVertexInterface(
 				TOutputDataVertexModel<FAudioBuffer>(TEXT("Out"), LOCTEXT("FilterOutputToolTip", "Audio Out"))
 			)
@@ -692,7 +702,7 @@ namespace Metasound
 
 #pragma region Biquad Filter
 	DECLARE_METASOUND_ENUM(Audio::EBiquadFilter::Type, Audio::EBiquadFilter::Lowpass,
-		METASOUNDSTANDARDNODES_API, FEnumEBiquadFilterType, FEnumBiQuadFilterTypeInfo, FEnumBiQuadFilterReadRef, FEnumBiQuadFilterWriteRef);
+	METASOUNDSTANDARDNODES_API, FEnumEBiquadFilterType, FEnumBiQuadFilterTypeInfo, FEnumBiQuadFilterReadRef, FEnumBiQuadFilterWriteRef);
 
 	DEFINE_METASOUND_ENUM_BEGIN(Audio::EBiquadFilter::Type, FEnumEBiquadFilterType, "BiquadFilterType")
 		DEFINE_METASOUND_ENUM_ENTRY(Audio::EBiquadFilter::Lowpass, LOCTEXT("LpDescription", "Low Pass"), LOCTEXT("LpDescriptionTT", "Low Pass Biquad Filter")),
@@ -704,11 +714,13 @@ namespace Metasound
 		DEFINE_METASOUND_ENUM_ENTRY(Audio::EBiquadFilter::HighShelf, LOCTEXT("HighShelfDescription", "High Shelf"), LOCTEXT("HighShelfDescriptionTT", "High Shelf Biquad Filter")),
 		DEFINE_METASOUND_ENUM_ENTRY(Audio::EBiquadFilter::AllPass, LOCTEXT("AllPassDescription", "All Pass"), LOCTEXT("AllPassDescriptionTT", "All Pass Biquad Filter")),
 		DEFINE_METASOUND_ENUM_ENTRY(Audio::EBiquadFilter::ButterworthLowPass, LOCTEXT("LowPassButterDescription", "Butterworth Low Pass"), LOCTEXT("LowPassButterDescriptionTT", "Butterworth Low Pass Biquad Filter")),
-		DEFINE_METASOUND_ENUM_ENTRY(Audio::EBiquadFilter::ButterworthHighPass,LOCTEXT("HighPassButterDescription", "Butterworth High Pass"), LOCTEXT("HighPassButterDescriptionTT", "Butterworth High Pass Biquad Filter"))
-	DEFINE_METASOUND_ENUM_END()
-	
-	class FBiquadFilterOperator : public TExecutableOperator<FBiquadFilterOperator>
+		DEFINE_METASOUND_ENUM_ENTRY(Audio::EBiquadFilter::ButterworthHighPass, LOCTEXT("HighPassButterDescription", "Butterworth High Pass"), LOCTEXT("HighPassButterDescriptionTT", "Butterworth High Pass Biquad Filter"))
+		DEFINE_METASOUND_ENUM_END()
+
+		class FBiquadFilterOperator : public TExecutableOperator<FBiquadFilterOperator>
 	{
+	private:
+		static constexpr float InvalidValue = -1.f;
 	public:
 		static const FNodeClassMetadata& GetNodeInfo();
 
@@ -722,6 +734,7 @@ namespace Metasound
 			InputDataReferences.AddDataReadReference(TEXT("In"), FAudioBufferReadRef(AudioInput));
 			InputDataReferences.AddDataReadReference(TEXT("Freq."), FFloatReadRef(Frequency));
 			InputDataReferences.AddDataReadReference(TEXT("Bandwidth"), FFloatReadRef(Bandwidth));
+			InputDataReferences.AddDataReadReference(TEXT("FilterGainDb"), FFloatReadRef(FilterGainDb));
 			InputDataReferences.AddDataReadReference(TEXT("Type"), FEnumBiQuadFilterReadRef(FilterType));
 
 			return InputDataReferences;
@@ -744,21 +757,24 @@ namespace Metasound
 		FAudioBufferReadRef AudioInput;
 		FFloatReadRef Frequency;
 		FFloatReadRef Bandwidth;
+		FFloatReadRef FilterGainDb;
 		FEnumBiQuadFilterReadRef FilterType;
 
 		// cached data
-		float LastFrequency;
-		float LastBandwidth;
+		float PreviousFrequency{ InvalidValue };
+		float PreviousBandwidth{ InvalidValue };
+		float PreviousFilterGainDb{ InvalidValue };
 
 		// output pins
 		FAudioBufferWriteRef AudioOutput;
 
 		// data
 		const int32 BlockSize;
-		FSampleRate SampleRate;
+		float SampleRate;
+		float MaxCutoffFrequency;
 
 		// dsp
-		Audio::EBiquadFilter::Type LastFilterType;
+		Audio::EBiquadFilter::Type PreviousFilterType;
 		Audio::FBiquadFilter BiquadFilter;
 
 
@@ -769,20 +785,22 @@ namespace Metasound
 			const FAudioBufferReadRef& InAudioInput,
 			const FFloatReadRef& InFrequency,
 			const FFloatReadRef& InBandwidth,
+			const FFloatReadRef& InFilterGainDb,
 			const FEnumBiQuadFilterReadRef& InFilterType
 		)
 			: AudioInput(InAudioInput)
 			, Frequency(InFrequency)
 			, Bandwidth(InBandwidth)
+			, FilterGainDb(InFilterGainDb)
 			, FilterType(InFilterType)
 			, AudioOutput(FAudioBufferWriteRef::CreateNew(InSettings))
 			, BlockSize(InSettings.GetNumFramesPerBlock())
 			, SampleRate(InSettings.GetSampleRate())
-			, LastFilterType(*FilterType)
+			, MaxCutoffFrequency(0.5f * SampleRate)
+			, PreviousFilterType(*FilterType)
 		{
 			// verify our buffer sizes:
 			check(AudioOutput->Num() == BlockSize);
-
 			BiquadFilter.Init(SampleRate, 1, *FilterType);
 		}
 	};
@@ -825,8 +843,9 @@ namespace Metasound
 		static const FVertexInterface Interface(
 			FInputVertexInterface(
 				TInputDataVertexModel<FAudioBuffer>(TEXT("In"), LOCTEXT("AudioInputTooltip", "Audio Input")),
-				TInputDataVertexModel<float>(TEXT("Freq."), LOCTEXT("FrequencyTooltip", "Controls Filter Cutoff")),
-				TInputDataVertexModel<float>(TEXT("Bandwidth"), LOCTEXT("BandwidthTooltip", "Bandwidth Control")),
+				TInputDataVertexModel<float>(TEXT("Freq."), LOCTEXT("FrequencyTooltip", "Controls Filter Cutoff"), 20000.f),
+				TInputDataVertexModel<float>(TEXT("Bandwidth"), LOCTEXT("BandwidthTooltip", "Bandwidth Control"), 1.f),
+				TInputDataVertexModel<float>(TEXT("FilterGainDb"), LOCTEXT("FilterGainDbTooltip", "Filter Gain Control (decibels)"), 0.f), 
 				TInputDataVertexModel<FEnumEBiquadFilterType>(TEXT("Type"), LOCTEXT("FilterTypeDescription", "Biquad Filter Type"))
 			),
 			FOutputVertexInterface(
@@ -847,39 +866,50 @@ namespace Metasound
 		FAudioBufferReadRef AudioIn = InputDataRefs.GetDataReadReferenceOrConstruct<FAudioBuffer>(TEXT("In"), InParams.OperatorSettings);
 		FFloatReadRef FrequencyIn = InputDataRefs.GetDataReadReferenceOrConstruct<float>(TEXT("Freq."));
 		FFloatReadRef BandwidthIn = InputDataRefs.GetDataReadReferenceOrConstruct<float>(TEXT("Bandwidth"));
-		FEnumBiQuadFilterReadRef FilterType = InputDataRefs.GetDataReadReferenceOrConstruct<FEnumEBiquadFilterType>(TEXT("Type"));
+		FFloatReadRef FilterGainDbIn = InputDataRefs.GetDataReadReferenceOrConstruct<float>(TEXT("FilterGainDb"));
+			FEnumBiQuadFilterReadRef FilterType = InputDataRefs.GetDataReadReferenceOrConstruct<FEnumEBiquadFilterType>(TEXT("Type"));
 
 		return MakeUnique<FBiquadFilterOperator>(
 			InParams.OperatorSettings
 			, AudioIn
 			, FrequencyIn
 			, BandwidthIn
+			, FilterGainDbIn
 			, FilterType
 			);
 	}
 
 	void FBiquadFilterOperator::Execute()
 	{
-		if (!FMath::IsNearlyEqual(LastFrequency, *Frequency))
+		const float CurrentFrequency = FMath::Clamp(*Frequency, 0.f, MaxCutoffFrequency);
+		const float CurrentBandwidth = FMath::Max(*Bandwidth, 0.f);
+		const float CurrentFilterGainDb = *FilterGainDb;
+
+		if (!FMath::IsNearlyEqual(PreviousFrequency, CurrentFrequency))
 		{
-			BiquadFilter.SetFrequency(*Frequency);
+			BiquadFilter.SetFrequency(CurrentFrequency);
+			PreviousFrequency = CurrentFrequency;
 		}
 
-		if (!FMath::IsNearlyEqual(LastBandwidth, *Bandwidth))
+		if (!FMath::IsNearlyEqual(PreviousBandwidth, CurrentBandwidth))
 		{
-			BiquadFilter.SetBandwidth(*Bandwidth);
+			BiquadFilter.SetBandwidth(CurrentBandwidth);
+			PreviousBandwidth = CurrentBandwidth;
 		}
 
-		if (*FilterType != LastFilterType)
+		if (!FMath::IsNearlyEqual(PreviousFilterGainDb, CurrentFilterGainDb))
+		{
+			BiquadFilter.SetGainDB(CurrentFilterGainDb);
+			PreviousFilterGainDb = CurrentFilterGainDb;
+		}
+
+		if (*FilterType != PreviousFilterType)
 		{
 			BiquadFilter.SetType(*FilterType);
-			LastFilterType = *FilterType;
+			PreviousFilterType = *FilterType;
 		}
 
 		BiquadFilter.ProcessAudio(AudioInput->GetData(), AudioInput->Num(), AudioOutput->GetData());
-
-		LastFrequency = *Frequency;
-		LastBandwidth = *Bandwidth;
 	}
 
 	METASOUND_REGISTER_NODE(FBiquadFilterNode);
