MetaSound array bug fixes
- Removing MultichannelAudioFormat registration because type is unused
- Implementing copy constructor and copy assignment operator on FWaveAsset to support array nodes
- Fixing index checks on array set/get/subset
- Renaming Array Slice to Array Subset

diff --git a/Engine/Plugins/Runtime/Metasound/Source/MetasoundEngine/Private/MetasoundWavePlayerNode.cpp b/Engine/Plugins/Runtime/Metasound/Source/MetasoundEngine/Private/MetasoundWavePlayerNode.cpp
--- a/Engine/Plugins/Runtime/Metasound/Source/MetasoundEngine/Private/MetasoundWavePlayerNode.cpp
+++ b/Engine/Plugins/Runtime/Metasound/Source/MetasoundEngine/Private/MetasoundWavePlayerNode.cpp
@@ -67,9 +67,9 @@ namespace Metasound
 			check(OutputSampleRate);
 			check(AudioBufferL->Num() == OutputBlockSizeInFrames && AudioBufferR->Num() == OutputBlockSizeInFrames);
 
-			if (Wave->SoundWaveProxy.IsValid())
+			if (Wave->IsSoundWaveValid())
 			{
-				CurrentSoundWaveName = Wave->SoundWaveProxy->GetFName();
+				CurrentSoundWaveName = (*Wave)->GetFName();
 			}
 		}
 
@@ -100,7 +100,7 @@ namespace Metasound
 			FName NewSoundWaveName = FName();
 			if (Wave->IsSoundWaveValid())
 			{
-				NewSoundWaveName = Wave->SoundWaveProxy->GetFName();
+				NewSoundWaveName = (*Wave)->GetFName();
 			}
 
 			if (NewSoundWaveName != CurrentSoundWaveName)
@@ -149,13 +149,13 @@ namespace Metasound
 				return false;
 			}
 
-			return Decoder.Initialize(Params, Wave->SoundWaveProxy);
+			return Decoder.Initialize(Params, Wave->GetSoundWaveProxy());
 		}
 
 
 		void ExecuteInternal(int32 StartFrame, int32 EndFrame)
 		{
-			bool bCanDecodeWave = Wave->IsSoundWaveValid() && (Wave->SoundWaveProxy->GetNumChannels() <= 2); // only support mono or stereo inputs
+			bool bCanDecodeWave = Wave->IsSoundWaveValid() && ((*Wave)->GetNumChannels() <= 2); // only support mono or stereo inputs
 
 			// note: output is hard-coded to stereo (dual-mono)
 			float* FinalOutputLeft = AudioBufferL->GetData() + StartFrame;
@@ -167,7 +167,7 @@ namespace Metasound
 				ensure(Decoder.CanGenerateAudio());
 
 
-				const int32 NumInputChannels = Wave->SoundWaveProxy->GetNumChannels();
+				const int32 NumInputChannels = (*Wave)->GetNumChannels();
 				const bool bNeedsUpmix = (NumInputChannels == 1);
 				const bool bNeedsDeinterleave = !bNeedsUpmix;
 
@@ -253,7 +253,7 @@ namespace Metasound
 			AddBuildError<FWavePlayerError>(OutErrors, WaveNode, LOCTEXT("NoSoundWave", "No Sound Wave"));
 
 		}
-		else if (Wave->SoundWaveProxy->GetNumChannels() > 2)
+		else if ((*Wave)->GetNumChannels() > 2)
 		{
 			AddBuildError<FWavePlayerError>(OutErrors, WaveNode, LOCTEXT("WavePlayerCurrentlyOnlySuportsMonoAssets", "Wave Player Currently Only Supports Mono Or Stereo Assets"));
 		}
diff --git a/Engine/Plugins/Runtime/Metasound/Source/MetasoundEngine/Public/MetasoundWave.h b/Engine/Plugins/Runtime/Metasound/Source/MetasoundEngine/Public/MetasoundWave.h
--- a/Engine/Plugins/Runtime/Metasound/Source/MetasoundEngine/Public/MetasoundWave.h
+++ b/Engine/Plugins/Runtime/Metasound/Source/MetasoundEngine/Public/MetasoundWave.h
@@ -20,24 +20,12 @@ namespace Metasound
 	// Metasound data type that holds onto a weak ptr. Mostly used as a placeholder until we have a proper proxy type.
 	class METASOUNDENGINE_API FWaveAsset
 	{
-	public:
-		// TODO: Implement using TSharedPtr<> to enable cheaper assignment when 
-		// choosing between multiple instances of FWaveAsset.
 		FSoundWaveProxyPtr SoundWaveProxy;
+	public:
 
 		FWaveAsset() = default;
-		FWaveAsset(const FWaveAsset&) = delete;
-
-		// Define move constructor and move assignment since
-		// object cannot be copied to due TUniquePtr member.
-		FWaveAsset(FWaveAsset&&) = default;
-		FWaveAsset& operator=(FWaveAsset&&) = default;
-
-		FWaveAsset& operator=(const FWaveAsset& Other)
-		{
-			SoundWaveProxy = Other.SoundWaveProxy;
-			return *this;
-		}
+		FWaveAsset(const FWaveAsset&) = default;
+		FWaveAsset& operator=(const FWaveAsset& Other) = default;
 
 		FWaveAsset(const TUniquePtr<Audio::IProxyData>& InInitData)
 		{
@@ -45,8 +33,11 @@ namespace Metasound
 
 			if (InInitData.IsValid())
 			{
-				// should we be getting handed a SharedPtr here?
-				SoundWaveProxy = MakeShared<FSoundWaveProxy, ESPMode::ThreadSafe>(InInitData->GetAs<FSoundWaveProxy>());
+				if (InInitData->CheckTypeCast<FSoundWaveProxy>())
+				{
+					// should we be getting handed a SharedPtr here?
+					SoundWaveProxy = MakeShared<FSoundWaveProxy, ESPMode::ThreadSafe>(InInitData->GetAs<FSoundWaveProxy>());
+				}
 			}
 		}
 
@@ -54,7 +45,21 @@ namespace Metasound
 		{
 			return SoundWaveProxy.IsValid();
 		}
-		
+
+		const FSoundWaveProxyPtr& GetSoundWaveProxy() const
+		{
+			return SoundWaveProxy;
+		}
+
+		const FSoundWaveProxy* operator->() const
+		{
+			return SoundWaveProxy.Get();
+		}
+
+		FSoundWaveProxy* operator->()
+		{
+			return SoundWaveProxy.Get();
+		}
 	};
 
 	DECLARE_METASOUND_DATA_REFERENCE_TYPES(FWaveAsset, METASOUNDENGINE_API, FWaveAssetTypeInfo, FWaveAssetReadRef, FWaveAssetWriteRef)
diff --git a/Engine/Plugins/Runtime/Metasound/Source/MetasoundFrontend/Public/MetasoundArrayNodes.h b/Engine/Plugins/Runtime/Metasound/Source/MetasoundFrontend/Public/MetasoundArrayNodes.h
--- a/Engine/Plugins/Runtime/Metasound/Source/MetasoundFrontend/Public/MetasoundArrayNodes.h
+++ b/Engine/Plugins/Runtime/Metasound/Source/MetasoundFrontend/Public/MetasoundArrayNodes.h
@@ -306,7 +306,7 @@ namespace Metasound
 				const int32 IndexValue = *Index;
 				const ArrayType& ArrayRef = *Array;
 
-				if (IndexValue < ArrayRef.Num())
+				if ((IndexValue >= 0) && (IndexValue < ArrayRef.Num()))
 				{
 					*Value = ArrayRef[IndexValue];
 				}
@@ -601,9 +601,9 @@ namespace Metasound
 		virtual ~TArrayConcatNode() = default;
 	};
 
-	/** TArraySliceOperator slices an array on trigger. */
+	/** TArraySubsetOperator slices an array on trigger. */
 	template<typename ArrayType>
-	class TArraySliceOperator : public TExecutableOperator<TArraySliceOperator<ArrayType>>
+	class TArraySubsetOperator : public TExecutableOperator<TArraySubsetOperator<ArrayType>>
 	{
 	public:
 		using FArrayDataReadReference = TDataReadReference<ArrayType>;
@@ -616,14 +616,14 @@ namespace Metasound
 
 			static const FVertexInterface DefaultInterface(
 				FInputVertexInterface(
-					TInputDataVertexModel<FTrigger>(GetInputTriggerName(), LOCTEXT("ArrayOpArraySliceTrigger", "Trigger to set value.")),
-					TInputDataVertexModel<ArrayType>(GetInputArrayName(), LOCTEXT("ArrayOpArraySliceInputLeft", "Input Array.")),
-					TInputDataVertexModel<int32>(GetInputStartIndexName(), LOCTEXT("ArrayOpArraySliceStartIndex", "First index to include.")),
-					TInputDataVertexModel<int32>(GetInputEndIndexName(), LOCTEXT("ArrayOpArraySliceEndIndex", "Last index to include."))
+					TInputDataVertexModel<FTrigger>(GetInputTriggerName(), LOCTEXT("ArrayOpArraySubsetTrigger", "Trigger to set value.")),
+					TInputDataVertexModel<ArrayType>(GetInputArrayName(), LOCTEXT("ArrayOpArraySubsetInputLeft", "Input Array.")),
+					TInputDataVertexModel<int32>(GetInputStartIndexName(), LOCTEXT("ArrayOpArraySubsetStartIndex", "First index to include.")),
+					TInputDataVertexModel<int32>(GetInputEndIndexName(), LOCTEXT("ArrayOpArraySubsetEndIndex", "Last index to include."))
 
 				),
 				FOutputVertexInterface(
-					TOutputDataVertexModel<ArrayType>(GetOutputArrayName(), LOCTEXT("ArrayOpArraySliceOutput", "Array after slicing."))
+					TOutputDataVertexModel<ArrayType>(GetOutputArrayName(), LOCTEXT("ArrayOpArraySubsetOutput", "Subset of input array."))
 				)
 			);
 
@@ -635,9 +635,9 @@ namespace Metasound
 			auto CreateNodeClassMetadata = []() -> FNodeClassMetadata
 			{
 				FName DataTypeName = GetMetasoundDataTypeName<ArrayType>();
-				FName OperatorName = TEXT("Slice"); 
-				FText NodeDisplayName = FText::Format(LOCTEXT("ArrayOpArraySliceDisplayNamePattern", "Slice ({0})"), FText::FromString(GetMetasoundDataTypeString<ArrayType>()));
-				FText NodeDescription = LOCTEXT("ArrayOpArraySliceDescription", "Slice array on trigger.");
+				FName OperatorName = TEXT("Subset"); 
+				FText NodeDisplayName = FText::Format(LOCTEXT("ArrayOpArraySubsetDisplayNamePattern", "Subset ({0})"), FText::FromString(GetMetasoundDataTypeString<ArrayType>()));
+				FText NodeDescription = LOCTEXT("ArrayOpArraySubsetDescription", "Subset array on trigger.");
 				FVertexInterface NodeInterface = GetDefaultInterface();
 			
 				return MetasoundArrayNodesPrivate::CreateArrayNodeClassMetadata(DataTypeName, OperatorName, NodeDisplayName, NodeDescription, NodeInterface);
@@ -664,11 +664,11 @@ namespace Metasound
 
 			FArrayDataWriteReference OutArray = TDataWriteReferenceFactory<ArrayType>::CreateExplicitArgs(InParams.OperatorSettings);
 
-			return MakeUnique<TArraySliceOperator>(Trigger, InArray, StartIndex, EndIndex, OutArray);
+			return MakeUnique<TArraySubsetOperator>(Trigger, InArray, StartIndex, EndIndex, OutArray);
 		}
 
 
-		TArraySliceOperator(TDataReadReference<FTrigger> InTrigger, FArrayDataReadReference InInputArray, TDataReadReference<int32> InStartIndex, TDataReadReference<int32> InEndIndex, FArrayDataWriteReference InOutputArray)
+		TArraySubsetOperator(TDataReadReference<FTrigger> InTrigger, FArrayDataReadReference InInputArray, TDataReadReference<int32> InStartIndex, TDataReadReference<int32> InEndIndex, FArrayDataWriteReference InOutputArray)
 		: Trigger(InTrigger)
 		, InputArray(InInputArray)
 		, StartIndex(InStartIndex)
@@ -677,7 +677,7 @@ namespace Metasound
 		{
 		}
 
-		virtual ~TArraySliceOperator() = default;
+		virtual ~TArraySubsetOperator() = default;
 
 		virtual FDataReferenceCollection GetInputs() const override
 		{
@@ -712,7 +712,7 @@ namespace Metasound
 
 				const ArrayType& InputArrayRef = *InputArray;
 				const int32 StartIndexValue = FMath::Max(0, *StartIndex);
-				const int32 EndIndexValue = FMath::Min(InputArrayRef.Num(), *EndIndex);
+				const int32 EndIndexValue = FMath::Min(InputArrayRef.Num(), *EndIndex + 1);
 
 				if (StartIndexValue < EndIndexValue)
 				{
@@ -732,15 +732,15 @@ namespace Metasound
 	};
 
 	template<typename ArrayType>
-	class TArraySliceNode : public FNodeFacade
+	class TArraySubsetNode : public FNodeFacade
 	{
 	public:
-		TArraySliceNode(const FNodeInitData& InInitData)
-		: FNodeFacade(InInitData.InstanceName, InInitData.InstanceID, TFacadeOperatorClass<TArraySliceOperator<ArrayType>>())
+		TArraySubsetNode(const FNodeInitData& InInitData)
+		: FNodeFacade(InInitData.InstanceName, InInitData.InstanceID, TFacadeOperatorClass<TArraySubsetOperator<ArrayType>>())
 		{
 		}
 
-		virtual ~TArraySliceNode() = default;
+		virtual ~TArraySubsetNode() = default;
 	};
 }
 
diff --git a/Engine/Plugins/Runtime/Metasound/Source/MetasoundFrontend/Public/MetasoundArrayNodesRegistration.h b/Engine/Plugins/Runtime/Metasound/Source/MetasoundFrontend/Public/MetasoundArrayNodesRegistration.h
--- a/Engine/Plugins/Runtime/Metasound/Source/MetasoundFrontend/Public/MetasoundArrayNodesRegistration.h
+++ b/Engine/Plugins/Runtime/Metasound/Source/MetasoundFrontend/Public/MetasoundArrayNodesRegistration.h
@@ -37,7 +37,7 @@ namespace Metasound
 			static constexpr bool bIsArrayConcatSupported = std::is_copy_constructible<ElementType>::value;
 
 			// Elements must be copy constructible
-			static constexpr bool bIsArraySliceSupported = std::is_copy_constructible<ElementType>::value;
+			static constexpr bool bIsArraySubsetSupported = std::is_copy_constructible<ElementType>::value;
 		};
 
 		template<typename ArrayType, typename std::enable_if<TArrayNodeSupport<ArrayType>::bIsArrayGetSupported, bool>::type = true>
@@ -71,18 +71,18 @@ namespace Metasound
 			return true;
 		}
 
-		template<typename ArrayType, typename std::enable_if<TArrayNodeSupport<ArrayType>::bIsArraySliceSupported, bool>::type = true>
-		bool RegisterArraySliceNode()
+		template<typename ArrayType, typename std::enable_if<TArrayNodeSupport<ArrayType>::bIsArraySubsetSupported, bool>::type = true>
+		bool RegisterArraySubsetNode()
 		{
-			using FNodeType = typename Metasound::TArraySliceNode<ArrayType>;
+			using FNodeType = typename Metasound::TArraySubsetNode<ArrayType>;
 
-			static_assert(TArrayNodeSupport<ArrayType>::bIsArraySliceSupported, "TArraySliceNode<> is not supported by array type");
+			static_assert(TArrayNodeSupport<ArrayType>::bIsArraySubsetSupported, "TArraySubsetNode<> is not supported by array type");
 
 			return RegisterNodeWithFrontend<FNodeType>();
 		}
 
-		template<typename ArrayType, typename std::enable_if<!TArrayNodeSupport<ArrayType>::bIsArraySliceSupported, bool>::type = true>
-		bool RegisterArraySliceNode()
+		template<typename ArrayType, typename std::enable_if<!TArrayNodeSupport<ArrayType>::bIsArraySubsetSupported, bool>::type = true>
+		bool RegisterArraySubsetNode()
 		{
 			// No op if not supported
 			return true;
@@ -126,7 +126,7 @@ namespace Metasound
 		bool bSuccess = RegisterArrayNumNode<ArrayType>();
 		bSuccess = bSuccess && RegisterArrayGetNode<ArrayType>();
 		bSuccess = bSuccess && RegisterArraySetNode<ArrayType>();
-		bSuccess = bSuccess && RegisterArraySliceNode<ArrayType>();
+		bSuccess = bSuccess && RegisterArraySubsetNode<ArrayType>();
 		bSuccess = bSuccess && RegisterArrayConcatNode<ArrayType>();
 
 		return bSuccess;
diff --git a/Engine/Plugins/Runtime/Metasound/Source/MetasoundStandardNodes/Private/MetasoundAudioFormats.cpp b/Engine/Plugins/Runtime/Metasound/Source/MetasoundStandardNodes/Private/MetasoundAudioFormats.cpp
--- a/Engine/Plugins/Runtime/Metasound/Source/MetasoundStandardNodes/Private/MetasoundAudioFormats.cpp
+++ b/Engine/Plugins/Runtime/Metasound/Source/MetasoundStandardNodes/Private/MetasoundAudioFormats.cpp
@@ -642,6 +642,6 @@ namespace Metasound
 // so that the registration macro has access to the specializations.
 REGISTER_METASOUND_DATATYPE(Metasound::FMonoAudioFormat, "Audio:Mono");
 REGISTER_METASOUND_DATATYPE(Metasound::FStereoAudioFormat, "Audio:Stereo");
-REGISTER_METASOUND_DATATYPE(Metasound::FMultichannelAudioFormat, "Audio:Multichannel", ELiteralType::Integer);
+//REGISTER_METASOUND_DATATYPE(Metasound::FMultichannelAudioFormat, "Audio:Multichannel", ELiteralType::Integer);
 
 #undef LOCTEXT_NAMESPACE // MetasoundStandardNodes
diff --git a/Engine/Plugins/Runtime/Metasound/Source/MetasoundStandardNodes/Public/MetasoundAudioFormats.h b/Engine/Plugins/Runtime/Metasound/Source/MetasoundStandardNodes/Public/MetasoundAudioFormats.h
--- a/Engine/Plugins/Runtime/Metasound/Source/MetasoundStandardNodes/Public/MetasoundAudioFormats.h
+++ b/Engine/Plugins/Runtime/Metasound/Source/MetasoundStandardNodes/Public/MetasoundAudioFormats.h
@@ -325,7 +325,8 @@ namespace Metasound
 	};
 
 
-	DECLARE_METASOUND_DATA_REFERENCE_TYPES(FMultichannelAudioFormat, METASOUNDSTANDARDNODES_API, FMultichannelAudioFormatTypeInfo, FMultichannelAudioFormatReadRef, FMultichannelAudioFormatWriteRef);
+	// TODO: currently unused. Commenting out to keep clean UX.
+	//DECLARE_METASOUND_DATA_REFERENCE_TYPES(FMultichannelAudioFormat, METASOUNDSTANDARDNODES_API, FMultichannelAudioFormatTypeInfo, FMultichannelAudioFormatReadRef, FMultichannelAudioFormatWriteRef);
 
 	DECLARE_METASOUND_DATA_REFERENCE_TYPES(FMonoAudioFormat, METASOUNDSTANDARDNODES_API, FMonoAudioFormatTypeInfo, FMonoAudioFormatReadRef, FMonoAudioFormatWriteRef);
 
