<description: restricted, no permission to view>

diff --git a/Engine/Plugins/Runtime/Metasound/Source/MetasoundFrontend/Public/MetasoundRouter.h b/Engine/Plugins/Runtime/Metasound/Source/MetasoundFrontend/Public/MetasoundRouter.h
--- a/Engine/Plugins/Runtime/Metasound/Source/MetasoundFrontend/Public/MetasoundRouter.h
+++ b/Engine/Plugins/Runtime/Metasound/Source/MetasoundFrontend/Public/MetasoundRouter.h
@@ -436,7 +436,7 @@ namespace Metasound
 		FSenderInitParams Params;
 
 		// The amount of pop operations before elements before this sender pushed appear.
-		uint32 NumElementsToDelayBy;
+		uint32 NumElementsToDelayBy = 0;
 	};
 
 	template<typename TDataType>
diff --git a/Engine/Plugins/Runtime/Metasound/Source/MetasoundGraphCore/Private/MetasoundOperatorBuilder.cpp b/Engine/Plugins/Runtime/Metasound/Source/MetasoundGraphCore/Private/MetasoundOperatorBuilder.cpp
--- a/Engine/Plugins/Runtime/Metasound/Source/MetasoundGraphCore/Private/MetasoundOperatorBuilder.cpp
+++ b/Engine/Plugins/Runtime/Metasound/Source/MetasoundGraphCore/Private/MetasoundOperatorBuilder.cpp
@@ -69,7 +69,7 @@ namespace Metasound
 	{
 		FOperatorBuilderSettings Settings;
 
-		Settings.PruningMode = EOperatorBuilderNodePruning::PruneNodesWithoutOutputDependency;
+		Settings.PruningMode = EOperatorBuilderNodePruning::PruneNodesWithoutExternalDependency;
 		Settings.bValidateNoCyclesInGraph = true;
 		Settings.bValidateNoDuplicateInputs = true;
 		Settings.bValidateVerticesExist = true;
@@ -84,7 +84,7 @@ namespace Metasound
 	{
 		FOperatorBuilderSettings Settings;
 
-		Settings.PruningMode = EOperatorBuilderNodePruning::PruneNodesWithoutOutputDependency;
+		Settings.PruningMode = EOperatorBuilderNodePruning::PruneNodesWithoutExternalDependency;
 		Settings.bValidateNoCyclesInGraph = false;
 		Settings.bValidateNoDuplicateInputs = false;
 		Settings.bValidateVerticesExist = false;
@@ -99,7 +99,7 @@ namespace Metasound
 	{
 		FOperatorBuilderSettings Settings;
 
-		Settings.PruningMode = EOperatorBuilderNodePruning::PruneNodesWithoutOutputDependency;
+		Settings.PruningMode = EOperatorBuilderNodePruning::PruneNodesWithoutExternalDependency;
 		Settings.bValidateNoCyclesInGraph = false;
 		Settings.bValidateNoDuplicateInputs = false;
 		Settings.bValidateVerticesExist = false;
diff --git a/Engine/Plugins/Runtime/Metasound/Source/MetasoundStandardNodes/Public/MetasoundMathNodes.h b/Engine/Plugins/Runtime/Metasound/Source/MetasoundStandardNodes/Public/MetasoundMathNodes.h
--- a/Engine/Plugins/Runtime/Metasound/Source/MetasoundStandardNodes/Public/MetasoundMathNodes.h
+++ b/Engine/Plugins/Runtime/Metasound/Source/MetasoundStandardNodes/Public/MetasoundMathNodes.h
@@ -777,7 +777,7 @@ namespace Metasound
 				return;
 			}
 
-			FMemory::Memcpy(OutResult->GetData(), InAdditionalOperands[0]->GetData(), sizeof(float) * NumSamples);
+			FMemory::Memcpy(OutResult->GetData(), InPrimaryOperand->GetData(), sizeof(float) * NumSamples);
 
 			for (int32 i = 0; i < InAdditionalOperands.Num(); ++i)
 			{
