Metasound live parameter updates

diff --git a/Engine/Plugins/Runtime/Metasound/Source/MetasoundEditor/Private/MetasoundEditor.cpp b/Engine/Plugins/Runtime/Metasound/Source/MetasoundEditor/Private/MetasoundEditor.cpp
--- a/Engine/Plugins/Runtime/Metasound/Source/MetasoundEditor/Private/MetasoundEditor.cpp
+++ b/Engine/Plugins/Runtime/Metasound/Source/MetasoundEditor/Private/MetasoundEditor.cpp
@@ -2,6 +2,7 @@
 
 #include "MetasoundEditor.h"
 
+#include "AudioDevice.h"
 #include "Components/AudioComponent.h"
 #include "EdGraph/EdGraphNode.h"
 #include "EdGraphUtilities.h"
@@ -525,6 +526,7 @@ namespace Metasound
 
 			if (USoundBase* MetasoundToPlay = Cast<USoundBase>(Metasound))
 			{
+
 				GEditor->PlayPreviewSound(MetasoundToPlay);
 
 				// Set the send to the audio bus that is used for analyzing the metasound output
@@ -534,8 +536,28 @@ namespace Metasound
 					{
 						PreviewComp->SetAudioBusSendPostEffect(MetasoundAudioBus.Get(), 1.0f);
 					}
+
+					TUniquePtr<IAudioInstanceTransmitter> MetasoundComm;
+
+					// Communicating with playing component requires audio component ID and sample rate. 
+					if (FAudioDevice* PreviewAudioDevice = PreviewComp->GetAudioDevice())
+					{
+						if(FMetasoundAssetBase* MetasoundAsset = IMetasoundUObjectRegistry::Get().GetObjectAsAssetBase(Metasound))
+						{
+							const uint64 AudioComponentID = PreviewComp->GetAudioComponentID();
+							float SampleRate = PreviewAudioDevice->GetSampleRate();
+
+							MetasoundComm = MetasoundToPlay->CreateInstanceTransmitter(FAudioInstanceTransmitterInitParams{AudioComponentID, SampleRate});
+						}
+					}
+
+					if (UMetasoundEditorGraph* Graph = CastChecked<UMetasoundEditorGraph>(MetasoundGraphEditor->GetCurrentGraph()))
+					{
+						Graph->SetMetasoundInstanceTransmitter(MoveTemp(MetasoundComm));
+					}
 				}
 
+
 				MetasoundGraphEditor->RegisterActiveTimer(0.0f,
 					FWidgetActiveTimerDelegate::CreateLambda([](double InCurrentTime, float InDeltaTime)
 						{
diff --git a/Engine/Plugins/Runtime/Metasound/Source/MetasoundEditor/Private/MetasoundEditorGraph.cpp b/Engine/Plugins/Runtime/Metasound/Source/MetasoundEditor/Private/MetasoundEditorGraph.cpp
--- a/Engine/Plugins/Runtime/Metasound/Source/MetasoundEditor/Private/MetasoundEditorGraph.cpp
+++ b/Engine/Plugins/Runtime/Metasound/Source/MetasoundEditor/Private/MetasoundEditorGraph.cpp
@@ -48,10 +48,26 @@ UObject& UMetasoundEditorGraph::GetMetasoundChecked() const
 	return *ParentMetasound;
 }
 
+void UMetasoundEditorGraph::SetMetasoundInstanceTransmitter(TUniquePtr<IAudioInstanceTransmitter>&& InTransmitter)
+{
+	Transmitter = MoveTemp(InTransmitter);
+}
+
+IAudioInstanceTransmitter* UMetasoundEditorGraph::GetMetasoundInstanceTransmitter()
+{
+	return Transmitter.Get();
+}
+
+const IAudioInstanceTransmitter* UMetasoundEditorGraph::GetMetasoundInstanceTransmitter() const
+{
+	return Transmitter.Get();
+}
+
 void UMetasoundEditorGraph::Synchronize()
 {
 	if (ParentMetasound)
 	{
 		Metasound::Editor::FGraphBuilder::SynchronizeGraph(*ParentMetasound);
 	}
-}
\ No newline at end of file
+}
+
diff --git a/Engine/Plugins/Runtime/Metasound/Source/MetasoundEditor/Private/MetasoundEditorGraphBuilder.cpp b/Engine/Plugins/Runtime/Metasound/Source/MetasoundEditor/Private/MetasoundEditorGraphBuilder.cpp
--- a/Engine/Plugins/Runtime/Metasound/Source/MetasoundEditor/Private/MetasoundEditorGraphBuilder.cpp
+++ b/Engine/Plugins/Runtime/Metasound/Source/MetasoundEditor/Private/MetasoundEditorGraphBuilder.cpp
@@ -891,20 +891,31 @@ namespace Metasound
 				InGraphNode.RemovePin(InGraphNode.Pins[i]);
 			}
 
-			TArray<FInputHandle> InputHandles = NodeHandle->GetInputs();
-			InputHandles = NodeHandle->GetInputStyle().SortDefaults(InputHandles);
-			for (const FInputHandle& InputHandle : InputHandles)
+			// Only add input pins of the node is not an input node. Input nodes
+			// have their input pins hidden because they cannot be connected internal
+			// to the graph.
+			if (EMetasoundFrontendClassType::Input != NodeHandle->GetClassType())
 			{
-				AddPinToNode(InGraphNode, InputHandle);
+				TArray<FInputHandle> InputHandles = NodeHandle->GetInputs();
+				InputHandles = NodeHandle->GetInputStyle().SortDefaults(InputHandles);
+				for (const FInputHandle& InputHandle : InputHandles)
+				{
+					AddPinToNode(InGraphNode, InputHandle);
+				}
 			}
 
-			TArray<FOutputHandle> OutputHandles = NodeHandle->GetOutputs();
-			OutputHandles = NodeHandle->GetOutputStyle().SortDefaults(OutputHandles);
-			for (const FOutputHandle& OutputHandle : OutputHandles)
+			// Only add output pins of the node is not an output node. Output nodes
+			// have their output pins hidden because they cannot be connected internal
+			// to the graph.
+			if (EMetasoundFrontendClassType::Output != NodeHandle->GetClassType())
 			{
-				AddPinToNode(InGraphNode, OutputHandle);
+				TArray<FOutputHandle> OutputHandles = NodeHandle->GetOutputs();
+				OutputHandles = NodeHandle->GetOutputStyle().SortDefaults(OutputHandles);
+				for (const FOutputHandle& OutputHandle : OutputHandles)
+				{
+					AddPinToNode(InGraphNode, OutputHandle);
+				}
 			}
-
 			InGraphNode.MarkPackageDirty();
 		}
 
@@ -1178,17 +1189,28 @@ namespace Metasound
 				}
 			}
 
-			// Add unmatched input pins
-			InputHandles = InNode->GetOutputStyle().SortDefaults(InputHandles);
-			for (Frontend::FInputHandle& InputHandle : InputHandles)
+			// Only add input pins of the node is not an input node. Input nodes
+			// have their input pins hidden because they cannot be connected internal
+			// to the graph.
+			if (EMetasoundFrontendClassType::Input != InNode->GetClassType())
 			{
-				AddPinToNode(InEditorNode, InputHandle);
+				InputHandles = InNode->GetOutputStyle().SortDefaults(InputHandles);
+				for (Frontend::FInputHandle& InputHandle : InputHandles)
+				{
+					AddPinToNode(InEditorNode, InputHandle);
+				}
 			}
 
-			OutputHandles = InNode->GetOutputStyle().SortDefaults(OutputHandles);
-			for (Frontend::FOutputHandle& OutputHandle : OutputHandles)
+			// Only add output pins of the node is not an output node. Output nodes
+			// have their output pins hidden because they cannot be connected internal
+			// to the graph.
+			if (EMetasoundFrontendClassType::Output != InNode->GetClassType())
 			{
-				AddPinToNode(InEditorNode, OutputHandle);
+				OutputHandles = InNode->GetOutputStyle().SortDefaults(OutputHandles);
+				for (Frontend::FOutputHandle& OutputHandle : OutputHandles)
+				{
+					AddPinToNode(InEditorNode, OutputHandle);
+				}
 			}
 			
 			if (bIsNodeDirty)
@@ -1231,6 +1253,13 @@ namespace Metasound
 
 				FNodeHandle Node = MetasoundEditorNode->GetNodeHandle();
 
+				if (EMetasoundFrontendClassType::Input == Node->GetClassType())
+				{
+					// Skip this node if it is an input node. Input pins on input 
+					// nodes are not connected internal to the graph.
+					continue;
+				}
+
 				TArray<UEdGraphPin*> Pins = MetasoundEditorNode->GetAllPins();
 				TArray<FInputHandle> NodeInputs = Node->GetInputs();
 
diff --git a/Engine/Plugins/Runtime/Metasound/Source/MetasoundEditor/Private/MetasoundInputNodeDetailCustomization.cpp b/Engine/Plugins/Runtime/Metasound/Source/MetasoundEditor/Private/MetasoundInputNodeDetailCustomization.cpp
--- a/Engine/Plugins/Runtime/Metasound/Source/MetasoundEditor/Private/MetasoundInputNodeDetailCustomization.cpp
+++ b/Engine/Plugins/Runtime/Metasound/Source/MetasoundEditor/Private/MetasoundInputNodeDetailCustomization.cpp
@@ -356,7 +356,7 @@ namespace Metasound
 				{
 					if (InputNode.IsValid())
 					{
-						InputNode->UpdateDocumentInput();
+						InputNode->OnLiteralChanged();
 					}
 				}
 			});
@@ -484,7 +484,7 @@ namespace Metasound
 					{
 						if (InGraphNode.IsValid())
 						{
-							InGraphNode->UpdateDocumentInput();
+							InGraphNode->OnLiteralChanged();
 						}
 					});
 
diff --git a/Engine/Plugins/Runtime/Metasound/Source/MetasoundEditor/Public/MetasoundEditorGraph.h b/Engine/Plugins/Runtime/Metasound/Source/MetasoundEditor/Public/MetasoundEditorGraph.h
--- a/Engine/Plugins/Runtime/Metasound/Source/MetasoundEditor/Public/MetasoundEditorGraph.h
+++ b/Engine/Plugins/Runtime/Metasound/Source/MetasoundEditor/Public/MetasoundEditorGraph.h
@@ -25,11 +25,23 @@ public:
 	UObject* GetMetasound() const;
 	UObject& GetMetasoundChecked() const;
 
+
 	virtual void Synchronize();
 
+	// Sets the transmitter interface for the current metasound preview.
+	void SetMetasoundInstanceTransmitter(TUniquePtr<IAudioInstanceTransmitter>&& InTransmitter);
+
+	// Gets the transmitter interface for the current metasound preview.
+	IAudioInstanceTransmitter* GetMetasoundInstanceTransmitter();
+
+	// Gets the transmitter interface for the current metasound preview.
+	const IAudioInstanceTransmitter* GetMetasoundInstanceTransmitter() const;
+
 private:
 	UPROPERTY()
 	UObject* ParentMetasound;
 
+	TUniquePtr<IAudioInstanceTransmitter> Transmitter;
+
 	friend class Metasound::Editor::FEditor;
 };
diff --git a/Engine/Plugins/Runtime/Metasound/Source/MetasoundEditor/Public/MetasoundEditorGraphInputNodes.h b/Engine/Plugins/Runtime/Metasound/Source/MetasoundEditor/Public/MetasoundEditorGraphInputNodes.h
--- a/Engine/Plugins/Runtime/Metasound/Source/MetasoundEditor/Public/MetasoundEditorGraphInputNodes.h
+++ b/Engine/Plugins/Runtime/Metasound/Source/MetasoundEditor/Public/MetasoundEditorGraphInputNodes.h
@@ -1,9 +1,10 @@
 // Copyright Epic Games, Inc. All Rights Reserved.
 #pragma once
 
+#include "Algo/Transform.h"
 #include "CoreMinimal.h"
 #include "EdGraph/EdGraphNode.h"
-#include "UObject/ObjectMacros.h"
+#include "MetasoundEditorGraph.h"
 #include "MetasoundEditorGraphNode.h"
 #include "MetasoundFrontend.h"
 #include "MetasoundFrontendController.h"
@@ -12,6 +13,7 @@
 #include "MetasoundFrontendRegistries.h"
 #include "Misc/Guid.h"
 #include "Sound/SoundWave.h"
+#include "UObject/ObjectMacros.h"
 #include "UObject/SoftObjectPath.h"
 
 #include "MetasoundEditorGraphInputNodes.generated.h"
@@ -42,6 +44,7 @@ public:
 		return FMetasoundFrontendLiteral();
 	}
 
+
 	void UpdateDocumentInput() const
 	{
 		using namespace Metasound::Frontend;
@@ -54,12 +57,36 @@ public:
 		GraphHandle->SetDefaultInput(VertexID, GetLiteralDefault());
 	}
 
+	virtual void UpdatePreviewInstance(const Metasound::FVertexKey& InParamterName, IAudioInstanceTransmitter& InInstanceTransmitter) const {}
+
+	void OnLiteralChanged() 
+	{
+		UpdateDocumentInput();
+
+		if (UMetasoundEditorGraph* MetasoundGraph = Cast<UMetasoundEditorGraph>(GetGraph()))
+		{
+			if (IAudioInstanceTransmitter* Transmitter = MetasoundGraph->GetMetasoundInstanceTransmitter())
+			{
+				// TODO: fix how this parameter name is determined. It should not be done with a "DisplayName" but rather 
+				// the vertex "Name". Currently user entered vertex names only have their "Names" stored as "DisplayNames" 
+				TArray<Metasound::Frontend::FConstInputHandle> Inputs = GetNodeHandle()->GetConstInputs();
+
+				// An input node should only have one input. 
+				if (ensure(Inputs.Num() == 1))
+				{
+					Metasound::FVertexKey VertexKey = Metasound::FVertexKey(Inputs[0]->GetDisplayName().ToString());
+					UpdatePreviewInstance(VertexKey, *Transmitter);
+				}
+			}
+		}
+	}
+
 #if WITH_EDITORONLY_DATA
 	virtual void PostEditUndo() override
 	{
 		Super::PostEditUndo();
 
-		UpdateDocumentInput();
+		OnLiteralChanged();
 	}
 #endif // WITH_EDITORONLY_DATA
 
@@ -102,6 +129,11 @@ public:
 	{
 		return EMetasoundFrontendLiteralType::Boolean;
 	}
+
+	void UpdatePreviewInstance(const Metasound::FVertexKey& InParameterName, IAudioInstanceTransmitter& InInstanceTransmitter) const override
+	{
+		InInstanceTransmitter.SetBoolParameter(*InParameterName, Default);
+	}
 };
 
 UCLASS(MinimalAPI)
@@ -126,6 +158,12 @@ public:
 	{
 		return EMetasoundFrontendLiteralType::BooleanArray;
 	}
+
+	void UpdatePreviewInstance(const Metasound::FVertexKey& InParameterName, IAudioInstanceTransmitter& InInstanceTransmitter) const override
+	{
+		// TODO:
+		//InInstanceTransmitter.SetBoolArrayParameter(InParameterName, Default);
+	}
 };
 
 // Broken out to be able to customize and swap enum behavior for basic integer literal behavior
@@ -160,6 +198,11 @@ public:
 	{
 		return EMetasoundFrontendLiteralType::Integer;
 	}
+
+	void UpdatePreviewInstance(const Metasound::FVertexKey& InParameterName, IAudioInstanceTransmitter& InInstanceTransmitter) const override
+	{
+		InInstanceTransmitter.SetIntParameter(*InParameterName, Default.Value);
+	}
 };
 
 UCLASS(MinimalAPI)
@@ -190,6 +233,16 @@ public:
 	{
 		return EMetasoundFrontendLiteralType::IntegerArray;
 	}
+
+	void UpdatePreviewInstance(const Metasound::FVertexKey& InParameterName, IAudioInstanceTransmitter& InInstanceTransmitter) const override
+	{
+
+		TArray<int32> IntArray;
+		Algo::Transform(Default, IntArray, [](const FMetasoundEditorGraphInputInt& InValue) { return InValue.Value; });
+
+		// TODO:
+		//InInstanceTransmitter.SetIntArrayParameter(*InParameterName, IntArray);
+	}
 };
 
 UCLASS(MinimalAPI)
@@ -214,6 +267,11 @@ public:
 	{
 		return EMetasoundFrontendLiteralType::Float;
 	}
+
+	void UpdatePreviewInstance(const Metasound::FVertexKey& InParameterName, IAudioInstanceTransmitter& InInstanceTransmitter) const override
+	{
+		InInstanceTransmitter.SetFloatParameter(*InParameterName, Default);
+	}
 };
 
 UCLASS(MinimalAPI)
@@ -238,6 +296,12 @@ public:
 	{
 		return EMetasoundFrontendLiteralType::FloatArray;
 	}
+
+	void UpdatePreviewInstance(const Metasound::FVertexKey& InParameterName, IAudioInstanceTransmitter& InInstanceTransmitter) const override
+	{
+		// TODO:
+		//InInstanceTransmitter.SetFloatArrayParameter(*InParameterName, Default);
+	}
 };
 
 UCLASS(MinimalAPI)
@@ -262,6 +326,12 @@ public:
 	{
 		return EMetasoundFrontendLiteralType::String;
 	}
+
+	void UpdatePreviewInstance(const Metasound::FVertexKey& InParamterName, IAudioInstanceTransmitter& InInstanceTransmitter) const override
+	{
+		// TODO:
+		//InInstanceTransmitter.SetStringParameter(*InParameterName, Default);
+	}
 };
 
 UCLASS(MinimalAPI)
@@ -286,6 +356,12 @@ public:
 	{
 		return EMetasoundFrontendLiteralType::StringArray;
 	}
+
+	void UpdatePreviewInstance(const Metasound::FVertexKey& InParameterName, IAudioInstanceTransmitter& InInstanceTransmitter) const override
+	{
+		// TODO:
+		//InInstanceTransmitter.SetStringArrayParameter(*InParameterName, Default);
+	}
 };
 
 // Broken out to be able to customize and swap AllowedClass based on provided object proxy
@@ -320,6 +396,12 @@ public:
 	{
 		return EMetasoundFrontendLiteralType::UObject;
 	}
+
+	void UpdatePreviewInstance(const Metasound::FVertexKey& InParameterName, IAudioInstanceTransmitter& InInstanceTransmitter) const override
+	{
+		// TODO:
+		//InInstanceTransmitter.SetUObjectParameter(*InParameterName, Default.Object);
+	}
 };
 
 UCLASS(MinimalAPI)
@@ -350,4 +432,13 @@ public:
 	{
 		return EMetasoundFrontendLiteralType::UObjectArray;
 	}
-};
\ No newline at end of file
+
+	void UpdatePreviewInstance(const Metasound::FVertexKey& InParameterName, IAudioInstanceTransmitter& InInstanceTransmitter) const override
+	{
+		TArray<UObject*> ObjectArray;
+		Algo::Transform(Default, ObjectArray, [](const FMetasoundEditorGraphInputObject& InValue) { return InValue.Object; });
+
+		// TODO:
+		//InInstanceTransmitter.SetUObjectArrayParameter(*InParameterName, Default.Object);
+	}
+};
diff --git a/Engine/Plugins/Runtime/Metasound/Source/MetasoundEngine/Private/MetasoundSource.cpp b/Engine/Plugins/Runtime/Metasound/Source/MetasoundEngine/Private/MetasoundSource.cpp
--- a/Engine/Plugins/Runtime/Metasound/Source/MetasoundEngine/Private/MetasoundSource.cpp
+++ b/Engine/Plugins/Runtime/Metasound/Source/MetasoundEngine/Private/MetasoundSource.cpp
@@ -414,18 +414,15 @@ TArray<UMetasoundSource::FSendInfoAndVertexName> UMetasoundSource::GetSendInfos(
 	for (const FString& VertexName : SendVertices)
 	{
 		FConstNodeHandle InputNode = RootGraph->GetInputNodeWithName(VertexName);
-		// TODO: re-expose InputNode Inputs. Currently they are hidden since they cannot be 
-		// connected. But instead, they should be accessible through this API, but fail
-		// to connect. As a quick fix, the output handles are used, but this is not
-		// sustainable for situations where the input node is specialized. 
-		for (FConstOutputHandle InputHandle : InputNode->GetConstOutputs())
+		for (FConstInputHandle InputHandle : InputNode->GetConstInputs())
 		{
 			FSendInfoAndVertexName Info;
 
 			// TODO: incorporate VertexID into address. But need to ensure that VertexID
 			// will be maintained after injecting Receive nodes. 
 			Info.SendInfo.Address = CreateSendAddress(InInstanceID, InputHandle->GetName(), InputHandle->GetDataType());
-			Info.SendInfo.ParameterName = FName(InputHandle->GetDisplayName().ToString()); // TODO: display name hack. Need to have nameing consistent in editor for inputs
+			Info.SendInfo.ParameterName = FName(InputHandle->GetDisplayName().ToString()); // TODO: display name hack. Need to have naming consistent in editor for inputs
+			//Info.SendInfo.ParameterName = FName(*InputHandle->GetName()); // TODO: this is the expected parameter name.
 			Info.SendInfo.TypeName = InputHandle->GetDataType();
 			Info.VertexName = VertexName;
 			
diff --git a/Engine/Plugins/Runtime/Metasound/Source/MetasoundFrontend/Private/MetasoundFrontendStandardController.cpp b/Engine/Plugins/Runtime/Metasound/Source/MetasoundFrontend/Private/MetasoundFrontendStandardController.cpp
--- a/Engine/Plugins/Runtime/Metasound/Source/MetasoundFrontend/Private/MetasoundFrontendStandardController.cpp
+++ b/Engine/Plugins/Runtime/Metasound/Source/MetasoundFrontend/Private/MetasoundFrontendStandardController.cpp
@@ -599,6 +599,59 @@ namespace Metasound
 			return Access;
 		}
 
+		//
+		// FOutputNodeOutputController
+		// 
+		FOutputNodeOutputController::FOutputNodeOutputController(const FOutputNodeOutputController::FInitParams& InParams)
+		: FBaseOutputController({InParams.ID, InParams.NodeVertexPtr, InParams.ClassOutputPtr, InParams.GraphPtr, InParams.OwningNode})
+		, OwningGraphClassOutputPtr(InParams.OwningGraphClassOutputPtr)
+		{
+		}
+
+		bool FOutputNodeOutputController::IsValid() const
+		{
+			return FBaseOutputController::IsValid() && OwningGraphClassOutputPtr.IsValid();
+		}
+
+		const FText& FOutputNodeOutputController::GetDisplayName() const
+		{
+			if (const FMetasoundFrontendClassOutput* OwningOutput = OwningGraphClassOutputPtr.Get())
+			{
+				return OwningOutput->Metadata.DisplayName;
+			}
+
+			return FrontendControllerIntrinsics::GetInvalidValueConstRef<FText>();
+		}
+
+		const FText& FOutputNodeOutputController::GetTooltip() const
+		{
+			if (OwningGraphClassOutputPtr.IsValid())
+			{
+				return OwningGraphClassOutputPtr->Metadata.Description;
+			}
+
+			return FrontendControllerIntrinsics::GetInvalidValueConstRef<FText>();
+		}
+
+		FConnectability FOutputNodeOutputController::CanConnectTo(const IInputController& InController) const 
+		{
+			// Cannot connect to a graph's output.
+			static const FConnectability Connectability = {FConnectability::EConnectable::No};
+
+			return Connectability;
+		}
+
+		bool FOutputNodeOutputController::Connect(IInputController& InController) 
+		{
+			return false;
+		}
+
+		bool FOutputNodeOutputController::ConnectWithConverterNode(IInputController& InController, const FConverterNodeInfo& InNodeClassName)
+		{
+			return false;
+		}
+
+
 		//
 		// FBaseInputController
 		// 
@@ -979,6 +1032,55 @@ namespace Metasound
 			return Access;
 		}
 
+		//
+		// FInputNodeInputController
+		//
+		FInputNodeInputController::FInputNodeInputController(const FInputNodeInputController::FInitParams& InParams)
+		: FBaseInputController({InParams.ID, InParams.NodeVertexPtr, InParams.ClassInputPtr, InParams.GraphPtr, InParams.OwningNode})
+		, OwningGraphClassInputPtr(InParams.OwningGraphClassInputPtr)
+		{
+		}
+
+		bool FInputNodeInputController::IsValid() const
+		{
+			return FBaseInputController::IsValid() && OwningGraphClassInputPtr.IsValid();
+		}
+
+		const FText& FInputNodeInputController::GetDisplayName() const
+		{
+			if (const FMetasoundFrontendClassInput* OwningInput = OwningGraphClassInputPtr.Get())
+			{
+				return OwningInput->Metadata.DisplayName;
+			}
+
+			return FrontendControllerIntrinsics::GetInvalidValueConstRef<FText>();
+		}
+
+		const FText& FInputNodeInputController::GetTooltip() const
+		{
+			if (const FMetasoundFrontendClassInput* OwningInput = OwningGraphClassInputPtr.Get())
+			{
+				return OwningInput->Metadata.Description;
+			}
+			
+			return FrontendControllerIntrinsics::GetInvalidValueConstRef<FText>();
+		}
+
+		FConnectability FInputNodeInputController::CanConnectTo(const IOutputController& InController) const 
+		{
+			static const FConnectability Connectability = {FConnectability::EConnectable::No};
+			return Connectability;
+		}
+
+		bool FInputNodeInputController::Connect(IOutputController& InController) 
+		{
+			return false;
+		}
+
+		bool FInputNodeInputController::ConnectWithConverterNode(IOutputController& InController, const FConverterNodeInfo& InNodeClassName)
+		{
+			return false;
+		}
 
 		//
 		// FBaseNodeController
@@ -1744,7 +1846,7 @@ namespace Metasound
 
 		FOutputHandle FOutputNodeController::CreateOutputController(FGuid InVertexID, FConstVertexAccessPtr InNodeVertexPtr, FConstClassOutputAccessPtr InClassOutputPtr, FNodeHandle InOwningNode) const
 		{
-			return FInvalidOutputController::GetInvalid();
+			return MakeShared<FOutputNodeOutputController>(FOutputNodeOutputController::FInitParams{InVertexID, InNodeVertexPtr, InClassOutputPtr, OwningGraphClassOutputPtr, GraphPtr, InOwningNode});
 		}
 
 		FDocumentAccess FOutputNodeController::ShareAccess() 
@@ -1834,7 +1936,7 @@ namespace Metasound
 
 		FInputHandle FInputNodeController::CreateInputController(FGuid InVertexID, FConstVertexAccessPtr InNodeVertexPtr, FConstClassInputAccessPtr InClassInputPtr, FNodeHandle InOwningNode) const
 		{
-			return FInvalidInputController::GetInvalid();
+			return MakeShared<FInputNodeInputController>(FInputNodeInputController::FInitParams{InVertexID, InNodeVertexPtr, InClassInputPtr, OwningGraphClassInputPtr, GraphPtr, InOwningNode});
 		}
 
 		FOutputHandle FInputNodeController::CreateOutputController(FGuid InVertexID, FConstVertexAccessPtr InNodeVertexPtr, FConstClassOutputAccessPtr InClassOutputPtr, FNodeHandle InOwningNode) const
diff --git a/Engine/Plugins/Runtime/Metasound/Source/MetasoundFrontend/Private/MetasoundFrontendStandardController.h b/Engine/Plugins/Runtime/Metasound/Source/MetasoundFrontend/Private/MetasoundFrontendStandardController.h
--- a/Engine/Plugins/Runtime/Metasound/Source/MetasoundFrontend/Private/MetasoundFrontendStandardController.h
+++ b/Engine/Plugins/Runtime/Metasound/Source/MetasoundFrontend/Private/MetasoundFrontendStandardController.h
@@ -126,6 +126,47 @@ namespace Metasound
 			FConstClassInputAccessPtr OwningGraphClassInputPtr;
 		};
 
+		/** FOutputNodeOutputController represents the output vertex of an input 
+		 * node. 
+		 *
+		 * FOutputNodeOutputController is largely to represent inputs coming into 
+		 * graph. 
+		 */
+		class METASOUNDFRONTEND_API FOutputNodeOutputController : public FBaseOutputController
+		{
+		public:
+			struct FInitParams
+			{
+				FGuid ID;
+
+				FConstVertexAccessPtr NodeVertexPtr;
+				FConstClassOutputAccessPtr ClassOutputPtr;
+				FConstClassOutputAccessPtr OwningGraphClassOutputPtr;
+				FGraphAccessPtr GraphPtr; 
+
+				/* Node handle which owns this output. */
+				FNodeHandle OwningNode;
+			};
+
+			/** Constructs the output controller. */
+			FOutputNodeOutputController(const FInitParams& InParams);
+
+			virtual ~FOutputNodeOutputController() = default;
+
+			bool IsValid() const override;
+
+			// Output metadata
+			const FText& GetDisplayName() const override;
+			const FText& GetTooltip() const override;
+
+			FConnectability CanConnectTo(const IInputController& InController) const override;
+			bool Connect(IInputController& InController) override;
+			bool ConnectWithConverterNode(IInputController& InController, const FConverterNodeInfo& InNodeClassName) override;
+
+		private:
+			FConstClassOutputAccessPtr OwningGraphClassOutputPtr;
+		};
+
 		
 		/** FBaseInputController provides common functionality for multiple derived
 		 * input controllers.
@@ -237,6 +278,46 @@ namespace Metasound
 			FConstClassOutputAccessPtr OwningGraphClassOutputPtr;
 		};
 
+		/** FInputNodeInputController represents the input vertex of an output 
+		 * node. 
+		 *
+		 * FInputNodeInputController is largely to represent outputs exposed from
+		 * a graph. 
+		 */
+		class METASOUNDFRONTEND_API FInputNodeInputController : public FBaseInputController 
+		{
+		public:
+			struct FInitParams
+			{
+				FGuid ID; 
+				FConstVertexAccessPtr NodeVertexPtr;
+				FConstClassInputAccessPtr ClassInputPtr;
+				FConstClassInputAccessPtr OwningGraphClassInputPtr;
+				FGraphAccessPtr GraphPtr; 
+				FNodeHandle OwningNode;
+			};
+
+			/** Constructs the input controller. */
+			FInputNodeInputController(const FInitParams& InParams);
+
+			bool IsValid() const override;
+
+			// Input metadata
+			const FText& GetDisplayName() const override;
+			const FText& GetTooltip() const override;
+
+			FConnectability CanConnectTo(const IOutputController& InController) const override;
+			bool Connect(IOutputController& InController) override;
+
+			// Connection controls.
+			bool ConnectWithConverterNode(IOutputController& InController, const FConverterNodeInfo& InNodeClassName) override;
+
+
+		private:
+
+			FConstClassInputAccessPtr OwningGraphClassInputPtr;
+		};
+
 		/** FBaseNodeController provides common functionality for multiple derived
 		 * node controllers.
 		 */
