Trigger Compare Node

diff --git a/Engine/Plugins/Runtime/Metasound/Source/MetasoundStandardNodes/Private/MetasoundTriggerCompareNode.cpp b/Engine/Plugins/Runtime/Metasound/Source/MetasoundStandardNodes/Private/MetasoundTriggerCompareNode.cpp
new file mode 100644
--- /dev/null
+++ b/Engine/Plugins/Runtime/Metasound/Source/MetasoundStandardNodes/Private/MetasoundTriggerCompareNode.cpp
@@ -0,0 +1,218 @@
+// Copyright Epic Games, Inc. All Rights Reserved.
+#include "MetasoundTriggerCompareNode.h"
+
+#include "CoreMinimal.h"
+
+#include "MetasoundBuilderInterface.h"
+#include "MetasoundDataReferenceCollection.h"
+#include "MetasoundExecutableOperator.h"
+#include "MetasoundFacade.h"
+#include "MetasoundNode.h"
+#include "MetasoundNodeInterface.h"
+#include "MetasoundNodeRegistrationMacro.h"
+#include "MetasoundDataTypeRegistrationMacro.h"
+#include "MetasoundOperatorInterface.h"
+#include "MetasoundPrimitives.h"
+#include "MetasoundStandardNodesNames.h"
+#include "MetasoundStandardNodesCategories.h"
+
+#define LOCTEXT_NAMESPACE "MetasoundStandardNodes"
+
+namespace Metasound
+{
+	namespace TriggerCompare
+	{
+		static const TCHAR* InParamNameCompare = TEXT("Compare");
+		static const TCHAR* InParamNameA = TEXT("A");
+		static const TCHAR* InParamNameB = TEXT("B");
+		static const TCHAR* InParamCompareType = TEXT("Type");
+		static const TCHAR* OutParamOnTrue = TEXT("True");
+		static const TCHAR* OutParamOnFalse = TEXT("False");
+	}
+
+	enum class ETriggerComparisonType
+	{
+		Equals,
+		NotEquals,
+		LessThan,
+		GreaterThan,
+		LessThanOrEquals,
+		GreaterThanOrEquals
+	};
+
+	DECLARE_METASOUND_ENUM(ETriggerComparisonType, ETriggerComparisonType::Equals, METASOUNDSTANDARDNODES_API,
+		FEnumTriggerComparisonType, FEnumTriggerComparisonTypeInfo, FEnumTriggerComparisonTypeReadRef, FEnumTriggerComparisonTypeWriteRef);
+
+	DEFINE_METASOUND_ENUM_BEGIN(ETriggerComparisonType, FEnumTriggerComparisonType, "TriggerComparisonType")
+		DEFINE_METASOUND_ENUM_ENTRY(ETriggerComparisonType::Equals, LOCTEXT("EqualsDescription", "Equals"), LOCTEXT("EqualsDescriptionTT", "True if A and B are equal.")),
+		DEFINE_METASOUND_ENUM_ENTRY(ETriggerComparisonType::NotEquals, LOCTEXT("NotEqualsDescriptioin", "Not Equals"), LOCTEXT("NotEqualsTT", "True if A and B are not equal.")),
+		DEFINE_METASOUND_ENUM_ENTRY(ETriggerComparisonType::LessThan, LOCTEXT("LessThanDescription", "Less Than"), LOCTEXT("LessThanTT", "True if A is less than B.")),
+		DEFINE_METASOUND_ENUM_ENTRY(ETriggerComparisonType::GreaterThan, LOCTEXT("GreaterThanDescription", "Greater Than"), LOCTEXT("GreaterThanTT", "True if A is greater than B.")),
+		DEFINE_METASOUND_ENUM_ENTRY(ETriggerComparisonType::LessThanOrEquals, LOCTEXT("LessThanOrEqualsDescription", "Less Than Or Equals"), LOCTEXT("LessThanOrEqualsTT", "True if A is less than or equal to B.")),
+		DEFINE_METASOUND_ENUM_ENTRY(ETriggerComparisonType::GreaterThanOrEquals, LOCTEXT("GreaterThanOrEqualsDescription", "Greater Than Or Equals"), LOCTEXT("GreaterThanOrEqualsTT", "True if A is greater than or equal to B.")),
+	DEFINE_METASOUND_ENUM_END()
+
+	class FTriggerCompareOperator : public TExecutableOperator<FTriggerCompareOperator>
+	{
+		public:
+			static const FNodeClassMetadata& GetNodeInfo();
+			static const FVertexInterface& GetVertexInterface();
+			static TUniquePtr<IOperator> CreateOperator(const FCreateOperatorParams& InParams, FBuildErrorArray& OutErrors);
+
+			FTriggerCompareOperator(const FOperatorSettings& InSettings, const FTriggerReadRef& InOnCompareTrigger, const FInt32ReadRef& InValueA, const FInt32ReadRef& InValueB, const FEnumTriggerComparisonTypeReadRef& InTriggerComparisonType);
+
+			virtual FDataReferenceCollection GetInputs() const override;
+			virtual FDataReferenceCollection GetOutputs() const override;
+
+			void Execute();
+
+		private:
+			FTriggerReadRef OnCompareTrigger;
+			FInt32ReadRef ValueA;
+			FInt32ReadRef ValueB;
+			FEnumTriggerComparisonTypeReadRef TriggerComparisonType;
+
+			FTriggerWriteRef TriggerOutOnTrue;
+			FTriggerWriteRef TriggerOutOnFalse;
+	};
+
+	FTriggerCompareOperator::FTriggerCompareOperator(const FOperatorSettings& InSettings, const FTriggerReadRef& InOnCompareTrigger, const FInt32ReadRef& InValueA, const FInt32ReadRef& InValueB, const FEnumTriggerComparisonTypeReadRef& InTriggerComparisonType)
+		: OnCompareTrigger(InOnCompareTrigger)
+		, ValueA(InValueA)
+		, ValueB(InValueB)
+		, TriggerComparisonType(InTriggerComparisonType)
+		, TriggerOutOnTrue(FTriggerWriteRef::CreateNew(InSettings))
+		, TriggerOutOnFalse(FTriggerWriteRef::CreateNew(InSettings))
+	{
+	}
+
+	FDataReferenceCollection FTriggerCompareOperator::GetInputs() const
+	{
+		FDataReferenceCollection InputDataReferences;
+		InputDataReferences.AddDataReadReference(TriggerCompare::InParamNameCompare, OnCompareTrigger);
+		InputDataReferences.AddDataReadReference(TriggerCompare::InParamNameA, ValueA);
+		InputDataReferences.AddDataReadReference(TriggerCompare::InParamNameB, ValueB);
+		InputDataReferences.AddDataReadReference(TriggerCompare::InParamCompareType, TriggerComparisonType);
+		return InputDataReferences;
+	}
+
+	FDataReferenceCollection FTriggerCompareOperator::GetOutputs() const
+	{
+		FDataReferenceCollection OutputDataReferences;
+		OutputDataReferences.AddDataReadReference(TriggerCompare::OutParamOnTrue, TriggerOutOnTrue);
+		OutputDataReferences.AddDataReadReference(TriggerCompare::OutParamOnFalse, TriggerOutOnFalse);
+
+		return OutputDataReferences;
+	}
+
+	void FTriggerCompareOperator::Execute()
+	{
+		TriggerOutOnTrue->AdvanceBlock();
+		TriggerOutOnFalse->AdvanceBlock();
+
+		OnCompareTrigger->ExecuteBlock(
+			[&](int32 StartFrame, int32 EndFrame)
+			{
+			},
+			[this](int32 StartFrame, int32 EndFrame)
+			{
+				int32 CurrValA = *ValueA;
+				int32 CurrValB = *ValueB;
+				bool bIsTrue = false;
+
+				switch (*TriggerComparisonType)
+				{
+					case ETriggerComparisonType::Equals:
+						bIsTrue = (CurrValA == CurrValB);
+						break;
+					case ETriggerComparisonType::NotEquals:
+						bIsTrue = (CurrValA != CurrValB);
+						break;
+					case ETriggerComparisonType::LessThan:
+						bIsTrue = (CurrValA < CurrValB);
+						break;
+					case ETriggerComparisonType::GreaterThan:
+						bIsTrue = (CurrValA > CurrValB);
+						break;
+					case ETriggerComparisonType::LessThanOrEquals:
+						bIsTrue = (CurrValA <= CurrValB);
+						break;
+					case ETriggerComparisonType::GreaterThanOrEquals:
+						bIsTrue = (CurrValA >= CurrValB);
+						break;
+				}
+
+				if (bIsTrue)
+				{
+					TriggerOutOnTrue->TriggerFrame(StartFrame);
+				}
+				else
+				{
+					TriggerOutOnFalse->TriggerFrame(StartFrame);
+				}
+			}
+		);
+	}
+
+	TUniquePtr<IOperator> FTriggerCompareOperator::CreateOperator(const FCreateOperatorParams& InParams, FBuildErrorArray& OutErrors)
+	{
+		const FDataReferenceCollection& InputCollection = InParams.InputDataReferences;
+		const FInputVertexInterface& InputInterface = GetVertexInterface().GetInputInterface();
+
+		FTriggerReadRef InOnTriggerCompare = InputCollection.GetDataReadReferenceOrConstruct<FTrigger>(TriggerCompare::InParamNameCompare, InParams.OperatorSettings);
+		FInt32ReadRef InValueA = InputCollection.GetDataReadReferenceOrConstructWithVertexDefault<int32>(InputInterface, TriggerCompare::InParamNameA);
+		FInt32ReadRef InValueB = InputCollection.GetDataReadReferenceOrConstructWithVertexDefault<int32>(InputInterface, TriggerCompare::InParamNameB);
+		FEnumTriggerComparisonTypeReadRef InComparison = InputCollection.GetDataReadReferenceOrConstruct<FEnumTriggerComparisonType>(TriggerCompare::InParamCompareType);
+
+		return MakeUnique<FTriggerCompareOperator>(InParams.OperatorSettings, InOnTriggerCompare, InValueA, InValueB, InComparison);
+	}
+
+	const FVertexInterface& FTriggerCompareOperator::GetVertexInterface()
+	{
+		static const FVertexInterface Interface(
+			FInputVertexInterface(
+				TInputDataVertexModel<FTrigger>(TriggerCompare::InParamNameCompare, LOCTEXT("TriggerCompareTT", "Trigger to compare A and B.")),
+				TInputDataVertexModel<int32>(TriggerCompare::InParamNameA, LOCTEXT("TriggerCompareATT", "The first value, A, to compare against.")),
+				TInputDataVertexModel<int32>(TriggerCompare::InParamNameB, LOCTEXT("TriggerCompareBTT", "The second value, B, to compare against.")),
+				TInputDataVertexModel<FEnumTriggerComparisonType>(TriggerCompare::InParamCompareType, LOCTEXT("TriggerCompareTypeTT", "How to compare A and B."))
+			),
+			FOutputVertexInterface(
+				TOutputDataVertexModel<FTrigger>(TriggerCompare::OutParamOnTrue, LOCTEXT("TriggerCompareOnTrueTT", "Output trigger for when the comparison is true.")),
+				TOutputDataVertexModel<FTrigger>(TriggerCompare::OutParamOnFalse, LOCTEXT("TriggerCompareOnFalseTT", "Output trigger for when the comparison is false."))
+			)
+		);
+
+		return Interface;
+	}
+
+	const FNodeClassMetadata& FTriggerCompareOperator::GetNodeInfo()
+	{
+		auto InitNodeInfo = []() -> FNodeClassMetadata
+		{
+			FNodeClassMetadata Info;
+			Info.ClassName = {Metasound::StandardNodes::Namespace, TEXT("Trigger Compare"), TEXT("")};
+			Info.MajorVersion = 1;
+			Info.MinorVersion = 0;
+			Info.DisplayName = LOCTEXT("Metasound_TriggerSelectNodeDisplayName", "Trigger Compare");
+			Info.Description = LOCTEXT("Metasound_TriggerSelectNodeDescription", "Output triggers based on comparing integer inputs, A and B.");
+			Info.Author = PluginAuthor;
+			Info.PromptIfMissing = PluginNodeMissingPrompt;
+			Info.DefaultInterface = GetVertexInterface();
+			Info.CategoryHierarchy.Emplace(StandardNodes::TriggerUtils);
+
+			return Info;
+		};
+
+		static const FNodeClassMetadata Info = InitNodeInfo();
+		return Info;
+	}
+
+	FTriggerCompareNode::FTriggerCompareNode(const FNodeInitData& InInitData)
+	:	FNodeFacade(InInitData.InstanceName, InInitData.InstanceID, TFacadeOperatorClass<FTriggerCompareOperator>())
+	{
+	}
+
+	METASOUND_REGISTER_NODE(FTriggerCompareNode)
+}
+
+#undef LOCTEXT_NAMESPACE // MetasoundTriggerDelayNode
diff --git a/Engine/Plugins/Runtime/Metasound/Source/MetasoundStandardNodes/Public/MetasoundTriggerCompareNode.h b/Engine/Plugins/Runtime/Metasound/Source/MetasoundStandardNodes/Public/MetasoundTriggerCompareNode.h
new file mode 100644
--- /dev/null
+++ b/Engine/Plugins/Runtime/Metasound/Source/MetasoundStandardNodes/Public/MetasoundTriggerCompareNode.h
@@ -0,0 +1,21 @@
+// Copyright Epic Games, Inc. All Rights Reserved.
+
+#pragma once
+
+#include "CoreMinimal.h"
+#include "MetasoundBuilderInterface.h"
+#include "MetasoundDataReferenceCollection.h"
+#include "MetasoundExecutableOperator.h"
+#include "MetasoundFacade.h"
+#include "MetasoundNodeInterface.h"
+
+namespace Metasound
+{
+	class METASOUNDSTANDARDNODES_API FTriggerCompareNode : public FNodeFacade
+	{
+		public:
+			FTriggerCompareNode(const FNodeInitData& InInitData);
+
+			virtual ~FTriggerCompareNode() = default;
+	};
+}
