Slight refactor and cleanup of the metasound value node

diff --git a/Engine/Plugins/Runtime/Metasound/Source/MetasoundStandardNodes/Private/MetasoundValueNode.cpp b/Engine/Plugins/Runtime/Metasound/Source/MetasoundStandardNodes/Private/MetasoundValueNode.cpp
--- a/Engine/Plugins/Runtime/Metasound/Source/MetasoundStandardNodes/Private/MetasoundValueNode.cpp
+++ b/Engine/Plugins/Runtime/Metasound/Source/MetasoundStandardNodes/Private/MetasoundValueNode.cpp
@@ -8,6 +8,58 @@
 
 namespace Metasound
 {
+	namespace MetasoundValueNodePrivate
+	{
+		FNodeClassMetadata CreateNodeClassMetadata(const FName& InDataTypeName, const FName& InOperatorName, const FText& InDisplayName, const FText& InDescription, const FVertexInterface& InDefaultInterface)
+		{
+			FNodeClassMetadata Metadata
+			{
+				FNodeClassName{FName("Value"), InOperatorName, InDataTypeName},
+				1, // Major Version
+				0, // Minor Version
+				InDisplayName,
+				InDescription,
+				PluginAuthor,
+				PluginNodeMissingPrompt,
+				InDefaultInterface,
+				{LOCTEXT("ValueCategory", "Value")},
+				{TEXT("Value")},
+				FNodeDisplayStyle{}
+			};
+
+			return Metadata;
+		}
+	}
+
+
+	namespace ValueVertexNames
+	{
+		const FString& GetInitValueName()
+		{
+			static const FString Name = TEXT("Init");
+			return Name;
+		}
+
+		const FString& GetSetValueName()
+		{
+			static const FString Name = TEXT("Value");
+			return Name;
+		}
+
+		const FString& GetInputTriggerName()
+		{
+			static const FString Name = TEXT("Set");
+			return Name;
+		}
+
+		const FString& GetOutputValueName()
+		{
+			static const FString Name = TEXT("Value");
+			return Name;
+		}
+	}
+
+
 	using FValueNodeInt32 = TValueNode<int32>;
  	METASOUND_REGISTER_NODE(FValueNodeInt32)
 
diff --git a/Engine/Plugins/Runtime/Metasound/Source/MetasoundStandardNodes/Public/MetasoundValueNode.h b/Engine/Plugins/Runtime/Metasound/Source/MetasoundStandardNodes/Public/MetasoundValueNode.h
--- a/Engine/Plugins/Runtime/Metasound/Source/MetasoundStandardNodes/Public/MetasoundValueNode.h
+++ b/Engine/Plugins/Runtime/Metasound/Source/MetasoundStandardNodes/Public/MetasoundValueNode.h
@@ -14,100 +14,27 @@
 
 namespace Metasound
 {
-	template<typename ValueType>
-	struct TValueDisplayName
-	{
-		bool bIsSupported = false;
-	};
-
-	template<>
-	struct TValueDisplayName<int32>
-	{
-		static FText GetDisplayName()
-		{
-			return LOCTEXT("IntValueOpDisplayName", "Int Value");
-		}
-
-		static FText GetDescription()
-		{
-			return LOCTEXT("ValueOpDescription", "Allows setting an int value to the output on trigger.");
-		}
-	};
-
-	template<>
-	struct TValueDisplayName<float>
-	{
-		static FText GetDisplayName()
-		{
-			return LOCTEXT("FloatValueOpDisplayName", "Float Value");
-		}
-
-		static FText GetDescription()
-		{
-			return LOCTEXT("FloatValueOpDisplayName", "Allows setting a float value to the output on trigger.");
-		}
-	};
-
-	template<>
-	struct TValueDisplayName<bool>
+	namespace MetasoundValueNodePrivate
 	{
-		static FText GetDisplayName()
-		{
-			return LOCTEXT("BoolValueOpDisplayName", "Bool Value");
-		}
-
-		static FText GetDescription()
-		{
-			return LOCTEXT("BoolValueOpDisplayName", "Allows setting a bool value to the output on trigger.");
-		}
-	};
+		METASOUNDSTANDARDNODES_API FNodeClassMetadata CreateNodeClassMetadata(const FName& InDataTypeName, const FName& InOperatorName, const FText& InDisplayName, const FText& InDescription, const FVertexInterface& InDefaultInterface);
+	}
 
-	template<>
-	struct TValueDisplayName<FString>
+	namespace ValueVertexNames
 	{
-		static FText GetDisplayName()
-		{
-			return LOCTEXT("StringValueOpDisplayName", "String Value");
-		}
-
-		static FText GetDescription()
-		{
-			return LOCTEXT("StringValueOpDisplayName", "Allows setting a string value to the output on trigger.");
-		}
-	};
+		METASOUNDSTANDARDNODES_API const FString& GetInitValueName();
+		METASOUNDSTANDARDNODES_API const FString& GetSetValueName();
+		METASOUNDSTANDARDNODES_API const FString& GetInputTriggerName();
+		METASOUNDSTANDARDNODES_API const FString& GetOutputValueName();
+	}
 
 	template<typename ValueType>
 	class TValueOperator : public TExecutableOperator<TValueOperator<ValueType>>
 	{
 	public:
-		using FArrayDataReadReference = TDataReadReference<ValueType>;
-
-		static const FString& GetInitValueName()
-		{
-			static const FString Name = TEXT("Init");
-			return Name;
-		}
-
-		static const FString& GetSetValueName()
-		{
-			static const FString Name = TEXT("Set");
-			return Name;
-		}
-
-		static const FString& GetInputTriggerName()
-		{
-			static const FString Name = TEXT("Trigger");
-			return Name;
-		}
-
-		static const FString& GetOutputValueName()
-		{
-			static const FString Name = TEXT("Output");
-			return Name;
-		}
-
 		static const FVertexInterface& GetDefaultInterface()
 		{
+			using namespace ValueVertexNames;
+
 			static const FVertexInterface DefaultInterface(
 				FInputVertexInterface(
 					TInputDataVertexModel<ValueType>(GetInitValueName(), LOCTEXT("InitValue", "Value to init the output to.")),
@@ -126,22 +53,13 @@ namespace Metasound
 		{
 			auto CreateNodeClassMetadata = []() -> FNodeClassMetadata
 			{
-				FNodeClassMetadata Metadata
-				{
-					FNodeClassName{FName("Value"), GetMetasoundDataTypeName<ValueType>(), TEXT("") },
-					1, // Major Version
-					0, // Minor Version
-					TValueDisplayName<ValueType>::GetDisplayName(),
-					LOCTEXT("ValueOpDescription", "Allows setting a value to output on trigger."),
-					PluginAuthor,
-					PluginNodeMissingPrompt,
-					GetDefaultInterface(),
-					{LOCTEXT("ValueCategory", "Value")},
-					{TEXT("Value")},
-					FNodeDisplayStyle{}
-				};
-
-				return Metadata;
+				FName DataTypeName = GetMetasoundDataTypeName<ValueType>();
+				FName OperatorName = TEXT("Value");
+				FText NodeDisplayName = FText::Format(LOCTEXT("ValueDisplayNamePattern", "Value ({0})"), FText::FromString(GetMetasoundDataTypeString<ValueType>()));
+				FText NodeDescription = LOCTEXT("ValueDescription", "Allows setting a value to output on trigger.");
+				FVertexInterface NodeInterface = GetDefaultInterface();
+
+				return MetasoundValueNodePrivate::CreateNodeClassMetadata(DataTypeName, OperatorName, NodeDisplayName, NodeDescription, NodeInterface);
 			};
 
 			static const FNodeClassMetadata Metadata = CreateNodeClassMetadata();
@@ -150,6 +68,8 @@ namespace Metasound
 
 		static TUniquePtr<IOperator> CreateOperator(const FCreateOperatorParams& InParams, TArray<TUniquePtr<IOperatorBuildError>>& OutErrors)
 		{
+			using namespace ValueVertexNames;
+
 			const FInputVertexInterface& InputInterface = InParams.Node.GetVertexInterface().GetInputInterface();
 			const FDataReferenceCollection& InputCollection = InParams.InputDataReferences;
 
@@ -175,8 +95,9 @@ namespace Metasound
 
 		virtual FDataReferenceCollection GetInputs() const override
 		{
-			FDataReferenceCollection Inputs;
+			using namespace ValueVertexNames;
 
+			FDataReferenceCollection Inputs;
 			Inputs.AddDataReadReference(GetInputTriggerName(), Trigger);
 			Inputs.AddDataReadReference(GetInitValueName(), InitValue);
 			Inputs.AddDataReadReference(GetSetValueName(), SetValue);
@@ -186,8 +107,9 @@ namespace Metasound
 
 		virtual FDataReferenceCollection GetOutputs() const override
 		{
-			FDataReferenceCollection Outputs;
+			using namespace ValueVertexNames;
 
+			FDataReferenceCollection Outputs;
 			Outputs.AddDataReadReference(GetOutputValueName(), OutputValue);
 
 			return Outputs;
