Random node refactor

-- using generic programming for random node
-- added random bool
-- added random node to new random utiliies category

diff --git a/Engine/Plugins/Runtime/Metasound/Source/MetasoundStandardNodes/Private/MetasoundRandomNode.cpp b/Engine/Plugins/Runtime/Metasound/Source/MetasoundStandardNodes/Private/MetasoundRandomNode.cpp
--- a/Engine/Plugins/Runtime/Metasound/Source/MetasoundStandardNodes/Private/MetasoundRandomNode.cpp
+++ b/Engine/Plugins/Runtime/Metasound/Source/MetasoundStandardNodes/Private/MetasoundRandomNode.cpp
@@ -3,404 +3,111 @@
 #include "MetasoundRandomNode.h"
 
 #include "Internationalization/Text.h"
-#include "MetasoundExecutableOperator.h"
 #include "MetasoundNodeRegistrationMacro.h"
-#include "MetasoundPrimitives.h"
-#include "MetasoundStandardNodesNames.h"
-#include "MetasoundTrigger.h"
-#include "Math/RandomStream.h"
-#include "HAL/PlatformTime.h"
 
 #define LOCTEXT_NAMESPACE "MetasoundStandardNodes"
 
 namespace Metasound
 {
-	static const TCHAR* InParamNameSeed = TEXT("Seed");
-	static const TCHAR* InParamNameMin = TEXT("Min");
-	static const TCHAR* InParamNameMax = TEXT("Max");
-	static const TCHAR* InParamNameNextTrigger = TEXT("Next");
-	static const TCHAR* InParamNameResetTrigger = TEXT("Reset");
-	static const TCHAR* OutParamNameValue = TEXT("Value");
-
-	//////////////////////////////////////////////////////////////////////////
-	// Random Int
-
-	class FRandomIntOperator : public TExecutableOperator<FRandomIntOperator>
-	{
-	public:
-
-		static const FNodeClassMetadata& GetNodeInfo();
-		static const FVertexInterface& GetVertexInterface();
-		static TUniquePtr<IOperator> CreateOperator(const FCreateOperatorParams& InParams, FBuildErrorArray& OutErrors);
-
-		FRandomIntOperator(const FOperatorSettings& InSettings,
-			const FInt32ReadRef& InSeed, 
-			const FInt32ReadRef& InMin, 
-			const FInt32ReadRef& InMax, 
-			const FTriggerReadRef& InTriggerNext, 
-			const FTriggerReadRef& InTriggerReset);
-
-		virtual FDataReferenceCollection GetInputs() const override;
-		virtual FDataReferenceCollection GetOutputs() const override;
-		void Execute();
-
-	private:
-		void EvaluateSeedChanges();
-
-		FInt32ReadRef Seed;
-		FInt32ReadRef Min;
-		FInt32ReadRef Max;
-		FTriggerReadRef TriggerNext;
-		FTriggerReadRef TriggerReset;
-
-		FInt32WriteRef ValueOutput;
-
-		FRandomStream RandomStream;
-
-		// Used to track if the random node starts non-zero and then becomes 0 (i.e. we want a new random seed)
-		bool bIsZeroSeeded = false;
-		bool bIsRandomStreamInitialized = false;
-	};
-
-	FRandomIntOperator::FRandomIntOperator(const FOperatorSettings& InSettings,
-		const FInt32ReadRef& InSeed,
-		const FInt32ReadRef& InMin,
-		const FInt32ReadRef& InMax,
-		const FTriggerReadRef& InTriggerNext,
-		const FTriggerReadRef& InTriggerReset)
-		: Seed(InSeed)
-		, Min(InMin)
-		, Max(InMax)
-		, TriggerNext(InTriggerNext)
-		, TriggerReset(InTriggerReset)
-		, ValueOutput(FInt32WriteRef::CreateNew(0))
-		, bIsZeroSeeded(*Seed == 0)
-		, bIsRandomStreamInitialized(false)
+	namespace RandomNodeNames
 	{
-		*ValueOutput = *Min;
-	}
-
-	FDataReferenceCollection FRandomIntOperator::GetInputs() const
-	{
-		FDataReferenceCollection InputDataReferences;
-		InputDataReferences.AddDataReadReference(InParamNameSeed, FInt32ReadRef(Seed));
-		InputDataReferences.AddDataReadReference(InParamNameMin, FInt32ReadRef(Min));
-		InputDataReferences.AddDataReadReference(InParamNameMax, FInt32ReadRef(Max));
-		InputDataReferences.AddDataReadReference(InParamNameNextTrigger, FTriggerReadRef(TriggerNext));
-		InputDataReferences.AddDataReadReference(InParamNameResetTrigger, FTriggerReadRef(TriggerReset));
-
-		return InputDataReferences;
-	}
-
-	FDataReferenceCollection FRandomIntOperator::GetOutputs() const
-	{
-		FDataReferenceCollection OutputDataReferences;
-		OutputDataReferences.AddDataReadReference(OutParamNameValue, FInt32ReadRef(ValueOutput));
-		return OutputDataReferences;
-	}
-
-	void FRandomIntOperator::EvaluateSeedChanges()
-	{
-		// if we have a non-zero seed
-		if (*Seed != 0)
+		const FString& GetInputNextTriggerName()
 		{
-			// If we were previously zero-seeded OR our seed has changed
-			if (bIsZeroSeeded || !bIsRandomStreamInitialized || *Seed != RandomStream.GetInitialSeed())
-			{
-				bIsRandomStreamInitialized = true;
-				bIsZeroSeeded = false;
-				RandomStream.Initialize(*Seed);
-			}
+			static const FString Name = TEXT("Next");
+			return Name;
 		}
-		// If we are zero-seeded now BUT were previously not, we need to randomize our seed
-		else if (!bIsZeroSeeded || !bIsRandomStreamInitialized)
+
+		const FString& GetInputResetTriggerName()
 		{
-			bIsRandomStreamInitialized = true;
-			bIsZeroSeeded = true;
-			RandomStream.Initialize(FPlatformTime::Cycles());
+			static const FString Name = TEXT("Reset");
+			return Name;
 		}
-	}
 
-	void FRandomIntOperator::Execute()
-	{
-		if (TriggerReset->IsTriggeredInBlock())
+		const FString& GetInputSeedName()
 		{
-			EvaluateSeedChanges();
-			RandomStream.Reset();
+			static const FString Name = TEXT("Seed");
+			return Name;
 		}
 
-		if (TriggerNext->IsTriggeredInBlock())
+		const FString& GetInputMinName()
 		{
-			EvaluateSeedChanges();
-			*ValueOutput = RandomStream.RandRange(*Min, *Max);
+			static const FString Name = TEXT("Min");
+			return Name;
 		}
 
-// 		TriggerReset->ExecuteBlock(
-// 			// OnPreTrigger
-// 			[&](int32 StartFrame, int32 EndFrame)
-// 			{
-// 			},
-// 			// OnTrigger
-// 				[&](int32 StartFrame, int32 EndFrame)
-// 			{
-// 				EvaluateSeedChanges();
-// 				RandomStream.Reset();
-// 
-// 				UE_LOG(LogTemp, Log, TEXT("RESET"));
-// 			});
-// 
-// 		TriggerNext->ExecuteBlock(
-// 			// OnPreTrigger
-// 			[&](int32 StartFrame, int32 EndFrame)
-// 			{
-// 			},
-// 			// OnTrigger
-// 				[&](int32 StartFrame, int32 EndFrame)
-// 			{
-// 				EvaluateSeedChanges();
-// 				*ValueOutput = RandomStream.RandRange(*Min, *Max);
-// 
-// 				UE_LOG(LogTemp, Log, TEXT("%d"), *ValueOutput);
-// 			})
-	}
-
-	const FVertexInterface& FRandomIntOperator::GetVertexInterface()
-	{
-		static const FVertexInterface Interface(
-			FInputVertexInterface(
-				TInputDataVertexModel<FTrigger>(InParamNameNextTrigger, LOCTEXT("NextTooltip", "Trigger to generate the next random integer.")),
-				TInputDataVertexModel<FTrigger>(InParamNameResetTrigger, LOCTEXT("ResetTooltip", "Trigger to reset the random sequence with the supplied seed. Useful to get randomized repetition.")),
-				TInputDataVertexModel<int32>(InParamNameSeed, LOCTEXT("SeedTooltip", "The seed value to use for the random node. Set to 0 to use a random seed."), 0),
-				TInputDataVertexModel<int32>(InParamNameMin, LOCTEXT("MinTooltip", "Min random integer."), 0),
-				TInputDataVertexModel<int32>(InParamNameMax, LOCTEXT("MaxTooltip", "Max random integer."), 100)
-			),
-			FOutputVertexInterface(
-				TOutputDataVertexModel<int32>(OutParamNameValue, LOCTEXT("ValueTooltip", "The current randomly generated random integer value."))
-			)
-		);
-
-		return Interface;
-	}
-
-	const FNodeClassMetadata& FRandomIntOperator::GetNodeInfo()
-	{
-		auto InitNodeInfo = []() -> FNodeClassMetadata
+		const FString& GetInputMaxName()
 		{
-			FNodeClassMetadata Info;
-			Info.ClassName = { Metasound::StandardNodes::Namespace, TEXT("Random"), TEXT("Int32") };
-			Info.MajorVersion = 1;
-			Info.MinorVersion = 0;
-			Info.DisplayName = LOCTEXT("Metasound_RandomIntDisplayName", "Random Int");
-			Info.Description = LOCTEXT("Metasound_RandomIntNodeDescription", "Generates a seedable random integer in the given value range.");
-			Info.Author = PluginAuthor;
-			Info.PromptIfMissing = PluginNodeMissingPrompt;
-			Info.DefaultInterface = GetVertexInterface();
-
-			return Info;
-		};
-
-		static const FNodeClassMetadata Info = InitNodeInfo();
-
-		return Info;
-	}
-
-	TUniquePtr<IOperator> FRandomIntOperator::CreateOperator(const FCreateOperatorParams& InParams, FBuildErrorArray& OutErrors)
-	{
-		const FRandomIntNode& RandomIntNode = static_cast<const FRandomIntNode&>(InParams.Node);
-		const FDataReferenceCollection& InputCollection = InParams.InputDataReferences;
-		const FInputVertexInterface& InputInterface = GetVertexInterface().GetInputInterface();
-
-		FInt32ReadRef SeedValue = InputCollection.GetDataReadReferenceOrConstructWithVertexDefault<int32>(InputInterface, InParamNameSeed);
-		FInt32ReadRef MinValue = InputCollection.GetDataReadReferenceOrConstructWithVertexDefault<int32>(InputInterface, InParamNameMin);
-		FInt32ReadRef MaxValue = InputCollection.GetDataReadReferenceOrConstructWithVertexDefault<int32>(InputInterface, InParamNameMax);
-		FTriggerReadRef TriggerNext = InputCollection.GetDataReadReferenceOrConstruct<FTrigger>(InParamNameNextTrigger, InParams.OperatorSettings);
-		FTriggerReadRef TriggerReset = InputCollection.GetDataReadReferenceOrConstruct<FTrigger>(InParamNameResetTrigger, InParams.OperatorSettings);
-
-		return MakeUnique<FRandomIntOperator>(InParams.OperatorSettings, SeedValue, MinValue, MaxValue, TriggerNext, TriggerReset);
-	}
-
-	FRandomIntNode::FRandomIntNode(const FNodeInitData& InitData)
-		: FNodeFacade(InitData.InstanceName, InitData.InstanceID, TFacadeOperatorClass<FRandomIntOperator>())
-	{
-	}
-
-	METASOUND_REGISTER_NODE(FRandomIntNode)
-
-	//////////////////////////////////////////////////////////////////////////
-	// Random Float
-
-	class FRandomFloatOperator : public TExecutableOperator<FRandomFloatOperator>
-	{
-	public:
-
-		static const FNodeClassMetadata& GetNodeInfo();
-		static const FVertexInterface& GetVertexInterface();
-		static TUniquePtr<IOperator> CreateOperator(const FCreateOperatorParams& InParams, FBuildErrorArray& OutErrors);
-
-		FRandomFloatOperator(const FOperatorSettings& InSettings,
-			const FInt32ReadRef& InSeed,
-			const FFloatReadRef& InMin,
-			const FFloatReadRef& InMax,
-			const FTriggerReadRef& InTriggerNext,
-			const FTriggerReadRef& InTriggerReset);
-
-		virtual FDataReferenceCollection GetInputs() const override;
-		virtual FDataReferenceCollection GetOutputs() const override;
-		void Execute();
-
-	private:
-		void EvaluateSeedChanges();
-
-		FInt32ReadRef Seed;
-		FFloatReadRef Min;
-		FFloatReadRef Max;
-		FTriggerReadRef TriggerNext;
-		FTriggerReadRef TriggerReset;
-
-		FFloatWriteRef ValueOutput;
-
-		FRandomStream RandomStream;
-
-		// Used to track if the random node starts non-zero and then becomes 0 (i.e. we want a new random seed)
-		bool bIsZeroSeeded = false;
-		bool bIsRandomStreamInitialized = false;
-	};
-
-	FRandomFloatOperator::FRandomFloatOperator(const FOperatorSettings& InSettings,
-		const FInt32ReadRef& InSeed,
-		const FFloatReadRef& InMin,
-		const FFloatReadRef& InMax,
-		const FTriggerReadRef& InTriggerNext,
-		const FTriggerReadRef& InTriggerReset)
-		: Seed(InSeed)
-		, Min(InMin)
-		, Max(InMax)
-		, TriggerNext(InTriggerNext)
-		, TriggerReset(InTriggerReset)
-		, ValueOutput(FFloatWriteRef::CreateNew(0.0f))
-		, bIsZeroSeeded(*Seed == 0)
-		, bIsRandomStreamInitialized(false)
-	{
-		*ValueOutput = *Min;
-	}
-
-	FDataReferenceCollection FRandomFloatOperator::GetInputs() const
-	{
-		FDataReferenceCollection InputDataReferences;
-		InputDataReferences.AddDataReadReference(InParamNameSeed, FInt32ReadRef(Seed));
-		InputDataReferences.AddDataReadReference(InParamNameMin, FFloatReadRef(Min));
-		InputDataReferences.AddDataReadReference(InParamNameMax, FFloatReadRef(Max));
-		InputDataReferences.AddDataReadReference(InParamNameNextTrigger, FTriggerReadRef(TriggerNext));
-		InputDataReferences.AddDataReadReference(InParamNameResetTrigger, FTriggerReadRef(TriggerReset));
-
-		return InputDataReferences;
-	}
-
-	FDataReferenceCollection FRandomFloatOperator::GetOutputs() const
-	{
-		FDataReferenceCollection OutputDataReferences;
-		OutputDataReferences.AddDataReadReference(OutParamNameValue, FFloatReadRef(ValueOutput));
-		return OutputDataReferences;
-	}
+			static const FString Name = TEXT("Max");
+			return Name;
+		}
 
-	void FRandomFloatOperator::EvaluateSeedChanges()
-	{
-		// if we have a non-zero seed
-		if (*Seed != 0)
+		const FString& GetOutputOnNextTriggerName()
 		{
-			// If we were previously zero-seeded OR our seed has changed
-			if (bIsZeroSeeded || !bIsRandomStreamInitialized || *Seed != RandomStream.GetInitialSeed())
-			{
-				bIsRandomStreamInitialized = true;
-				bIsZeroSeeded = false;
-				RandomStream.Initialize(*Seed);
-			}
+			static const FString Name = TEXT("On Next");
+			return Name;
 		}
-		// If we are zero-seeded now BUT were previously not, we need to randomize our seed
-		else if (!bIsZeroSeeded || !bIsRandomStreamInitialized)
+
+		const FString& GetOutputOnResetTriggerName()
 		{
-			bIsRandomStreamInitialized = true;
-			bIsZeroSeeded = true;
-			RandomStream.Initialize(FPlatformTime::Cycles());
+			static const FString Name = TEXT("On Reset");
+			return Name;
 		}
-	}
 
-	void FRandomFloatOperator::Execute()
-	{
-		if (TriggerReset->IsTriggeredInBlock())
+		const FString& GetOutputValueName()
 		{
-			EvaluateSeedChanges();
-			RandomStream.Reset();
+			static const FString Name = TEXT("Value");
+			return Name;
 		}
 
-		if (TriggerNext->IsTriggeredInBlock())
+		static FText GetNextTriggerDescription()
 		{
-			EvaluateSeedChanges();
-			*ValueOutput = RandomStream.FRandRange(*Min, *Max);
+			return LOCTEXT("RandomNodeNextTT", "Trigger to generate the next random integer.");
 		}
-	}
 
-	const FVertexInterface& FRandomFloatOperator::GetVertexInterface()
-	{
-		static const FVertexInterface Interface(
-			FInputVertexInterface(
-				TInputDataVertexModel<FTrigger>(InParamNameNextTrigger, LOCTEXT("NextTooltip", "Trigger to generate the next random float.")),
-				TInputDataVertexModel<FTrigger>(InParamNameResetTrigger, LOCTEXT("ResetTooltip", "Trigger to reset the random sequence with the supplied seed. Useful to get randomized repetition.")),
-				TInputDataVertexModel<int32>(InParamNameSeed, LOCTEXT("SeedTooltip", "The seed value to use for the random node. Set to 0 to use a random seed."), 0),
-				TInputDataVertexModel<float>(InParamNameMin, LOCTEXT("MinTooltip", "Min random float."), 0.0f),
-				TInputDataVertexModel<float>(InParamNameMax, LOCTEXT("MaxTooltip", "Max random float."), 1.0f)
-			),
-			FOutputVertexInterface(
-				TOutputDataVertexModel<float>(OutParamNameValue, LOCTEXT("ValueTooltip", "The current randomly generated float value."))
-			)
-		);
-
-		return Interface;
-	}
-
-	const FNodeClassMetadata& FRandomFloatOperator::GetNodeInfo()
-	{
-		auto InitNodeInfo = []() -> FNodeClassMetadata
+		static FText GetResetDescription()
 		{
-			FNodeClassMetadata Info;
-			Info.ClassName = { Metasound::StandardNodes::Namespace, TEXT("Random"), TEXT("Float") };
-			Info.MajorVersion = 1;
-			Info.MinorVersion = 0;
-			Info.DisplayName = LOCTEXT("Metasound_RandomFloatDisplayName", "Random Float");
-			Info.Description = LOCTEXT("Metasound_RandomFloatNodeDescription", "Generates a seedable random float in the given value range.");
-			Info.Author = PluginAuthor;
-			Info.PromptIfMissing = PluginNodeMissingPrompt;
-			Info.DefaultInterface = GetVertexInterface();
-
-			return Info;
-		};
+			return LOCTEXT("RandomNodeResetTT", "Trigger to reset the random sequence with the supplied seed. Useful to get randomized repetition.");
+		}
 
-		static const FNodeClassMetadata Info = InitNodeInfo();
+		static FText GetSeedDescription()
+		{
+			return LOCTEXT("RandomNodeSeedTT", "The seed value to use for the random node. Set to -1 to use a random seed.");
+		}
 
-		return Info;
-	}
+		static FText GetMinDescription()
+		{
+			return LOCTEXT("RandomNodeMinTT", "Min random value.");
+		}
 
-	TUniquePtr<IOperator> FRandomFloatOperator::CreateOperator(const FCreateOperatorParams& InParams, FBuildErrorArray& OutErrors)
-	{
-		const FRandomFloatNode& RandomIntNode = static_cast<const FRandomFloatNode&>(InParams.Node);
-		const FDataReferenceCollection& InputCollection = InParams.InputDataReferences;
-		const FInputVertexInterface& InputInterface = GetVertexInterface().GetInputInterface();
+		static FText GetMaxDescription()
+		{
+			return LOCTEXT("RandNodeMaxTT", "Max random value.");
+		}
 
-		FInt32ReadRef SeedValue = InputCollection.GetDataReadReferenceOrConstructWithVertexDefault<int32>(InputInterface, InParamNameSeed);
-		FFloatReadRef MinValue = InputCollection.GetDataReadReferenceOrConstructWithVertexDefault<float>(InputInterface, InParamNameMin);
-		FFloatReadRef MaxValue = InputCollection.GetDataReadReferenceOrConstructWithVertexDefault<float>(InputInterface, InParamNameMax);
-		FTriggerReadRef TriggerNext = InputCollection.GetDataReadReferenceOrConstruct<FTrigger>(InParamNameNextTrigger, InParams.OperatorSettings);
-		FTriggerReadRef TriggerReset = InputCollection.GetDataReadReferenceOrConstruct<FTrigger>(InParamNameResetTrigger, InParams.OperatorSettings);
+		static FText GetOutputDescription()
+		{
+			return LOCTEXT("RandomNodeOutputTT", "The randomly generated value.");
+		}
 
-		return MakeUnique<FRandomFloatOperator>(InParams.OperatorSettings, SeedValue, MinValue, MaxValue, TriggerNext, TriggerReset);
-	}
+		static FText GetOutputOnNextDescription()
+		{
+			return LOCTEXT("RandomNodeOutputNextTT", "Triggers when next is triggered.");
+		}
 
-	FRandomFloatNode::FRandomFloatNode(const FNodeInitData& InitData)
-		: FNodeFacade(InitData.InstanceName, InitData.InstanceID, TFacadeOperatorClass<FRandomFloatOperator>())
-	{
+		static FText GetOutputOnResetDescription()
+		{
+			return LOCTEXT("RandomNodeOutputNextTT", "Triggers when reset is triggered.");
+		}
 	}
 
-	METASOUND_REGISTER_NODE(FRandomFloatNode)
+ 	using FRandomNodeInt32 = TRandomNode<int32>;
+ 	METASOUND_REGISTER_NODE(FRandomNodeInt32)
+ 
+ 	using FRandomNodeFloat = TRandomNode<float>;
+ 	METASOUND_REGISTER_NODE(FRandomNodeFloat)
+ 
+	using FRandomNodeBool = TRandomNode<bool>;
+	METASOUND_REGISTER_NODE(FRandomNodeBool)
 }
 
 #undef LOCTEXT_NAMESPACE
diff --git a/Engine/Plugins/Runtime/Metasound/Source/MetasoundStandardNodes/Private/MetasoundStandardNodesCategories.cpp b/Engine/Plugins/Runtime/Metasound/Source/MetasoundStandardNodes/Private/MetasoundStandardNodesCategories.cpp
--- a/Engine/Plugins/Runtime/Metasound/Source/MetasoundStandardNodes/Private/MetasoundStandardNodesCategories.cpp
+++ b/Engine/Plugins/Runtime/Metasound/Source/MetasoundStandardNodes/Private/MetasoundStandardNodesCategories.cpp
@@ -10,7 +10,8 @@ namespace Metasound
 	{
 		const FText Generators = { LOCTEXT("Metasound_GeneratorsCategory", "Generators") };
 		const FText Filters = { LOCTEXT("Metasound_FiltersCategory", "Filters") };
-		const FText TriggerUtils = { LOCTEXT("Metasound_TriggerUtilsCategory", "Trigger Utility") };
+		const FText TriggerUtils = { LOCTEXT("Metasound_TriggerUtilsCategory", "Trigger Utilities") };
+		const FText RandomUtils = { LOCTEXT("Metasound_RandomUtilsCategory", "Random Utilities") };
 	}
 }
 
diff --git a/Engine/Plugins/Runtime/Metasound/Source/MetasoundStandardNodes/Public/MetasoundRandomNode.h b/Engine/Plugins/Runtime/Metasound/Source/MetasoundStandardNodes/Public/MetasoundRandomNode.h
--- a/Engine/Plugins/Runtime/Metasound/Source/MetasoundStandardNodes/Public/MetasoundRandomNode.h
+++ b/Engine/Plugins/Runtime/Metasound/Source/MetasoundStandardNodes/Public/MetasoundRandomNode.h
@@ -4,35 +4,381 @@
 
 #include "CoreMinimal.h"
 #include "MetasoundFacade.h"
+#include "MetasoundExecutableOperator.h"
+#include "MetasoundPrimitives.h"
+#include "MetasoundStandardNodesNames.h"
+#include "MetasoundTrigger.h"
+#include "Internationalization/Text.h"
+#include "MetasoundStandardNodesCategories.h"
+
+#define LOCTEXT_NAMESPACE "MetasoundStandardNodes"
 
 namespace Metasound
 {
-	/** FRandomIntNode
-	 *
-	 *  Generates a random integer value when triggered.
-	 */
-	class METASOUNDSTANDARDNODES_API FRandomIntNode : public FNodeFacade
+	template<typename ValueType>
+	struct TRandomNodeSpecialization
+	{
+		bool bSupported = false;
+	};
+
+	template<>
+	struct TRandomNodeSpecialization<int32>
+	{
+		static FName GetClassName()
+		{
+			return FName("RandomInt32");
+		}
+
+		static FText GetDisplayName()
+		{
+			return LOCTEXT("Int32RandomValueDisplayName", "Random Int");
+		}
+
+		static FText GetDescription()
+		{
+			return LOCTEXT("Int32RandomDescription", "Generates a seedable random integer in the given value range.");
+		}
+
+		static bool HasRange()
+		{
+			return true;
+		}
+
+		static void GetDefaultRange(int32& OutMin, int32& OutMax)
+		{
+			OutMin = 0;
+			OutMax = 100;
+		}
+
+		static int32 GetNextValue(FRandomStream& InStream, int32 InMin, int32 InMax)
+		{
+			return InStream.RandRange(InMin, InMax);
+		}
+	};
+
+	template<>
+	struct TRandomNodeSpecialization<float>
+	{
+		static FName GetClassName()
+		{
+			return FName("RandomFloat");
+		}
+
+		static FText GetDisplayName()
+		{
+			return LOCTEXT("FloatRandomValueDisplayName", "Random Float");
+		}
+
+		static FText GetDescription()
+		{
+			return LOCTEXT("FloatRandomDescription", "Generates a seedable random float in the given value range.");
+		}
+
+		static bool HasRange()
+		{
+			return true;
+		}
+
+		static void GetDefaultRange(float& OutMin, float& OutMax)
+		{
+			OutMin = 0.0f;
+			OutMax = 1.0f;
+		}
+
+		static float GetNextValue(FRandomStream& InStream, float InMin, float InMax)
+		{
+			return InStream.FRandRange(InMin, InMax);
+		}
+	};
+
+	template<>
+	struct TRandomNodeSpecialization<bool>
+	{
+		static FName GetClassName()
+		{
+			return FName("RandomBool");
+		}
+
+		static FText GetDisplayName()
+		{
+			return LOCTEXT("BoolRandomValueDisplayName", "Random Bool");
+		}
+
+		static FText GetDescription()
+		{
+			return LOCTEXT("BoolRandomDiscription", "Generates a random bool value.");
+		}
+
+		static bool HasRange()
+		{
+			return false;
+		}
+
+		static void GetDefaultRange(bool& OutMin, bool& OutMax)
+		{
+		}
+
+		static bool GetNextValue(FRandomStream& InStream, bool, bool)
+		{
+			return (bool)InStream.RandRange(0, 1);
+		}
+	};
+
+	namespace RandomNodeNames
+	{
+		/** Input vertex names. */
+		const FString& GetInputNextTriggerName();
+		const FString& GetInputResetTriggerName();
+		const FString& GetInputSeedName();
+		const FString& GetInputMinName();
+		const FString& GetInputMaxName();
+
+		/** Output vertex names. */
+		const FString& GetOutputOnNextTriggerName();
+		const FString& GetOutputOnResetTriggerName();
+		const FString& GetOutputValueName();
+
+		/** Descriptions. */
+		static FText GetNextTriggerDescription();
+		static FText GetResetDescription();
+		static FText GetSeedDescription();
+		static FText GetMinDescription();
+		static FText GetMaxDescription();
+		static FText GetOutputDescription();
+		static FText GetOutputOnNextDescription();
+		static FText GetOutputOnResetDescription();
+	}
+	
+	template<typename ValueType>
+	class TRandomNodeOperator : public TExecutableOperator<TRandomNodeOperator<ValueType>>
 	{
 	public:
-		/**
-		 * Constructor used by the Metasound Frontend.
-		 */
-		FRandomIntNode(const FNodeInitData& InInitData);
+		static constexpr int32 DefaultSeed = -1;
+
+		static const FVertexInterface& GetDefaultInterface()
+		{
+			using namespace RandomNodeNames;
+
+			if (TRandomNodeSpecialization<ValueType>::HasRange())
+			{
+				ValueType MinValue;
+				ValueType MaxValue;
+				TRandomNodeSpecialization<ValueType>::GetDefaultRange(MinValue, MaxValue);
+
+				static const FVertexInterface DefaultInterface(
+					FInputVertexInterface(
+						TInputDataVertexModel<FTrigger>(GetInputNextTriggerName(), GetNextTriggerDescription()),
+						TInputDataVertexModel<FTrigger>(GetInputResetTriggerName(), GetResetDescription()),
+						TInputDataVertexModel<int32>(GetInputSeedName(), GetSeedDescription(), DefaultSeed),
+						TInputDataVertexModel<ValueType>(GetInputMinName(), GetMinDescription(), MinValue),
+						TInputDataVertexModel<ValueType>(GetInputMaxName(), GetMaxDescription(), MaxValue)
+					),
+					FOutputVertexInterface(
+						TOutputDataVertexModel<FTrigger>(GetOutputOnNextTriggerName(), GetOutputOnNextDescription()),
+						TOutputDataVertexModel<FTrigger>(GetOutputOnResetTriggerName(), GetOutputOnResetDescription()),
+						TOutputDataVertexModel<ValueType>(GetOutputValueName(), GetOutputDescription())
+					)
+				);
+
+				return DefaultInterface;
+			}
+			else
+			{
+				static const FVertexInterface DefaultInterface(
+					FInputVertexInterface(
+						TInputDataVertexModel<FTrigger>(GetInputNextTriggerName(), GetNextTriggerDescription()),
+						TInputDataVertexModel<FTrigger>(GetInputResetTriggerName(), GetResetDescription()),
+						TInputDataVertexModel<int32>(GetInputSeedName(), GetSeedDescription(), DefaultSeed)
+					),
+					FOutputVertexInterface(
+						TOutputDataVertexModel<FTrigger>(GetOutputOnNextTriggerName(), GetOutputOnNextDescription()),
+						TOutputDataVertexModel<FTrigger>(GetOutputOnResetTriggerName(), GetOutputOnResetDescription()),
+						TOutputDataVertexModel<ValueType>(GetOutputValueName(), GetOutputDescription())
+					)
+				);
+
+				return DefaultInterface;
+			}
+		}
+
+		static const FNodeClassMetadata& GetNodeInfo()
+		{
+			auto InitNodeInfo = []() -> FNodeClassMetadata
+			{
+				FNodeClassMetadata Info;
+				Info.ClassName = { Metasound::StandardNodes::Namespace, TRandomNodeSpecialization<ValueType>::GetClassName(), TEXT("") };
+				Info.MajorVersion = 1;
+				Info.MinorVersion = 0;
+				Info.DisplayName = TRandomNodeSpecialization<ValueType>::GetDisplayName();
+				Info.Description = LOCTEXT("RandomOpDescription", "Generates a Random value.");
+				Info.Author = PluginAuthor;
+				Info.PromptIfMissing = PluginNodeMissingPrompt;
+				Info.DefaultInterface = GetDefaultInterface();
+				Info.CategoryHierarchy.Emplace(StandardNodes::RandomUtils);
+
+				return Info;
+			};
+
+			static const FNodeClassMetadata Info = InitNodeInfo();
+			return Info;
+		}
+
+		static TUniquePtr<IOperator> CreateOperator(const FCreateOperatorParams& InParams, TArray<TUniquePtr<IOperatorBuildError>>& OutErrors)
+		{
+ 			const FInputVertexInterface& InputInterface = InParams.Node.GetVertexInterface().GetInputInterface();
+ 			const FDataReferenceCollection& InputCollection = InParams.InputDataReferences;
+			using namespace RandomNodeNames;
+
+ 			FTriggerReadRef InNextTrigger = InputCollection.GetDataReadReferenceOrConstruct<FTrigger>(GetInputNextTriggerName(), InParams.OperatorSettings);
+			FTriggerReadRef InResetTrigger = InputCollection.GetDataReadReferenceOrConstruct<FTrigger>(GetInputResetTriggerName(), InParams.OperatorSettings);
+			FInt32ReadRef InSeedValue = InputCollection.GetDataReadReferenceOrConstructWithVertexDefault<int32>(InputInterface, GetInputSeedName());
+
+			// note: this is ignored by random bool
+			TDataReadReference<ValueType> InMinValue = InputCollection.GetDataReadReferenceOrConstructWithVertexDefault<ValueType>(InputInterface, GetInputMinName());
+			TDataReadReference<ValueType> InMaxValue = InputCollection.GetDataReadReferenceOrConstructWithVertexDefault<ValueType>(InputInterface, GetInputMaxName());
+
+			return MakeUnique<TRandomNodeOperator<ValueType>>(InParams.OperatorSettings, InNextTrigger, InResetTrigger, InSeedValue, InMinValue, InMaxValue);
+		}
+
+
+		TRandomNodeOperator(const FOperatorSettings& InSettings, 
+			TDataReadReference<FTrigger> InNextTrigger, 
+			TDataReadReference<FTrigger> InResetTrigger, 
+			TDataReadReference<int32> InSeedValue, 
+			TDataReadReference<ValueType> InMinValue,
+			TDataReadReference<ValueType> InMaxValue)
+ 			: NextTrigger(InNextTrigger)
+			, ResetTrigger(InResetTrigger)
+			, SeedValue(InSeedValue)
+ 			, MinValue(InMinValue)
+ 			, MaxValue(InMaxValue)
+			, TriggerOutOnNext(FTriggerWriteRef::CreateNew(InSettings))
+			, TriggerOutOnReset(FTriggerWriteRef::CreateNew(InSettings))
+ 			, OutputValue(TDataWriteReferenceFactory<ValueType>::CreateAny(InSettings))
+			, bIsDefaultSeeded(*SeedValue == DefaultSeed)
+			, bIsRandomStreamInitialized(false)
+		{
+			// We need to initialize the output value to *something*
+			*OutputValue = *MinValue;
+		}
+
+		virtual ~TRandomNodeOperator() = default;
+
+
+		virtual FDataReferenceCollection GetInputs() const override
+		{
+			using namespace RandomNodeNames;
+
+			FDataReferenceCollection Inputs;
+			Inputs.AddDataReadReference(GetInputNextTriggerName(), NextTrigger);
+ 			Inputs.AddDataReadReference(GetInputResetTriggerName(), ResetTrigger);
+ 			Inputs.AddDataReadReference(GetInputSeedName(), SeedValue);
+
+			// If the type doesn't have a range no need for input pins to define it
+			if (TRandomNodeSpecialization<ValueType>::HasRange())
+			{
+				Inputs.AddDataReadReference(GetInputMinName(), MinValue);
+				Inputs.AddDataReadReference(GetInputMaxName(), MinValue);
+			}
+			return Inputs;
+		}
+
+		virtual FDataReferenceCollection GetOutputs() const override
+		{
+			using namespace RandomNodeNames;
+
+			FDataReferenceCollection Outputs;
+			Outputs.AddDataReadReference(GetOutputOnNextTriggerName(), TriggerOutOnNext);
+			Outputs.AddDataReadReference(GetOutputOnResetTriggerName(), TriggerOutOnReset);
+			Outputs.AddDataReadReference(GetOutputValueName(), OutputValue);
+			return Outputs;
+		}
+
+		void Execute()
+		{
+			TriggerOutOnReset->AdvanceBlock();
+			TriggerOutOnNext->AdvanceBlock();
+
+			ResetTrigger->ExecuteBlock(
+				[&](int32 StartFrame, int32 EndFrame)
+				{
+				},
+				[this](int32 StartFrame, int32 EndFrame)
+				{
+					EvaluateSeedChanges();
+					RandomStream.Reset();
+					TriggerOutOnReset->TriggerFrame(StartFrame);
+				}
+			);
+
+			NextTrigger->ExecuteBlock(
+				[&](int32 StartFrame, int32 EndFrame)
+				{
+				},
+				[this](int32 StartFrame, int32 EndFrame)
+				{
+					*OutputValue = TRandomNodeSpecialization<ValueType>::GetNextValue(RandomStream, *MinValue, *MaxValue);
+					TriggerOutOnNext->TriggerFrame(StartFrame);
+				}
+			);
+		}
+
+	private:
+		void EvaluateSeedChanges()
+		{
+			// if we have a non-zero seed
+			if (*SeedValue != DefaultSeed)
+			{
+				// If we were previously zero-seeded OR our seed has changed
+				if (bIsDefaultSeeded || !bIsRandomStreamInitialized || *SeedValue != RandomStream.GetInitialSeed())
+				{
+					bIsRandomStreamInitialized = true;
+					bIsDefaultSeeded = false;
+					RandomStream.Initialize(*SeedValue);
+				}
+			}
+			// If we are zero-seeded now BUT were previously not, we need to randomize our seed
+			else if (!bIsDefaultSeeded || !bIsRandomStreamInitialized)
+			{
+				bIsRandomStreamInitialized = true;
+				bIsDefaultSeeded = true;
+				RandomStream.Initialize(FPlatformTime::Cycles());
+			}
+		}
+
+
+		FTriggerReadRef NextTrigger;
+		FTriggerReadRef ResetTrigger;
+		FInt32ReadRef SeedValue;
+		TDataReadReference<ValueType> MinValue;
+		TDataReadReference<ValueType> MaxValue;
+
+		FTriggerWriteRef TriggerOutOnNext;
+		FTriggerWriteRef TriggerOutOnReset;
+		TDataWriteReference<ValueType> OutputValue;
+
+		FRandomStream RandomStream;
+		bool bIsDefaultSeeded = false;
+		bool bIsRandomStreamInitialized = false;
 	};
 
-	/** FRandomFloatNode
+	/** TRandomNode
 	 *
 	 *  Generates a random float value when triggered.
 	 */
-	class METASOUNDSTANDARDNODES_API FRandomFloatNode : public FNodeFacade
+	template<typename ValueType>
+	class METASOUNDSTANDARDNODES_API TRandomNode : public FNodeFacade
 	{
 	public:
 		/**
 		 * Constructor used by the Metasound Frontend.
 		 */
-		FRandomFloatNode(const FNodeInitData& InInitData);
-	};
+		TRandomNode(const FNodeInitData& InInitData)
+			: FNodeFacade(InInitData.InstanceName, InInitData.InstanceID, TFacadeOperatorClass<TRandomNodeOperator<ValueType>>())
+		{}
 
+		virtual ~TRandomNode() = default;
+	};
 
 } // namespace Metasound
 
diff --git a/Engine/Plugins/Runtime/Metasound/Source/MetasoundStandardNodes/Public/MetasoundStandardNodesCategories.h b/Engine/Plugins/Runtime/Metasound/Source/MetasoundStandardNodes/Public/MetasoundStandardNodesCategories.h
--- a/Engine/Plugins/Runtime/Metasound/Source/MetasoundStandardNodes/Public/MetasoundStandardNodesCategories.h
+++ b/Engine/Plugins/Runtime/Metasound/Source/MetasoundStandardNodes/Public/MetasoundStandardNodesCategories.h
@@ -11,5 +11,6 @@ namespace Metasound
 		extern const FText Generators;
 		extern const FText Filters;
 		extern const FText TriggerUtils;
+		extern const FText RandomUtils;
 	}
 }
