metasound Trigger toggle node

diff --git a/Engine/Plugins/Runtime/Metasound/Source/MetasoundStandardNodes/Private/MetasoundTriggerToggleNode.cpp b/Engine/Plugins/Runtime/Metasound/Source/MetasoundStandardNodes/Private/MetasoundTriggerToggleNode.cpp
new file mode 100644
--- /dev/null
+++ b/Engine/Plugins/Runtime/Metasound/Source/MetasoundStandardNodes/Private/MetasoundTriggerToggleNode.cpp
@@ -0,0 +1,173 @@
+// Copyright Epic Games, Inc. All Rights Reserved.
+#include "MetasoundTriggerToggleNode.h"
+
+#include "CoreMinimal.h"
+
+#include "MetasoundBuilderInterface.h"
+#include "MetasoundDataReferenceCollection.h"
+#include "MetasoundExecutableOperator.h"
+#include "MetasoundFacade.h"
+#include "MetasoundNode.h"
+#include "MetasoundNodeInterface.h"
+#include "MetasoundNodeRegistrationMacro.h"
+#include "MetasoundOperatorInterface.h"
+#include "MetasoundPrimitives.h"
+#include "MetasoundStandardNodesNames.h"
+#include "MetasoundStandardNodesCategories.h"
+
+#define LOCTEXT_NAMESPACE "MetasoundStandardNodes"
+
+namespace Metasound
+{
+	namespace TriggerToggle
+	{
+		const FString GetInputOnTriggerName() { return TEXT("On"); }
+		const FString GetInputOffTriggerName() { return TEXT("Off"); }
+		const FString GetInputInitName() { return TEXT("Init"); }
+		const FString GetOutputTriggerName() { return TEXT("Out"); }
+		const FString GetOutputValueName() { return TEXT("Value"); }
+	}
+
+	class FTriggerToggleOperator : public TExecutableOperator<FTriggerToggleOperator>
+	{
+		public:
+			static const FNodeClassMetadata& GetNodeInfo();
+			static const FVertexInterface& GetVertexInterface();
+			static TUniquePtr<IOperator> CreateOperator(const FCreateOperatorParams& InParams, FBuildErrorArray& OutErrors);
+
+			FTriggerToggleOperator(const FOperatorSettings& InSettings, const FTriggerReadRef& InTriggerOn, const FTriggerReadRef& InTriggerOff, const FBoolReadRef& InInitValue);
+
+			virtual FDataReferenceCollection GetInputs() const override;
+			virtual FDataReferenceCollection GetOutputs() const override;
+
+			void Execute();
+
+		private:
+
+			FTriggerReadRef TriggerOn;
+			FTriggerReadRef TriggerOff;
+			FBoolReadRef InitValue;
+			FTriggerWriteRef TriggerOutput;
+			FBoolWriteRef ValueOutput;
+	};
+
+	FTriggerToggleOperator::FTriggerToggleOperator(const FOperatorSettings& InSettings, const FTriggerReadRef& InTriggerOn, const FTriggerReadRef& InTriggerOff, const FBoolReadRef& InInitValue)
+	: TriggerOn(InTriggerOn)
+	, TriggerOff(InTriggerOff)
+	, InitValue(InInitValue)
+	, TriggerOutput(FTriggerWriteRef::CreateNew(InSettings))
+	, ValueOutput(FBoolWriteRef::CreateNew(*InInitValue))
+	{
+	}
+
+	FDataReferenceCollection FTriggerToggleOperator::GetInputs() const
+	{
+		using namespace TriggerToggle;
+
+		FDataReferenceCollection InputDataReferences;
+		InputDataReferences.AddDataReadReference(GetInputOnTriggerName(), TriggerOn);
+		InputDataReferences.AddDataReadReference(GetInputOffTriggerName(), TriggerOff);
+		InputDataReferences.AddDataReadReference(GetInputInitName(), InitValue);
+		return InputDataReferences;
+	}
+
+	FDataReferenceCollection FTriggerToggleOperator::GetOutputs() const
+	{
+		using namespace TriggerToggle;
+
+		FDataReferenceCollection OutputDataReferences;
+		OutputDataReferences.AddDataReadReference(GetOutputTriggerName(), TriggerOutput);
+		OutputDataReferences.AddDataReadReference(GetOutputValueName(), ValueOutput);
+
+		return OutputDataReferences;
+	}
+
+	void FTriggerToggleOperator::Execute()
+	{
+		TriggerOutput->AdvanceBlock();
+
+		TriggerOn->ExecuteBlock(
+			[&](int32 StartFrame, int32 EndFrame)
+			{
+			},
+			[this](int32 StartFrame, int32 EndFrame)
+			{
+				*ValueOutput = true;
+				TriggerOutput->TriggerFrame(StartFrame);
+			}
+		);
+
+		TriggerOff->ExecuteBlock(
+			[&](int32 StartFrame, int32 EndFrame)
+			{
+			},
+			[this](int32 StartFrame, int32 EndFrame)
+			{
+				*ValueOutput = false;
+				TriggerOutput->TriggerFrame(StartFrame);
+			}
+		);
+	}
+
+	TUniquePtr<IOperator> FTriggerToggleOperator::CreateOperator(const FCreateOperatorParams& InParams, FBuildErrorArray& OutErrors)
+	{
+		using namespace TriggerToggle;
+
+		const FInputVertexInterface& InputInterface = GetVertexInterface().GetInputInterface();
+		FTriggerReadRef InTriggerOn = InParams.InputDataReferences.GetDataReadReferenceOrConstruct<FTrigger>(GetInputOnTriggerName(), InParams.OperatorSettings);
+		FTriggerReadRef InTriggerOff = InParams.InputDataReferences.GetDataReadReferenceOrConstruct<FTrigger>(GetInputOffTriggerName(), InParams.OperatorSettings);
+		FBoolReadRef InInitValue = InParams.InputDataReferences.GetDataReadReferenceOrConstructWithVertexDefault<bool>(InputInterface, GetInputInitName());
+
+		return MakeUnique<FTriggerToggleOperator>(InParams.OperatorSettings, InTriggerOn, InTriggerOff, InInitValue);
+	}
+
+	const FVertexInterface& FTriggerToggleOperator::GetVertexInterface()
+	{
+		using namespace TriggerToggle;
+
+		static const FVertexInterface Interface(
+			FInputVertexInterface(
+				TInputDataVertexModel<FTrigger>(GetInputOnTriggerName(), LOCTEXT("TriggerGateOnTT", "Trigger to toggle gate output to 1.")),
+				TInputDataVertexModel<FTrigger>(GetInputOffTriggerName(), LOCTEXT("TriggerGateOffTT", "Trigger to toggle gate output to 0.")),
+				TInputDataVertexModel<bool>(GetInputInitName(), LOCTEXT("TriggerGateInitTT", "Initial value of the output gate."))
+			),
+			FOutputVertexInterface(
+				TOutputDataVertexModel<FTrigger>(GetOutputTriggerName(), LOCTEXT("TriggerGateOutputTriggerTT", "Triggers output when gate is toggled.")),
+				TOutputDataVertexModel<bool>(GetOutputValueName(), LOCTEXT("TriggerGateOutputTriggerTT", "Current output value of the toggle."))
+			)
+		);
+
+		return Interface;
+	}
+
+	const FNodeClassMetadata& FTriggerToggleOperator::GetNodeInfo()
+	{
+		auto InitNodeInfo = []() -> FNodeClassMetadata
+		{
+			FNodeClassMetadata Info;
+			Info.ClassName = {Metasound::StandardNodes::Namespace, TEXT("Trigger Toggle"), TEXT("")};
+			Info.MajorVersion = 1;
+			Info.MinorVersion = 0;
+			Info.DisplayName = LOCTEXT("Metasound_TriggerToggleNodeDisplayName", "Trigger Toggle");
+			Info.Description = LOCTEXT("Metasound_TriggerToggleNodeDescription", "Toggles a boolean value on or off.");
+			Info.Author = PluginAuthor;
+			Info.PromptIfMissing = PluginNodeMissingPrompt;
+			Info.DefaultInterface = GetVertexInterface();
+			Info.CategoryHierarchy.Emplace(StandardNodes::TriggerUtils);
+
+			return Info;
+		};
+
+		static const FNodeClassMetadata Info = InitNodeInfo();
+		return Info;
+	}
+
+	FTriggerToggleNode::FTriggerToggleNode(const FNodeInitData& InInitData)
+	:	FNodeFacade(InInitData.InstanceName, InInitData.InstanceID, TFacadeOperatorClass<FTriggerToggleOperator>())
+	{
+	}
+
+	METASOUND_REGISTER_NODE(FTriggerToggleNode)
+}
+
+#undef LOCTEXT_NAMESPACE // MetasoundTriggerDelayNode
diff --git a/Engine/Plugins/Runtime/Metasound/Source/MetasoundStandardNodes/Public/MetasoundTriggerToggleNode.h b/Engine/Plugins/Runtime/Metasound/Source/MetasoundStandardNodes/Public/MetasoundTriggerToggleNode.h
new file mode 100644
--- /dev/null
+++ b/Engine/Plugins/Runtime/Metasound/Source/MetasoundStandardNodes/Public/MetasoundTriggerToggleNode.h
@@ -0,0 +1,26 @@
+// Copyright Epic Games, Inc. All Rights Reserved.
+
+#pragma once
+
+#include "CoreMinimal.h"
+#include "MetasoundBuilderInterface.h"
+#include "MetasoundDataReferenceCollection.h"
+#include "MetasoundExecutableOperator.h"
+#include "MetasoundFacade.h"
+#include "MetasoundNode.h"
+#include "MetasoundNodeInterface.h"
+#include "MetasoundOperatorInterface.h"
+#include "MetasoundTime.h"
+#include "MetasoundTrigger.h"
+
+
+namespace Metasound
+{
+	class METASOUNDSTANDARDNODES_API FTriggerToggleNode : public FNodeFacade
+	{
+		public:
+			FTriggerToggleNode(const FNodeInitData& InInitData);
+
+			virtual ~FTriggerToggleNode() = default;
+	};
+}
