Metasound shuffle node

diff --git a/Engine/Plugins/Runtime/Metasound/Source/MetasoundFrontend/Private/MetasoundArrayNodes.cpp b/Engine/Plugins/Runtime/Metasound/Source/MetasoundFrontend/Private/MetasoundArrayNodes.cpp
--- a/Engine/Plugins/Runtime/Metasound/Source/MetasoundFrontend/Private/MetasoundArrayNodes.cpp
+++ b/Engine/Plugins/Runtime/Metasound/Source/MetasoundFrontend/Private/MetasoundArrayNodes.cpp
@@ -99,6 +99,233 @@ namespace Metasound
 			return Name;
 		}
 	}
+
+	namespace ArrayNodeShuffleVertexNames
+	{
+		const FString& GetInputTriggerNextName()
+		{
+			static const FString Name = TEXT("Next");
+			return Name;
+		}
+
+		const FString& GetInputTriggerShuffleName()
+		{
+			static const FString Name = TEXT("Shuffle");
+			return Name;
+		}
+
+		const FString& GetInputTriggerResetName()
+		{
+			static const FString Name = TEXT("Reset Seed");
+			return Name;
+		}
+
+		const FString& GetInputShuffleArrayName()
+		{
+			static const FString Name = TEXT("In Array");
+			return Name;
+		}
+
+		const FString& GetInputSeedName()
+		{
+			static const FString Name = TEXT("Seed");
+			return Name;
+		}
+
+		const FString& GetInputAutoShuffleName()
+		{
+			static const FString Name = TEXT("Auto Shuffle");
+			return Name;
+		}
+
+		const FString& GetInputNamespaceName()
+		{
+			static const FString Name = TEXT("Shuffle Namespace");
+			return Name;
+		}
+
+		const FString& GetOutputTriggerOnNextName()
+		{
+			static const FString Name = TEXT("On Next");
+			return Name;
+		}
+
+		const FString& GetOutputTriggerOnShuffleName()
+		{
+			static const FString Name = TEXT("On Shuffle");
+			return Name;
+		}
+
+		const FString& GetOutputTriggerOnResetName()
+		{
+			static const FString Name = TEXT("On Reset Seed");
+			return Name;
+		}
+
+		const FString& GetOutputValueName()
+		{
+			static const FString Name = TEXT("Value");
+			return Name;
+		}
+	}
+
+	uint32 GetTypeHash(const FGlobalArrayShuffleKey& InKey)
+	{
+		return InKey.Hash;
+	}
+
+	/**
+	* FArrayIndexShuffler
+	*/
+
+	FArrayIndexShuffler::FArrayIndexShuffler(int32 InSeed, int32 InMaxIndicies)
+	{
+		Init(InSeed, InMaxIndicies);
+	}
+
+	void FArrayIndexShuffler::Init(int32 InSeed, int32 InMaxIndicies)
+	{
+		SetSeed(InSeed);
+		if (InMaxIndicies > 0)
+		{
+			ShuffleIndices.AddUninitialized(InMaxIndicies);
+			for (int32 i = 0; i < ShuffleIndices.Num(); ++i)
+			{
+				ShuffleIndices[i] = i;
+			}
+			ShuffleArray();
+		}
+	}
+
+	void FArrayIndexShuffler::SetSeed(int32 InSeed)
+	{
+		if (InSeed == INDEX_NONE)
+		{
+			RandomStream.Initialize(FPlatformTime::Cycles());
+		}
+		else
+		{
+			RandomStream.Initialize(InSeed);
+		}
+
+		ResetSeed();
+	}
+
+	void FArrayIndexShuffler::ResetSeed()
+	{
+		RandomStream.Reset();
+	}
+
+	bool FArrayIndexShuffler::NextValue(bool bAutoShuffle, int32& OutIndex)
+	{
+		bool bShuffled = false;
+		if (CurrentIndex == ShuffleIndices.Num())
+		{
+			if (bAutoShuffle)
+			{
+				ShuffleArray();
+				bShuffled = true;
+			}
+			else
+			{
+				CurrentIndex = 0;
+			}
+		}
+
+		check(CurrentIndex < ShuffleIndices.Num());
+		PrevValue = ShuffleIndices[CurrentIndex];
+		OutIndex = PrevValue;
+		++CurrentIndex;
+
+		return bShuffled;
+	}
+
+	// Shuffle the array with the given max indices
+	void FArrayIndexShuffler::ShuffleArray()
+	{
+		// Randomize the shuffled array by randomly swapping indicies
+		for (int32 i = 0; i < ShuffleIndices.Num(); ++i)
+		{
+			RandomSwap(i, 0, ShuffleIndices.Num() - 1);
+		}
+
+		// Reset the current index back to 0
+		CurrentIndex = 0;
+
+		// Fix up the new current index if the value is our previous value and we have an array larger than 1
+		if (ShuffleIndices.Num() > 1 && ShuffleIndices[CurrentIndex] == PrevValue)
+		{
+			RandomSwap(0, 1, ShuffleIndices.Num() - 1);
+		}
+	}
+
+	void FArrayIndexShuffler::RandomSwap(int32 InCurrentIndex, int32 InStartIndex, int32 InEndIndex)
+	{
+		int32 ShuffleIndex = RandomStream.RandRange(InStartIndex, InEndIndex);
+		int32 Temp = ShuffleIndices[ShuffleIndex];
+		ShuffleIndices[ShuffleIndex] = ShuffleIndices[InCurrentIndex];
+		ShuffleIndices[InCurrentIndex] = Temp;
+	}
+
+	/** 
+	* FGlobalArrayShuffleManager
+	*/
+
+	FGlobalArrayShuffleManager& FGlobalArrayShuffleManager::Get()
+	{
+		static FGlobalArrayShuffleManager GShufflerManager;
+		return GShufflerManager;
+	}
+
+	bool FGlobalArrayShuffleManager::NextValue(FGlobalArrayShuffleKey& InKey, bool bAutoShuffle, int32& OutIndex)
+	{
+		check(!InKey.Namespace.IsEmpty());
+
+		FScopeLock Lock(&CritSect);
+
+		TUniquePtr<FArrayIndexShuffler>& ShufflerPtr = GetShuffler(InKey);
+		return ShufflerPtr->NextValue(bAutoShuffle, OutIndex);
+	}
+
+	void FGlobalArrayShuffleManager::SetSeed(FGlobalArrayShuffleKey& InKey, int32 InSeed)
+	{
+		check(!InKey.Namespace.IsEmpty());
+
+		FScopeLock Lock(&CritSect);
+
+		TUniquePtr<FArrayIndexShuffler>& ShufflerPtr = GetShuffler(InKey);
+		ShufflerPtr->SetSeed(InSeed);
+	}
+
+	void FGlobalArrayShuffleManager::ResetSeed(FGlobalArrayShuffleKey& InKey)
+	{
+		check(!InKey.Namespace.IsEmpty());
+
+		FScopeLock Lock(&CritSect);
+
+		TUniquePtr<FArrayIndexShuffler>& ShufflerPtr = GetShuffler(InKey);
+		ShufflerPtr->ResetSeed();
+	}
+
+	void FGlobalArrayShuffleManager::ShuffleArray(FGlobalArrayShuffleKey& InKey)
+	{
+		check(!InKey.Namespace.IsEmpty());
+
+		FScopeLock Lock(&CritSect);
+
+		TUniquePtr<FArrayIndexShuffler>& ShufflerPtr = GetShuffler(InKey);
+		ShufflerPtr->ShuffleArray();
+	}
+
+	TUniquePtr<FArrayIndexShuffler>& FGlobalArrayShuffleManager::GetShuffler(FGlobalArrayShuffleKey& InKey)
+	{
+		TUniquePtr<FArrayIndexShuffler>& ShufflerPtr = GlobalShuffleIndicies.FindOrAdd(InKey);
+		if (!ShufflerPtr.IsValid())
+		{
+			ShufflerPtr = MakeUnique<FArrayIndexShuffler>(InKey.Seed, InKey.NumElements);
+		}
+		return ShufflerPtr;
+	}
 }
 
 #undef LOCTEXT_NAMESPACE
diff --git a/Engine/Plugins/Runtime/Metasound/Source/MetasoundFrontend/Public/MetasoundArrayNodes.h b/Engine/Plugins/Runtime/Metasound/Source/MetasoundFrontend/Public/MetasoundArrayNodes.h
--- a/Engine/Plugins/Runtime/Metasound/Source/MetasoundFrontend/Public/MetasoundArrayNodes.h
+++ b/Engine/Plugins/Runtime/Metasound/Source/MetasoundFrontend/Public/MetasoundArrayNodes.h
@@ -742,6 +742,419 @@ namespace Metasound
 
 		virtual ~TArraySubsetNode() = default;
 	};
+
+	/** Shuffle Node Vertex Names */
+	namespace ArrayNodeShuffleVertexNames
+	{
+		/** Input Vertex Names */
+		METASOUNDFRONTEND_API const FString& GetInputTriggerNextName();
+		METASOUNDFRONTEND_API const FString& GetInputTriggerShuffleName();
+		METASOUNDFRONTEND_API const FString& GetInputTriggerResetName();
+		METASOUNDFRONTEND_API const FString& GetInputShuffleArrayName();
+		METASOUNDFRONTEND_API const FString& GetInputSeedName();
+		METASOUNDFRONTEND_API const FString& GetInputAutoShuffleName();
+		METASOUNDFRONTEND_API const FString& GetInputNamespaceName();
+
+		METASOUNDFRONTEND_API const FString& GetOutputTriggerOnNextName();
+		METASOUNDFRONTEND_API const FString& GetOutputTriggerOnShuffleName();
+		METASOUNDFRONTEND_API const FString& GetOutputTriggerOnResetName();
+		METASOUNDFRONTEND_API const FString& GetOutputValueName();
+	}
+
+	class METASOUNDFRONTEND_API FArrayIndexShuffler
+	{
+	public:
+		FArrayIndexShuffler() = default;
+		FArrayIndexShuffler(int32 InSeed, int32 MaxIndicies);
+
+		void Init(int32 InSeed, int32 MaxIndicies);
+		void SetSeed(int32 InSeed);
+		void ResetSeed();
+
+		// Returns the next value in the array indices. Returns true if the array was re-shuffled automatically.
+		bool NextValue(bool bAutoShuffle, int32& OutIndex);
+
+		// Shuffle the array with the given max indices
+		void ShuffleArray();
+
+	private:
+		// Helper function to swap the current index with a random index
+		void RandomSwap(int32 InCurrentIndex, int32 InStartIndex, int32 InEndIndex);
+
+		// The current index into the array of indicies (wraps between 0 and ShuffleIndices.Num())
+		int32 CurrentIndex = 0;
+
+		// The previously returned value. Used to avoid repeating the last value on shuffle.
+		int32 PrevValue = INDEX_NONE;
+
+		// Array of indices (in order 0 to Num), shuffled
+		TArray<int32> ShuffleIndices;
+
+		// Random stream to use to randomize the shuffling
+		FRandomStream RandomStream;
+	};
+
+	// Key used for global shuffle.
+	struct METASOUNDFRONTEND_API FGlobalArrayShuffleKey
+	{
+		FString Namespace;
+		int32 NumElements = 0;
+		int32 Seed = 0;
+		uint32 Hash = INDEX_NONE;
+
+		FGlobalArrayShuffleKey() = default;
+
+		FGlobalArrayShuffleKey(const FString& InNamespace, int32 InNumElements, int32 InSeed)
+			: Namespace(InNamespace)
+			, NumElements(InNumElements)
+			, Seed(InSeed)
+		{
+			// Hash on the input values
+			Hash = HashCombine(GetTypeHash(InNamespace), GetTypeHash(InNumElements));
+			Hash = HashCombine(Hash, GetTypeHash(Seed));
+		}
+
+		bool operator==(const FGlobalArrayShuffleKey& Other) const
+		{
+			return Other.Namespace == Namespace &&
+				Other.NumElements == NumElements &&
+				Other.Seed == Seed;
+		}
+	};
+
+	uint32 GetTypeHash(const FGlobalArrayShuffleKey& InKey);
+
+	class METASOUNDFRONTEND_API FGlobalArrayShuffleManager
+	{
+	public:
+		FGlobalArrayShuffleManager() = default;
+		~FGlobalArrayShuffleManager() = default;
+
+		// Retrieves the shuffle manager
+		static FGlobalArrayShuffleManager& Get();
+
+		// API that wraps the array shuffler
+		bool NextValue(FGlobalArrayShuffleKey& InKey, bool bAutoShuffle, int32& OutIndex);
+		void SetSeed(FGlobalArrayShuffleKey& InKey, int32 InSeed);
+		void ResetSeed(FGlobalArrayShuffleKey& InKey);
+		void ShuffleArray(FGlobalArrayShuffleKey& InKey);
+
+	private:
+
+		TUniquePtr<FArrayIndexShuffler>& GetShuffler(FGlobalArrayShuffleKey& InKey);
+
+		// Protects access to the shuffle namespace data
+		FCriticalSection CritSect;
+
+		// Global array map for namespaced array shufflers
+		TMap<FGlobalArrayShuffleKey, TUniquePtr<FArrayIndexShuffler>> GlobalShuffleIndicies;
+	};
+
+	/** TArrayShuffleOperator shuffles an array on trigger and outputs values sequentially on "next". It avoids repeating shuffled elements and supports auto-shuffling.*/
+	template<typename ArrayType>
+	class TArrayShuffleOperator : public TExecutableOperator<TArrayShuffleOperator<ArrayType>>
+	{
+	public:
+		using FArrayDataReadReference = TDataReadReference<ArrayType>;
+		using ElementType = typename MetasoundArrayNodesPrivate::TArrayElementType<ArrayType>::Type;
+		using FElementTypeWriteReference = TDataWriteReference<ElementType>;
+
+		static const FVertexInterface& GetDefaultInterface()
+		{
+			using namespace ArrayNodeShuffleVertexNames;
+
+			static const FVertexInterface DefaultInterface(
+				FInputVertexInterface(
+					TInputDataVertexModel<FTrigger>(GetInputTriggerNextName(), LOCTEXT("ShuffleOpInputTriggerNextNameTT", "Trigger to get the next value in the shuffled array.")),
+					TInputDataVertexModel<FTrigger>(GetInputTriggerShuffleName(), LOCTEXT("ShuffleOpInputTriggerShuffleNameTT", "Trigger to shuffle the array manually.")),
+					TInputDataVertexModel<FTrigger>(GetInputTriggerResetName(), LOCTEXT("ShuffleOpInputTriggerResetNameTT", "Trigger to reset the random seed stream of the shuffle node.")),
+					TInputDataVertexModel<ArrayType>(GetInputShuffleArrayName(), LOCTEXT("ShuffleOpInputShuffleArrayNameTT", "Input Array.")),
+					TInputDataVertexModel<int32>(GetInputSeedName(), LOCTEXT("ShuffleOpInputSeedNameTT", "Seed to use for the the random shuffle."), -1),
+					TInputDataVertexModel<bool>(GetInputAutoShuffleName(), LOCTEXT("ShuffleOpInputAutoShuffleNameTT", "Set to true to automatically shuffle when the array has been read."), true),
+					TInputDataVertexModel<FString>(GetInputNamespaceName(), LOCTEXT("ShuffleOpInputNamespaceNameTT", "Set to a string value to share shuffle namespace with this array. Shuffling state of array will be shared in this namespace."))
+				),
+				FOutputVertexInterface(
+					TOutputDataVertexModel<FTrigger>(GetOutputTriggerOnNextName(), LOCTEXT("ShuffleOpOutputTriggerOnNextNameTT", "Triggers when the \"Next\" input is triggered.")),
+					TOutputDataVertexModel<FTrigger>(GetOutputTriggerOnShuffleName(), LOCTEXT("ShuffleOpOutputTriggerOnShuffleNameTT", "Triggers when the \"Shuffle\" input is triggered or if the array is auto-shuffled.")),
+					TOutputDataVertexModel<FTrigger>(GetOutputTriggerOnResetName(), LOCTEXT("ShuffleOpOutputTriggerOnResetNameTT", "Triggers when the \"Reset Seed\" input is triggered.")),
+					TOutputDataVertexModel<ElementType>(GetOutputValueName(), LOCTEXT("ShuffleOpOutputValueTT", "Value of the current shuffled element."))
+				)
+			);
+
+			return DefaultInterface;
+		}
+
+		static const FNodeClassMetadata& GetNodeInfo()
+		{
+			auto CreateNodeClassMetadata = []() -> FNodeClassMetadata
+			{
+				FName DataTypeName = GetMetasoundDataTypeName<ArrayType>();
+				FName OperatorName = TEXT("Shuffle");
+				FText NodeDisplayName = FText::Format(LOCTEXT("ArrayOpArrayShuffleDisplayNamePattern", "Shuffle ({0})"), FText::FromString(GetMetasoundDataTypeString<ArrayType>()));
+				FText NodeDescription = LOCTEXT("ArrayOpArrayShuffleDescription", "Output next element of a shuffled array on trigger.");
+				FVertexInterface NodeInterface = GetDefaultInterface();
+
+				return MetasoundArrayNodesPrivate::CreateArrayNodeClassMetadata(DataTypeName, OperatorName, NodeDisplayName, NodeDescription, NodeInterface);
+			};
+
+			static const FNodeClassMetadata Metadata = CreateNodeClassMetadata();
+
+			return Metadata;
+		}
+
+		static TUniquePtr<IOperator> CreateOperator(const FCreateOperatorParams& InParams, TArray<TUniquePtr<IOperatorBuildError>>& OutErrors)
+		{
+			using namespace ArrayNodeShuffleVertexNames;
+			using namespace MetasoundArrayNodesPrivate;
+
+			const FInputVertexInterface& Inputs = InParams.Node.GetVertexInterface().GetInputInterface();
+
+			TDataReadReference<FTrigger> InTriggerNext = GetDataReadReferenceOrConstructWithVertexDefault<FTrigger>(InParams.InputDataReferences, Inputs, GetInputTriggerNextName(), InParams.OperatorSettings);
+			TDataReadReference<FTrigger> InTriggerShuffle = GetDataReadReferenceOrConstructWithVertexDefault<FTrigger>(InParams.InputDataReferences, Inputs, GetInputTriggerShuffleName(), InParams.OperatorSettings);
+			TDataReadReference<FTrigger> InTriggerReset = GetDataReadReferenceOrConstructWithVertexDefault<FTrigger>(InParams.InputDataReferences, Inputs, GetInputTriggerResetName(), InParams.OperatorSettings);
+			FArrayDataReadReference InInputArray = GetDataReadReferenceOrConstructWithVertexDefault<ArrayType>(InParams.InputDataReferences, Inputs, GetInputShuffleArrayName(), InParams.OperatorSettings);
+			TDataReadReference<int32> InSeedValue = GetDataReadReferenceOrConstructWithVertexDefault<int32>(InParams.InputDataReferences, Inputs, GetInputSeedName(), InParams.OperatorSettings);
+			TDataReadReference<bool> bInAutoShuffle = GetDataReadReferenceOrConstructWithVertexDefault<bool>(InParams.InputDataReferences, Inputs, GetInputAutoShuffleName(), InParams.OperatorSettings);
+			TDataReadReference<FString> InShuffleNamespace = GetDataReadReferenceOrConstructWithVertexDefault<FString>(InParams.InputDataReferences, Inputs, GetInputNamespaceName(), InParams.OperatorSettings);
+
+			return MakeUnique<TArrayShuffleOperator>(InParams.OperatorSettings, InTriggerNext, InTriggerShuffle, InTriggerReset, InInputArray, InSeedValue, bInAutoShuffle, InShuffleNamespace);
+		}
+
+		TArrayShuffleOperator(
+			const FOperatorSettings& InSettings,
+			TDataReadReference<FTrigger> InTriggerNext, 
+			TDataReadReference<FTrigger> InTriggerShuffle,
+			TDataReadReference<FTrigger> InTriggerReset,
+			FArrayDataReadReference InInputArray,
+			TDataReadReference<int32> InSeedValue,
+			TDataReadReference<bool> bInAutoShuffle,
+			TDataReadReference<FString> InShuffleNamespace)
+			: TriggerNext(InTriggerNext)
+			, TriggerShuffle(InTriggerShuffle)
+			, TriggerReset(InTriggerReset)
+			, InputArray(InInputArray)
+			, SeedValue(InSeedValue)
+			, bAutoShuffle(bInAutoShuffle)
+			, ShuffleNamespace(InShuffleNamespace)
+			, TriggerOnNext(FTriggerWriteRef::CreateNew(InSettings))
+			, TriggerOnShuffle(FTriggerWriteRef::CreateNew(InSettings))
+			, TriggerOnReset(FTriggerWriteRef::CreateNew(InSettings))
+			, OutValue(TDataWriteReferenceFactory<ElementType>::CreateAny(InSettings))
+		{
+			// Check to see if this is a global shuffler or a local one. 
+			// Global shuffler will use a namespace to opt into it.
+			const ArrayType& InputArrayRef = *InputArray;
+			PrevNamespaceString = *ShuffleNamespace;
+			PrevArrayNum = InputArrayRef.Num();
+			PrevSeedValue = *SeedValue;
+
+			// Make a key!
+			if (PrevNamespaceString.Len() > 0)
+			{
+				ArrayShuffleKey = FGlobalArrayShuffleKey(PrevNamespaceString, PrevArrayNum, PrevSeedValue);
+			}
+			else if (PrevArrayNum > 0)
+			{
+				ArrayIndexShuffler = MakeUnique<FArrayIndexShuffler>(PrevSeedValue, PrevArrayNum);
+			}
+
+			if (PrevArrayNum == 0)
+			{
+				UE_LOG(LogMetasound, Error, TEXT("Array Shuffle: Can't shuffle an empty array"));
+			}
+		}
+
+		virtual ~TArrayShuffleOperator() = default;
+
+		virtual FDataReferenceCollection GetInputs() const override
+		{
+			using namespace ArrayNodeShuffleVertexNames;
+
+			FDataReferenceCollection Inputs;
+			Inputs.AddDataReadReference(GetInputTriggerNextName(), TriggerNext);
+			Inputs.AddDataReadReference(GetInputTriggerShuffleName(), TriggerShuffle);
+			Inputs.AddDataReadReference(GetInputTriggerResetName(), TriggerReset);
+			Inputs.AddDataReadReference(GetInputShuffleArrayName(), InputArray);
+			Inputs.AddDataReadReference(GetInputSeedName(), SeedValue);
+			Inputs.AddDataReadReference(GetInputAutoShuffleName(), bAutoShuffle);
+			Inputs.AddDataReadReference(GetInputNamespaceName(), ShuffleNamespace);
+
+			return Inputs;
+		}
+
+		virtual FDataReferenceCollection GetOutputs() const override
+		{
+			using namespace ArrayNodeShuffleVertexNames;
+
+			FDataReferenceCollection Outputs;
+			Outputs.AddDataReadReference(GetOutputTriggerOnNextName(), TriggerOnNext);
+			Outputs.AddDataReadReference(GetOutputTriggerOnShuffleName(), TriggerOnShuffle);
+			Outputs.AddDataReadReference(GetOutputTriggerOnResetName(), TriggerOnReset);
+			Outputs.AddDataReadReference(GetOutputValueName(), OutValue);
+
+			return Outputs;
+		}
+
+		void Execute()
+		{
+			TriggerOnNext->AdvanceBlock();
+			TriggerOnShuffle->AdvanceBlock();
+			TriggerOnReset->AdvanceBlock();
+
+			const ArrayType& InputArrayRef = *InputArray;
+ 
+			// Check for any input pin changes to see if we need to swap behavior
+			if (PrevArrayNum != InputArrayRef.Num() || 
+				PrevSeedValue != *SeedValue || 
+				PrevNamespaceString != *ShuffleNamespace)
+			{
+				PrevArrayNum = InputArrayRef.Num();
+				PrevSeedValue = *SeedValue;
+				PrevNamespaceString = *ShuffleNamespace;
+
+				if (PrevArrayNum > 0)
+				{
+					if (PrevNamespaceString.Len() > 0)
+					{
+						// Make new shuffle key and clear out our old index shuffler
+						ArrayShuffleKey = FGlobalArrayShuffleKey(PrevNamespaceString, PrevArrayNum, PrevSeedValue);
+						ArrayIndexShuffler = nullptr;
+					}
+					else
+					{
+						// Clear out our shuffle key and make a new array index shuffler
+						ArrayShuffleKey = FGlobalArrayShuffleKey();
+						ArrayIndexShuffler = MakeUnique<FArrayIndexShuffler>(PrevSeedValue, PrevArrayNum);
+					}
+				}
+				else
+				{
+					UE_LOG(LogMetasound, Error, TEXT("Array Shuffle: Can't shuffle an empty array."));
+				}
+			}
+
+			TriggerReset->ExecuteBlock(
+				[&](int32 StartFrame, int32 EndFrame)
+				{
+				},
+				[this](int32 StartFrame, int32 EndFrame)
+				{
+					if (ArrayShuffleKey.Hash != INDEX_NONE)
+					{
+						FGlobalArrayShuffleManager& SM = FGlobalArrayShuffleManager::Get();
+						SM.ResetSeed(ArrayShuffleKey);
+					}
+					else
+					{
+						check(ArrayIndexShuffler.IsValid());
+						ArrayIndexShuffler->ResetSeed();
+					}
+
+					TriggerOnReset->TriggerFrame(StartFrame);
+				}
+			);
+
+			TriggerShuffle->ExecuteBlock(
+				[&](int32 StartFrame, int32 EndFrame)
+				{
+				},
+				[this](int32 StartFrame, int32 EndFrame)
+				{
+					if (PrevArrayNum > 0)
+					{
+						if (ArrayShuffleKey.Hash != INDEX_NONE)
+						{
+							FGlobalArrayShuffleManager& SM = FGlobalArrayShuffleManager::Get();
+							SM.ShuffleArray(ArrayShuffleKey);
+						}
+						else
+						{
+							check(ArrayIndexShuffler.IsValid());
+							ArrayIndexShuffler->ShuffleArray();
+						}
+
+						TriggerOnShuffle->TriggerFrame(StartFrame);
+					}
+ 				}
+			);
+
+			TriggerNext->ExecuteBlock(
+				[&](int32 StartFrame, int32 EndFrame)
+				{
+				},
+				[this](int32 StartFrame, int32 EndFrame)
+				{
+					const ArrayType& InputArrayRef = *InputArray;
+
+					if (PrevArrayNum > 0)
+					{
+						bool bShuffleTriggered = false;
+						int32 OutShuffleIndex = INDEX_NONE;
+						if (ArrayShuffleKey.Hash != INDEX_NONE)
+						{
+							FGlobalArrayShuffleManager& SM = FGlobalArrayShuffleManager::Get();
+							bShuffleTriggered = SM.NextValue(ArrayShuffleKey, *bAutoShuffle, OutShuffleIndex);
+						}
+						else
+						{
+							check(ArrayIndexShuffler.IsValid());
+							bShuffleTriggered = ArrayIndexShuffler->NextValue(*bAutoShuffle, OutShuffleIndex);
+						}
+
+						check(OutShuffleIndex != INDEX_NONE);
+						check(OutShuffleIndex < InputArrayRef.Num());
+
+						// Write out the value of the shuffled index
+						*OutValue = InputArrayRef[OutShuffleIndex];
+
+						TriggerOnNext->TriggerFrame(StartFrame);
+
+						// Trigger out if the array was auto-shuffled
+						if (bShuffleTriggered)
+						{
+							TriggerOnShuffle->TriggerFrame(StartFrame);
+						}
+					}
+				}
+			);
+		}
+
+	private:
+
+		// Inputs
+		TDataReadReference<FTrigger> TriggerNext;
+		TDataReadReference<FTrigger> TriggerShuffle;
+		TDataReadReference<FTrigger> TriggerReset;
+		FArrayDataReadReference InputArray;
+		TDataReadReference<int32> SeedValue;
+		TDataReadReference<bool> bAutoShuffle;
+		TDataReadReference<FString> ShuffleNamespace;
+
+		// Outputs
+		TDataWriteReference<FTrigger> TriggerOnNext;
+		TDataWriteReference<FTrigger> TriggerOnShuffle;
+		TDataWriteReference<FTrigger> TriggerOnReset;
+		TDataWriteReference<ElementType> OutValue;
+
+		// Data
+		TUniquePtr<FArrayIndexShuffler> ArrayIndexShuffler;
+		FGlobalArrayShuffleKey ArrayShuffleKey;
+		FString PrevNamespaceString;
+		int32 PrevArrayNum = INDEX_NONE;
+		int32 PrevSeedValue = INDEX_NONE;
+	};
+
+	template<typename ArrayType>
+	class TArrayShuffleNode : public FNodeFacade
+	{
+	public:
+		TArrayShuffleNode(const FNodeInitData& InInitData)
+			: FNodeFacade(InInitData.InstanceName, InInitData.InstanceID, TFacadeOperatorClass<TArrayShuffleOperator<ArrayType>>())
+		{
+		}
+
+		virtual ~TArrayShuffleNode() = default;
+	};
 }
 
 #undef LOCTEXT_NAMESPACE
diff --git a/Engine/Plugins/Runtime/Metasound/Source/MetasoundFrontend/Public/MetasoundArrayNodesRegistration.h b/Engine/Plugins/Runtime/Metasound/Source/MetasoundFrontend/Public/MetasoundArrayNodesRegistration.h
--- a/Engine/Plugins/Runtime/Metasound/Source/MetasoundFrontend/Public/MetasoundArrayNodesRegistration.h
+++ b/Engine/Plugins/Runtime/Metasound/Source/MetasoundFrontend/Public/MetasoundArrayNodesRegistration.h
@@ -38,6 +38,9 @@ namespace Metasound
 
 			// Elements must be copy constructible
 			static constexpr bool bIsArraySubsetSupported = std::is_copy_constructible<ElementType>::value;
+
+			// Array shuffle is supported for all types that get is supported for.
+			static constexpr bool bIsArrayShuffleSupported = bIsArrayGetSupported;
 		};
 
 		template<typename ArrayType, typename std::enable_if<TArrayNodeSupport<ArrayType>::bIsArrayGetSupported, bool>::type = true>
@@ -112,6 +115,20 @@ namespace Metasound
 
 			return ensureAlways(RegisterNodeWithFrontend<Metasound::TArrayNumNode<ArrayType>>());
 		}
+
+		template<typename ArrayType, typename std::enable_if<TArrayNodeSupport<ArrayType>::bIsArrayShuffleSupported, bool>::type = true>
+		bool RegisterArrayShuffleNode()
+		{
+			using FNodeType = typename Metasound::TArrayShuffleNode<ArrayType>;
+			return RegisterNodeWithFrontend<FNodeType>();
+		}
+
+		template<typename ArrayType, typename std::enable_if<!TArrayNodeSupport<ArrayType>::bIsArrayShuffleSupported, bool>::type = true>
+		bool RegisterArrayShuffleNode()
+		{
+			// No op if not supported
+			return true;
+		}
 	}
 
 	/** Registers all available array nodes which can be instantiated for the given
@@ -128,7 +145,7 @@ namespace Metasound
 		bSuccess = bSuccess && RegisterArraySetNode<ArrayType>();
 		bSuccess = bSuccess && RegisterArraySubsetNode<ArrayType>();
 		bSuccess = bSuccess && RegisterArrayConcatNode<ArrayType>();
-
+		bSuccess = bSuccess && RegisterArrayShuffleNode<ArrayType>();
 		return bSuccess;
 	}
 }
