Metasound Communications Blueprint API

diff --git a/Engine/Plugins/Runtime/Metasound/Source/MetasoundEditor/Public/MetasoundEditorGraphInputNodes.h b/Engine/Plugins/Runtime/Metasound/Source/MetasoundEditor/Public/MetasoundEditorGraphInputNodes.h
--- a/Engine/Plugins/Runtime/Metasound/Source/MetasoundEditor/Public/MetasoundEditorGraphInputNodes.h
+++ b/Engine/Plugins/Runtime/Metasound/Source/MetasoundEditor/Public/MetasoundEditorGraphInputNodes.h
@@ -132,7 +132,7 @@ public:
 
 	void UpdatePreviewInstance(const Metasound::FVertexKey& InParameterName, IAudioInstanceTransmitter& InInstanceTransmitter) const override
 	{
-		InInstanceTransmitter.SetBoolParameter(*InParameterName, Default);
+		InInstanceTransmitter.SetParameter(*InParameterName, Default);
 	}
 };
 
@@ -161,8 +161,7 @@ public:
 
 	void UpdatePreviewInstance(const Metasound::FVertexKey& InParameterName, IAudioInstanceTransmitter& InInstanceTransmitter) const override
 	{
-		// TODO:
-		//InInstanceTransmitter.SetBoolArrayParameter(InParameterName, Default);
+		InInstanceTransmitter.SetParameter(*InParameterName, TArray<bool>{Default});
 	}
 };
 
@@ -201,7 +200,7 @@ public:
 
 	void UpdatePreviewInstance(const Metasound::FVertexKey& InParameterName, IAudioInstanceTransmitter& InInstanceTransmitter) const override
 	{
-		InInstanceTransmitter.SetIntParameter(*InParameterName, Default.Value);
+		InInstanceTransmitter.SetParameter(*InParameterName, Default.Value);
 	}
 };
 
@@ -236,12 +235,10 @@ public:
 
 	void UpdatePreviewInstance(const Metasound::FVertexKey& InParameterName, IAudioInstanceTransmitter& InInstanceTransmitter) const override
 	{
-
 		TArray<int32> IntArray;
 		Algo::Transform(Default, IntArray, [](const FMetasoundEditorGraphInputInt& InValue) { return InValue.Value; });
 
-		// TODO:
-		//InInstanceTransmitter.SetIntArrayParameter(*InParameterName, IntArray);
+		InInstanceTransmitter.SetParameter(*InParameterName, TArray<int32>{ IntArray });
 	}
 };
 
@@ -270,7 +267,7 @@ public:
 
 	void UpdatePreviewInstance(const Metasound::FVertexKey& InParameterName, IAudioInstanceTransmitter& InInstanceTransmitter) const override
 	{
-		InInstanceTransmitter.SetFloatParameter(*InParameterName, Default);
+		InInstanceTransmitter.SetParameter(*InParameterName, Default);
 	}
 };
 
@@ -299,8 +296,7 @@ public:
 
 	void UpdatePreviewInstance(const Metasound::FVertexKey& InParameterName, IAudioInstanceTransmitter& InInstanceTransmitter) const override
 	{
-		// TODO:
-		//InInstanceTransmitter.SetFloatArrayParameter(*InParameterName, Default);
+		InInstanceTransmitter.SetParameter(*InParameterName, TArray<float>{Default});
 	}
 };
 
@@ -327,10 +323,9 @@ public:
 		return EMetasoundFrontendLiteralType::String;
 	}
 
-	void UpdatePreviewInstance(const Metasound::FVertexKey& InParamterName, IAudioInstanceTransmitter& InInstanceTransmitter) const override
+	void UpdatePreviewInstance(const Metasound::FVertexKey& InParameterName, IAudioInstanceTransmitter& InInstanceTransmitter) const override
 	{
-		// TODO:
-		//InInstanceTransmitter.SetStringParameter(*InParameterName, Default);
+		InInstanceTransmitter.SetParameter(*InParameterName, FString{ Default });
 	}
 };
 
@@ -359,8 +354,7 @@ public:
 
 	void UpdatePreviewInstance(const Metasound::FVertexKey& InParameterName, IAudioInstanceTransmitter& InInstanceTransmitter) const override
 	{
-		// TODO:
-		//InInstanceTransmitter.SetStringArrayParameter(*InParameterName, Default);
+		InInstanceTransmitter.SetParameter(*InParameterName, TArray<FString>{Default});
 	}
 };
 
@@ -399,8 +393,8 @@ public:
 
 	void UpdatePreviewInstance(const Metasound::FVertexKey& InParameterName, IAudioInstanceTransmitter& InInstanceTransmitter) const override
 	{
-		// TODO:
-		//InInstanceTransmitter.SetUObjectParameter(*InParameterName, Default.Object);
+		// TODO. We need proxy object here safely.
+		//InInstanceTransmitter.SetParameter(*InParameterName, Default.Object);
 	}
 };
 
@@ -437,8 +431,8 @@ public:
 	{
 		TArray<UObject*> ObjectArray;
 		Algo::Transform(Default, ObjectArray, [](const FMetasoundEditorGraphInputObject& InValue) { return InValue.Object; });
-
-		// TODO:
-		//InInstanceTransmitter.SetUObjectArrayParameter(*InParameterName, Default.Object);
+		
+		// TODO. We need proxy object here safely.
+		//InInstanceTransmitter.SetParameter(*InParameterName, Default.Object);
 	}
 };
diff --git a/Engine/Plugins/Runtime/Metasound/Source/MetasoundFrontend/Private/MetasoundInstanceTransmitter.cpp b/Engine/Plugins/Runtime/Metasound/Source/MetasoundFrontend/Private/MetasoundInstanceTransmitter.cpp
--- a/Engine/Plugins/Runtime/Metasound/Source/MetasoundFrontend/Private/MetasoundInstanceTransmitter.cpp
+++ b/Engine/Plugins/Runtime/Metasound/Source/MetasoundFrontend/Private/MetasoundInstanceTransmitter.cpp
@@ -16,21 +16,56 @@ namespace Metasound
 		return InstanceID;
 	}
 
-	bool FMetasoundInstanceTransmitter::SetFloatParameter(const FName& InParameterName, float InValue)
+	bool FMetasoundInstanceTransmitter::SetParameter(const FName& InParameterName, float InValue)
 	{
 		return SetParameterWithLiteral(InParameterName, FLiteral(InValue));
 	}
 
-	bool FMetasoundInstanceTransmitter::SetIntParameter(const FName& InParameterName, int32 InValue) 
+	bool FMetasoundInstanceTransmitter::SetParameter(const FName& InParameterName, int32 InValue) 
 	{
 		return SetParameterWithLiteral(InParameterName, FLiteral(InValue));
 	}
 
-	bool FMetasoundInstanceTransmitter::SetBoolParameter(const FName& InParameterName, bool InValue) 
+	bool FMetasoundInstanceTransmitter::SetParameter(const FName& InParameterName, bool InValue) 
 	{
 		return SetParameterWithLiteral(InParameterName, FLiteral(InValue));
 	}
 
+	bool FMetasoundInstanceTransmitter::SetParameter(const FName& InParameterName, TArray<bool>&& InValue)
+	{
+		return SetParameterWithLiteral(InParameterName, FLiteral(MoveTemp(InValue)));
+	}
+
+	bool FMetasoundInstanceTransmitter::SetParameter(const FName& InParameterName, TArray<int32>&& InValue)
+	{
+		return SetParameterWithLiteral(InParameterName, FLiteral(MoveTemp(InValue)));
+	}
+
+	bool FMetasoundInstanceTransmitter::SetParameter(const FName& InParameterName, TArray<float>&& InValue)
+	{
+		return SetParameterWithLiteral(InParameterName, FLiteral(MoveTemp(InValue)));
+	}
+
+	bool FMetasoundInstanceTransmitter::SetParameter(const FName& InParameterName, FString&& InValue)
+	{
+		return SetParameterWithLiteral(InParameterName, FLiteral(MoveTemp(InValue)));
+	}
+
+	bool FMetasoundInstanceTransmitter::SetParameter(const FName& InParameterName, TArray<FString>&& InValue)
+	{
+		return SetParameterWithLiteral(InParameterName, FLiteral(MoveTemp(InValue)));
+	}
+
+	bool FMetasoundInstanceTransmitter::SetParameter(const FName& InParameterName, Audio::IProxyDataPtr&& InValue)
+	{
+		return SetParameterWithLiteral(InParameterName, FLiteral(MoveTemp(InValue)));
+	}
+
+	bool FMetasoundInstanceTransmitter::SetParameter(const FName& InParameterName, TArray<Audio::IProxyDataPtr>&& InValue)
+	{
+		return SetParameterWithLiteral(InParameterName, FLiteral(MoveTemp(InValue)));
+	}
+
 	bool FMetasoundInstanceTransmitter::SetParameterWithLiteral(const FName& InParameterName, const FLiteral& InLiteral) 
 	{
 		ISender* Sender = FindSender(InParameterName);
diff --git a/Engine/Plugins/Runtime/Metasound/Source/MetasoundFrontend/Public/MetasoundInstanceTransmitter.h b/Engine/Plugins/Runtime/Metasound/Source/MetasoundFrontend/Public/MetasoundInstanceTransmitter.h
--- a/Engine/Plugins/Runtime/Metasound/Source/MetasoundFrontend/Public/MetasoundInstanceTransmitter.h
+++ b/Engine/Plugins/Runtime/Metasound/Source/MetasoundFrontend/Public/MetasoundInstanceTransmitter.h
@@ -65,46 +65,109 @@ namespace Metasound
 
 		/** Returns ID of the MetaSound instance associated with this transmitter. */
 		uint64 GetInstanceID() const override;
+		
+		/** Set a parameter using a literal.
+		 *
+		 * @param InParameterName - Name of MetaSound instance parameter.
+		 * @param InValue - Literal value used to construct paramter value. 
+		 *
+		 * @return true on success, false on failure. 
+		 */
+		bool SetParameterWithLiteral(const FName& InParameterName, const FLiteral& InValue);
+
+		/** Duplicate this transmitter interface. The transmitters association with
+		 * the MetaSound instance will be maintained. */
+		TUniquePtr<IAudioInstanceTransmitter> Clone() const override;
+		
+		/** Set a parameter (bool)
+		 *
+		 * @param InParameterName - Name of MetaSound instance parameter.
+		 * @param InValue - bool value used to construct parameter value.
+		 *
+		 * @return true on success, false on failure.
+		 */
+		bool SetParameter(const FName& InName, bool InValue) override;
+
+		/** Set a parameter (bool array)
+		 *
+		 * @param InParameterName - Name of MetaSound instance parameter.
+		 * @param InValue - bool array r-value value used to construct parameter value.
+		 *
+		 * @return true on success, false on failure.
+		 */
+		bool SetParameter(const FName& InName, TArray<bool>&& InValue) override;		
+
+		/** Set a parameter (int32)
+		 *
+		 * @param InParameterName - Name of MetaSound instance parameter.
+		 * @param InValue - in32 value used to construct parameter value.
+		 *
+		 * @return true on success, false on failure.
+		 */
+		bool SetParameter(const FName& InName, int32 InValue) override;
 
-		/** Set a float parameter on the MetaSound instance by name. 
+		/** Set a parameter overloads (int32 array)
 		 *
 		 * @param InParameterName - Name of MetaSound instance parameter.
-		 * @param InValue - Value to set. 
+		 * @param InValue - int32 array r-value value used to construct parameter value.
 		 *
 		 * @return true on success, false on failure.
 		 */
-		bool SetFloatParameter(const FName& InParameterName, float InValue) override;
+		bool SetParameter(const FName& InName, TArray<int32>&& InValue) override;		
 
-		/** Set a int parameter on the MetaSound instance by name. 
+		/** Set a parameter overloads (float)
 		 *
 		 * @param InParameterName - Name of MetaSound instance parameter.
-		 * @param InValue - Value to set. 
+		 * @param InValue - float value used to construct parameter value.
 		 *
 		 * @return true on success, false on failure.
 		 */
-		bool SetIntParameter(const FName& InParameterName, int32 InValue) override;
+		bool SetParameter(const FName& InName, float InValue) override;
 
-		/** Set a bool parameter on the MetaSound instance by name. 
+		/** Set a parameter overloads (float array)
 		 *
 		 * @param InParameterName - Name of MetaSound instance parameter.
-		 * @param InValue - Value to set. 
+		 * @param InValue - float array r-value value used to construct parameter value.
 		 *
 		 * @return true on success, false on failure.
 		 */
-		bool SetBoolParameter(const FName& InParameterName, bool InValue) override;
+		bool SetParameter(const FName& InName, TArray<float>&& InValue) override;
 		
-		/** Set a parameter using a literal.
+		/** Set a parameter overloads (string)
 		 *
 		 * @param InParameterName - Name of MetaSound instance parameter.
-		 * @param InValue - Literal value used to construct paramter value. 
+		 * @param InValue - string r-value value used to construct parameter value.
 		 *
-		 * @return true on success, false on failure. 
+		 * @return true on success, false on failure.
 		 */
-		bool SetParameterWithLiteral(const FName& InParameterName, const FLiteral& InValue);
+		bool SetParameter(const FName& InName, FString&& InValue) override;
 
-		/** Duplicate this transmitter interface. The transmitters association with
-		 * the MetaSound instance will be maintained. */
-		TUniquePtr<IAudioInstanceTransmitter> Clone() const override;
+		/** Set a parameter overloads (string array)
+		 *
+		 * @param InParameterName - Name of MetaSound instance parameter.
+		 * @param InValue - string array r-value value used to construct parameter value.
+		 *
+		 * @return true on success, false on failure.
+		 */
+		bool SetParameter(const FName& InName, TArray<FString>&& InValue) override;
+		
+		/** Set a parameter overloads (Audio::IProxyData)
+		 *
+		 * @param InParameterName - Name of MetaSound instance parameter.
+		 * @param InValue - Audio::IProxyData r-value value used to construct parameter value.
+		 *
+		 * @return true on success, false on failure.
+		 */
+		bool SetParameter(const FName& InName, Audio::IProxyDataPtr&& ) override;
+
+		/** Set a parameter overloads (Audio::IProxyData array)
+		 *
+		 * @param InParameterName - Name of MetaSound instance parameter.
+		 * @param InValue - Audio::IProxyData array r-value value used to construct parameter value.
+		 *
+		 * @return true on success, false on failure.
+		 */
+		bool SetParameter(const FName& InName, TArray<Audio::IProxyDataPtr>&& InValue) override;
 
 	private:
 		// Find FSendInfo by parameter name. 
diff --git a/Engine/Source/Runtime/AudioExtensions/Public/IAudioExtensionPlugin.h b/Engine/Source/Runtime/AudioExtensions/Public/IAudioExtensionPlugin.h
--- a/Engine/Source/Runtime/AudioExtensions/Public/IAudioExtensionPlugin.h
+++ b/Engine/Source/Runtime/AudioExtensions/Public/IAudioExtensionPlugin.h
@@ -9,6 +9,7 @@
 #include "Modules/ModuleManager.h"
 #include "UObject/ObjectMacros.h"
 #include "AudioDefines.h"
+#include "IAudioProxyInitializer.h"
 
 #include "IAudioExtensionPlugin.generated.h"
 
@@ -627,16 +628,19 @@ class IAudioInstanceTransmitter
 
 		// Return the instance ID
 		virtual uint64 GetInstanceID() const = 0;
-
-		// Set a float parameter by name. Return true on success, false on failure.
-		virtual bool SetFloatParameter(const FName& InParameterName, float InValue) = 0;
-
-		// Set an int parameter by name. Return true on success, false on failure.
-		virtual bool SetIntParameter(const FName& InParameterName, int32 InValue) = 0;
-
-		// Set a bool parameter by name. Return true on success, false on failure.
-		virtual bool SetBoolParameter(const FName& InParameterName, bool InValue) = 0;
-
+	
+		// Set Parameter Overloads
+		virtual bool SetParameter(const FName& InName, bool InValue) = 0;
+		virtual bool SetParameter(const FName& InName, TArray<bool>&& InValue) = 0;
+		virtual bool SetParameter(const FName& InName, int32 InValue) = 0;
+		virtual bool SetParameter(const FName& InName, TArray<int32>&& InValue) = 0;
+		virtual bool SetParameter(const FName& InName, float InValue) = 0;
+		virtual bool SetParameter(const FName& InName, TArray<float>&& InValue) = 0;
+		virtual bool SetParameter(const FName& InName, FString&& InValue) = 0;
+		virtual bool SetParameter(const FName& InName, TArray<FString>&& InValue) = 0;
+		virtual bool SetParameter(const FName& InName, Audio::IProxyDataPtr&&) = 0;
+		virtual bool SetParameter(const FName& InName, TArray<Audio::IProxyDataPtr>&& InValue) = 0;
+	
 		// Create a copy of the instance transmitter.
 		virtual TUniquePtr<IAudioInstanceTransmitter> Clone() const = 0;
 };
diff --git a/Engine/Source/Runtime/Engine/Classes/Components/AudioComponent.h b/Engine/Source/Runtime/Engine/Classes/Components/AudioComponent.h
--- a/Engine/Source/Runtime/Engine/Classes/Components/AudioComponent.h
+++ b/Engine/Source/Runtime/Engine/Classes/Components/AudioComponent.h
@@ -13,6 +13,7 @@
 #include "Sound/QuartzQuantizationUtilities.h"
 #include "Quartz/AudioMixerClockHandle.h"
 #include "Quartz/AudioMixerQuantizedCommands.h"
+#include "Sound/AudioCommunicationInterface.h"
 
 #include "AudioComponent.generated.h"
 
@@ -528,6 +529,12 @@ public:
 	UFUNCTION(BlueprintCallable, Category="Audio|Components|Audio")
 	void SetFloatParameter(FName InName, float InFloat);
 
+	UPROPERTY(Transient)
+	TScriptInterface<IAudioCommunicationInterface> CommunicationInterface;
+
+	UFUNCTION(BlueprintCallable, Category = "Communications")
+	TScriptInterface<IAudioCommunicationInterface> GetCommunicationInterface() const;
+
 	/** Allows the designer to set the Wave Parameter on the SoundCue whose name matches the name indicated.
 	 * @param InName The name of the Wave to set. It must match the name set in SoundCue's WaveParam Node
 	 * @param InWave The value to set the Parameter to
diff --git a/Engine/Source/Runtime/Engine/Classes/Sound/AudioCommunicationInterface.h b/Engine/Source/Runtime/Engine/Classes/Sound/AudioCommunicationInterface.h
new file mode 100644
--- /dev/null
+++ b/Engine/Source/Runtime/Engine/Classes/Sound/AudioCommunicationInterface.h
@@ -0,0 +1,64 @@
+// Copyright Epic Games, Inc. All Rights Reserved.
+
+#pragma once
+
+#include "CoreMinimal.h"
+#include "UObject/Interface.h"
+
+#include "AudioCommunicationInterface.generated.h"
+
+UINTERFACE(MinimalApi, BlueprintType, meta = (CannotImplementInterfaceInBlueprint))
+class UAudioCommunicationInterface : public UInterface
+{
+	GENERATED_UINTERFACE_BODY()
+};
+
+class IAudioCommunicationInterface : public IInterface
+{
+	GENERATED_IINTERFACE_BODY()
+
+public:
+	// Triggers a named trigger 
+	UFUNCTION(BlueprintCallable, Category = "Communication|Triggers")
+	virtual void Trigger(FName Name) = 0;
+
+	// Sets a named Boolean 
+	UFUNCTION(BlueprintCallable, Category = "Communication|Primitive")
+	virtual void SetBool(FName Name, bool Value) = 0;
+	
+	// Sets a named Boolean Array 
+	UFUNCTION(BlueprintCallable, Category = "Communication|Array")
+	virtual void SetBoolArray(FName Name,const TArray<bool>& Value) = 0;
+
+	// Sets a named Int32 
+	UFUNCTION(BlueprintCallable, Category = "Communication|Primitive")
+	virtual void SetInt(FName Name, int32 Value) = 0;
+
+	// Sets a named Int32 Array
+	UFUNCTION(BlueprintCallable, Category = "Communication|Array")
+	virtual void SetIntArray(FName Name, const TArray<int32>& Value) = 0;
+
+	// Sets a named Float
+	UFUNCTION(BlueprintCallable, Category = "Communication|Primitive")
+	virtual void SetFloat(FName Name, float Value) = 0;
+	
+	// Sets a named Float Array
+	UFUNCTION(BlueprintCallable, Category = "Communication|Array")
+	virtual void SetFloatArray(FName Name, const TArray<float>& Value) = 0;
+	
+	// Sets a named String
+	UFUNCTION(BlueprintCallable, Category = "Communication|Primitive")
+	virtual void SetString(FName Name, const FString& Value) = 0;
+
+	// Sets a named String Array
+	UFUNCTION(BlueprintCallable, Category = "Communication|Array")
+	virtual void SetStringArray(FName Name, const TArray<FString>& Value) = 0;
+	
+	// Sets a named UObject 
+	UFUNCTION(BlueprintCallable, Category = "Communication|Primitive")
+	virtual void SetObject(FName Name, UObject* Value) = 0;
+
+	// Sets a named UObject Array
+	UFUNCTION(BlueprintCallable, Category = "Communication|Array")
+	virtual void SetObjectArray(FName Name, const TArray<UObject*>& Value) = 0;
+};
\ No newline at end of file
diff --git a/Engine/Source/Runtime/Engine/Classes/Sound/AudioComponentCommuncation.h b/Engine/Source/Runtime/Engine/Classes/Sound/AudioComponentCommuncation.h
new file mode 100644
--- /dev/null
+++ b/Engine/Source/Runtime/Engine/Classes/Sound/AudioComponentCommuncation.h
@@ -0,0 +1,32 @@
+// Copyright Epic Games, Inc. All Rights Reserved.
+
+#pragma once
+
+#include "AudioCommunicationInterface.h"
+#include "UObject/Object.h"
+
+#include "AudioComponentCommuncation.generated.h"
+
+UCLASS()
+class ENGINE_API UAudioComponentCommunication	: public UObject
+												, public IAudioCommunicationInterface
+{
+	GENERATED_BODY()
+
+protected:
+	// IAudioCommunicationInterface
+	void Trigger(FName InName) override;
+	void SetBool(FName InName, bool InValue) override;
+	void SetBoolArray(FName InName, const TArray<bool>& InValue) override;
+	void SetInt(FName InName, int32 InValue) override;
+	void SetIntArray(FName InName, const TArray<int32>& InValue) override;
+	void SetFloat(FName InName, float InValue) override;
+	void SetFloatArray(FName InName, const TArray<float>& InValue) override;
+	void SetString(FName InName, const FString& InValue) override;
+	void SetStringArray(FName InName, const TArray<FString>& InValue) override;
+	void SetObject(FName InName, UObject* InValue) override;
+	void SetObjectArray(FName InName, const TArray<UObject*>& InValue) override;
+
+	template<typename T>
+	void SetValue(FName InName, T&& InX);
+};
diff --git a/Engine/Source/Runtime/Engine/Private/ActiveSound.cpp b/Engine/Source/Runtime/Engine/Private/ActiveSound.cpp
--- a/Engine/Source/Runtime/Engine/Private/ActiveSound.cpp
+++ b/Engine/Source/Runtime/Engine/Private/ActiveSound.cpp
@@ -1287,31 +1287,20 @@ void FActiveSound::SetFloatParameter(const FName InName, const float InFloat)
 {
 	if (InName != NAME_None)
 	{		
-		if (InstanceTransmitter.IsValid())
+		// First see if an entry for this name already exists
+		for (FAudioComponentParam& P : InstanceParameters)
 		{
-			bool bSuccess = InstanceTransmitter->SetFloatParameter(InName, InFloat);
-			if (!bSuccess)
+			if (P.ParamName == InName)
 			{
-				UE_LOG(LogAudio, Warning, TEXT("Failed to set instance float parameter [InstanceID:%d, ParameterName:%s]"), InstanceTransmitter->GetInstanceID(), *InName.ToString());
+				P.FloatParam = InFloat;
+				return;
 			}
 		}
-		else
-		{
-			// First see if an entry for this name already exists
-			for (FAudioComponentParam& P : InstanceParameters)
-			{
-				if (P.ParamName == InName)
-				{
-					P.FloatParam = InFloat;
-					return;
-				}
-			}
 
-			// We didn't find one, so create a new one.
-			const int32 NewParamIndex = InstanceParameters.AddDefaulted();
-			InstanceParameters[NewParamIndex].ParamName = InName;
-			InstanceParameters[NewParamIndex].FloatParam = InFloat;
-		}
+		// We didn't find one, so create a new one.
+		const int32 NewParamIndex = InstanceParameters.AddDefaulted();
+		InstanceParameters[NewParamIndex].ParamName = InName;
+		InstanceParameters[NewParamIndex].FloatParam = InFloat;
 	}
 }
 
@@ -1376,31 +1365,20 @@ void FActiveSound::SetBoolParameter(const FName InName, const bool InBool)
 {
 	if (InName != NAME_None)
 	{
-		if (InstanceTransmitter.IsValid())
+		// First see if an entry for this name already exists
+		for (FAudioComponentParam& P : InstanceParameters)
 		{
-			bool bSuccess = InstanceTransmitter->SetBoolParameter(InName, InBool);
-			if (!bSuccess)
+			if (P.ParamName == InName)
 			{
-				UE_LOG(LogAudio, Warning, TEXT("Failed to set instance bool parameter [InstanceID:%d, ParameterName:%s]"), InstanceTransmitter->GetInstanceID(), *InName.ToString());
+				P.BoolParam = InBool;
+				return;
 			}
 		}
-		else
-		{
-			// First see if an entry for this name already exists
-			for (FAudioComponentParam& P : InstanceParameters)
-			{
-				if (P.ParamName == InName)
-				{
-					P.BoolParam = InBool;
-					return;
-				}
-			}
 
-			// We didn't find one, so create a new one.
-			const int32 NewParamIndex = InstanceParameters.AddDefaulted();
-			InstanceParameters[NewParamIndex].ParamName = InName;
-			InstanceParameters[NewParamIndex].BoolParam = InBool;
-		}
+		// We didn't find one, so create a new one.
+		const int32 NewParamIndex = InstanceParameters.AddDefaulted();
+		InstanceParameters[NewParamIndex].ParamName = InName;
+		InstanceParameters[NewParamIndex].BoolParam = InBool;
 	}
 }
 
@@ -1426,31 +1404,20 @@ void FActiveSound::SetIntParameter(const FName InName, const int32 InInt)
 {
 	if (InName != NAME_None)
 	{
-		if (InstanceTransmitter.IsValid())
+		// First see if an entry for this name already exists
+		for (FAudioComponentParam& P : InstanceParameters)
 		{
-			bool bSuccess = InstanceTransmitter->SetIntParameter(InName, InInt);
-			if (!bSuccess)
+			if (P.ParamName == InName)
 			{
-				UE_LOG(LogAudio, Warning, TEXT("Failed to set instance int parameter [InstanceID:%d, ParameterName:%s]"), InstanceTransmitter->GetInstanceID(), *InName.ToString());
+				P.IntParam = InInt;
+				return;
 			}
 		}
-		else
-		{
-			// First see if an entry for this name already exists
-			for (FAudioComponentParam& P : InstanceParameters)
-			{
-				if (P.ParamName == InName)
-				{
-					P.IntParam = InInt;
-					return;
-				}
-			}
 
-			// We didn't find one, so create a new one.
-			const int32 NewParamIndex = InstanceParameters.AddDefaulted();
-			InstanceParameters[NewParamIndex].ParamName = InName;
-			InstanceParameters[NewParamIndex].IntParam = InInt;
-		}
+		// We didn't find one, so create a new one.
+		const int32 NewParamIndex = InstanceParameters.AddDefaulted();
+		InstanceParameters[NewParamIndex].ParamName = InName;
+		InstanceParameters[NewParamIndex].IntParam = InInt;
 	}
 }
 
diff --git a/Engine/Source/Runtime/Engine/Private/Audio/AudioCommunicationInterface.cpp b/Engine/Source/Runtime/Engine/Private/Audio/AudioCommunicationInterface.cpp
new file mode 100644
--- /dev/null
+++ b/Engine/Source/Runtime/Engine/Private/Audio/AudioCommunicationInterface.cpp
@@ -0,0 +1,7 @@
+// Copyright Epic Games, Inc. All Rights Reserved.
+
+#include "Sound/AudioCommunicationInterface.h"
+
+UAudioCommunicationInterface::UAudioCommunicationInterface(FObjectInitializer const& InOI)
+	: UInterface(InOI)
+{}
\ No newline at end of file
diff --git a/Engine/Source/Runtime/Engine/Private/AudioComponentCommunication.cpp b/Engine/Source/Runtime/Engine/Private/AudioComponentCommunication.cpp
new file mode 100644
--- /dev/null
+++ b/Engine/Source/Runtime/Engine/Private/AudioComponentCommunication.cpp
@@ -0,0 +1,127 @@
+// Copyright Epic Games, Inc. All Rights Reserved.
+
+#include "Sound/AudioComponentCommuncation.h"
+#include "AudioDevice.h"
+#include "AudioThread.h"
+#include "Components/AudioComponent.h"
+#include "ActiveSound.h"
+#include "IAudioExtensionPlugin.h"
+
+void UAudioComponentCommunication::Trigger(FName InName)
+{
+	// Trigger is just a (true) bool param currently.
+	SetValue(InName, true);
+}
+
+void UAudioComponentCommunication::SetBool(FName InName, bool InValue)
+{
+	SetValue(InName, InValue);
+}
+
+void UAudioComponentCommunication::SetBoolArray(FName InName, const TArray<bool>& InValue)
+{
+	SetValue(InName, InValue);
+}
+
+void UAudioComponentCommunication::SetInt(FName InName, int32 InValue)
+{
+	SetValue(InName, InValue);
+}
+
+void UAudioComponentCommunication::SetIntArray(FName InName, const TArray<int32>& InValue)
+{
+	SetValue(InName, InValue);
+}
+
+void UAudioComponentCommunication::SetFloat(FName InName, float InValue)
+{
+	SetValue(InName, InValue);
+}
+
+void UAudioComponentCommunication::SetFloatArray(FName InName, const TArray<float>& InValue)
+{
+	SetValue(InName, InValue);
+}
+
+void UAudioComponentCommunication::SetString(FName InName, const FString& InValue)
+{
+	SetValue(InName, InValue);
+}
+
+void UAudioComponentCommunication::SetStringArray(FName InName, const TArray<FString>& InValue)
+{
+	SetValue(InName, InValue);
+}
+
+void UAudioComponentCommunication::SetObject(FName InName, UObject* InValue)
+{
+	IAudioProxyDataFactory* ObjectAsFactory = Cast<USoundWave>(InValue);	// FIXME. Possible to query for IAudioProxy support? Metasound used SFINAE
+
+	// The creation of the Proxy object needs to happen where its safe to access UObject, 
+	// namely the Game thread. 
+
+	if (ensureAlways(ObjectAsFactory))
+	{
+		static FName ProxySubsystemName = TEXT("Metasound");
+		Audio::FProxyDataInitParams ProxyInitParams;
+		ProxyInitParams.NameOfFeatureRequestingProxy = ProxySubsystemName;
+
+		TUniquePtr<Audio::IProxyData> Proxy{ ObjectAsFactory->CreateNewProxyData(ProxyInitParams) };
+		SetValue(InName, MoveTemp(Proxy));
+	}
+}
+
+void UAudioComponentCommunication::SetObjectArray(FName InName, const TArray<UObject*>& InValue)
+{
+	static FName ProxySubsystemName = TEXT("Metasound");
+	Audio::FProxyDataInitParams ProxyInitParams;
+	ProxyInitParams.NameOfFeatureRequestingProxy = ProxySubsystemName;
+
+	// The creation of the Proxy object needs to happen where its safe to access UObject, 
+	// namely the Game thread. 
+
+	TArray<Audio::IProxyDataPtr> ProxiedInputs;
+	for (UObject* i : InValue)
+	{
+		if (IAudioProxyDataFactory* ProxyFactory = Cast<USoundWave>(i)) // FIXME. Query for IAudioProxyDataFactory
+		{
+			ProxiedInputs.Emplace(ProxyFactory->CreateNewProxyData(ProxyInitParams));
+		}
+	}
+	SetValue(InName, MoveTemp(ProxiedInputs));
+}
+
+template<typename T>
+void UAudioComponentCommunication::SetValue(FName InName, T&& InX)
+{
+	if (InName == NAME_None)
+	{
+		return;
+	}
+
+	if (UAudioComponent* OwningComponent = Cast<UAudioComponent>(GetOuter()))
+	{
+		if (OwningComponent->IsActive())
+		{
+			if (FAudioDevice* AudioDevice = OwningComponent->GetAudioDevice())
+			{
+				DECLARE_CYCLE_STAT(TEXT("FAudioThreadTask.SetValue"), STAT_AudioSetSoundParameter, STATGROUP_AudioThreadCommands);
+
+				const uint64 MyAudioComponentID = OwningComponent->GetAudioComponentID();
+				FAudioThread::RunCommandOnAudioThread([AudioDevice, MyAudioComponentID, InName, InX{ MoveTempIfPossible(InX) }]() mutable
+				{
+					if (FActiveSound* ActiveSound = AudioDevice->FindActiveSound(MyAudioComponentID))
+					{
+						if (IAudioInstanceTransmitter* Transmitter = ActiveSound->GetTransmitter())
+						{
+							if (!Transmitter->SetParameter(InName, MoveTempIfPossible(InX)))
+							{
+								UE_LOG(LogTemp, Warning, TEXT("Failed to SetParameter '%s'"), *InName.ToString() );
+							}
+						}
+					}
+				}, GET_STATID(STAT_AudioSetSoundParameter));
+			}
+		}
+	}
+}
diff --git a/Engine/Source/Runtime/Engine/Private/Components/AudioComponent.cpp b/Engine/Source/Runtime/Engine/Private/Components/AudioComponent.cpp
--- a/Engine/Source/Runtime/Engine/Private/Components/AudioComponent.cpp
+++ b/Engine/Source/Runtime/Engine/Private/Components/AudioComponent.cpp
@@ -13,8 +13,7 @@
 #include "UObject/FrameworkObjectVersion.h"
 #include "Misc/App.h"
 #include "Kismet/GameplayStatics.h"
-
-
+#include "Sound/AudioComponentCommuncation.h"
 
 DECLARE_CYCLE_STAT(TEXT("AudioComponent Play"), STAT_AudioComp_Play, STATGROUP_Audio);
 
@@ -87,6 +86,10 @@ UAudioComponent::UAudioComponent(const FObjectInitializer& ObjectInitializer)
 		FScopeLock Lock(&AudioIDToComponentMapLock);
 		AudioIDToComponentMap.Add(AudioComponentID, this);
 	}
+
+	// Create communications object.
+	static const FName NAME_CommunicationSubObj{ TEXT("AudioComponentCommunication") };
+	CommunicationInterface = CreateDefaultSubobject<UAudioComponentCommunication>(NAME_CommunicationSubObj, true);
 }
 
 UAudioComponent* UAudioComponent::GetAudioComponentFromID(uint64 AudioComponentID)
@@ -1107,6 +1110,11 @@ void UAudioComponent::SetFloatParameter( const FName InName, const float InFloat
 	}
 }
 
+TScriptInterface<IAudioCommunicationInterface> UAudioComponent::GetCommunicationInterface() const
+{
+	return CommunicationInterface;
+}
+
 void UAudioComponent::SetWaveParameter( const FName InName, USoundWave* InWave )
 {
 	if (InName != NAME_None)
diff --git a/Engine/Source/Runtime/Engine/Public/ActiveSound.h b/Engine/Source/Runtime/Engine/Public/ActiveSound.h
--- a/Engine/Source/Runtime/Engine/Public/ActiveSound.h
+++ b/Engine/Source/Runtime/Engine/Public/ActiveSound.h
@@ -403,6 +403,15 @@ private:
 			return TransmitterImpl.IsValid();
 		}
 
+		IAudioInstanceTransmitter* Get()
+		{
+			return TransmitterImpl.Get();
+		}
+		const IAudioInstanceTransmitter* Get() const
+		{
+			return TransmitterImpl.Get();
+		}
+
 		IAudioInstanceTransmitter* operator->()
 		{
 			return TransmitterImpl.Get();
@@ -420,6 +429,15 @@ private:
 	FInstanceTransmitter InstanceTransmitter;
 
 public:
+	IAudioInstanceTransmitter* GetTransmitter() 
+	{
+		return InstanceTransmitter.Get();
+	}
+	const IAudioInstanceTransmitter* GetTransmitter() const
+	{
+		return InstanceTransmitter.Get();
+	}
+
 	enum class EFadeOut : uint8
 	{
 		// Sound is not currently fading out
