Send Receive channel name fix

diff --git a/Engine/Plugins/Runtime/Metasound/Source/MetasoundFrontend/Private/MetasoundRouter.cpp b/Engine/Plugins/Runtime/Metasound/Source/MetasoundFrontend/Private/MetasoundRouter.cpp
--- a/Engine/Plugins/Runtime/Metasound/Source/MetasoundFrontend/Private/MetasoundRouter.cpp
+++ b/Engine/Plugins/Runtime/Metasound/Source/MetasoundFrontend/Private/MetasoundRouter.cpp
@@ -7,28 +7,27 @@
 #include "CoreMinimal.h"
 #include "HAL/IConsoleManager.h"
 
+
 // Convenience exec commands to push values to global params.
 static FAutoConsoleCommand GPushFloatCommand(
 	TEXT("au.metasound.SetFloat"),
-	TEXT("Use this with au.metasound.SetFloat [address] [value]. Pushes a parameter value directly to a global address, which can then be received by Metasounds using a Receive node."),
+	TEXT("Use this with au.metasound.SetFloat [type] [address] [value]. Pushes a parameter value directly to a global address, which can then be received by Metasounds using a Receive node."),
 	FConsoleCommandWithArgsDelegate::CreateStatic(
 		[](const TArray< FString >& Args)
 		{
-			if (Args.Num() < 2)
+			if (Args.Num() < 3)
 			{
-				UE_LOG(LogTemp, Warning, TEXT("au.metasound.SetFloat should be called with two args- the address to send to and the value to send."));
+				UE_LOG(LogTemp, Warning, TEXT("au.metasound.Set* should be called with three args- the data type of the channel, the address to send to and the value to send."));
 				return;
 			}
 
-			FName ChannelName = FName(*Args[0]);
-			float ValueToPush = TCString<TCHAR>::Atof(*Args[1]);
+			FName DataType = FName(*Args[0]);
+			FName ChannelName = FName(*Args[1]);
+			float ValueToPush = TCString<TCHAR>::Atof(*Args[2]);
 
 			Metasound::FLiteral LiteralParam(ValueToPush);
 
-			if (!Metasound::FDataTransmissionCenter::Get().PushLiteral(ChannelName, LiteralParam))
-			{
-				UE_LOG(LogTemp, Warning, TEXT("au.metasound.SetFloat failed, likely because the channel did not support floats or the channel did not exist."));
-			}
+			Metasound::FDataTransmissionCenter::Get().PushLiteral(DataType, ChannelName, LiteralParam);
 		})
 );
 
@@ -41,20 +40,19 @@ static FAutoConsoleCommand GPushBoolCommand(
 			if (Args.Num() < 2)
 			{
 				UE_LOG(LogTemp, Warning, TEXT("au.metasound.SetBool should be called with two args- the address to send to and the value to send."));
+				UE_LOG(LogTemp, Warning, TEXT("au.metasound.Set* should be called with three args- the data type of the channel, the address to send to and the value to send."));
 				return;
 			}
 
-			FName ChannelName = FName(*Args[0]);
-			int32 ValueAsInt = TCString<TCHAR>::Atoi(*Args[1]);
+			FName DataType = FName(*Args[0]);
+			FName ChannelName = FName(*Args[1]);
+			int32 ValueAsInt = TCString<TCHAR>::Atoi(*Args[2]);
 
 			bool ValueToPush = ValueAsInt != 0;
 
 			Metasound::FLiteral LiteralParam(ValueToPush);
 
-			if (!Metasound::FDataTransmissionCenter::Get().PushLiteral(ChannelName, LiteralParam))
-			{
-				UE_LOG(LogTemp, Warning, TEXT("au.metasound.SetBool failed, likely because the channel did not support floats or the channel did not exist."));
-			}
+			Metasound::FDataTransmissionCenter::Get().PushLiteral(DataType, ChannelName, LiteralParam);
 		})
 );
 
@@ -66,19 +64,17 @@ static FAutoConsoleCommand GPushIntCommand(
 		{
 			if (Args.Num() < 2)
 			{
-				UE_LOG(LogTemp, Warning, TEXT("au.metasound.SetBool should be called with two args- the address to send to and the value to send."));
+				UE_LOG(LogTemp, Warning, TEXT("au.metasound.Set* should be called with three args- the data type of the channel, the address to send to and the value to send."));
 				return;
 			}
 
-			FName ChannelName = FName(*Args[0]);
-			int32 ValueToPush = TCString<TCHAR>::Atoi(*Args[1]);
+			FName DataType = FName(*Args[0]);
+			FName ChannelName = FName(*Args[1]);
+			int32 ValueToPush = TCString<TCHAR>::Atoi(*Args[2]);
 
 			Metasound::FLiteral LiteralParam(ValueToPush);
 
-			if (!Metasound::FDataTransmissionCenter::Get().PushLiteral(ChannelName, LiteralParam))
-			{
-				UE_LOG(LogTemp, Warning, TEXT("au.metasound.SetInt failed, likely because the channel did not support floats or the channel did not exist."));
-			}
+			Metasound::FDataTransmissionCenter::Get().PushLiteral(DataType, ChannelName, LiteralParam);
 		})
 );
 
@@ -88,20 +84,19 @@ static FAutoConsoleCommand GPushStringCommand(
 	FConsoleCommandWithArgsDelegate::CreateStatic(
 		[](const TArray< FString >& Args)
 		{
-			if (Args.Num() < 2)
+			if (Args.Num() < 3)
 			{
-				UE_LOG(LogTemp, Warning, TEXT("au.metasound.SetBool should be called with two args- the address to send to and the value to send."));
+				UE_LOG(LogTemp, Warning, TEXT("au.metasound.SetBool should be called with three args- the data type of the channel, the address to send to and the value to send."));
+				UE_LOG(LogTemp, Warning, TEXT("au.metasound.Set* should be called with three args- the data type of the channel, the address to send to and the value to send."));
 				return;
 			}
 
-			FName ChannelName = FName(*Args[0]);
+			FName DataType = FName(*Args[0]);
+			FName ChannelName = FName(*Args[1]);
 
-			Metasound::FLiteral LiteralParam(Args[1]);
+			Metasound::FLiteral LiteralParam(Args[2]);
 
-			if (!Metasound::FDataTransmissionCenter::Get().PushLiteral(ChannelName, LiteralParam))
-			{
-				UE_LOG(LogTemp, Warning, TEXT("au.metasound.SetString failed, likely because the channel did not support floats or the channel did not exist."));
-			}
+			Metasound::FDataTransmissionCenter::Get().PushLiteral(DataType, ChannelName, LiteralParam);
 		})
 );
 
@@ -205,9 +200,9 @@ namespace Metasound
 		return MoveTemp(Receiver);
 	}
 
-	bool FDataTransmissionCenter::PushLiteral(FName GlobalChannelName, const FLiteral& InParam)
+	bool FDataTransmissionCenter::PushLiteral(FName DataType, FName GlobalChannelName, const FLiteral& InParam)
 	{
-		return GlobalRouter.PushLiteral(GlobalChannelName, InParam);
+		return GlobalRouter.PushLiteral(DataType, GlobalChannelName, InParam);
 	}
 
 	void FDataTransmissionCenter::RegisterSubsystem(ITransmissionSubsystem* InSystem, FName InSubsytemName)
@@ -230,57 +225,48 @@ namespace Metasound
 		SubsystemRouters.Remove(InSubsystemName);
 	}
 
-	TArray<FName> FAddressRouter::GetAvailableChannels()
+	FAddressRouter::FDataChannelKey FAddressRouter::GetDataChannelKey(const FName& InDataTypeName, const FName& InChannelName) const
 	{
-		FScopeLock ScopeLock(&DataChannelMapMutationLock);
-
-		TArray<FName> AvailableChannels;
-
-		for (auto& DataChannelPair : DataChannelMap)
-		{
-			AvailableChannels.Add(DataChannelPair.Key);
-		}
-
-		return AvailableChannels;
+		return FString::Format(TEXT("{0}[{1}]"), { InChannelName.ToString(), InDataTypeName.ToString() });
 	}
 
-	FName FAddressRouter::GetDatatypeForChannel(FName InChannelName)
+	TSharedPtr<IDataChannel, ESPMode::ThreadSafe> FAddressRouter::FindDataChannel(const FName& InDataTypeName, const FName& InChannelName)
 	{
-		FScopeLock ScopeLock(&DataChannelMapMutationLock);
+		TSharedPtr<IDataChannel, ESPMode::ThreadSafe> Channel;
+
+		const FDataChannelKey ChannelKey = GetDataChannelKey(InDataTypeName, InChannelName);
 
-		if (TSharedRef<IDataChannel, ESPMode::ThreadSafe>* FoundChannel = DataChannelMap.Find(InChannelName))
-		{
-			return (*FoundChannel)->GetDataType();
-		}
-		else
 		{
-			return FName();
+			FScopeLock ScopeLock(&DataChannelMapMutationLock);
+
+			if (TSharedRef<IDataChannel, ESPMode::ThreadSafe>* ExistingChannelPtr = DataChannelMap.Find(ChannelKey))
+			{
+				Channel = *ExistingChannelPtr;
+			}
 		}
+
+		return Channel;
 	}
 
 	TSharedPtr<IDataChannel, ESPMode::ThreadSafe> FAddressRouter::GetDataChannel(const FName& InDataTypeName, const FName& InChannelName, const FOperatorSettings& InOperatorSettings)
 	{
-		TSharedPtr<IDataChannel, ESPMode::ThreadSafe> DataChannel; 
+		TSharedPtr<IDataChannel, ESPMode::ThreadSafe> DataChannel = FindDataChannel(InDataTypeName, InChannelName);
 
+		if (!DataChannel.IsValid())
 		{
 			FScopeLock ScopeLock(&DataChannelMapMutationLock);
 
-			if (TSharedRef<IDataChannel, ESPMode::ThreadSafe>* ExistingChannelPtr = DataChannelMap.Find(InChannelName))
-			{
-				DataChannel = *ExistingChannelPtr;
-			}
-			else
-			{
-				// This is the first time we're seeing this, add it to the map.
-				FMetasoundFrontendRegistryContainer* Registry = FMetasoundFrontendRegistryContainer::Get();
+			const FDataChannelKey ChannelKey = GetDataChannelKey(InDataTypeName, InChannelName);
+			
+			// This is the first time we're seeing this, add it to the map.
+			FMetasoundFrontendRegistryContainer* Registry = FMetasoundFrontendRegistryContainer::Get();
 
-				if (ensure(nullptr != Registry))
+			if (ensure(nullptr != Registry))
+			{
+				DataChannel = Registry->CreateDataChannelForDataType(InDataTypeName, InOperatorSettings);
+				if (DataChannel.IsValid())
 				{
-					DataChannel = Registry->CreateDataChannelForDataType(InDataTypeName, InOperatorSettings);
-					if (DataChannel.IsValid())
-					{
-						DataChannelMap.Add(InChannelName, DataChannel.ToSharedRef());
-					}
+					DataChannelMap.Add(ChannelKey, DataChannel.ToSharedRef());
 				}
 			}
 		}
@@ -318,34 +304,6 @@ namespace Metasound
 		}
 	}
 
-	TArray<FName> FInstanceLocalRouter::GetAvailableChannels(uint64 InInstanceID)
-	{
-		FScopeLock ScopeLock(&InstanceRouterMapMutationLock);
-
-		if (FAddressRouter* AddressRouter = InstanceRouterMap.Find(InInstanceID))
-		{
-			return AddressRouter->GetAvailableChannels();
-		}
-		else
-		{
-			return TArray<FName>();
-		}
-	}
-
-	FName FInstanceLocalRouter::GetDatatypeForChannel(uint64 InInstanceID, FName InChannelName)
-	{
-		FScopeLock ScopeLock(&InstanceRouterMapMutationLock);
-
-		if (FAddressRouter* AddressRouter = InstanceRouterMap.Find(InInstanceID))
-		{
-			return AddressRouter->GetDatatypeForChannel(InChannelName);
-		}
-		else
-		{
-			return FName();
-		}
-	}
-
 	TUniquePtr<IReceiver> FInstanceLocalRouter::RegisterNewReceiver(uint64 InInstanceID, const FName& InDataTypeName, const FName& InChannelName, const FReceiverInitParams& InitParams)
 	{
 		FScopeLock ScopeLock(&InstanceRouterMapMutationLock);
diff --git a/Engine/Plugins/Runtime/Metasound/Source/MetasoundFrontend/Public/MetasoundRouter.h b/Engine/Plugins/Runtime/Metasound/Source/MetasoundFrontend/Public/MetasoundRouter.h
--- a/Engine/Plugins/Runtime/Metasound/Source/MetasoundFrontend/Public/MetasoundRouter.h
+++ b/Engine/Plugins/Runtime/Metasound/Source/MetasoundFrontend/Public/MetasoundRouter.h
@@ -877,11 +877,11 @@ namespace Metasound
 	// Basic router that takes an FName address, 
 	class METASOUNDFRONTEND_API FAddressRouter
 	{
-	public:
-		TArray<FName> GetAvailableChannels();
 
-		FName GetDatatypeForChannel(FName InChannelName);
+	public:
+		using FDataChannelKey = FString;
 
+		TSharedPtr<IDataChannel, ESPMode::ThreadSafe> FindDataChannel(const FName& InDataTypeName, const FName& InChannelName);
 		TSharedPtr<IDataChannel, ESPMode::ThreadSafe> GetDataChannel(const FName& InDataTypeName, const FName& InChannelName, const FOperatorSettings& InOperatorSettings);
 
 		TUniquePtr<ISender> RegisterNewSender(const FName& InDataTypeName, const FName& InChannelName, const FSenderInitParams& InitParams);
@@ -917,13 +917,12 @@ namespace Metasound
 			return MoveTemp(Sender);
 		}
 
-		bool PushLiteral(const FName& InChannelName, const FLiteral& InParam)
+		bool PushLiteral(const FName& InDataTypeName, const FName& InChannelName, const FLiteral& InParam)
 		{
-			FScopeLock ScopeLock(&DataChannelMapMutationLock);
-			if (TSharedRef<IDataChannel, ESPMode::ThreadSafe>* ExistingChannelPtr = DataChannelMap.Find(InChannelName))
+			TSharedPtr<IDataChannel, ESPMode::ThreadSafe> Channel = FindDataChannel(InDataTypeName, InChannelName);
+			if (Channel.IsValid())
 			{
-				TSharedRef<IDataChannel, ESPMode::ThreadSafe>& ExistingChannel = *ExistingChannelPtr;
-				return ExistingChannel->PushLiteral(InParam);
+				return Channel->PushLiteral(InParam);
 			}
 			else
 			{
@@ -940,7 +939,10 @@ namespace Metasound
 		{}
 
 	private:
-		TMap<FName, TSharedRef<IDataChannel, ESPMode::ThreadSafe>> DataChannelMap;
+		FDataChannelKey GetDataChannelKey(const FName& InDataTypeName, const FName& InChannelName) const;
+
+
+		TMap<FDataChannelKey, TSharedRef<IDataChannel, ESPMode::ThreadSafe>> DataChannelMap;
 		FCriticalSection DataChannelMapMutationLock;
 	};
 
@@ -949,10 +951,6 @@ namespace Metasound
 	{
 	public:
 
-		TArray<FName> GetAvailableChannels(uint64 InInstanceID);
-
-		FName GetDatatypeForChannel(uint64 InInstanceID, FName InChannelName);
-
 		TUniquePtr<ISender> RegisterNewSender(uint64 InInstanceID, const FName& InDataTypeName, const FName& InChannelName, const FSenderInitParams& InitParams);
 
 		template<typename TDataType>
@@ -1102,7 +1100,7 @@ namespace Metasound
 
 		// Pushes a literal parameter to a specific data channel in the global router.
 		// returns false if the literal type isn't supported.
-		bool PushLiteral(FName GlobalChannelName, const FLiteral& InParam);
+		bool PushLiteral(FName DataTypeName, FName GlobalChannelName, const FLiteral& InParam);
 
 	private:
 		// Single map of FNames to IDataChannels
diff --git a/Engine/Plugins/Runtime/Metasound/Source/MetasoundFrontend/Public/MetasoundSendNode.h b/Engine/Plugins/Runtime/Metasound/Source/MetasoundFrontend/Public/MetasoundSendNode.h
--- a/Engine/Plugins/Runtime/Metasound/Source/MetasoundFrontend/Public/MetasoundSendNode.h
+++ b/Engine/Plugins/Runtime/Metasound/Source/MetasoundFrontend/Public/MetasoundSendNode.h
@@ -129,10 +129,18 @@ namespace Metasound
 
 				virtual TUniquePtr<IOperator> CreateOperator(const FCreateOperatorParams& InParams, FBuildErrorArray& OutErrors) override
 				{
-					return MakeUnique<TSendOperator>(InParams.InputDataReferences.GetDataReadReference<TDataType>(GetSendInputName()),
-						InParams.InputDataReferences.GetDataReadReference<FSendAddress>(GetAddressInputName()),
-						InParams.OperatorSettings
-						);
+					if (InParams.InputDataReferences.ContainsDataReadReference<TDataType>(GetSendInputName()))
+					{
+						return MakeUnique<TSendOperator>(InParams.InputDataReferences.GetDataReadReference<TDataType>(GetSendInputName()),
+							InParams.InputDataReferences.GetDataReadReference<FSendAddress>(GetAddressInputName()),
+							InParams.OperatorSettings
+							);
+					}
+					else
+					{
+						// No input hook up to send, so this node can no-op
+						return MakeUnique<FNoOpOperator>();
+					}
 				}
 		};
 
diff --git a/Engine/Plugins/Runtime/Metasound/Source/MetasoundFrontend/Public/MetasoundTime.h b/Engine/Plugins/Runtime/Metasound/Source/MetasoundFrontend/Public/MetasoundTime.h
--- a/Engine/Plugins/Runtime/Metasound/Source/MetasoundFrontend/Public/MetasoundTime.h
+++ b/Engine/Plugins/Runtime/Metasound/Source/MetasoundFrontend/Public/MetasoundTime.h
@@ -26,6 +26,11 @@ namespace Metasound
 			{
 			}
 
+			explicit FTime(float InTime)
+			: Time(static_cast<TimeType>(InTime))
+			{
+			}
+			
 			FTime(const FTime& InOtherTime)
 				: Time(InOtherTime.Time)
 			{
diff --git a/Engine/Plugins/Runtime/Metasound/Source/MetasoundGraphCore/Public/MetasoundExecutableOperator.h b/Engine/Plugins/Runtime/Metasound/Source/MetasoundGraphCore/Public/MetasoundExecutableOperator.h
--- a/Engine/Plugins/Runtime/Metasound/Source/MetasoundGraphCore/Public/MetasoundExecutableOperator.h
+++ b/Engine/Plugins/Runtime/Metasound/Source/MetasoundGraphCore/Public/MetasoundExecutableOperator.h
@@ -53,10 +53,20 @@ namespace Metasound
 		public:
 			virtual ~FNoOpOperator() {}
 
-			virtual FExecuteFunction GetExecuteFunction() override
+			FExecuteFunction GetExecuteFunction() override
 			{
 				return nullptr;
 			}
+
+			FDataReferenceCollection GetInputs() const override
+			{
+				return FDataReferenceCollection{};
+			}
+
+			FDataReferenceCollection GetOutputs() const override
+			{
+				return FDataReferenceCollection{};
+			}
 	};
 
 	/** FExecuter
