Metasound trigger route node

diff --git a/Engine/Plugins/Runtime/Metasound/Source/MetasoundStandardNodes/Private/MetasoundTriggerRouteNode.cpp b/Engine/Plugins/Runtime/Metasound/Source/MetasoundStandardNodes/Private/MetasoundTriggerRouteNode.cpp
new file mode 100644
--- /dev/null
+++ b/Engine/Plugins/Runtime/Metasound/Source/MetasoundStandardNodes/Private/MetasoundTriggerRouteNode.cpp
@@ -0,0 +1,115 @@
+// Copyright Epic Games, Inc. All Rights Reserved.
+#include "MetasoundTriggerRouteNode.h"
+
+#include "CoreMinimal.h"
+
+#include "MetasoundBuilderInterface.h"
+#include "MetasoundDataReferenceCollection.h"
+#include "MetasoundExecutableOperator.h"
+#include "MetasoundFacade.h"
+#include "MetasoundNode.h"
+#include "MetasoundNodeInterface.h"
+#include "MetasoundNodeRegistrationMacro.h"
+#include "MetasoundOperatorInterface.h"
+#include "MetasoundPrimitives.h"
+#include "MetasoundStandardNodesNames.h"
+#include "MetasoundStandardNodesCategories.h"
+
+#define LOCTEXT_NAMESPACE "MetasoundStandardNodes"
+
+#define REGISTER_TRIGGER_ROUTE_NOTE(DataType, Number) \
+	using FTriggerRouteNode##DataType##_##Number = TTriggerRouteNode<DataType, Number>; \
+	METASOUND_REGISTER_NODE(FTriggerRouteNode##DataType##_##Number) \
+
+namespace Metasound
+{
+	namespace MetasoundTriggerRouteNodePrivate
+	{
+		FNodeClassMetadata CreateNodeClassMetadata(const FName& InDataTypeName, const FName& InOperatorName, const FText& InDisplayName, const FText& InDescription, const FVertexInterface& InDefaultInterface)
+		{
+			FNodeClassMetadata Metadata
+			{
+				FNodeClassName{FName("TriggerRoute"), InOperatorName, InDataTypeName},
+				1, // Major Version
+				0, // Minor Version
+				InDisplayName,
+				InDescription,
+				PluginAuthor,
+				PluginNodeMissingPrompt,
+				InDefaultInterface,
+				{StandardNodes::TriggerUtils},
+				{TEXT("TriggerRoute")},
+				FNodeDisplayStyle{}
+			};
+
+			return Metadata;
+		}
+	}
+
+
+	namespace TriggerRouteVertexNames
+	{
+		const FString GetInputTriggerName(uint32 InIndex)
+		{
+			return FText::Format(LOCTEXT("TriggerRouteTriggerInputName", "Set {0}"), InIndex).ToString();
+		}
+
+		const FText GetInputTriggerDescription(uint32 InIndex)
+		{
+			return FText::Format(LOCTEXT("TriggerRouteInputTriggerDesc", "The input trigger {0} to cause the corresponding input value to route to the output value."), InIndex);
+		}
+
+		const FString GetInputValueName(uint32 InIndex)
+		{
+			return FText::Format(LOCTEXT("TriggerRouteValueInputName", "Value {0}"), InIndex).ToString();
+		}
+
+		const FText GetInputValueDescription(uint32 InIndex)
+		{
+			static const FText Desc = FText::Format(LOCTEXT("TriggerRouteValueDesc", "The input value ({0}) to route to the output when triggered by Set {0}."), InIndex);
+			return Desc;
+		}
+
+		const FString& GetOutputTriggerName()
+		{
+			static const FString Name = TEXT("On Set");
+			return Name;
+		}
+
+		const FText& GetOutputTriggerDescription()
+		{
+			static const FText Desc = LOCTEXT("TriggerRouteOnSetDesc", "Triggered when any of the input triggers are set.");
+			return Desc;
+		}
+
+		const FString& GetOutputValueName()
+		{
+			static const FString Name = TEXT("Value");
+			return Name;
+		}
+
+		const FText& GetOutputValueDescription()
+		{
+			static const FText Desc = LOCTEXT("TriggerRouteOutputValueDesc", "The output value set by the input triggers.");
+			return Desc;
+		}
+
+	}
+
+	REGISTER_TRIGGER_ROUTE_NOTE(int32, 2)
+	REGISTER_TRIGGER_ROUTE_NOTE(int32, 3)
+	REGISTER_TRIGGER_ROUTE_NOTE(int32, 4)
+	REGISTER_TRIGGER_ROUTE_NOTE(int32, 5)
+
+	REGISTER_TRIGGER_ROUTE_NOTE(float, 2)
+	REGISTER_TRIGGER_ROUTE_NOTE(float, 3)
+	REGISTER_TRIGGER_ROUTE_NOTE(float, 4)
+	REGISTER_TRIGGER_ROUTE_NOTE(float, 5)
+
+	REGISTER_TRIGGER_ROUTE_NOTE(bool, 2)
+	REGISTER_TRIGGER_ROUTE_NOTE(bool, 3)
+	REGISTER_TRIGGER_ROUTE_NOTE(bool, 4)
+	REGISTER_TRIGGER_ROUTE_NOTE(bool, 5)
+}
+
+#undef LOCTEXT_NAMESPACE // MetasoundTriggerDelayNode
diff --git a/Engine/Plugins/Runtime/Metasound/Source/MetasoundStandardNodes/Public/MetasoundTriggerRouteNode.h b/Engine/Plugins/Runtime/Metasound/Source/MetasoundStandardNodes/Public/MetasoundTriggerRouteNode.h
new file mode 100644
--- /dev/null
+++ b/Engine/Plugins/Runtime/Metasound/Source/MetasoundStandardNodes/Public/MetasoundTriggerRouteNode.h
@@ -0,0 +1,181 @@
+// Copyright Epic Games, Inc. All Rights Reserved.
+
+#pragma once
+
+#include "CoreMinimal.h"
+#include "MetasoundFacade.h"
+#include "MetasoundExecutableOperator.h"
+#include "MetasoundPrimitives.h"
+#include "MetasoundStandardNodesNames.h"
+#include "MetasoundTrigger.h"
+#include "Internationalization/Text.h"
+
+#define LOCTEXT_NAMESPACE "MetasoundStandardNodes"
+
+
+namespace Metasound
+{
+	namespace MetasoundTriggerRouteNodePrivate
+	{
+		METASOUNDSTANDARDNODES_API FNodeClassMetadata CreateNodeClassMetadata(const FName& InDataTypeName, const FName& InOperatorName, const FText& InDisplayName, const FText& InDescription, const FVertexInterface& InDefaultInterface);
+	}
+
+	namespace TriggerRouteVertexNames
+	{
+		METASOUNDSTANDARDNODES_API const FString GetInputTriggerName(uint32 InIndex);
+		METASOUNDSTANDARDNODES_API const FText GetInputTriggerDescription(uint32 InIndex);
+		METASOUNDSTANDARDNODES_API const FString GetInputValueName(uint32 InIndex);
+		METASOUNDSTANDARDNODES_API const FText GetInputValueDescription(uint32 InIndex);
+
+		METASOUNDSTANDARDNODES_API const FString& GetOutputTriggerName();
+		METASOUNDSTANDARDNODES_API const FText& GetOutputTriggerDescription();
+		METASOUNDSTANDARDNODES_API const FString& GetOutputValueName();
+		METASOUNDSTANDARDNODES_API const FText& GetOutputValueDescription();
+	}
+
+	template<typename ValueType, uint32 NumInputs>
+	class TTriggerRouteOperator : public TExecutableOperator<TTriggerRouteOperator<ValueType, NumInputs>>
+	{
+	public:
+		static const FVertexInterface& GetDefaultInterface()
+		{
+			using namespace TriggerRouteVertexNames;
+
+			auto CreateDefaultInterface = []() -> FVertexInterface
+			{
+				FInputVertexInterface InputInterface;
+				for (uint32 i = 0; i < NumInputs; ++i)
+				{
+					InputInterface.Add(TInputDataVertexModel<FTrigger>(GetInputTriggerName(i), GetInputTriggerDescription(i)));
+					InputInterface.Add(TInputDataVertexModel<ValueType>(GetInputValueName(i), GetInputValueDescription(i)));
+				}
+
+				FOutputVertexInterface OutputInterface;
+				OutputInterface.Add(TOutputDataVertexModel<FTrigger>(GetOutputTriggerName(), GetOutputTriggerDescription()));
+				OutputInterface.Add(TOutputDataVertexModel<ValueType>(GetOutputValueName(), GetOutputValueDescription()));
+				return FVertexInterface(InputInterface, OutputInterface);
+			};
+
+			static const FVertexInterface DefaultInterface = CreateDefaultInterface();
+			return DefaultInterface;
+		}
+
+		static const FNodeClassMetadata& GetNodeInfo()
+		{
+			auto CreateNodeClassMetadata = []() -> FNodeClassMetadata
+			{
+				FName DataTypeName = GetMetasoundDataTypeName<ValueType>();
+				FName OperatorName = *FString::Printf(TEXT("Trigger Route (%s, %d)"), *DataTypeName.ToString(), NumInputs);
+				FText NodeDisplayName = FText::Format(LOCTEXT("TriggerRouteDisplayNamePattern", "Trigger Route ({0}, {1})"), FText::FromString(GetMetasoundDataTypeString<ValueType>()), NumInputs);
+				FText NodeDescription = LOCTEXT("TriggerRouteDescription", "Allows routing different values to the same output pin depending on trigger inputs.");
+				FVertexInterface NodeInterface = GetDefaultInterface();
+
+				return MetasoundTriggerRouteNodePrivate::CreateNodeClassMetadata(DataTypeName, OperatorName, NodeDisplayName, NodeDescription, NodeInterface);
+			};
+
+			static const FNodeClassMetadata Metadata = CreateNodeClassMetadata();
+			return Metadata;
+		}
+
+		static TUniquePtr<IOperator> CreateOperator(const FCreateOperatorParams& InParams, TArray<TUniquePtr<IOperatorBuildError>>& OutErrors)
+		{
+			using namespace TriggerRouteVertexNames;
+
+			const FInputVertexInterface& InputInterface = InParams.Node.GetVertexInterface().GetInputInterface();
+			const FDataReferenceCollection& InputCollection = InParams.InputDataReferences;
+
+			TArray<FTriggerReadRef> InputTriggers;
+			TArray<TDataReadReference<ValueType>> InputValues;
+
+			for (uint32 i = 0; i < NumInputs; ++i)
+			{
+				InputTriggers.Add(InputCollection.GetDataReadReferenceOrConstruct<FTrigger>(GetInputTriggerName(i), InParams.OperatorSettings));
+				InputValues.Add(InputCollection.GetDataReadReferenceOrConstructWithVertexDefault<ValueType>(InputInterface, GetInputValueName(i)));
+			}
+
+			return MakeUnique<TTriggerRouteOperator<ValueType, NumInputs>>(InParams.OperatorSettings, MoveTemp(InputTriggers), MoveTemp(InputValues));
+		}
+
+		TTriggerRouteOperator(const FOperatorSettings& InSettings, const TArray<FTriggerReadRef>&& InInputTriggers, TArray<TDataReadReference<ValueType>>&& InInputValues)
+			: InputTriggers(InInputTriggers)
+			, InputValues(InInputValues)
+			, OutputTrigger(FTriggerWriteRef::CreateNew(InSettings))
+			, OutputValue(TDataWriteReferenceFactory<ValueType>::CreateAny(InSettings))
+		{
+			*OutputValue = *InputValues[0];
+		}
+
+		virtual ~TTriggerRouteOperator() = default;
+
+
+		virtual FDataReferenceCollection GetInputs() const override
+		{
+			using namespace TriggerRouteVertexNames;
+
+			FDataReferenceCollection Inputs;
+			for (uint32 i = 0; i < NumInputs; ++i)
+			{
+				Inputs.AddDataReadReference(GetInputTriggerName(i), InputTriggers[i]);
+				Inputs.AddDataReadReference(GetInputValueName(i), InputValues[i]);
+			}
+
+			return Inputs;
+		}
+
+		virtual FDataReferenceCollection GetOutputs() const override
+		{
+			using namespace TriggerRouteVertexNames;
+
+			FDataReferenceCollection Outputs;
+			Outputs.AddDataReadReference(GetOutputTriggerName(), OutputTrigger);
+			Outputs.AddDataReadReference(GetOutputValueName(), OutputValue);
+
+			return Outputs;
+		}
+
+		void Execute()
+		{
+			OutputTrigger->AdvanceBlock();
+
+			for (uint32 i = 0; i < NumInputs; ++i)
+			{
+				InputTriggers[i]->ExecuteBlock(
+					[&](int32 StartFrame, int32 EndFrame)
+					{
+					},
+					[this, i](int32 StartFrame, int32 EndFrame)
+					{
+						*OutputValue = *InputValues[i];
+						OutputTrigger->TriggerFrame(StartFrame);
+					}
+				);
+			}
+		}
+
+	private:
+
+		TArray<FTriggerReadRef> InputTriggers;
+		TArray<TDataReadReference<ValueType>> InputValues;
+
+		FTriggerWriteRef OutputTrigger;
+		TDataWriteReference<ValueType> OutputValue;
+	};
+
+	/** TTriggerRouteNode
+	 *
+	 *  Routes values from multiple input pins to a single output pin based on trigger inputs.
+	 */
+	template<typename ValueType, uint32 NumInputs>
+	class METASOUNDSTANDARDNODES_API TTriggerRouteNode : public FNodeFacade
+	{
+	public:
+		/**
+		 * Constructor used by the Metasound Frontend.
+		 */
+		TTriggerRouteNode(const FNodeInitData& InInitData)
+			: FNodeFacade(InInitData.InstanceName, InInitData.InstanceID, TFacadeOperatorClass<TTriggerRouteOperator<ValueType, NumInputs>>())
+		{}
+
+		virtual ~TTriggerRouteNode() = default;
+	};
+}
