Trigger Accumulator redesign

diff --git a/Engine/Plugins/Runtime/Metasound/Source/MetasoundStandardNodes/Private/MetasoundTriggerAccumulatorNode.cpp b/Engine/Plugins/Runtime/Metasound/Source/MetasoundStandardNodes/Private/MetasoundTriggerAccumulatorNode.cpp
--- a/Engine/Plugins/Runtime/Metasound/Source/MetasoundStandardNodes/Private/MetasoundTriggerAccumulatorNode.cpp
+++ b/Engine/Plugins/Runtime/Metasound/Source/MetasoundStandardNodes/Private/MetasoundTriggerAccumulatorNode.cpp
@@ -2,233 +2,85 @@
 #include "MetasoundTriggerAccumulatorNode.h"
 
 #include "Internationalization/Text.h"
-#include "MetasoundExecutableOperator.h"
 #include "MetasoundNodeRegistrationMacro.h"
-#include "MetasoundPrimitives.h"
-#include "MetasoundTrigger.h"
 #include "MetasoundStandardNodesNames.h"
 #include "MetasoundVertex.h"
 #include "MetasoundStandardNodesCategories.h"
 
 #define LOCTEXT_NAMESPACE "MetasoundStandardNodes"
 
+#define REGISTER_TRIGGER_ACCUMULATOR_NODE(Number) \
+	using FTriggerAccumulatorNode_##Number = TTriggerAccumulatorNode<Number>; \
+	METASOUND_REGISTER_NODE(FTriggerAccumulatorNode_##Number) \
+
 namespace Metasound
 {
-	METASOUND_REGISTER_NODE(FTriggerAccumulatorNode)
-
-	class FTriggerAccumulatorOperator : public TExecutableOperator<FTriggerAccumulatorOperator>
-	{
-		public:
-			static const FNodeClassMetadata& GetNodeInfo();
-			static FVertexInterface DeclareVertexInterface();
-			static TUniquePtr<IOperator> CreateOperator(const FCreateOperatorParams& InParams, FBuildErrorArray& OutErrors);
-
-			~FTriggerAccumulatorOperator() = default;
-			
-			FTriggerAccumulatorOperator(
-				const FOperatorSettings& InSettings,
-				const FTriggerReadRef& InInputTriggerToAccumulate,
-				const FTriggerReadRef& InInputTriggerResetCount,
-				const FInt32ReadRef& InInputTriggerAtCount);
-
-			virtual FDataReferenceCollection GetInputs() const override;
-			virtual FDataReferenceCollection GetOutputs() const override;
-			
-			void Execute();
-
-		private:
-			void Reset(int32 InCurretFrameIndex);
-			void Accumulate(int32 CurrentFrameIndex);
-
-			// Input. (triggers)
-			FTriggerReadRef TriggerToAccumulate;
-			FTriggerReadRef ResetTrigger;
-			
-			// Input (values)
-			FInt32ReadRef TriggerCount;
-
-			// Outputs (triggers)
-			FTriggerWriteRef CountReached;
-
-			// Output (values).
-			FInt32WriteRef CurrentCount;
-
-			static constexpr const TCHAR* AccumulatePinName = TEXT("In");
-			static constexpr const TCHAR* TriggerCountPinName = TEXT("TriggerCount");
-			static constexpr const TCHAR* ResetPinName = TEXT("Reset");
-			static constexpr const TCHAR* CountReachedPinName = TEXT("Out");
-			static constexpr const TCHAR* CurrentCountPinName = TEXT("Count");
-	};
-
-	FTriggerAccumulatorOperator::FTriggerAccumulatorOperator(
-		const FOperatorSettings& InSettings,
-		const FTriggerReadRef& InInputTriggerToAccumulate,
-		const FTriggerReadRef& InInputTriggerResetCount,
-		const FInt32ReadRef& InInputTriggerAtCount)
-		: TriggerToAccumulate(InInputTriggerToAccumulate)
-		, ResetTrigger(InInputTriggerResetCount)
-		, TriggerCount(InInputTriggerAtCount)
-		, CountReached(FTriggerWriteRef::CreateNew(InSettings))
-		, CurrentCount(FInt32WriteRef::CreateNew(0))
-	{}
-
-	FDataReferenceCollection FTriggerAccumulatorOperator::GetInputs() const
+	namespace MetasoundTriggerAccumulatorNodePrivate
 	{
-		FDataReferenceCollection InputDataReferences;
-		InputDataReferences.AddDataReadReference(AccumulatePinName, FTriggerReadRef(TriggerToAccumulate));
-		InputDataReferences.AddDataReadReference(ResetPinName,	FTriggerReadRef(ResetTrigger));
-		InputDataReferences.AddDataReadReference(TriggerCountPinName, FInt32ReadRef(TriggerCount));
-		return InputDataReferences;
-	}
-
-	FDataReferenceCollection FTriggerAccumulatorOperator::GetOutputs() const
-	{
-		FDataReferenceCollection OutputDataReferences;
-		OutputDataReferences.AddDataReadReference(CountReachedPinName, FTriggerWriteRef(CountReached));
-		OutputDataReferences.AddDataReadReference(CurrentCountPinName, FInt32WriteRef(CurrentCount));
-		return OutputDataReferences;
+		FNodeClassMetadata CreateNodeClassMetadata(const FName& InOperatorName, const FText& InDisplayName, const FText& InDescription, const FVertexInterface& InDefaultInterface)
+		{
+			FNodeClassMetadata Metadata
+			{
+				FNodeClassName{FName("TriggerAccumulator"), InOperatorName, TEXT("")},
+				1, // Major Version
+				0, // Minor Version
+				InDisplayName,
+				InDescription,
+				PluginAuthor,
+				PluginNodeMissingPrompt,
+				InDefaultInterface,
+				{StandardNodes::TriggerUtils},
+				{TEXT("TriggerAccumulator")},
+				FNodeDisplayStyle{}
+			};
+
+			return Metadata;
+		}
 	}
 
-	void FTriggerAccumulatorOperator::Reset(int32 InCurrentFrameIndex)
+	namespace TriggerAccumulatorVertexNames
 	{
-		*CurrentCount = 0;
-	}
-	void FTriggerAccumulatorOperator::Accumulate(int32 InCurrentFrameIndex)
-	{
-		++(*CurrentCount);
-		if (*CurrentCount >= *TriggerCount)
+		const FString GetInputAutoResetName()
 		{
-			Reset(InCurrentFrameIndex);
-			CountReached->TriggerFrame(InCurrentFrameIndex);
+			return TEXT("Auto Reset");
 		}
-	}
-
-	void FTriggerAccumulatorOperator::Execute()
-	{		
-		// We need to advance our own output triggers.
-		CountReached->AdvanceBlock();
-
-		// Query the number of triggers in the current block.
-		int32 NumResets = ResetTrigger->NumTriggeredInBlock();
-		int32 NumAccumulates = TriggerToAccumulate->NumTriggeredInBlock();
-		int32 MaxTriggers = FMath::Max(NumResets, NumAccumulates);
 
-		// Indices into each trigger array
-		int32 ResetIndex = 0;
-		int32 AccumulateIndex = 0;
-		
-		// Move forwards through the trigger array and return next trigger frame index, or INDEX_NONE if there's none remaining.
-		static const auto GetNextTriggerFrame = [](const FTriggerReadRef& InTrigger, int32& InOutIndex, const int32 InNumTriggers) 
-			-> int32 { return InOutIndex < InNumTriggers ? (*InTrigger)[InOutIndex++] : INDEX_NONE; };
-		
-		// Convenience lambdas for getting the next Reset/Accumulate frame index from either the Trigger or Reset.
-		auto GetNextResetFrame = [this, &ResetIndex, NumResets]() -> int32 { return GetNextTriggerFrame(ResetTrigger, ResetIndex, NumResets); };
-		auto GetNextAccumulateFrame = [this, &AccumulateIndex, NumAccumulates]() -> int32 { return GetNextTriggerFrame(TriggerToAccumulate, AccumulateIndex, NumAccumulates); };
-
-		int32 NextResetFrame = GetNextResetFrame();
-		int32 NextAccumulateFrame = GetNextAccumulateFrame();
-
-		// Walk through each trigger in order.
-		for (int32 i = 0; i < MaxTriggers; ++i)
+		const FText GetInputAutoResetDescription()
 		{
-			// If there's no Reset, we can just process all Accumulates in order.
-			if (NextResetFrame == INDEX_NONE)
-			{
-				Accumulate(NextAccumulateFrame);
-				NextAccumulateFrame = GetNextAccumulateFrame();
-			}
-			// ... likewise if there's no Accumulate, we can just do the resets.
-			else if (NextAccumulateFrame == INDEX_NONE)
-			{
-				Reset(NextResetFrame);
-				NextResetFrame = GetNextResetFrame();
-			}
-			// ... both resets and accumulates remaining.
-			else
-			{
-				if (NextAccumulateFrame < NextResetFrame)
-				{
-					Accumulate(NextAccumulateFrame);
-					NextAccumulateFrame = GetNextAccumulateFrame();
-				}
-				else if(NextResetFrame < NextAccumulateFrame)
-				{
-					Reset(NextResetFrame);
-					NextResetFrame = GetNextResetFrame();
-				}
-				else // Both on same frame, make sure Reset happens first.
-				{
-					Reset(NextResetFrame);
-					Accumulate(NextAccumulateFrame);
-					NextAccumulateFrame = GetNextAccumulateFrame();
-					NextResetFrame = GetNextResetFrame();
-				}
-			}
+			return LOCTEXT("TriggerAccumulatorAutoResetDesc", "Set to true to automatically reset the trigger accumulator state once the output is triggered.");
 		}
-	}
 
-	FVertexInterface FTriggerAccumulatorOperator::DeclareVertexInterface()
-	{
-		static const FVertexInterface Interface(
-			FInputVertexInterface(
-				TInputDataVertexModel<FTrigger>(AccumulatePinName, LOCTEXT("TriggerToAccumulateTooltip", "Trigger to accumulate")),
-				TInputDataVertexModel<FTrigger>(ResetPinName, LOCTEXT("ResetAccumulatorTooltip", "Reset the accumulator")),
-				TInputDataVertexModel<int32>(TriggerCountPinName, LOCTEXT("CountToTriggerAtTooltip", "Count to trigger at"), 1)
-			),
-			FOutputVertexInterface(
-				TOutputDataVertexModel<FTrigger>(CountReachedPinName, LOCTEXT("CountReachedTooltop", "Triggered when the accumulated count is reached")),
-				TOutputDataVertexModel<int32>(CurrentCountPinName, LOCTEXT("CountTooltop", "Current accumulator count"))
-			)
-		);
-
-		return Interface;
-	}
-
-	const FNodeClassMetadata& FTriggerAccumulatorOperator::GetNodeInfo()
-	{
-		auto InitNodeInfo = []() -> FNodeClassMetadata
+		const FString GetInputTriggerName(uint32 InIndex)
 		{
-			FNodeClassMetadata Info;
-			Info.ClassName = {Metasound::StandardNodes::Namespace, TEXT("TriggerAccumulator"), TEXT("")};
-			Info.MajorVersion = 1;
-			Info.MinorVersion = 0;
-			Info.DisplayName = LOCTEXT("Metasound_TriggerAccumulatorNodeDisplayName", "Trigger Accumulator");
-			Info.Description = LOCTEXT("Metasound_TriggerAccumulatorNodeDescription", "Accumulates Triggers until a Count is hit and then it Triggers");
-			Info.Author = PluginAuthor;
-			Info.PromptIfMissing = PluginNodeMissingPrompt;
-			Info.DefaultInterface = DeclareVertexInterface();
-			Info.CategoryHierarchy.Emplace(StandardNodes::TriggerUtils);
-
-			return Info;
-		};
+			return FText::Format(LOCTEXT("TriggerAccumulatorTriggerInputName", "In {0}"), InIndex).ToString();
+		}
 
-		static const FNodeClassMetadata Info = InitNodeInfo();
+		const FText GetInputTriggerDescription(uint32 InIndex)
+		{
+			return FText::Format(LOCTEXT("TriggerAccumulatorInputTriggerDesc", "Trigger {0} input. All trigger inputs must be triggered before the output trigger is hit."), InIndex);
+		}
 
-		return Info;
-	}
+		const FString& GetOutputTriggerName()
+		{
+			static const FString Name = TEXT("Out");
+			return Name;
+		}
 
-	TUniquePtr<IOperator> FTriggerAccumulatorOperator::CreateOperator(const FCreateOperatorParams& InParams, FBuildErrorArray& OutErrors)
-	{
-		const FTriggerAccumulatorNode& TaNode = static_cast<const FTriggerAccumulatorNode&>(InParams.Node);
-		const FDataReferenceCollection& InputCollection = InParams.InputDataReferences;
-				
-		return MakeUnique<FTriggerAccumulatorOperator>(
-			InParams.OperatorSettings, 
-			InputCollection.GetDataReadReferenceOrConstruct<FTrigger>(AccumulatePinName, InParams.OperatorSettings),
-			InputCollection.GetDataReadReferenceOrConstruct<FTrigger>(ResetPinName, InParams.OperatorSettings),
-			InputCollection.GetDataReadReferenceOrConstruct<int32>(TriggerCountPinName, TaNode.GetDefaultTriggerAt() )
-		);
+		const FText& GetOutputTriggerDescription()
+		{
+			static const FText Desc = LOCTEXT("TriggerAccumulateOutputTriggerDesc", "Triggered when all input triggers have been triggered. Call Reset to reset the state or use \"Auto Reset\"");
+			return Desc;
+		}
 	}
 
-	// Node.
-	FTriggerAccumulatorNode::FTriggerAccumulatorNode(const FString& InName, const FGuid& InInstanceID, int32 InDefaultTriggerAtCount)
-		: FNodeFacade(InName, InInstanceID, TFacadeOperatorClass<FTriggerAccumulatorOperator>())
-		, DefaultTriggerAtCount(InDefaultTriggerAtCount)
-	{}
-
-	FTriggerAccumulatorNode::FTriggerAccumulatorNode(const FNodeInitData& InitData)
-		: FTriggerAccumulatorNode(InitData.InstanceName, InitData.InstanceID, 1)
-	{}
+	REGISTER_TRIGGER_ACCUMULATOR_NODE(1)
+	REGISTER_TRIGGER_ACCUMULATOR_NODE(2)
+	REGISTER_TRIGGER_ACCUMULATOR_NODE(3)
+	REGISTER_TRIGGER_ACCUMULATOR_NODE(4)
+	REGISTER_TRIGGER_ACCUMULATOR_NODE(5)
+	REGISTER_TRIGGER_ACCUMULATOR_NODE(6)
+	REGISTER_TRIGGER_ACCUMULATOR_NODE(7)
+	REGISTER_TRIGGER_ACCUMULATOR_NODE(8)
 }
 
 #undef LOCTEXT_NAMESPACE
diff --git a/Engine/Plugins/Runtime/Metasound/Source/MetasoundStandardNodes/Public/MetasoundTriggerAccumulatorNode.h b/Engine/Plugins/Runtime/Metasound/Source/MetasoundStandardNodes/Public/MetasoundTriggerAccumulatorNode.h
--- a/Engine/Plugins/Runtime/Metasound/Source/MetasoundStandardNodes/Public/MetasoundTriggerAccumulatorNode.h
+++ b/Engine/Plugins/Runtime/Metasound/Source/MetasoundStandardNodes/Public/MetasoundTriggerAccumulatorNode.h
@@ -2,39 +2,202 @@
 #pragma once
 
 #include "CoreMinimal.h"
-
 #include "MetasoundFacade.h"
+#include "MetasoundExecutableOperator.h"
+#include "MetasoundPrimitives.h"
+#include "MetasoundStandardNodesNames.h"
+#include "MetasoundTrigger.h"
+#include "Internationalization/Text.h"
+
+#define LOCTEXT_NAMESPACE "MetasoundStandardNodes"
 
 namespace Metasound
 {
-	/** FTriggerAccumulatorNode
-	 *
-	 *  Creates a Trigger Accumulator Node, that accumulate triggers until hitting a count, then Trigger and reset.
-	 */
-	class METASOUNDSTANDARDNODES_API FTriggerAccumulatorNode : public FNodeFacade
+	namespace MetasoundTriggerAccumulatorNodePrivate
+	{
+		METASOUNDSTANDARDNODES_API FNodeClassMetadata CreateNodeClassMetadata(const FName& InOperatorName, const FText& InDisplayName, const FText& InDescription, const FVertexInterface& InDefaultInterface);
+	}
+
+	namespace TriggerAccumulatorVertexNames
+	{
+		METASOUNDSTANDARDNODES_API const FString GetInputAutoResetName();
+		METASOUNDSTANDARDNODES_API const FText GetInputAutoResetDescription();
+
+		METASOUNDSTANDARDNODES_API const FString GetInputTriggerName(uint32 InIndex);
+		METASOUNDSTANDARDNODES_API const FText GetInputTriggerDescription(uint32 InIndex);
+
+		METASOUNDSTANDARDNODES_API const FString& GetOutputTriggerName();
+		METASOUNDSTANDARDNODES_API const FText& GetOutputTriggerDescription();
+	}
+
+	template<uint32 NumInputs>
+	class TTriggerAccumulatorOperator : public TExecutableOperator<TTriggerAccumulatorOperator<NumInputs>>
 	{
 	public:
+		static const FVertexInterface& GetDefaultInterface()
+		{
+			using namespace TriggerAccumulatorVertexNames;
 
-		/** Trigger Accumulator node constructor.
-		 *
-		 * @param InName - Name of this node.
-		 * @param InTriggerAtCount - The count at which after accumulated triggers we trigger.
-		 */
-		FTriggerAccumulatorNode(const FString& InName, const FGuid& InInstanceID, int32 InDefaultTriggerAtCount);
+			auto CreateDefaultInterface = []() -> FVertexInterface
+			{
+				FInputVertexInterface InputInterface;
 
-		/**
-		 * Constructor used by the Metasound Frontend.
-		 */
-		FTriggerAccumulatorNode(const FNodeInitData& InitData);
+				InputInterface.Add(TInputDataVertexModel<bool>(GetInputAutoResetName(), GetInputAutoResetDescription()));
 
-		virtual ~FTriggerAccumulatorNode() = default;
+				for (uint32 i = 0; i < NumInputs; ++i)
+				{
+					InputInterface.Add(TInputDataVertexModel<FTrigger>(GetInputTriggerName(i), GetInputTriggerDescription(i)));
+				}
 
-		/** 
-		 * Get Default Trigger At Value.
-		 */
-		int32 GetDefaultTriggerAt() const { return DefaultTriggerAtCount; }
+				FOutputVertexInterface OutputInterface;
+				OutputInterface.Add(TOutputDataVertexModel<FTrigger>(GetOutputTriggerName(), GetOutputTriggerDescription()));
+
+				return FVertexInterface(InputInterface, OutputInterface);
+			};
+
+			static const FVertexInterface DefaultInterface = CreateDefaultInterface();
+			return DefaultInterface;
+		}
+
+		static const FNodeClassMetadata& GetNodeInfo()
+		{
+			auto CreateNodeClassMetadata = []() -> FNodeClassMetadata
+			{
+				FName OperatorName = *FString::Printf(TEXT("Trigger Accumulate (%d)") , NumInputs);
+				FText NodeDisplayName = FText::Format(LOCTEXT("TriggerAccumulateDisplayNamePattern", "Trigger Accumulate ({0})"), NumInputs);
+				FText NodeDescription = LOCTEXT("TriggerAccumulateDescription", "Will trigger output once all input triggers have been hit at some point in the past.");
+				FVertexInterface NodeInterface = GetDefaultInterface();
+
+				return MetasoundTriggerAccumulatorNodePrivate::CreateNodeClassMetadata(OperatorName, NodeDisplayName, NodeDescription, NodeInterface);
+			};
+
+			static const FNodeClassMetadata Metadata = CreateNodeClassMetadata();
+			return Metadata;
+		}
+
+		static TUniquePtr<IOperator> CreateOperator(const FCreateOperatorParams& InParams, TArray<TUniquePtr<IOperatorBuildError>>& OutErrors)
+		{
+			using namespace TriggerAccumulatorVertexNames;
+
+			const FInputVertexInterface& InputInterface = InParams.Node.GetVertexInterface().GetInputInterface();
+			const FDataReferenceCollection& InputCollection = InParams.InputDataReferences;
+
+			FBoolReadRef bInAutoReset = InputCollection.GetDataReadReferenceOrConstructWithVertexDefault<bool>(InputInterface, GetInputAutoResetName());
+			TArray<FTriggerReadRef> InputTriggers;
+
+			for (uint32 i = 0; i < NumInputs; ++i)
+			{
+				InputTriggers.Add(InputCollection.GetDataReadReferenceOrConstruct<FTrigger>(GetInputTriggerName(i), InParams.OperatorSettings));
+			}
+
+			return MakeUnique<TTriggerAccumulatorOperator<NumInputs>>(InParams.OperatorSettings, bInAutoReset, MoveTemp(InputTriggers));
+		}
+
+		TTriggerAccumulatorOperator(const FOperatorSettings& InSettings, 
+			const FBoolReadRef& bInAutoReset, 
+			const TArray<FTriggerReadRef>&& InInputTriggers)
+			: bAutoReset(bInAutoReset)
+			, InputTriggers(InInputTriggers)
+			, OutputTrigger(FTriggerWriteRef::CreateNew(InSettings))
+		{
+			ResetTriggerState();
+		}
+
+		virtual ~TTriggerAccumulatorOperator() = default;
+
+
+		virtual FDataReferenceCollection GetInputs() const override
+		{
+			using namespace TriggerAccumulatorVertexNames;
+
+			FDataReferenceCollection Inputs;
+			Inputs.AddDataReadReference(GetInputAutoResetName(), bAutoReset);
+			for (uint32 i = 0; i < NumInputs; ++i)
+			{
+				Inputs.AddDataReadReference(GetInputTriggerName(i), InputTriggers[i]);
+			}
+
+			return Inputs;
+		}
+
+		virtual FDataReferenceCollection GetOutputs() const override
+		{
+			using namespace TriggerAccumulatorVertexNames;
+
+			FDataReferenceCollection Outputs;
+			Outputs.AddDataReadReference(GetOutputTriggerName(), OutputTrigger);
+
+			return Outputs;
+		}
+
+		void TriggerOutputIfReady(int32 InStartFrame)
+		{
+			for (bool WasTriggered : bInputWasTriggered)
+			{
+				if (!WasTriggered)
+				{
+					return;
+				}
+			}
+
+			if (*bAutoReset)
+			{
+				ResetTriggerState();
+			}
+
+			OutputTrigger->TriggerFrame(InStartFrame);
+		}
+
+		void ResetTriggerState()
+		{
+			bInputWasTriggered.Reset();
+			bInputWasTriggered.AddDefaulted(InputTriggers.Num());
+		}
+
+		void Execute()
+		{
+			OutputTrigger->AdvanceBlock();
+
+			for (uint32 i = 0; i < NumInputs; ++i)
+			{
+				InputTriggers[i]->ExecuteBlock(
+					[&](int32 StartFrame, int32 EndFrame)
+					{
+					},
+					[this, i](int32 StartFrame, int32 EndFrame)
+					{
+						bInputWasTriggered[i] = true;
+						TriggerOutputIfReady(StartFrame);
+					}
+				);
+			}
+		}
 
 	private:
-		int32 DefaultTriggerAtCount = 1;
+
+		FBoolReadRef bAutoReset;
+		TArray<FTriggerReadRef> InputTriggers;
+		FTriggerWriteRef OutputTrigger;
+
+		TArray<bool> bInputWasTriggered;
+	};
+
+	/** TTriggerAccumulatorNode
+	*
+	*  Routes values from multiple input pins to a single output pin based on trigger inputs.
+	*/
+	template<uint32 NumInputs>
+	class METASOUNDSTANDARDNODES_API TTriggerAccumulatorNode : public FNodeFacade
+	{
+	public:
+		/**
+		 * Constructor used by the Metasound Frontend.
+		 */
+		TTriggerAccumulatorNode(const FNodeInitData& InInitData)
+			: FNodeFacade(InInitData.InstanceName, InInitData.InstanceID, TFacadeOperatorClass<TTriggerAccumulatorOperator<NumInputs>>())
+		{}
+
+		virtual ~TTriggerAccumulatorNode() = default;
 	};
+
 } // namespace Metasound
