Trigger Any Node

diff --git a/Engine/Plugins/Runtime/Metasound/Source/MetasoundStandardNodes/Private/MetasoundTriggerAnyNode.cpp b/Engine/Plugins/Runtime/Metasound/Source/MetasoundStandardNodes/Private/MetasoundTriggerAnyNode.cpp
new file mode 100644
--- /dev/null
+++ b/Engine/Plugins/Runtime/Metasound/Source/MetasoundStandardNodes/Private/MetasoundTriggerAnyNode.cpp
@@ -0,0 +1,76 @@
+// Copyright Epic Games, Inc. All Rights Reserved.
+#include "MetasoundTriggerAnyNode.h"
+
+#include "Internationalization/Text.h"
+#include "MetasoundNodeRegistrationMacro.h"
+#include "MetasoundStandardNodesNames.h"
+#include "MetasoundVertex.h"
+#include "MetasoundStandardNodesCategories.h"
+
+#define LOCTEXT_NAMESPACE "MetasoundStandardNodes"
+
+#define REGISTER_TRIGGER_ANY_NODE(Number) \
+	using FTriggerAnyNode_##Number = TTriggerAnyNode<Number>; \
+	METASOUND_REGISTER_NODE(FTriggerAnyNode_##Number) \
+
+namespace Metasound
+{
+	namespace MetasoundTriggerAnyNodePrivate
+	{
+		FNodeClassMetadata CreateNodeClassMetadata(const FName& InOperatorName, const FText& InDisplayName, const FText& InDescription, const FVertexInterface& InDefaultInterface)
+		{
+			FNodeClassMetadata Metadata
+			{
+				FNodeClassName{FName("TriggerAny"), InOperatorName, TEXT("")},
+				1, // Major Version
+				0, // Minor Version
+				InDisplayName,
+				InDescription,
+				PluginAuthor,
+				PluginNodeMissingPrompt,
+				InDefaultInterface,
+				{StandardNodes::TriggerUtils},
+				{TEXT("TriggerAny")},
+				FNodeDisplayStyle{}
+			};
+
+			return Metadata;
+		}
+	}
+
+	namespace TriggerAnyVertexNames
+	{
+		const FString GetInputTriggerName(uint32 InIndex)
+		{
+			return FText::Format(LOCTEXT("TriggerAnyTriggerInputName", "In {0}"), InIndex).ToString();
+		}
+
+		const FText GetInputTriggerDescription(uint32 InIndex)
+		{
+			return FText::Format(LOCTEXT("TriggerAnyInputTriggerDesc", "Trigger {0} input. All trigger inputs must be triggered before the output trigger is hit."), InIndex);
+		}
+
+		const FString& GetOutputTriggerName()
+		{
+			static const FString Name = TEXT("Out");
+			return Name;
+		}
+
+		const FText& GetOutputTriggerDescription()
+		{
+			static const FText Desc = LOCTEXT("TriggerAccumulateOutputTriggerDesc", "Triggered when all input triggers have been triggered. Call Reset to reset the state or use \"Auto Reset\"");
+			return Desc;
+		}
+	}
+
+	REGISTER_TRIGGER_ANY_NODE(1)
+	REGISTER_TRIGGER_ANY_NODE(2)
+	REGISTER_TRIGGER_ANY_NODE(3)
+	REGISTER_TRIGGER_ANY_NODE(4)
+	REGISTER_TRIGGER_ANY_NODE(5)
+	REGISTER_TRIGGER_ANY_NODE(6)
+	REGISTER_TRIGGER_ANY_NODE(7)
+	REGISTER_TRIGGER_ANY_NODE(8)
+}
+
+#undef LOCTEXT_NAMESPACE
diff --git a/Engine/Plugins/Runtime/Metasound/Source/MetasoundStandardNodes/Public/MetasoundTriggerAnyNode.h b/Engine/Plugins/Runtime/Metasound/Source/MetasoundStandardNodes/Public/MetasoundTriggerAnyNode.h
new file mode 100644
--- /dev/null
+++ b/Engine/Plugins/Runtime/Metasound/Source/MetasoundStandardNodes/Public/MetasoundTriggerAnyNode.h
@@ -0,0 +1,164 @@
+// Copyright Epic Games, Inc. All Rights Reserved.
+#pragma once
+
+#include "CoreMinimal.h"
+#include "MetasoundFacade.h"
+#include "MetasoundExecutableOperator.h"
+#include "MetasoundPrimitives.h"
+#include "MetasoundStandardNodesNames.h"
+#include "MetasoundTrigger.h"
+#include "Internationalization/Text.h"
+
+#define LOCTEXT_NAMESPACE "MetasoundStandardNodes"
+
+namespace Metasound
+{
+	namespace MetasoundTriggerAnyNodePrivate
+	{
+		METASOUNDSTANDARDNODES_API FNodeClassMetadata CreateNodeClassMetadata(const FName& InOperatorName, const FText& InDisplayName, const FText& InDescription, const FVertexInterface& InDefaultInterface);
+	}
+
+	namespace TriggerAnyVertexNames
+	{
+		METASOUNDSTANDARDNODES_API const FString GetInputTriggerName(uint32 InIndex);
+		METASOUNDSTANDARDNODES_API const FText GetInputTriggerDescription(uint32 InIndex);
+
+		METASOUNDSTANDARDNODES_API const FString& GetOutputTriggerName();
+		METASOUNDSTANDARDNODES_API const FText& GetOutputTriggerDescription();
+	}
+
+	template<uint32 NumInputs>
+	class TTriggerAnyOperator : public TExecutableOperator<TTriggerAnyOperator<NumInputs>>
+	{
+	public:
+		static const FVertexInterface& GetDefaultInterface()
+		{
+			using namespace TriggerAnyVertexNames;
+
+			auto CreateDefaultInterface = []() -> FVertexInterface
+			{
+				FInputVertexInterface InputInterface;
+
+				for (uint32 i = 0; i < NumInputs; ++i)
+				{
+					InputInterface.Add(TInputDataVertexModel<FTrigger>(GetInputTriggerName(i), GetInputTriggerDescription(i)));
+				}
+
+				FOutputVertexInterface OutputInterface;
+				OutputInterface.Add(TOutputDataVertexModel<FTrigger>(GetOutputTriggerName(), GetOutputTriggerDescription()));
+
+				return FVertexInterface(InputInterface, OutputInterface);
+			};
+
+			static const FVertexInterface DefaultInterface = CreateDefaultInterface();
+			return DefaultInterface;
+		}
+
+		static const FNodeClassMetadata& GetNodeInfo()
+		{
+			auto CreateNodeClassMetadata = []() -> FNodeClassMetadata
+			{
+				FName OperatorName = *FString::Printf(TEXT("Trigger Any (%d)") , NumInputs);
+				FText NodeDisplayName = FText::Format(LOCTEXT("TriggerAnyDisplayNamePattern", "Trigger Any ({0})"), NumInputs);
+				FText NodeDescription = LOCTEXT("TriggerAnyDescription", "Will trigger output on any of the input triggers.");
+				FVertexInterface NodeInterface = GetDefaultInterface();
+
+				return MetasoundTriggerAnyNodePrivate::CreateNodeClassMetadata(OperatorName, NodeDisplayName, NodeDescription, NodeInterface);
+			};
+
+			static const FNodeClassMetadata Metadata = CreateNodeClassMetadata();
+			return Metadata;
+		}
+
+		static TUniquePtr<IOperator> CreateOperator(const FCreateOperatorParams& InParams, TArray<TUniquePtr<IOperatorBuildError>>& OutErrors)
+		{
+			using namespace TriggerAnyVertexNames;
+
+			const FInputVertexInterface& InputInterface = InParams.Node.GetVertexInterface().GetInputInterface();
+			const FDataReferenceCollection& InputCollection = InParams.InputDataReferences;
+
+			TArray<FTriggerReadRef> InputTriggers;
+
+			for (uint32 i = 0; i < NumInputs; ++i)
+			{
+				InputTriggers.Add(InputCollection.GetDataReadReferenceOrConstruct<FTrigger>(GetInputTriggerName(i), InParams.OperatorSettings));
+			}
+
+			return MakeUnique<TTriggerAnyOperator<NumInputs>>(InParams.OperatorSettings, MoveTemp(InputTriggers));
+		}
+
+		TTriggerAnyOperator(const FOperatorSettings& InSettings, const TArray<FTriggerReadRef>&& InInputTriggers)
+			: InputTriggers(InInputTriggers)
+			, OutputTrigger(FTriggerWriteRef::CreateNew(InSettings))
+		{
+		}
+
+		virtual ~TTriggerAnyOperator() = default;
+
+
+		virtual FDataReferenceCollection GetInputs() const override
+		{
+			using namespace TriggerAnyVertexNames;
+
+			FDataReferenceCollection Inputs;
+			for (uint32 i = 0; i < NumInputs; ++i)
+			{
+				Inputs.AddDataReadReference(GetInputTriggerName(i), InputTriggers[i]);
+			}
+
+			return Inputs;
+		}
+
+		virtual FDataReferenceCollection GetOutputs() const override
+		{
+			using namespace TriggerAnyVertexNames;
+
+			FDataReferenceCollection Outputs;
+			Outputs.AddDataReadReference(GetOutputTriggerName(), OutputTrigger);
+
+			return Outputs;
+		}
+
+		void Execute()
+		{
+			OutputTrigger->AdvanceBlock();
+
+			for (uint32 i = 0; i < NumInputs; ++i)
+			{
+				InputTriggers[i]->ExecuteBlock(
+					[&](int32 StartFrame, int32 EndFrame)
+					{
+					},
+					[this, i](int32 StartFrame, int32 EndFrame)
+					{
+						OutputTrigger->TriggerFrame(StartFrame);
+					}
+				);
+			}
+		}
+
+	private:
+
+		TArray<FTriggerReadRef> InputTriggers;
+		FTriggerWriteRef OutputTrigger;
+	};
+
+	/** TTriggerAnyNode
+	*
+	*  Will output a trigger whenever any of its input triggers are set.
+	*/
+	template<uint32 NumInputs>
+	class METASOUNDSTANDARDNODES_API TTriggerAnyNode : public FNodeFacade
+	{
+	public:
+		/**
+		 * Constructor used by the Metasound Frontend.
+		 */
+		TTriggerAnyNode(const FNodeInitData& InInitData)
+			: FNodeFacade(InInitData.InstanceName, InInitData.InstanceID, TFacadeOperatorClass<TTriggerAnyOperator<NumInputs>>())
+		{}
+
+		virtual ~TTriggerAnyNode() = default;
+	};
+
+} // namespace Metasound
