Remove envelope follower peak mode

diff --git a/Engine/Plugins/Runtime/Metasound/Source/MetasoundStandardNodes/Private/MetasoundEnvelopeFollowerNode.cpp b/Engine/Plugins/Runtime/Metasound/Source/MetasoundStandardNodes/Private/MetasoundEnvelopeFollowerNode.cpp
--- a/Engine/Plugins/Runtime/Metasound/Source/MetasoundStandardNodes/Private/MetasoundEnvelopeFollowerNode.cpp
+++ b/Engine/Plugins/Runtime/Metasound/Source/MetasoundStandardNodes/Private/MetasoundEnvelopeFollowerNode.cpp
@@ -19,21 +19,11 @@
 
 namespace Metasound
 {
-	DECLARE_METASOUND_ENUM(Audio::EPeakMode::Type, Audio::EPeakMode::Peak, METASOUNDSTANDARDNODES_API,
-		FEnumPeakModeType, FEnumPeakModeTypeInfo, FEnumPeakModeTypeReadRef, FEnumPeakModeTypeWriteRef);
-
-	DEFINE_METASOUND_ENUM_BEGIN(Audio::EPeakMode::Type, FEnumPeakModeType, "PeakModeType")
-		DEFINE_METASOUND_ENUM_ENTRY(Audio::EPeakMode::MeanSquared, LOCTEXT("MeanSquaredDescription", "Mean Squared"), LOCTEXT("MeanSquaredDescriptionTT", "Sets the envelope to follow the mean-squared of the audio signal")),
-		DEFINE_METASOUND_ENUM_ENTRY(Audio::EPeakMode::RootMeanSquared, LOCTEXT("RootMeanSquaredDescription", "Root Mean Squared"), LOCTEXT("RootMeanSquaredTT", "Sets the envelope to follow the root-mean-squared of the audio signal")),
-		DEFINE_METASOUND_ENUM_ENTRY(Audio::EPeakMode::Peak, LOCTEXT("PeakDescription", "Peak"), LOCTEXT("PeakDescriptionTT", "Sets the envelope to follow the peaks of the audio signal")),
-	DEFINE_METASOUND_ENUM_END()
-
 	namespace EnvelopeFollower
 	{
 		static const TCHAR* InParamNameAudioInput = TEXT("In");
 		static const TCHAR* InParamNameAttackTime = TEXT("Attack Time");
 		static const TCHAR* InParamNameReleaseTime = TEXT("Release Time");
-		static const TCHAR* InParamNamePeakModeType = TEXT("Peak Mode");
 		static const TCHAR* OutParamNameEnvelope = TEXT("Envelope");
 	}
 
@@ -48,8 +38,7 @@ namespace Metasound
 		FEnvelopeFollowerOperator(const FOperatorSettings& InSettings, 
 			const FAudioBufferReadRef& InAudioInput, 
 			const FTimeReadRef& InAttackTime, 
-			const FTimeReadRef& InReleaseTime,
-			const FEnumPeakModeTypeReadRef& InPeakModeType);
+			const FTimeReadRef& InReleaseTime);
 
 		virtual FDataReferenceCollection GetInputs() const override;
 		virtual FDataReferenceCollection GetOutputs() const override;
@@ -65,8 +54,6 @@ namespace Metasound
 		// The amount of release time
 		FTimeReadRef ReleaseTime;
 
-		FEnumPeakModeTypeReadRef PeakModeType;
-
 		// The audio output
 		FFloatWriteRef EnvelopeOutput;
 
@@ -75,25 +62,21 @@ namespace Metasound
 
 		double PrevAttackTime = 0.0;
 		double PrevReleaseTime = 0.0;
-		Audio::EPeakMode::Type PrevPeakModeType = Audio::EPeakMode::MeanSquared;
 	};
 
 	FEnvelopeFollowerOperator::FEnvelopeFollowerOperator(const FOperatorSettings& InSettings, 
 		const FAudioBufferReadRef& InAudioInput, 
 		const FTimeReadRef& InAttackTime,
-		const FTimeReadRef& InReleaseTime,
-		const FEnumPeakModeTypeReadRef& InPeakModeType)
+		const FTimeReadRef& InReleaseTime)
 		: AudioInput(InAudioInput)
 		, AttackTime(InAttackTime)
 		, ReleaseTime(InReleaseTime)
-		, PeakModeType(InPeakModeType)
 		, EnvelopeOutput(FFloatWriteRef::CreateNew(0.0f))
 	{
 		PrevAttackTime = FMath::Max(FTime::ToMilliseconds(*AttackTime), 0.0);
 		PrevReleaseTime = FMath::Max(FTime::ToMilliseconds(*ReleaseTime), 0.0);
-		PrevPeakModeType = *PeakModeType;
 
-		EnvelopeFollower.Init(InSettings.GetSampleRate(), PrevAttackTime, PrevReleaseTime, PrevPeakModeType);
+		EnvelopeFollower.Init(InSettings.GetSampleRate(), PrevAttackTime, PrevReleaseTime);
 	}
 
 	FDataReferenceCollection FEnvelopeFollowerOperator::GetInputs() const
@@ -102,7 +85,6 @@ namespace Metasound
 		InputDataReferences.AddDataReadReference(EnvelopeFollower::InParamNameAudioInput, AudioInput);
 		InputDataReferences.AddDataReadReference(EnvelopeFollower::InParamNameAttackTime, AttackTime);
 		InputDataReferences.AddDataReadReference(EnvelopeFollower::InParamNameReleaseTime, ReleaseTime);
-		InputDataReferences.AddDataReadReference(EnvelopeFollower::InParamNamePeakModeType, PeakModeType);
 
 		return InputDataReferences;
 	}
@@ -131,13 +113,6 @@ namespace Metasound
 			EnvelopeFollower.SetReleaseTime(CurrentReleaseTime);
 		}
 
-		Audio::EPeakMode::Type CurrentPeakModeType = *PeakModeType;
-		if (CurrentPeakModeType != PrevPeakModeType)
-		{
-			PrevPeakModeType = CurrentPeakModeType;
-			EnvelopeFollower.SetMode(CurrentPeakModeType);
-		}
-
 		// Process the audio through the envelope follower
 		const float* InputAudio = AudioInput->GetData();
 		int32 NumFrames = AudioInput->Num();
@@ -153,8 +128,7 @@ namespace Metasound
 			FInputVertexInterface(
 				TInputDataVertexModel<FAudioBuffer>(EnvelopeFollower::InParamNameAudioInput, LOCTEXT("AudioInputToolTT", "Audio input.")),
 				TInputDataVertexModel<FTime>(EnvelopeFollower::InParamNameAttackTime, LOCTEXT("AttackTimeTT", "The attack time of the envelope follower."), 0.01f),
-				TInputDataVertexModel<FTime>(EnvelopeFollower::InParamNameReleaseTime, LOCTEXT("ReleaseTimeTT", "The release time of the envelope follower."), 0.1f),
-				TInputDataVertexModel<FEnumPeakModeType>(EnvelopeFollower::InParamNamePeakModeType, LOCTEXT("PeakModeTypeTT", "The peak mode of the envelope follower."))
+				TInputDataVertexModel<FTime>(EnvelopeFollower::InParamNameReleaseTime, LOCTEXT("ReleaseTimeTT", "The release time of the envelope follower."), 0.1f)
 			),
 			FOutputVertexInterface(
 				TOutputDataVertexModel<float>(EnvelopeFollower::OutParamNameEnvelope, LOCTEXT("EnvelopeFollowerOutputTT", "The output envelope value of the audio signal."))
@@ -195,9 +169,8 @@ namespace Metasound
 		FAudioBufferReadRef AudioIn = InputCollection.GetDataReadReferenceOrConstruct<FAudioBuffer>(EnvelopeFollower::InParamNameAudioInput, InParams.OperatorSettings);
 		FTimeReadRef AttackTime = InputCollection.GetDataReadReferenceOrConstructWithVertexDefault<FTime, float>(InputInterface, EnvelopeFollower::InParamNameAttackTime);
 		FTimeReadRef ReleaseTime = InputCollection.GetDataReadReferenceOrConstructWithVertexDefault<FTime, float>(InputInterface, EnvelopeFollower::InParamNameReleaseTime);
-		FEnumPeakModeTypeReadRef EnvelopePeakModeType = InputCollection.GetDataReadReferenceOrConstruct<FEnumPeakModeType>(EnvelopeFollower::InParamNamePeakModeType);
 
-		return MakeUnique<FEnvelopeFollowerOperator>(InParams.OperatorSettings, AudioIn, AttackTime, ReleaseTime, EnvelopePeakModeType);
+		return MakeUnique<FEnvelopeFollowerOperator>(InParams.OperatorSettings, AudioIn, AttackTime, ReleaseTime);
 	}
 
 	FEnvelopeFollowerNode::FEnvelopeFollowerNode(const FNodeInitData& InitData)
