Fix up for triangle wave and square wave.

Added in a fade on all generators to avoid pops when starting.

diff --git a/Engine/Plugins/Runtime/Metasound/Source/MetasoundStandardNodes/Private/MetasoundOscillatorNodes.cpp b/Engine/Plugins/Runtime/Metasound/Source/MetasoundStandardNodes/Private/MetasoundOscillatorNodes.cpp
--- a/Engine/Plugins/Runtime/Metasound/Source/MetasoundStandardNodes/Private/MetasoundOscillatorNodes.cpp
+++ b/Engine/Plugins/Runtime/Metasound/Source/MetasoundStandardNodes/Private/MetasoundOscillatorNodes.cpp
@@ -88,6 +88,9 @@ namespace Metasound
 			float Phase = 0.f;
 			Oscillator Osc;
 			PhaseWrap Wrap;
+			float CurrentFade = 0.0f;
+			float FadeSmooth = 0.01f;
+
 			void operator()(const FGeneratorArgs& InArgs)
 			{
 				int32 RemainingSamplesInBlock = InArgs.AlignedBuffer.Num();
@@ -100,7 +103,9 @@ namespace Metasound
 				{
 					Wrap(Phase);
 
-					*Out++ = Osc(Phase, DeltaPhase, PulseWidth);
+					*Out++ = CurrentFade * Osc(Phase, DeltaPhase, InArgs);
+
+					CurrentFade = CurrentFade + FadeSmooth * (1.0f - CurrentFade);
 
 					Phase += DeltaPhase;
 
@@ -120,6 +125,7 @@ namespace Metasound
 			float Phase = 0.f;
 			Oscillator Osc;
 			PhaseWrap Wrap;
+
 			void operator()(const FGeneratorArgs& InArgs)
 			{
 				int32 RemainingSamplesInBlock = InArgs.AlignedBuffer.Num();
@@ -133,7 +139,7 @@ namespace Metasound
 					float PerSampleFreq = InArgs.FrequencyHz + *FM++;
 					float DeltaPhase = PerSampleFreq * OneOverSampleRate;
 
-					*Out++ = Osc(Phase, DeltaPhase, PulseWidth);
+					*Out++ = Osc(Phase, DeltaPhase, InArgs);
 
 					Phase += DeltaPhase;
 
@@ -147,7 +153,7 @@ namespace Metasound
 		// Sine types and generators.
 		struct FSinfGenerator
 		{
-			FORCEINLINE float operator()(float InPhase, float, float)
+			FORCEINLINE float operator()(float InPhase, float, const FGeneratorArgs&)
 			{
 				static constexpr float TwoPi = 2.f * PI;
 				return FMath::Sin(InPhase * TwoPi);
@@ -155,7 +161,7 @@ namespace Metasound
 		};
 		struct FBhaskaraGenerator
 		{
-			float operator()(float InPhase, float, float) const
+			float operator()(float InPhase, float, const FGeneratorArgs&) const
 			{
 				static const float TwoPi = 2.f * PI;
 				float PhaseRadians = InPhase * TwoPi;
@@ -183,14 +189,14 @@ namespace Metasound
 		// Saws.
 		struct FSawGenerator
 		{
-			FORCEINLINE float operator()(float InPhase, float, float)
+			FORCEINLINE float operator()(float InPhase, float, const FGeneratorArgs&)
 			{
 				return -1.f + (2.f * InPhase);
 			}
 		};
 		struct FSawPolySmoothGenerator
 		{
-			FORCEINLINE float operator()(float InPhase, float InPhaseDelta, float)
+			FORCEINLINE float operator()(float InPhase, float InPhaseDelta, const FGeneratorArgs&)
 			{
 				// Two-sided wave-shaped sawtooth
 				static const float A = Audio::FastTanh(1.5f);
@@ -210,7 +216,7 @@ namespace Metasound
 		// Square.
 		struct FSquareGenerator
 		{
-			FORCEINLINE float operator()(float InPhase, float, float)
+			FORCEINLINE float operator()(float InPhase, float, const FGeneratorArgs&)
 			{
 				// Does not obey pulse width.
 				return InPhase >= 0.5f ? 1.f : -1.f;
@@ -218,7 +224,7 @@ namespace Metasound
 		};
 		struct FSquarePolysmoothGenerator
 		{
-			FORCEINLINE float operator()(float InPhase, float InPhaseDelta, float InCurrentPulseWidth)
+			FORCEINLINE float operator()(float InPhase, float InPhaseDelta, const FGeneratorArgs& InArgs)
 			{
 				// Taken piecemeal from Osc.cpp
 				// Lots of branches.
@@ -231,7 +237,7 @@ namespace Metasound
 				float NewPhase = 0.0f;
 				if (InPhaseDelta > 0.0f)
 				{
-					NewPhase = InPhase + InCurrentPulseWidth;
+					NewPhase = InPhase + InArgs.PulseWidth;
 					if (NewPhase >= 1.0f)
 					{
 						NewPhase -= 1.0f;
@@ -239,7 +245,7 @@ namespace Metasound
 				}
 				else
 				{
-					NewPhase = InPhase - InCurrentPulseWidth;
+					NewPhase = InPhase - InArgs.PulseWidth;
 					if (NewPhase <= 0.0f)
 					{
 						NewPhase += 1.0f;
@@ -253,8 +259,8 @@ namespace Metasound
 				float Output = 0.5f * SquareSaw1 - 0.5f * SquareSaw2;
 
 				// Apply DC correction
-				const float Correction = (InCurrentPulseWidth < 0.5f) ? 
-					(1.0f / (1.0f - InCurrentPulseWidth)) : (1.0f / InCurrentPulseWidth);
+				const float Correction = (InArgs.PulseWidth < 0.5f) ?
+					(1.0f / (1.0f - InArgs.PulseWidth)) : (1.0f / InArgs.PulseWidth);
 
 				Output *= Correction;
 				return Output;
@@ -270,13 +276,37 @@ namespace Metasound
 		// Triangle.
 		struct FTriangleGenerator
 		{
-			FORCEINLINE float operator()(float InPhase, float, float)
+			float TriangleSign = -1.0f;
+			float PrevPhase = -1.0f;
+			float DPW_z1 = 0.0f;
+
+			FORCEINLINE float operator()(float InPhase, float InPhaseDelta, const FGeneratorArgs& InArgs)
 			{
-				return InPhase <= 0.5f ? InPhase - 0.25f * 4.f : ((1.f - InPhase) - 0.25f) * 4.f;
+				// if the oscillator phase wrapped, our previous phase will be higher than current phase
+				if (PrevPhase > InPhase)
+				{
+					// Flip the sign if we wrapped phase
+					TriangleSign *= -1.0f;
+				}
+				PrevPhase = InPhase;
+
+
+				// Get saw trivial output
+				float SawToothOutput = Audio::GetBipolar(InPhase);
+				const float SawSquaredInvMod = (1.0f - SawToothOutput * SawToothOutput) * TriangleSign;
+
+				const float Differentiated = SawSquaredInvMod - DPW_z1;
+				DPW_z1 = SawSquaredInvMod;
+				return Differentiated * InArgs.SampleRate / (4.0f * InArgs.FrequencyHz * (1.0f - InPhaseDelta));
 			}
 		};
+
 		using FTriangle						= TGenerateBlock<FTriangleGenerator>;
 		using FTriangleWithFm				= TGenerateBlockFM<FTriangleGenerator>;
+
+		// TODO: make a true polysmooth version
+		using FTrianglePolysmooth = TGenerateBlock<FTriangleGenerator>;
+		using FTrianglePolysmoothWithFm = TGenerateBlockFM<FTriangleGenerator>;
 	}
 
 	// Base class of Oscillator factories which holds common the interface.
@@ -707,7 +737,7 @@ namespace Metasound
 			switch (*Type)
 			{
 			default:
-			case ESawGenerationType::Trivial: return MakeUnique<TOscillatorOperator<FTriangle>>(OpParams);
+			case ESawGenerationType::Trivial: return MakeUnique<TOscillatorOperator<FSaw>>(OpParams);
 			case ESawGenerationType::PolySmooth: return MakeUnique<TOscillatorOperator<FSawPolysmooth>>(OpParams);
 			}
 		}
@@ -885,7 +915,7 @@ namespace Metasound
 			{
 				FVertexInterface Interface = GetCommmonVertexInterface();
 				Interface.GetInputInterface().Add(TInputDataVertexModel<FEnumSquareGenerationType>(SquareTypePinName, LOCTEXT("SquareTypeDescription", "The generator type to make the squarewave")));
-				Interface.GetInputInterface().Add(TInputDataVertexModel<float>(PulseWidthPinName, LOCTEXT("PhaseWidthDescription", "The Width of the square part of the wave"), 1.f));
+				Interface.GetInputInterface().Add(TInputDataVertexModel<float>(PulseWidthPinName, LOCTEXT("PhaseWidthDescription", "The Width of the square part of the wave"), 0.5f));
 				return Interface;
 			};
 			static const FVertexInterface Interface = MakeInterface();
@@ -957,6 +987,7 @@ namespace Metasound
 				switch (*Type)
 				{
 				default:
+				case ETriangleGenerationType::PolySmooth: return MakeUnique<TOscillatorOperatorFM<FTrianglePolysmoothWithFm>>(OpParams, FmBuffer);
 				case ETriangleGenerationType::Trivial: return MakeUnique<TOscillatorOperatorFM<FTriangleWithFm>>(OpParams, FmBuffer);
 				}
 			}
@@ -965,6 +996,7 @@ namespace Metasound
 				switch (*Type)
 				{
 				default:
+				case ETriangleGenerationType::PolySmooth: return MakeUnique<TOscillatorOperator<FTrianglePolysmooth>>(OpParams);
 				case ETriangleGenerationType::Trivial: return MakeUnique<TOscillatorOperator<FTriangle>>(OpParams);
 				}
 			}
