- Changing the method of implementation for the global shared array shuffle
- Implementing a random get array node which retrieves a random element from an array given a supplied weights array
- Pipe down new metasound environment variable for the USoundBase UniqueObjectID to use as a key for shared data across instances of a metasound

diff --git a/Engine/Plugins/Runtime/Metasound/Source/MetasoundEngine/Private/MetasoundSource.cpp b/Engine/Plugins/Runtime/Metasound/Source/MetasoundEngine/Private/MetasoundSource.cpp
--- a/Engine/Plugins/Runtime/Metasound/Source/MetasoundEngine/Private/MetasoundSource.cpp
+++ b/Engine/Plugins/Runtime/Metasound/Source/MetasoundEngine/Private/MetasoundSource.cpp
@@ -17,7 +17,7 @@
 #include "MetasoundPrimitives.h"
 #include "MetasoundReceiveNode.h"
 #include "MetasoundTrigger.h"
-
+#include "MetasoundEnvironment.h"
 
 #if WITH_EDITORONLY_DATA
 #include "EdGraph/EdGraph.h"
@@ -135,6 +135,9 @@ ISoundGeneratorPtr UMetasoundSource::CreateSoundGenerator(const FSoundGeneratorI
 	}
 	Environment.SetValue<FAudioDeviceHandle>(GetAudioDeviceHandleVariableName(), DeviceHandle);
 
+	// Set the unique object ID as an environment variable
+	Environment.SetValue<uint32>(GetSoundUniqueIdName(), GetUniqueID());
+
 	const FMetasoundFrontendDocument* OriginalDoc = GetDocument().Get();
 	if (nullptr == OriginalDoc)
 	{
@@ -464,6 +467,13 @@ const FString& UMetasoundSource::GetAudioDeviceHandleVariableName()
 	return AudioDeviceHandleVarName;
 }
 
+const FString& UMetasoundSource::GetSoundUniqueIdName()
+{
+	static const FString SoundUniqueIdVarName = TEXT("SoundUniqueId");
+	return SoundUniqueIdVarName;
+}
+
+
 const FMetasoundFrontendArchetype& UMetasoundSource::GetBaseArchetype()
 {
 	auto CreateBaseArchetype = []() -> FMetasoundFrontendArchetype
diff --git a/Engine/Plugins/Runtime/Metasound/Source/MetasoundEngine/Public/MetasoundEngineEnvironment.h b/Engine/Plugins/Runtime/Metasound/Source/MetasoundEngine/Public/MetasoundEngineEnvironment.h
--- a/Engine/Plugins/Runtime/Metasound/Source/MetasoundEngine/Public/MetasoundEngineEnvironment.h
+++ b/Engine/Plugins/Runtime/Metasound/Source/MetasoundEngine/Public/MetasoundEngineEnvironment.h
@@ -5,5 +5,4 @@
 #include "AudioDeviceManager.h"
 #include "MetasoundEnvironment.h"
 
-DECLARE_METASOUND_ENVIRONMENT_VARIABLE_TYPE(FAudioDeviceHandle);
-
+DECLARE_METASOUND_ENVIRONMENT_VARIABLE_TYPE(METASOUNDENGINE_API, FAudioDeviceHandle);
diff --git a/Engine/Plugins/Runtime/Metasound/Source/MetasoundEngine/Public/MetasoundSource.h b/Engine/Plugins/Runtime/Metasound/Source/MetasoundEngine/Public/MetasoundSource.h
--- a/Engine/Plugins/Runtime/Metasound/Source/MetasoundEngine/Public/MetasoundSource.h
+++ b/Engine/Plugins/Runtime/Metasound/Source/MetasoundEngine/Public/MetasoundSource.h
@@ -159,6 +159,7 @@ private:
 	static const FString& GetAudioOutputName();
 	static const FString& GetIsFinishedOutputName();
 	static const FString& GetAudioDeviceHandleVariableName();
+	static const FString& GetSoundUniqueIdName();
 	static const FMetasoundFrontendArchetype& GetBaseArchetype();
 	static const FMetasoundFrontendArchetype& GetMonoSourceArchetype();
 	static const FMetasoundFrontendArchetype& GetStereoSourceArchetype();
diff --git a/Engine/Plugins/Runtime/Metasound/Source/MetasoundFrontend/Private/MetasoundArrayNodes.cpp b/Engine/Plugins/Runtime/Metasound/Source/MetasoundFrontend/Private/MetasoundArrayNodes.cpp
--- a/Engine/Plugins/Runtime/Metasound/Source/MetasoundFrontend/Private/MetasoundArrayNodes.cpp
+++ b/Engine/Plugins/Runtime/Metasound/Source/MetasoundFrontend/Private/MetasoundArrayNodes.cpp
@@ -99,233 +99,6 @@ namespace Metasound
 			return Name;
 		}
 	}
-
-	namespace ArrayNodeShuffleVertexNames
-	{
-		const FString& GetInputTriggerNextName()
-		{
-			static const FString Name = TEXT("Next");
-			return Name;
-		}
-
-		const FString& GetInputTriggerShuffleName()
-		{
-			static const FString Name = TEXT("Shuffle");
-			return Name;
-		}
-
-		const FString& GetInputTriggerResetName()
-		{
-			static const FString Name = TEXT("Reset Seed");
-			return Name;
-		}
-
-		const FString& GetInputShuffleArrayName()
-		{
-			static const FString Name = TEXT("In Array");
-			return Name;
-		}
-
-		const FString& GetInputSeedName()
-		{
-			static const FString Name = TEXT("Seed");
-			return Name;
-		}
-
-		const FString& GetInputAutoShuffleName()
-		{
-			static const FString Name = TEXT("Auto Shuffle");
-			return Name;
-		}
-
-		const FString& GetInputNamespaceName()
-		{
-			static const FString Name = TEXT("Shuffle Namespace");
-			return Name;
-		}
-
-		const FString& GetOutputTriggerOnNextName()
-		{
-			static const FString Name = TEXT("On Next");
-			return Name;
-		}
-
-		const FString& GetOutputTriggerOnShuffleName()
-		{
-			static const FString Name = TEXT("On Shuffle");
-			return Name;
-		}
-
-		const FString& GetOutputTriggerOnResetName()
-		{
-			static const FString Name = TEXT("On Reset Seed");
-			return Name;
-		}
-
-		const FString& GetOutputValueName()
-		{
-			static const FString Name = TEXT("Value");
-			return Name;
-		}
-	}
-
-	uint32 GetTypeHash(const FGlobalArrayShuffleKey& InKey)
-	{
-		return InKey.Hash;
-	}
-
-	/**
-	* FArrayIndexShuffler
-	*/
-
-	FArrayIndexShuffler::FArrayIndexShuffler(int32 InSeed, int32 InMaxIndicies)
-	{
-		Init(InSeed, InMaxIndicies);
-	}
-
-	void FArrayIndexShuffler::Init(int32 InSeed, int32 InMaxIndicies)
-	{
-		SetSeed(InSeed);
-		if (InMaxIndicies > 0)
-		{
-			ShuffleIndices.AddUninitialized(InMaxIndicies);
-			for (int32 i = 0; i < ShuffleIndices.Num(); ++i)
-			{
-				ShuffleIndices[i] = i;
-			}
-			ShuffleArray();
-		}
-	}
-
-	void FArrayIndexShuffler::SetSeed(int32 InSeed)
-	{
-		if (InSeed == INDEX_NONE)
-		{
-			RandomStream.Initialize(FPlatformTime::Cycles());
-		}
-		else
-		{
-			RandomStream.Initialize(InSeed);
-		}
-
-		ResetSeed();
-	}
-
-	void FArrayIndexShuffler::ResetSeed()
-	{
-		RandomStream.Reset();
-	}
-
-	bool FArrayIndexShuffler::NextValue(bool bAutoShuffle, int32& OutIndex)
-	{
-		bool bShuffled = false;
-		if (CurrentIndex == ShuffleIndices.Num())
-		{
-			if (bAutoShuffle)
-			{
-				ShuffleArray();
-				bShuffled = true;
-			}
-			else
-			{
-				CurrentIndex = 0;
-			}
-		}
-
-		check(CurrentIndex < ShuffleIndices.Num());
-		PrevValue = ShuffleIndices[CurrentIndex];
-		OutIndex = PrevValue;
-		++CurrentIndex;
-
-		return bShuffled;
-	}
-
-	// Shuffle the array with the given max indices
-	void FArrayIndexShuffler::ShuffleArray()
-	{
-		// Randomize the shuffled array by randomly swapping indicies
-		for (int32 i = 0; i < ShuffleIndices.Num(); ++i)
-		{
-			RandomSwap(i, 0, ShuffleIndices.Num() - 1);
-		}
-
-		// Reset the current index back to 0
-		CurrentIndex = 0;
-
-		// Fix up the new current index if the value is our previous value and we have an array larger than 1
-		if (ShuffleIndices.Num() > 1 && ShuffleIndices[CurrentIndex] == PrevValue)
-		{
-			RandomSwap(0, 1, ShuffleIndices.Num() - 1);
-		}
-	}
-
-	void FArrayIndexShuffler::RandomSwap(int32 InCurrentIndex, int32 InStartIndex, int32 InEndIndex)
-	{
-		int32 ShuffleIndex = RandomStream.RandRange(InStartIndex, InEndIndex);
-		int32 Temp = ShuffleIndices[ShuffleIndex];
-		ShuffleIndices[ShuffleIndex] = ShuffleIndices[InCurrentIndex];
-		ShuffleIndices[InCurrentIndex] = Temp;
-	}
-
-	/** 
-	* FGlobalArrayShuffleManager
-	*/
-
-	FGlobalArrayShuffleManager& FGlobalArrayShuffleManager::Get()
-	{
-		static FGlobalArrayShuffleManager GShufflerManager;
-		return GShufflerManager;
-	}
-
-	bool FGlobalArrayShuffleManager::NextValue(FGlobalArrayShuffleKey& InKey, bool bAutoShuffle, int32& OutIndex)
-	{
-		check(!InKey.Namespace.IsEmpty());
-
-		FScopeLock Lock(&CritSect);
-
-		TUniquePtr<FArrayIndexShuffler>& ShufflerPtr = GetShuffler(InKey);
-		return ShufflerPtr->NextValue(bAutoShuffle, OutIndex);
-	}
-
-	void FGlobalArrayShuffleManager::SetSeed(FGlobalArrayShuffleKey& InKey, int32 InSeed)
-	{
-		check(!InKey.Namespace.IsEmpty());
-
-		FScopeLock Lock(&CritSect);
-
-		TUniquePtr<FArrayIndexShuffler>& ShufflerPtr = GetShuffler(InKey);
-		ShufflerPtr->SetSeed(InSeed);
-	}
-
-	void FGlobalArrayShuffleManager::ResetSeed(FGlobalArrayShuffleKey& InKey)
-	{
-		check(!InKey.Namespace.IsEmpty());
-
-		FScopeLock Lock(&CritSect);
-
-		TUniquePtr<FArrayIndexShuffler>& ShufflerPtr = GetShuffler(InKey);
-		ShufflerPtr->ResetSeed();
-	}
-
-	void FGlobalArrayShuffleManager::ShuffleArray(FGlobalArrayShuffleKey& InKey)
-	{
-		check(!InKey.Namespace.IsEmpty());
-
-		FScopeLock Lock(&CritSect);
-
-		TUniquePtr<FArrayIndexShuffler>& ShufflerPtr = GetShuffler(InKey);
-		ShufflerPtr->ShuffleArray();
-	}
-
-	TUniquePtr<FArrayIndexShuffler>& FGlobalArrayShuffleManager::GetShuffler(FGlobalArrayShuffleKey& InKey)
-	{
-		TUniquePtr<FArrayIndexShuffler>& ShufflerPtr = GlobalShuffleIndicies.FindOrAdd(InKey);
-		if (!ShufflerPtr.IsValid())
-		{
-			ShufflerPtr = MakeUnique<FArrayIndexShuffler>(InKey.Seed, InKey.NumElements);
-		}
-		return ShufflerPtr;
-	}
 }
 
 #undef LOCTEXT_NAMESPACE
diff --git a/Engine/Plugins/Runtime/Metasound/Source/MetasoundFrontend/Private/MetasoundArrayRandomNode.cpp b/Engine/Plugins/Runtime/Metasound/Source/MetasoundFrontend/Private/MetasoundArrayRandomNode.cpp
new file mode 100644
--- /dev/null
+++ b/Engine/Plugins/Runtime/Metasound/Source/MetasoundFrontend/Private/MetasoundArrayRandomNode.cpp
@@ -0,0 +1,241 @@
+// Copyright Epic Games, Inc. All Rights Reserved.
+
+#include "MetasoundArrayRandomNode.h"
+
+#define LOCTEXT_NAMESPACE "MetasoundFrontend"
+
+namespace Metasound
+{
+	namespace ArrayNodeRandomGetVertexNames
+	{
+		const FString& GetInputTriggerNextName()
+		{
+			static const FString Name = TEXT("Next");
+			return Name;
+		}
+
+		const FString& GetInputTriggerResetName()
+		{
+			static const FString Name = TEXT("Reset");
+			return Name;
+		}
+
+		const FString& GetInputRandomArrayName()
+		{
+			static const FString Name = TEXT("In Array");
+			return Name;
+		}
+
+		const FString& GetInputWeightsName()
+		{
+			static const FString Name = TEXT("Weights");
+			return Name;
+		}
+
+		const FString& GetInputSeedName()
+		{
+			static const FString Name = TEXT("Seed");
+			return Name;
+		}
+
+		const FString& GetInputNoRepeatOrderName()
+		{
+			static const FString Name = TEXT("No Repeats");
+			return Name;
+		}
+
+		const FString& GetInputEnableSharedStateName()
+		{
+			static const FString Name = TEXT("Enable Shared State");
+			return Name;
+		}
+
+		const FString& GetOutputTriggerOnNextName()
+		{
+			static const FString Name = TEXT("On Next");
+			return Name;
+		}
+
+		const FString& GetOutputTriggerOnResetName()
+		{
+			static const FString Name = TEXT("On Reset");
+			return Name;
+		}
+
+		const FString& GetOutputValueName()
+		{
+			static const FString Name = TEXT("Value");
+			return Name;
+		}
+	}
+
+	FArrayRandomGet::FArrayRandomGet(int32 InSeed, int32 InMaxIndex, int32 InNoRepeatOrder)
+	{
+		Init(InSeed, InMaxIndex, InNoRepeatOrder);
+	}
+
+	void FArrayRandomGet::Init(int32 InSeed, int32 InMaxIndex, int32 InNoRepeatOrder)
+	{
+		SetSeed(InSeed);
+		MaxIndex = InMaxIndex;
+		SetNoRepeatOrder(InNoRepeatOrder);
+		PreviousIndicesQueue = MakeUnique<TCircularQueue<int32>>(InNoRepeatOrder);
+		PreviousIndices.Reset();
+
+	}
+
+	void FArrayRandomGet::SetSeed(int32 InSeed)
+	{
+		if (InSeed == INDEX_NONE)
+		{
+			RandomStream.Initialize(FPlatformTime::Cycles());
+		}
+		else
+		{
+			RandomStream.Initialize(InSeed);
+		}
+
+		ResetSeed();
+	}
+
+	void FArrayRandomGet::SetNoRepeatOrder(int32 InNoRepeatOrder)
+	{
+		// Make sure the no repeat order is between 0 and one less than our max index
+		InNoRepeatOrder = FMath::Clamp(InNoRepeatOrder, 0, MaxIndex - 1);
+
+		if (InNoRepeatOrder != NoRepeatOrder)
+		{
+			PreviousIndicesQueue = MakeUnique<TCircularQueue<int32>>(InNoRepeatOrder);
+			PreviousIndices.Reset();
+			NoRepeatOrder = InNoRepeatOrder;
+		}
+	}
+
+	void FArrayRandomGet::ResetSeed()
+	{
+		RandomStream.Reset();
+	}
+
+	// Returns the next random weighted value in the array indices. 
+	int32 FArrayRandomGet::NextValue()
+	{
+		// First compute the total size of the weights
+		float TotalWeight = 0.0f;
+		if (RandomWeights.Num() > 0)
+		{
+			for (int32 i = 0; i < MaxIndex; ++i)
+			{
+				// If the index exists in previous indices, continue
+				if (PreviousIndices.Contains(i))
+				{
+					continue;
+				}
+
+				// We modulus on the weight array to determine weights for the input array
+				// I.e. if weights is 2 elements, the weights will alternate in application to the input array
+				TotalWeight += RandomWeights[i % RandomWeights.Num()];
+			}
+		}
+		else
+		{
+			// If we don't have a random weights array, everything is equal weight so we only need to consider the total number of un-picked indices
+			TotalWeight = (float)(FMath::Max(MaxIndex - PreviousIndices.Num(), 1));
+		}
+
+		// Make a random choice based on the total weight
+		float Choice = RandomStream.FRandRange(0.0f, TotalWeight);
+
+		// Now find the index this choice matches up to
+		TotalWeight = 0.0f;
+		int32 ChosenIndex = INDEX_NONE;
+		for (int32 i = 0; i < MaxIndex; ++i)
+		{
+			if (PreviousIndices.Contains(i))
+			{
+				continue;
+			}
+
+			float NextTotalWeight = TotalWeight;
+			if (RandomWeights.Num() > 0)
+			{
+				NextTotalWeight += RandomWeights[i % RandomWeights.Num()];
+			}
+			else
+			{
+				NextTotalWeight += 1.0f;
+			}
+
+			if (Choice >= TotalWeight && Choice < NextTotalWeight)
+			{
+				ChosenIndex = i;
+				break;
+			}
+			TotalWeight = NextTotalWeight;
+		}
+		check(ChosenIndex != INDEX_NONE);
+
+		// Dequeue and remove the oldest previous index
+		if (PreviousIndices.Num() == NoRepeatOrder)
+		{
+			int32 OldPrevIndex;
+			PreviousIndicesQueue->Dequeue(OldPrevIndex);
+			PreviousIndices.Remove(OldPrevIndex);
+			check(PreviousIndices.Num() == NoRepeatOrder - 1);
+		}
+
+		// We should only have less than *or* equal to the NoRepeatOrder
+		check(PreviousIndices.Num() < NoRepeatOrder);
+		PreviousIndicesQueue->Enqueue(ChosenIndex);
+		check(!PreviousIndices.Contains(ChosenIndex));
+		PreviousIndices.Add(ChosenIndex);
+
+		return ChosenIndex;
+	}
+
+	FSharedStateRandomGetManager& FSharedStateRandomGetManager::Get()
+	{
+		static FSharedStateRandomGetManager RGM;
+		return RGM;
+	}
+
+	void FSharedStateRandomGetManager::InitSharedState(uint32 InSharedStateId, int32 InSeed, int32 InNumElements, int32 InNoRepeatOrder)
+	{
+		FScopeLock Lock(&CritSect);
+
+		if (!RandomGets.Contains(InSharedStateId))
+		{
+			RandomGets.Add(InSharedStateId, MakeUnique<FArrayRandomGet>(InSeed, InNumElements, InNoRepeatOrder));
+		}
+	}
+
+	int32 FSharedStateRandomGetManager::NextValue(uint32 InSharedStateId)
+	{
+		FScopeLock Lock(&CritSect);
+		TUniquePtr<FArrayRandomGet>* RG = RandomGets.Find(InSharedStateId);
+		return (*RG)->NextValue();
+	}
+
+	void FSharedStateRandomGetManager::SetSeed(uint32 InSharedStateId, int32 InSeed)
+	{
+		FScopeLock Lock(&CritSect);
+		TUniquePtr<FArrayRandomGet>* RG = RandomGets.Find(InSharedStateId);
+		(*RG)->SetSeed(InSeed);
+	}
+
+	void FSharedStateRandomGetManager::SetNoRepeatOrder(uint32 InSharedStateId, int32 InNoRepeatOrder)
+	{
+		FScopeLock Lock(&CritSect);
+		TUniquePtr<FArrayRandomGet>* RG = RandomGets.Find(InSharedStateId);
+		(*RG)->SetNoRepeatOrder(InNoRepeatOrder);
+	}
+
+	void FSharedStateRandomGetManager::ResetSeed(uint32 InSharedStateId)
+	{
+		FScopeLock Lock(&CritSect);
+		TUniquePtr<FArrayRandomGet>* RG = RandomGets.Find(InSharedStateId);
+		(*RG)->ResetSeed();
+	}
+
+}
+
+#undef LOCTEXT_NAMESPACE
diff --git a/Engine/Plugins/Runtime/Metasound/Source/MetasoundFrontend/Private/MetasoundArrayShuffleNode.cpp b/Engine/Plugins/Runtime/Metasound/Source/MetasoundFrontend/Private/MetasoundArrayShuffleNode.cpp
new file mode 100644
--- /dev/null
+++ b/Engine/Plugins/Runtime/Metasound/Source/MetasoundFrontend/Private/MetasoundArrayShuffleNode.cpp
@@ -0,0 +1,201 @@
+// Copyright Epic Games, Inc. All Rights Reserved.
+
+#include "MetasoundArrayShuffleNode.h"
+#include "MetasoundDataTypeRegistrationMacro.h"
+
+#define LOCTEXT_NAMESPACE "MetasoundFrontend"
+
+namespace Metasound
+{
+	namespace ArrayNodeShuffleVertexNames
+	{
+		const FString& GetInputTriggerNextName()
+		{
+			static const FString Name = TEXT("Next");
+			return Name;
+		}
+
+		const FString& GetInputTriggerShuffleName()
+		{
+			static const FString Name = TEXT("Shuffle");
+			return Name;
+		}
+
+		const FString& GetInputTriggerResetName()
+		{
+			static const FString Name = TEXT("Reset Seed");
+			return Name;
+		}
+
+		const FString& GetInputShuffleArrayName()
+		{
+			static const FString Name = TEXT("In Array");
+			return Name;
+		}
+
+		const FString& GetInputSeedName()
+		{
+			static const FString Name = TEXT("Seed");
+			return Name;
+		}
+
+		const FString& GetInputAutoShuffleName()
+		{
+			static const FString Name = TEXT("Auto Shuffle");
+			return Name;
+		}
+
+		const FString& GetInputEnableSharedStateName()
+		{
+			static const FString Name = TEXT("Enable Shared State");
+			return Name;
+		}
+
+		const FString& GetOutputTriggerOnNextName()
+		{
+			static const FString Name = TEXT("On Next");
+			return Name;
+		}
+
+		const FString& GetOutputTriggerOnShuffleName()
+		{
+			static const FString Name = TEXT("On Shuffle");
+			return Name;
+		}
+
+		const FString& GetOutputTriggerOnResetName()
+		{
+			static const FString Name = TEXT("On Reset Seed");
+			return Name;
+		}
+
+		const FString& GetOutputValueName()
+		{
+			static const FString Name = TEXT("Value");
+			return Name;
+		}
+	}
+
+	/**
+	* FArrayIndexShuffler
+	*/
+
+	FArrayIndexShuffler::FArrayIndexShuffler(int32 InSeed, int32 InMaxIndices)
+	{
+		Init(InSeed, InMaxIndices);
+	}
+
+	void FArrayIndexShuffler::Init(int32 InSeed, int32 InMaxIndices)
+	{
+		SetSeed(InSeed);
+		if (InMaxIndices > 0)
+		{
+			ShuffleIndices.AddUninitialized(InMaxIndices);
+			for (int32 i = 0; i < ShuffleIndices.Num(); ++i)
+			{
+				ShuffleIndices[i] = i;
+			}
+			ShuffleArray();
+		}
+	}
+
+	void FArrayIndexShuffler::SetSeed(int32 InSeed)
+	{
+		if (InSeed == INDEX_NONE)
+		{
+			RandomStream.Initialize(FPlatformTime::Cycles());
+		}
+		else
+		{
+			RandomStream.Initialize(InSeed);
+		}
+
+		ResetSeed();
+	}
+
+	void FArrayIndexShuffler::ResetSeed()
+	{
+		RandomStream.Reset();
+	}
+
+	bool FArrayIndexShuffler::NextValue(bool bAutoShuffle, int32& OutIndex)
+	{
+		bool bShuffled = false;
+		if (CurrentIndex == ShuffleIndices.Num())
+		{
+			if (bAutoShuffle)
+			{
+				ShuffleArray();
+				bShuffled = true;
+			}
+			else
+			{
+				CurrentIndex = 0;
+			}
+		}
+
+		check(CurrentIndex < ShuffleIndices.Num());
+		PrevValue = ShuffleIndices[CurrentIndex];
+		OutIndex = PrevValue;
+		++CurrentIndex;
+
+		return bShuffled;
+	}
+
+	// Shuffle the array with the given max indices
+	void FArrayIndexShuffler::ShuffleArray()
+	{
+		// Randomize the shuffled array by randomly swapping indicies
+		for (int32 i = 0; i < ShuffleIndices.Num(); ++i)
+		{
+			RandomSwap(i, 0, ShuffleIndices.Num() - 1);
+		}
+
+		// Reset the current index back to 0
+		CurrentIndex = 0;
+
+		// Fix up the new current index if the value is our previous value and we have an array larger than 1
+		if (ShuffleIndices.Num() > 1 && ShuffleIndices[CurrentIndex] == PrevValue)
+		{
+			RandomSwap(0, 1, ShuffleIndices.Num() - 1);
+		}
+	}
+
+	void FArrayIndexShuffler::RandomSwap(int32 InCurrentIndex, int32 InStartIndex, int32 InEndIndex)
+	{
+		int32 ShuffleIndex = RandomStream.RandRange(InStartIndex, InEndIndex);
+		int32 Temp = ShuffleIndices[ShuffleIndex];
+		ShuffleIndices[ShuffleIndex] = ShuffleIndices[InCurrentIndex];
+		ShuffleIndices[InCurrentIndex] = Temp;
+	}
+
+	namespace ArrayNodeGetGlobalArrayKeyVertexNames
+	{
+		const FString& GetInputNamespaceName()
+		{
+			static const FString Name = TEXT("Namespace");
+			return Name;		
+		}
+
+		const FString& GetInputArraySizeName()
+		{
+			static const FString Name = TEXT("Array Size");
+			return Name;
+		}
+
+		const FString& GetInputSeedName()
+		{
+			static const FString Name = TEXT("Seed");
+			return Name;
+		}
+
+		const FString& GetOutputArrayKeyName()
+		{
+			static const FString Name = TEXT("Global Array Key");
+			return Name;
+		}
+	}
+
+}
+
+#undef LOCTEXT_NAMESPACE
diff --git a/Engine/Plugins/Runtime/Metasound/Source/MetasoundFrontend/Public/MetasoundArrayNodes.h b/Engine/Plugins/Runtime/Metasound/Source/MetasoundFrontend/Public/MetasoundArrayNodes.h
--- a/Engine/Plugins/Runtime/Metasound/Source/MetasoundFrontend/Public/MetasoundArrayNodes.h
+++ b/Engine/Plugins/Runtime/Metasound/Source/MetasoundFrontend/Public/MetasoundArrayNodes.h
@@ -14,6 +14,7 @@
 #include "MetasoundOperatorInterface.h"
 #include "MetasoundPrimitives.h"
 #include "MetasoundTrigger.h"
+#include "Containers/CircularQueue.h"
 
 #include <type_traits>
 
@@ -742,419 +743,6 @@ namespace Metasound
 
 		virtual ~TArraySubsetNode() = default;
 	};
-
-	/** Shuffle Node Vertex Names */
-	namespace ArrayNodeShuffleVertexNames
-	{
-		/** Input Vertex Names */
-		METASOUNDFRONTEND_API const FString& GetInputTriggerNextName();
-		METASOUNDFRONTEND_API const FString& GetInputTriggerShuffleName();
-		METASOUNDFRONTEND_API const FString& GetInputTriggerResetName();
-		METASOUNDFRONTEND_API const FString& GetInputShuffleArrayName();
-		METASOUNDFRONTEND_API const FString& GetInputSeedName();
-		METASOUNDFRONTEND_API const FString& GetInputAutoShuffleName();
-		METASOUNDFRONTEND_API const FString& GetInputNamespaceName();
-
-		METASOUNDFRONTEND_API const FString& GetOutputTriggerOnNextName();
-		METASOUNDFRONTEND_API const FString& GetOutputTriggerOnShuffleName();
-		METASOUNDFRONTEND_API const FString& GetOutputTriggerOnResetName();
-		METASOUNDFRONTEND_API const FString& GetOutputValueName();
-	}
-
-	class METASOUNDFRONTEND_API FArrayIndexShuffler
-	{
-	public:
-		FArrayIndexShuffler() = default;
-		FArrayIndexShuffler(int32 InSeed, int32 MaxIndicies);
-
-		void Init(int32 InSeed, int32 MaxIndicies);
-		void SetSeed(int32 InSeed);
-		void ResetSeed();
-
-		// Returns the next value in the array indices. Returns true if the array was re-shuffled automatically.
-		bool NextValue(bool bAutoShuffle, int32& OutIndex);
-
-		// Shuffle the array with the given max indices
-		void ShuffleArray();
-
-	private:
-		// Helper function to swap the current index with a random index
-		void RandomSwap(int32 InCurrentIndex, int32 InStartIndex, int32 InEndIndex);
-
-		// The current index into the array of indicies (wraps between 0 and ShuffleIndices.Num())
-		int32 CurrentIndex = 0;
-
-		// The previously returned value. Used to avoid repeating the last value on shuffle.
-		int32 PrevValue = INDEX_NONE;
-
-		// Array of indices (in order 0 to Num), shuffled
-		TArray<int32> ShuffleIndices;
-
-		// Random stream to use to randomize the shuffling
-		FRandomStream RandomStream;
-	};
-
-	// Key used for global shuffle.
-	struct METASOUNDFRONTEND_API FGlobalArrayShuffleKey
-	{
-		FString Namespace;
-		int32 NumElements = 0;
-		int32 Seed = 0;
-		uint32 Hash = INDEX_NONE;
-
-		FGlobalArrayShuffleKey() = default;
-
-		FGlobalArrayShuffleKey(const FString& InNamespace, int32 InNumElements, int32 InSeed)
-			: Namespace(InNamespace)
-			, NumElements(InNumElements)
-			, Seed(InSeed)
-		{
-			// Hash on the input values
-			Hash = HashCombine(GetTypeHash(InNamespace), GetTypeHash(InNumElements));
-			Hash = HashCombine(Hash, GetTypeHash(Seed));
-		}
-
-		bool operator==(const FGlobalArrayShuffleKey& Other) const
-		{
-			return Other.Namespace == Namespace &&
-				Other.NumElements == NumElements &&
-				Other.Seed == Seed;
-		}
-	};
-
-	uint32 GetTypeHash(const FGlobalArrayShuffleKey& InKey);
-
-	class METASOUNDFRONTEND_API FGlobalArrayShuffleManager
-	{
-	public:
-		FGlobalArrayShuffleManager() = default;
-		~FGlobalArrayShuffleManager() = default;
-
-		// Retrieves the shuffle manager
-		static FGlobalArrayShuffleManager& Get();
-
-		// API that wraps the array shuffler
-		bool NextValue(FGlobalArrayShuffleKey& InKey, bool bAutoShuffle, int32& OutIndex);
-		void SetSeed(FGlobalArrayShuffleKey& InKey, int32 InSeed);
-		void ResetSeed(FGlobalArrayShuffleKey& InKey);
-		void ShuffleArray(FGlobalArrayShuffleKey& InKey);
-
-	private:
-
-		TUniquePtr<FArrayIndexShuffler>& GetShuffler(FGlobalArrayShuffleKey& InKey);
-
-		// Protects access to the shuffle namespace data
-		FCriticalSection CritSect;
-
-		// Global array map for namespaced array shufflers
-		TMap<FGlobalArrayShuffleKey, TUniquePtr<FArrayIndexShuffler>> GlobalShuffleIndicies;
-	};
-
-	/** TArrayShuffleOperator shuffles an array on trigger and outputs values sequentially on "next". It avoids repeating shuffled elements and supports auto-shuffling.*/
-	template<typename ArrayType>
-	class TArrayShuffleOperator : public TExecutableOperator<TArrayShuffleOperator<ArrayType>>
-	{
-	public:
-		using FArrayDataReadReference = TDataReadReference<ArrayType>;
-		using ElementType = typename MetasoundArrayNodesPrivate::TArrayElementType<ArrayType>::Type;
-		using FElementTypeWriteReference = TDataWriteReference<ElementType>;
-
-		static const FVertexInterface& GetDefaultInterface()
-		{
-			using namespace ArrayNodeShuffleVertexNames;
-
-			static const FVertexInterface DefaultInterface(
-				FInputVertexInterface(
-					TInputDataVertexModel<FTrigger>(GetInputTriggerNextName(), LOCTEXT("ShuffleOpInputTriggerNextNameTT", "Trigger to get the next value in the shuffled array.")),
-					TInputDataVertexModel<FTrigger>(GetInputTriggerShuffleName(), LOCTEXT("ShuffleOpInputTriggerShuffleNameTT", "Trigger to shuffle the array manually.")),
-					TInputDataVertexModel<FTrigger>(GetInputTriggerResetName(), LOCTEXT("ShuffleOpInputTriggerResetNameTT", "Trigger to reset the random seed stream of the shuffle node.")),
-					TInputDataVertexModel<ArrayType>(GetInputShuffleArrayName(), LOCTEXT("ShuffleOpInputShuffleArrayNameTT", "Input Array.")),
-					TInputDataVertexModel<int32>(GetInputSeedName(), LOCTEXT("ShuffleOpInputSeedNameTT", "Seed to use for the the random shuffle."), -1),
-					TInputDataVertexModel<bool>(GetInputAutoShuffleName(), LOCTEXT("ShuffleOpInputAutoShuffleNameTT", "Set to true to automatically shuffle when the array has been read."), true),
-					TInputDataVertexModel<FString>(GetInputNamespaceName(), LOCTEXT("ShuffleOpInputNamespaceNameTT", "Set to a string value to share shuffle namespace with this array. Shuffling state of array will be shared in this namespace."))
-				),
-				FOutputVertexInterface(
-					TOutputDataVertexModel<FTrigger>(GetOutputTriggerOnNextName(), LOCTEXT("ShuffleOpOutputTriggerOnNextNameTT", "Triggers when the \"Next\" input is triggered.")),
-					TOutputDataVertexModel<FTrigger>(GetOutputTriggerOnShuffleName(), LOCTEXT("ShuffleOpOutputTriggerOnShuffleNameTT", "Triggers when the \"Shuffle\" input is triggered or if the array is auto-shuffled.")),
-					TOutputDataVertexModel<FTrigger>(GetOutputTriggerOnResetName(), LOCTEXT("ShuffleOpOutputTriggerOnResetNameTT", "Triggers when the \"Reset Seed\" input is triggered.")),
-					TOutputDataVertexModel<ElementType>(GetOutputValueName(), LOCTEXT("ShuffleOpOutputValueTT", "Value of the current shuffled element."))
-				)
-			);
-
-			return DefaultInterface;
-		}
-
-		static const FNodeClassMetadata& GetNodeInfo()
-		{
-			auto CreateNodeClassMetadata = []() -> FNodeClassMetadata
-			{
-				FName DataTypeName = GetMetasoundDataTypeName<ArrayType>();
-				FName OperatorName = TEXT("Shuffle");
-				FText NodeDisplayName = FText::Format(LOCTEXT("ArrayOpArrayShuffleDisplayNamePattern", "Shuffle ({0})"), FText::FromString(GetMetasoundDataTypeString<ArrayType>()));
-				FText NodeDescription = LOCTEXT("ArrayOpArrayShuffleDescription", "Output next element of a shuffled array on trigger.");
-				FVertexInterface NodeInterface = GetDefaultInterface();
-
-				return MetasoundArrayNodesPrivate::CreateArrayNodeClassMetadata(DataTypeName, OperatorName, NodeDisplayName, NodeDescription, NodeInterface);
-			};
-
-			static const FNodeClassMetadata Metadata = CreateNodeClassMetadata();
-
-			return Metadata;
-		}
-
-		static TUniquePtr<IOperator> CreateOperator(const FCreateOperatorParams& InParams, TArray<TUniquePtr<IOperatorBuildError>>& OutErrors)
-		{
-			using namespace ArrayNodeShuffleVertexNames;
-			using namespace MetasoundArrayNodesPrivate;
-
-			const FInputVertexInterface& Inputs = InParams.Node.GetVertexInterface().GetInputInterface();
-
-			TDataReadReference<FTrigger> InTriggerNext = GetDataReadReferenceOrConstructWithVertexDefault<FTrigger>(InParams.InputDataReferences, Inputs, GetInputTriggerNextName(), InParams.OperatorSettings);
-			TDataReadReference<FTrigger> InTriggerShuffle = GetDataReadReferenceOrConstructWithVertexDefault<FTrigger>(InParams.InputDataReferences, Inputs, GetInputTriggerShuffleName(), InParams.OperatorSettings);
-			TDataReadReference<FTrigger> InTriggerReset = GetDataReadReferenceOrConstructWithVertexDefault<FTrigger>(InParams.InputDataReferences, Inputs, GetInputTriggerResetName(), InParams.OperatorSettings);
-			FArrayDataReadReference InInputArray = GetDataReadReferenceOrConstructWithVertexDefault<ArrayType>(InParams.InputDataReferences, Inputs, GetInputShuffleArrayName(), InParams.OperatorSettings);
-			TDataReadReference<int32> InSeedValue = GetDataReadReferenceOrConstructWithVertexDefault<int32>(InParams.InputDataReferences, Inputs, GetInputSeedName(), InParams.OperatorSettings);
-			TDataReadReference<bool> bInAutoShuffle = GetDataReadReferenceOrConstructWithVertexDefault<bool>(InParams.InputDataReferences, Inputs, GetInputAutoShuffleName(), InParams.OperatorSettings);
-			TDataReadReference<FString> InShuffleNamespace = GetDataReadReferenceOrConstructWithVertexDefault<FString>(InParams.InputDataReferences, Inputs, GetInputNamespaceName(), InParams.OperatorSettings);
-
-			return MakeUnique<TArrayShuffleOperator>(InParams.OperatorSettings, InTriggerNext, InTriggerShuffle, InTriggerReset, InInputArray, InSeedValue, bInAutoShuffle, InShuffleNamespace);
-		}
-
-		TArrayShuffleOperator(
-			const FOperatorSettings& InSettings,
-			TDataReadReference<FTrigger> InTriggerNext, 
-			TDataReadReference<FTrigger> InTriggerShuffle,
-			TDataReadReference<FTrigger> InTriggerReset,
-			FArrayDataReadReference InInputArray,
-			TDataReadReference<int32> InSeedValue,
-			TDataReadReference<bool> bInAutoShuffle,
-			TDataReadReference<FString> InShuffleNamespace)
-			: TriggerNext(InTriggerNext)
-			, TriggerShuffle(InTriggerShuffle)
-			, TriggerReset(InTriggerReset)
-			, InputArray(InInputArray)
-			, SeedValue(InSeedValue)
-			, bAutoShuffle(bInAutoShuffle)
-			, ShuffleNamespace(InShuffleNamespace)
-			, TriggerOnNext(FTriggerWriteRef::CreateNew(InSettings))
-			, TriggerOnShuffle(FTriggerWriteRef::CreateNew(InSettings))
-			, TriggerOnReset(FTriggerWriteRef::CreateNew(InSettings))
-			, OutValue(TDataWriteReferenceFactory<ElementType>::CreateAny(InSettings))
-		{
-			// Check to see if this is a global shuffler or a local one. 
-			// Global shuffler will use a namespace to opt into it.
-			const ArrayType& InputArrayRef = *InputArray;
-			PrevNamespaceString = *ShuffleNamespace;
-			PrevArrayNum = InputArrayRef.Num();
-			PrevSeedValue = *SeedValue;
-
-			// Make a key!
-			if (PrevNamespaceString.Len() > 0)
-			{
-				ArrayShuffleKey = FGlobalArrayShuffleKey(PrevNamespaceString, PrevArrayNum, PrevSeedValue);
-			}
-			else if (PrevArrayNum > 0)
-			{
-				ArrayIndexShuffler = MakeUnique<FArrayIndexShuffler>(PrevSeedValue, PrevArrayNum);
-			}
-
-			if (PrevArrayNum == 0)
-			{
-				UE_LOG(LogMetasound, Error, TEXT("Array Shuffle: Can't shuffle an empty array"));
-			}
-		}
-
-		virtual ~TArrayShuffleOperator() = default;
-
-		virtual FDataReferenceCollection GetInputs() const override
-		{
-			using namespace ArrayNodeShuffleVertexNames;
-
-			FDataReferenceCollection Inputs;
-			Inputs.AddDataReadReference(GetInputTriggerNextName(), TriggerNext);
-			Inputs.AddDataReadReference(GetInputTriggerShuffleName(), TriggerShuffle);
-			Inputs.AddDataReadReference(GetInputTriggerResetName(), TriggerReset);
-			Inputs.AddDataReadReference(GetInputShuffleArrayName(), InputArray);
-			Inputs.AddDataReadReference(GetInputSeedName(), SeedValue);
-			Inputs.AddDataReadReference(GetInputAutoShuffleName(), bAutoShuffle);
-			Inputs.AddDataReadReference(GetInputNamespaceName(), ShuffleNamespace);
-
-			return Inputs;
-		}
-
-		virtual FDataReferenceCollection GetOutputs() const override
-		{
-			using namespace ArrayNodeShuffleVertexNames;
-
-			FDataReferenceCollection Outputs;
-			Outputs.AddDataReadReference(GetOutputTriggerOnNextName(), TriggerOnNext);
-			Outputs.AddDataReadReference(GetOutputTriggerOnShuffleName(), TriggerOnShuffle);
-			Outputs.AddDataReadReference(GetOutputTriggerOnResetName(), TriggerOnReset);
-			Outputs.AddDataReadReference(GetOutputValueName(), OutValue);
-
-			return Outputs;
-		}
-
-		void Execute()
-		{
-			TriggerOnNext->AdvanceBlock();
-			TriggerOnShuffle->AdvanceBlock();
-			TriggerOnReset->AdvanceBlock();
-
-			const ArrayType& InputArrayRef = *InputArray;
- 
-			// Check for any input pin changes to see if we need to swap behavior
-			if (PrevArrayNum != InputArrayRef.Num() || 
-				PrevSeedValue != *SeedValue || 
-				PrevNamespaceString != *ShuffleNamespace)
-			{
-				PrevArrayNum = InputArrayRef.Num();
-				PrevSeedValue = *SeedValue;
-				PrevNamespaceString = *ShuffleNamespace;
-
-				if (PrevArrayNum > 0)
-				{
-					if (PrevNamespaceString.Len() > 0)
-					{
-						// Make new shuffle key and clear out our old index shuffler
-						ArrayShuffleKey = FGlobalArrayShuffleKey(PrevNamespaceString, PrevArrayNum, PrevSeedValue);
-						ArrayIndexShuffler = nullptr;
-					}
-					else
-					{
-						// Clear out our shuffle key and make a new array index shuffler
-						ArrayShuffleKey = FGlobalArrayShuffleKey();
-						ArrayIndexShuffler = MakeUnique<FArrayIndexShuffler>(PrevSeedValue, PrevArrayNum);
-					}
-				}
-				else
-				{
-					UE_LOG(LogMetasound, Error, TEXT("Array Shuffle: Can't shuffle an empty array."));
-				}
-			}
-
-			TriggerReset->ExecuteBlock(
-				[&](int32 StartFrame, int32 EndFrame)
-				{
-				},
-				[this](int32 StartFrame, int32 EndFrame)
-				{
-					if (ArrayShuffleKey.Hash != INDEX_NONE)
-					{
-						FGlobalArrayShuffleManager& SM = FGlobalArrayShuffleManager::Get();
-						SM.ResetSeed(ArrayShuffleKey);
-					}
-					else
-					{
-						check(ArrayIndexShuffler.IsValid());
-						ArrayIndexShuffler->ResetSeed();
-					}
-
-					TriggerOnReset->TriggerFrame(StartFrame);
-				}
-			);
-
-			TriggerShuffle->ExecuteBlock(
-				[&](int32 StartFrame, int32 EndFrame)
-				{
-				},
-				[this](int32 StartFrame, int32 EndFrame)
-				{
-					if (PrevArrayNum > 0)
-					{
-						if (ArrayShuffleKey.Hash != INDEX_NONE)
-						{
-							FGlobalArrayShuffleManager& SM = FGlobalArrayShuffleManager::Get();
-							SM.ShuffleArray(ArrayShuffleKey);
-						}
-						else
-						{
-							check(ArrayIndexShuffler.IsValid());
-							ArrayIndexShuffler->ShuffleArray();
-						}
-
-						TriggerOnShuffle->TriggerFrame(StartFrame);
-					}
- 				}
-			);
-
-			TriggerNext->ExecuteBlock(
-				[&](int32 StartFrame, int32 EndFrame)
-				{
-				},
-				[this](int32 StartFrame, int32 EndFrame)
-				{
-					const ArrayType& InputArrayRef = *InputArray;
-
-					if (PrevArrayNum > 0)
-					{
-						bool bShuffleTriggered = false;
-						int32 OutShuffleIndex = INDEX_NONE;
-						if (ArrayShuffleKey.Hash != INDEX_NONE)
-						{
-							FGlobalArrayShuffleManager& SM = FGlobalArrayShuffleManager::Get();
-							bShuffleTriggered = SM.NextValue(ArrayShuffleKey, *bAutoShuffle, OutShuffleIndex);
-						}
-						else
-						{
-							check(ArrayIndexShuffler.IsValid());
-							bShuffleTriggered = ArrayIndexShuffler->NextValue(*bAutoShuffle, OutShuffleIndex);
-						}
-
-						check(OutShuffleIndex != INDEX_NONE);
-						check(OutShuffleIndex < InputArrayRef.Num());
-
-						// Write out the value of the shuffled index
-						*OutValue = InputArrayRef[OutShuffleIndex];
-
-						TriggerOnNext->TriggerFrame(StartFrame);
-
-						// Trigger out if the array was auto-shuffled
-						if (bShuffleTriggered)
-						{
-							TriggerOnShuffle->TriggerFrame(StartFrame);
-						}
-					}
-				}
-			);
-		}
-
-	private:
-
-		// Inputs
-		TDataReadReference<FTrigger> TriggerNext;
-		TDataReadReference<FTrigger> TriggerShuffle;
-		TDataReadReference<FTrigger> TriggerReset;
-		FArrayDataReadReference InputArray;
-		TDataReadReference<int32> SeedValue;
-		TDataReadReference<bool> bAutoShuffle;
-		TDataReadReference<FString> ShuffleNamespace;
-
-		// Outputs
-		TDataWriteReference<FTrigger> TriggerOnNext;
-		TDataWriteReference<FTrigger> TriggerOnShuffle;
-		TDataWriteReference<FTrigger> TriggerOnReset;
-		TDataWriteReference<ElementType> OutValue;
-
-		// Data
-		TUniquePtr<FArrayIndexShuffler> ArrayIndexShuffler;
-		FGlobalArrayShuffleKey ArrayShuffleKey;
-		FString PrevNamespaceString;
-		int32 PrevArrayNum = INDEX_NONE;
-		int32 PrevSeedValue = INDEX_NONE;
-	};
-
-	template<typename ArrayType>
-	class TArrayShuffleNode : public FNodeFacade
-	{
-	public:
-		TArrayShuffleNode(const FNodeInitData& InInitData)
-			: FNodeFacade(InInitData.InstanceName, InInitData.InstanceID, TFacadeOperatorClass<TArrayShuffleOperator<ArrayType>>())
-		{
-		}
-
-		virtual ~TArrayShuffleNode() = default;
-	};
 }
 
 #undef LOCTEXT_NAMESPACE
diff --git a/Engine/Plugins/Runtime/Metasound/Source/MetasoundFrontend/Public/MetasoundArrayNodesRegistration.h b/Engine/Plugins/Runtime/Metasound/Source/MetasoundFrontend/Public/MetasoundArrayNodesRegistration.h
--- a/Engine/Plugins/Runtime/Metasound/Source/MetasoundFrontend/Public/MetasoundArrayNodesRegistration.h
+++ b/Engine/Plugins/Runtime/Metasound/Source/MetasoundFrontend/Public/MetasoundArrayNodesRegistration.h
@@ -3,6 +3,8 @@
 #pragma once
 
 #include "MetasoundArrayNodes.h"
+#include "MetasoundArrayShuffleNode.h"
+#include "MetasoundArrayRandomNode.h"
 #include "MetasoundNodeRegistrationMacro.h"
 
 #include <type_traits>
@@ -41,6 +43,9 @@ namespace Metasound
 
 			// Array shuffle is supported for all types that get is supported for.
 			static constexpr bool bIsArrayShuffleSupported = bIsArrayGetSupported;
+
+			// Random get is supported for all types that get is supported for.
+			static constexpr bool bIsArrayRandomGetSupported = bIsArrayGetSupported;
 		};
 
 		template<typename ArrayType, typename std::enable_if<TArrayNodeSupport<ArrayType>::bIsArrayGetSupported, bool>::type = true>
@@ -129,6 +134,20 @@ namespace Metasound
 			// No op if not supported
 			return true;
 		}
+
+		template<typename ArrayType, typename std::enable_if<TArrayNodeSupport<ArrayType>::bIsArrayRandomGetSupported, bool>::type = true>
+		bool RegisterArrayRandomGetNode()
+		{
+			using FNodeType = typename Metasound::TArrayRandomGetNode<ArrayType>;
+			return RegisterNodeWithFrontend<FNodeType>();
+		}
+
+		template<typename ArrayType, typename std::enable_if<!TArrayNodeSupport<ArrayType>::bIsArrayRandomGetSupported, bool>::type = true>
+		bool RegisterArrayShuffleNode()
+		{
+			// No op if not supported
+			return true;
+		}
 	}
 
 	/** Registers all available array nodes which can be instantiated for the given
@@ -146,6 +165,7 @@ namespace Metasound
 		bSuccess = bSuccess && RegisterArraySubsetNode<ArrayType>();
 		bSuccess = bSuccess && RegisterArrayConcatNode<ArrayType>();
 		bSuccess = bSuccess && RegisterArrayShuffleNode<ArrayType>();
+		bSuccess = bSuccess && RegisterArrayRandomGetNode<ArrayType>();
 		return bSuccess;
 	}
 }
diff --git a/Engine/Plugins/Runtime/Metasound/Source/MetasoundFrontend/Public/MetasoundArrayRandomNode.h b/Engine/Plugins/Runtime/Metasound/Source/MetasoundFrontend/Public/MetasoundArrayRandomNode.h
new file mode 100644
--- /dev/null
+++ b/Engine/Plugins/Runtime/Metasound/Source/MetasoundFrontend/Public/MetasoundArrayRandomNode.h
@@ -0,0 +1,371 @@
+// Copyright Epic Games, Inc. All Rights Reserved.
+
+#pragma once
+
+#include "CoreMinimal.h"
+#include "Internationalization/Text.h"
+#include "MetasoundBuilderInterface.h"
+#include "MetasoundDataFactory.h"
+#include "MetasoundExecutableOperator.h"
+#include "MetasoundFacade.h"
+#include "MetasoundLog.h"
+#include "MetasoundNodeInterface.h"
+#include "MetasoundNodeRegistrationMacro.h"
+#include "MetasoundOperatorInterface.h"
+#include "MetasoundPrimitives.h"
+#include "MetasoundTrigger.h"
+#include "Containers/CircularQueue.h"
+#include "MetasoundArrayNodes.h"
+
+#include <type_traits>
+
+#define LOCTEXT_NAMESPACE "MetasoundFrontend"
+
+
+namespace Metasound
+{
+	/** Random Get Node Vertex Names */
+	namespace ArrayNodeRandomGetVertexNames
+	{
+		METASOUNDFRONTEND_API const FString& GetInputTriggerNextName();
+		METASOUNDFRONTEND_API const FString& GetInputTriggerResetName();
+		METASOUNDFRONTEND_API const FString& GetInputRandomArrayName();
+		METASOUNDFRONTEND_API const FString& GetInputWeightsName();
+		METASOUNDFRONTEND_API const FString& GetInputSeedName();
+		METASOUNDFRONTEND_API const FString& GetInputNoRepeatOrderName();
+		METASOUNDFRONTEND_API const FString& GetInputEnableSharedStateName();
+
+		METASOUNDFRONTEND_API const FString& GetOutputTriggerOnNextName();
+		METASOUNDFRONTEND_API const FString& GetOutputTriggerOnResetName();
+		METASOUNDFRONTEND_API const FString& GetOutputValueName();
+	}
+
+	class METASOUNDFRONTEND_API FArrayRandomGet
+	{
+	public:
+		FArrayRandomGet() = default;
+		FArrayRandomGet(int32 InSeed, int32 InMaxIndex, int32 InNoRepeatOrder);
+		void Init(int32 InSeed, int32 InMaxIndex, int32 InNoRepeatOrder);
+		void SetSeed(int32 InSeed);
+		void SetNoRepeatOrder(int32 InNoRepeatOrder);
+		void ResetSeed();
+		int32 NextValue();
+
+	private:
+		// The current index into the array of indicies (wraps between 0 and ShuffleIndices.Num())
+		TArray<int32> PreviousIndices;
+		TUniquePtr<TCircularQueue<int32>> PreviousIndicesQueue;
+		int32 NoRepeatOrder = 0;
+
+		// Array of indices (in order 0 to Num)
+		int32 MaxIndex = 0;
+		TArray<float> RandomWeights;
+
+		// Random stream to use to randomize the shuffling
+		FRandomStream RandomStream;
+	};
+
+	class METASOUNDFRONTEND_API FSharedStateRandomGetManager
+	{
+	public:
+		static FSharedStateRandomGetManager& Get();
+
+		void InitSharedState(uint32 InSharedStateId, int32 InSeed, int32 InNumElements, int32 InNoRepeatOrder);
+		int32 NextValue(uint32 InSharedStateId);
+		void SetSeed(uint32 InSharedStateId, int32 InSeed);
+		void SetNoRepeatOrder(uint32 InSharedStateId, int32 InNoRepeatOrder);
+		void ResetSeed(uint32 InSharedStateId);
+
+	private:
+		FSharedStateRandomGetManager() = default;
+		~FSharedStateRandomGetManager() = default;
+
+		FCriticalSection CritSect;
+
+		TMap<uint32, TUniquePtr<FArrayRandomGet>> RandomGets;
+	};
+
+	/** TArrayShuffleOperator shuffles an array on trigger and outputs values sequentially on "next". It avoids repeating shuffled elements and supports auto-shuffling.*/
+	template<typename ArrayType>
+	class TArrayRandomGetOperator : public TExecutableOperator<TArrayRandomGetOperator<ArrayType>>
+	{
+	public:
+		using FArrayDataReadReference = TDataReadReference<ArrayType>;
+		using ElementType = typename MetasoundArrayNodesPrivate::TArrayElementType<ArrayType>::Type;
+		using FElementTypeWriteReference = TDataWriteReference<ElementType>;
+		using WeightArrayType = TArray<float>;
+
+		static const FVertexInterface& GetDefaultInterface()
+		{
+			using namespace ArrayNodeRandomGetVertexNames;
+
+			static const FVertexInterface DefaultInterface(
+				FInputVertexInterface(
+					TInputDataVertexModel<FTrigger>(GetInputTriggerNextName(), LOCTEXT("RandomGetArrayOpInputTriggerNextTT", "Trigger to get the next value in the randomized array.")),
+					TInputDataVertexModel<FTrigger>(GetInputTriggerResetName(), LOCTEXT("RandomGetArrayOpInputTriggerResetTT", "Trigger to reset the seed for the randomized array.")),
+					TInputDataVertexModel<ArrayType>(GetInputRandomArrayName(), LOCTEXT("RandomGetArrayOpArrayTT", "Input array to randomized.")),
+					TInputDataVertexModel<WeightArrayType>(GetInputWeightsName(), LOCTEXT("RandomGetOpInputSeedTT", "Input array of weights to use for random selection. Will repeat if this array is shorter than the input array to select from.")),
+					TInputDataVertexModel<int32>(GetInputSeedName(), LOCTEXT("RandomGetOpInputSeedTT", "Seed to use for the the random shuffle."), -1),
+					TInputDataVertexModel<int32>(GetInputNoRepeatOrderName(), LOCTEXT("RandomGetOpInputNoRepeatOrderTT", "The number of elements to track to avoid repeating in a row."), 1),
+					TInputDataVertexModel<bool>(GetInputEnableSharedStateName(), LOCTEXT("RandomGetOpEnableSharedStateNameTT", "Set to enabled shared state across instances of this metasound."), false)
+				),
+				FOutputVertexInterface(
+					TOutputDataVertexModel<FTrigger>(GetOutputTriggerOnNextName(), LOCTEXT("RandomGetOpOutputTriggerOnNextNameTT", "Triggers when the \"Next\" input is triggered.")),
+					TOutputDataVertexModel<FTrigger>(GetOutputTriggerOnResetName(), LOCTEXT("RandomGetOpOutputTriggerOnShuffleNameTT", "Triggers when the \"Shuffle\" input is triggered or if the array is auto-shuffled.")),
+					TOutputDataVertexModel<ElementType>(GetOutputValueName(), LOCTEXT("RandomGetOpOutputValueTT", "Value of the current shuffled element."))
+				)
+			);
+
+			return DefaultInterface;
+		}
+
+		static const FNodeClassMetadata& GetNodeInfo()
+		{
+			auto CreateNodeClassMetadata = []() -> FNodeClassMetadata
+			{
+				FName DataTypeName = GetMetasoundDataTypeName<ArrayType>();
+				FName OperatorName = TEXT("Random Get");
+				FText NodeDisplayName = FText::Format(LOCTEXT("RandomGetArrayOpDisplayNamePattern", "Random Get ({0})"), FText::FromString(GetMetasoundDataTypeString<ArrayType>()));
+				FText NodeDescription = LOCTEXT("RandomGetArrayDescription", "Randomly retrieve data from input array using the supplied weights.");
+				FVertexInterface NodeInterface = GetDefaultInterface();
+
+				return MetasoundArrayNodesPrivate::CreateArrayNodeClassMetadata(DataTypeName, OperatorName, NodeDisplayName, NodeDescription, NodeInterface);
+			};
+
+			static const FNodeClassMetadata Metadata = CreateNodeClassMetadata();
+
+			return Metadata;
+		}
+
+		static TUniquePtr<IOperator> CreateOperator(const FCreateOperatorParams& InParams, TArray<TUniquePtr<IOperatorBuildError>>& OutErrors)
+		{
+			using namespace ArrayNodeRandomGetVertexNames;
+			using namespace MetasoundArrayNodesPrivate;
+
+			const FInputVertexInterface& Inputs = InParams.Node.GetVertexInterface().GetInputInterface();
+
+			FTriggerReadRef InTriggerNext = GetDataReadReferenceOrConstructWithVertexDefault<FTrigger>(InParams.InputDataReferences, Inputs, GetInputTriggerNextName(), InParams.OperatorSettings);
+			FTriggerReadRef InTriggerReset = GetDataReadReferenceOrConstructWithVertexDefault<FTrigger>(InParams.InputDataReferences, Inputs, GetInputTriggerResetName(), InParams.OperatorSettings);
+			FArrayDataReadReference InInputArray = GetDataReadReferenceOrConstructWithVertexDefault<ArrayType>(InParams.InputDataReferences, Inputs, GetInputRandomArrayName(), InParams.OperatorSettings);
+			TDataReadReference<WeightArrayType> InInputWeightsArray = GetDataReadReferenceOrConstructWithVertexDefault<WeightArrayType>(InParams.InputDataReferences, Inputs, GetInputRandomArrayName(), InParams.OperatorSettings);
+			FInt32ReadRef InSeedValue = GetDataReadReferenceOrConstructWithVertexDefault<int32>(InParams.InputDataReferences, Inputs, GetInputSeedName(), InParams.OperatorSettings);
+			FInt32ReadRef InNoRepeatOrder = GetDataReadReferenceOrConstructWithVertexDefault<int32>(InParams.InputDataReferences, Inputs, GetInputNoRepeatOrderName(), InParams.OperatorSettings);
+			FBoolReadRef bInEnableSharedState = GetDataReadReferenceOrConstructWithVertexDefault<bool>(InParams.InputDataReferences, Inputs, GetInputEnableSharedStateName(), InParams.OperatorSettings);
+
+			return MakeUnique<TArrayRandomGetOperator>(InParams, InTriggerNext, InTriggerReset, InInputArray, InInputWeightsArray, InSeedValue, InNoRepeatOrder, bInEnableSharedState);
+		}
+
+		TArrayRandomGetOperator(
+			const FCreateOperatorParams& InParams,
+			const FTriggerReadRef& InTriggerNext,
+			const FTriggerReadRef& InTriggerReset,
+			const FArrayDataReadReference& InInputArray,
+			const TDataReadReference<WeightArrayType>& InInputWeightsArray,
+			const FInt32ReadRef& InSeedValue,
+			const FInt32ReadRef& InNoRepeatOrder,
+			const FBoolReadRef& bInEnableSharedState)
+			: TriggerNext(InTriggerNext)
+			, TriggerReset(InTriggerReset)
+			, InputArray(InInputArray)
+			, InputWeightsArray(InInputWeightsArray)
+			, SeedValue(InSeedValue)
+			, NoRepeatOrder(InNoRepeatOrder)
+			, bEnableSharedState(bInEnableSharedState)
+			, TriggerOnNext(FTriggerWriteRef::CreateNew(InParams.OperatorSettings))
+			, TriggerOnReset(FTriggerWriteRef::CreateNew(InParams.OperatorSettings))
+			, OutValue(TDataWriteReferenceFactory<ElementType>::CreateAny(InParams.OperatorSettings))
+		{
+			// Check to see if this is a global shuffler or a local one. 
+			// Global shuffler will use a namespace to opt into it.
+			PrevSeedValue = *SeedValue;
+			PrevNoRepeatOrder = FMath::Max(*NoRepeatOrder, 0);
+
+			const ArrayType& InputArrayRef = *InputArray;
+			int32 ArraySize = InputArrayRef.Num();
+
+			if (ArraySize > 0)
+			{
+				if (*bEnableSharedState)
+				{
+					// Get the environment variable for the unique ID of the sound
+					SharedStateUniqueId = InParams.Environment.GetValue<uint32>(TEXT("SoundUniqueId"));
+					check(SharedStateUniqueId != INDEX_NONE);
+
+					FSharedStateRandomGetManager& RGM = FSharedStateRandomGetManager::Get();
+					RGM.InitSharedState(SharedStateUniqueId, PrevSeedValue, ArraySize, PrevNoRepeatOrder);
+				}
+				else
+				{
+					ArrayRandomGet = MakeUnique<FArrayRandomGet>(PrevSeedValue, ArraySize, PrevNoRepeatOrder);
+				}
+			}
+			else
+			{
+				UE_LOG(LogMetasound, Error, TEXT("Array Random Get: Can't retrieve random elements from an empty array"));
+			}
+		}
+
+		virtual ~TArrayRandomGetOperator() = default;
+
+		virtual FDataReferenceCollection GetInputs() const override
+		{
+			using namespace ArrayNodeRandomGetVertexNames;
+
+			FDataReferenceCollection Inputs;
+			Inputs.AddDataReadReference(GetInputTriggerNextName(), TriggerNext);
+			Inputs.AddDataReadReference(GetInputTriggerResetName(), TriggerReset);
+			Inputs.AddDataReadReference(GetInputRandomArrayName(), InputArray);
+			Inputs.AddDataReadReference(GetInputWeightsName(), InputWeightsArray);
+			Inputs.AddDataReadReference(GetInputSeedName(), SeedValue);
+			Inputs.AddDataReadReference(GetInputNoRepeatOrderName(), NoRepeatOrder);
+			Inputs.AddDataReadReference(GetInputEnableSharedStateName(), bEnableSharedState);
+
+			return Inputs;
+		}
+
+		virtual FDataReferenceCollection GetOutputs() const override
+		{
+			using namespace ArrayNodeRandomGetVertexNames;
+
+			FDataReferenceCollection Outputs;
+			Outputs.AddDataReadReference(GetOutputTriggerOnNextName(), TriggerOnNext);
+			Outputs.AddDataReadReference(GetOutputTriggerOnResetName(), TriggerOnReset);
+			Outputs.AddDataReadReference(GetOutputValueName(), OutValue);
+
+			return Outputs;
+		}
+
+		void Execute()
+		{
+			TriggerOnNext->AdvanceBlock();
+			TriggerOnReset->AdvanceBlock();
+ 
+ 			const ArrayType& InputArrayRef = *InputArray;
+ 
+ 			// Check for a seed change
+ 			if (PrevSeedValue != *SeedValue)
+ 			{
+				PrevSeedValue = *SeedValue;
+
+				if (SharedStateUniqueId != INDEX_NONE)
+				{
+					FSharedStateRandomGetManager& RGM = FSharedStateRandomGetManager::Get();
+					RGM.SetSeed(SharedStateUniqueId, PrevSeedValue);
+				}
+				else
+				{
+					check(ArrayRandomGet.IsValid());
+					ArrayRandomGet->SetSeed(PrevSeedValue);
+				}
+			}
+
+			if (PrevNoRepeatOrder != *NoRepeatOrder)
+			{	
+				PrevNoRepeatOrder = *NoRepeatOrder;
+				if (SharedStateUniqueId != INDEX_NONE)
+				{
+					FSharedStateRandomGetManager& RGM = FSharedStateRandomGetManager::Get();
+					RGM.SetNoRepeatOrder(SharedStateUniqueId, PrevNoRepeatOrder);
+				}
+				else
+				{
+					check(ArrayRandomGet.IsValid());
+					ArrayRandomGet->SetNoRepeatOrder(PrevNoRepeatOrder);
+				}
+			}
+ 
+			// Don't do anything if our array is empty
+			if (InputArrayRef.Num() == 0)
+			{
+				return;
+			}
+ 
+ 			TriggerReset->ExecuteBlock(
+				[&](int32 StartFrame, int32 EndFrame)
+				{
+				},
+				[this](int32 StartFrame, int32 EndFrame)
+				{
+					if (SharedStateUniqueId != INDEX_NONE)
+					{
+						FSharedStateRandomGetManager& RGM = FSharedStateRandomGetManager::Get();
+						RGM.ResetSeed(SharedStateUniqueId);
+					}
+					else
+					{
+						check(ArrayRandomGet.IsValid());
+						ArrayRandomGet->ResetSeed();
+					}
+					TriggerOnReset->TriggerFrame(StartFrame);
+				}
+			);
+ 
+			TriggerNext->ExecuteBlock(
+				[&](int32 StartFrame, int32 EndFrame)
+				{
+				},
+				[this](int32 StartFrame, int32 EndFrame)
+				{
+					const ArrayType& InputArrayRef = *InputArray;
+					int32 OutRandomIndex = INDEX_NONE;
+
+					if (SharedStateUniqueId != INDEX_NONE)
+					{
+						FSharedStateRandomGetManager& RGM = FSharedStateRandomGetManager::Get();
+						OutRandomIndex = RGM.NextValue(SharedStateUniqueId);
+					}
+					else
+					{
+						check(ArrayRandomGet.IsValid());
+						OutRandomIndex = ArrayRandomGet->NextValue();
+					}
+
+					check(OutRandomIndex != INDEX_NONE);
+
+					// The input array size may have changed, so make sure it's wrapped into range of the input array
+					*OutValue = InputArrayRef[OutRandomIndex % InputArrayRef.Num()];
+
+					TriggerOnNext->TriggerFrame(StartFrame);
+				}
+			);
+		}
+
+	private:
+
+		// Inputs
+		FTriggerReadRef TriggerNext;
+		FTriggerReadRef TriggerReset;
+		FArrayDataReadReference InputArray;
+		TDataReadReference<WeightArrayType> InputWeightsArray;
+		FInt32ReadRef SeedValue;
+		FInt32ReadRef NoRepeatOrder;
+		FBoolReadRef bEnableSharedState;
+
+		// Outputs
+		FTriggerWriteRef TriggerOnNext;
+		FTriggerWriteRef TriggerOnReset;
+		TDataWriteReference<ElementType> OutValue;
+
+		// Data
+		TUniquePtr<FArrayRandomGet> ArrayRandomGet;
+		int32 PrevSeedValue = INDEX_NONE;
+		int32 PrevNoRepeatOrder = INDEX_NONE;
+		uint32 SharedStateUniqueId = INDEX_NONE;
+	};
+
+	template<typename ArrayType>
+	class TArrayRandomGetNode : public FNodeFacade
+	{
+	public:
+		TArrayRandomGetNode(const FNodeInitData& InInitData)
+			: FNodeFacade(InInitData.InstanceName, InInitData.InstanceID, TFacadeOperatorClass<TArrayRandomGetOperator<ArrayType>>())
+		{
+		}
+
+		virtual ~TArrayRandomGetNode() = default;
+	};
+
+
+}
+
+#undef LOCTEXT_NAMESPACE
diff --git a/Engine/Plugins/Runtime/Metasound/Source/MetasoundFrontend/Public/MetasoundArrayShuffleNode.h b/Engine/Plugins/Runtime/Metasound/Source/MetasoundFrontend/Public/MetasoundArrayShuffleNode.h
new file mode 100644
--- /dev/null
+++ b/Engine/Plugins/Runtime/Metasound/Source/MetasoundFrontend/Public/MetasoundArrayShuffleNode.h
@@ -0,0 +1,430 @@
+// Copyright Epic Games, Inc. All Rights Reserved.
+
+#pragma once
+
+#include "CoreMinimal.h"
+#include "Internationalization/Text.h"
+#include "MetasoundBuilderInterface.h"
+#include "MetasoundDataFactory.h"
+#include "MetasoundExecutableOperator.h"
+#include "MetasoundFacade.h"
+#include "MetasoundLog.h"
+#include "MetasoundNodeInterface.h"
+#include "MetasoundNodeRegistrationMacro.h"
+#include "MetasoundOperatorInterface.h"
+#include "MetasoundPrimitives.h"
+#include "MetasoundTrigger.h"
+#include "Containers/CircularQueue.h"
+#include "MetasoundArrayNodes.h"
+
+#include <type_traits>
+
+#define LOCTEXT_NAMESPACE "MetasoundFrontend"
+
+
+namespace Metasound
+{
+	/** Shuffle Node Vertex Names */
+	namespace ArrayNodeShuffleVertexNames
+	{
+		/** Input Vertex Names */
+		METASOUNDFRONTEND_API const FString& GetInputTriggerNextName();
+		METASOUNDFRONTEND_API const FString& GetInputTriggerShuffleName();
+		METASOUNDFRONTEND_API const FString& GetInputTriggerResetName();
+		METASOUNDFRONTEND_API const FString& GetInputShuffleArrayName();
+		METASOUNDFRONTEND_API const FString& GetInputSeedName();
+		METASOUNDFRONTEND_API const FString& GetInputAutoShuffleName();
+		METASOUNDFRONTEND_API const FString& GetInputEnableSharedStateName();
+
+		METASOUNDFRONTEND_API const FString& GetOutputTriggerOnNextName();
+		METASOUNDFRONTEND_API const FString& GetOutputTriggerOnShuffleName();
+		METASOUNDFRONTEND_API const FString& GetOutputTriggerOnResetName();
+		METASOUNDFRONTEND_API const FString& GetOutputValueName();
+	}
+
+	class METASOUNDFRONTEND_API FArrayIndexShuffler
+	{
+	public:
+		FArrayIndexShuffler() = default;
+		FArrayIndexShuffler(int32 InSeed, int32 MaxIndices);
+
+		void Init(int32 InSeed, int32 MaxIndices);
+		void SetSeed(int32 InSeed);
+		void ResetSeed();
+
+		// Returns the next value in the array indices. Returns true if the array was re-shuffled automatically.
+		bool NextValue(bool bAutoShuffle, int32& OutIndex);
+
+		// Shuffle the array with the given max indices
+		void ShuffleArray();
+
+	private:
+		// Helper function to swap the current index with a random index
+		void RandomSwap(int32 InCurrentIndex, int32 InStartIndex, int32 InEndIndex);
+
+		// The current index into the array of indicies (wraps between 0 and ShuffleIndices.Num())
+		int32 CurrentIndex = 0;
+
+		// The previously returned value. Used to avoid repeating the last value on shuffle.
+		int32 PrevValue = INDEX_NONE;
+
+		// Array of indices (in order 0 to Num), shuffled
+		TArray<int32> ShuffleIndices;
+
+		// Random stream to use to randomize the shuffling
+		FRandomStream RandomStream;
+	};
+
+	class FSharedStateShuffleManager
+	{
+	public:
+		static FSharedStateShuffleManager& Get()
+		{
+			static FSharedStateShuffleManager GSM;
+			return GSM;
+		}
+
+		void InitSharedState(uint32 InSharedStateId, int32 InSeed, int32 InNumElements)
+		{
+			FScopeLock Lock(&CritSect);
+
+			if (!Shufflers.Contains(InSharedStateId))
+			{
+				Shufflers.Add(InSharedStateId, MakeUnique<FArrayIndexShuffler>(InSeed, InNumElements));
+			}
+		}
+
+		bool NextValue(uint32 InSharedStateId, bool bAutoShuffle, int32& OutIndex)
+		{
+			FScopeLock Lock(&CritSect);
+			TUniquePtr<FArrayIndexShuffler>* Shuffler = Shufflers.Find(InSharedStateId);
+			return (*Shuffler)->NextValue(bAutoShuffle, OutIndex);
+		}
+
+		void SetSeed(uint32 InSharedStateId, int32 InSeed)
+		{
+			FScopeLock Lock(&CritSect);
+			TUniquePtr<FArrayIndexShuffler>* Shuffler = Shufflers.Find(InSharedStateId);
+			(*Shuffler)->SetSeed(InSeed);
+		}
+
+		void ResetSeed(uint32 InSharedStateId)
+		{
+			FScopeLock Lock(&CritSect);
+			TUniquePtr<FArrayIndexShuffler>* Shuffler = Shufflers.Find(InSharedStateId);
+			(*Shuffler)->ResetSeed();
+		}
+
+		void ShuffleArray(uint32 InSharedStateId)
+		{
+			FScopeLock Lock(&CritSect);
+			TUniquePtr<FArrayIndexShuffler>* Shuffler = Shufflers.Find(InSharedStateId);
+			(*Shuffler)->ShuffleArray();
+		}
+
+	private:
+		FSharedStateShuffleManager() = default;
+		~FSharedStateShuffleManager() = default;
+
+		FCriticalSection CritSect;
+
+		TMap<uint32, TUniquePtr<FArrayIndexShuffler>> Shufflers;	
+	};
+
+	/** TArrayShuffleOperator shuffles an array on trigger and outputs values sequentially on "next". It avoids repeating shuffled elements and supports auto-shuffling.*/
+	template<typename ArrayType>
+	class TArrayShuffleOperator : public TExecutableOperator<TArrayShuffleOperator<ArrayType>>
+	{
+	public:
+		using FArrayDataReadReference = TDataReadReference<ArrayType>;
+		using ElementType = typename MetasoundArrayNodesPrivate::TArrayElementType<ArrayType>::Type;
+		using FElementTypeWriteReference = TDataWriteReference<ElementType>;
+
+		static const FVertexInterface& GetDefaultInterface()
+		{
+			using namespace ArrayNodeShuffleVertexNames;
+
+			static const FVertexInterface DefaultInterface(
+				FInputVertexInterface(
+					TInputDataVertexModel<FTrigger>(GetInputTriggerNextName(), LOCTEXT("ShuffleOpInputTriggerNextTT", "Trigger to get the next value in the shuffled array.")),
+					TInputDataVertexModel<FTrigger>(GetInputTriggerShuffleName(), LOCTEXT("ShuffleOpInputTriggerShuffleTT", "Trigger to shuffle the array manually.")),
+					TInputDataVertexModel<FTrigger>(GetInputTriggerResetName(), LOCTEXT("ShuffleOpInputTriggerResetTT", "Trigger to reset the random seed stream of the shuffle node.")),
+					TInputDataVertexModel<ArrayType>(GetInputShuffleArrayName(), LOCTEXT("ShuffleOpInputShuffleArrayTT", "Input Array.")),
+					TInputDataVertexModel<int32>(GetInputSeedName(), LOCTEXT("ShuffleOpInputSeedTT", "Seed to use for the the random shuffle."), -1),
+					TInputDataVertexModel<bool>(GetInputAutoShuffleName(), LOCTEXT("ShuffleOpInputAutoShuffleTT", "Set to true to automatically shuffle when the array has been read."), true),
+					TInputDataVertexModel<bool>(GetInputEnableSharedStateName(), LOCTEXT("ShuffleOpInputEnableSharedStatTT", "Set to enabled shared state across instances of this metasound."), false)
+					),
+				FOutputVertexInterface(
+					TOutputDataVertexModel<FTrigger>(GetOutputTriggerOnNextName(), LOCTEXT("ShuffleOpOutputTriggerOnNextNameTT", "Triggers when the \"Next\" input is triggered.")),
+					TOutputDataVertexModel<FTrigger>(GetOutputTriggerOnShuffleName(), LOCTEXT("ShuffleOpOutputTriggerOnShuffleNameTT", "Triggers when the \"Shuffle\" input is triggered or if the array is auto-shuffled.")),
+					TOutputDataVertexModel<FTrigger>(GetOutputTriggerOnResetName(), LOCTEXT("ShuffleOpOutputTriggerOnResetNameTT", "Triggers when the \"Reset Seed\" input is triggered.")),
+					TOutputDataVertexModel<ElementType>(GetOutputValueName(), LOCTEXT("ShuffleOpOutputValueTT", "Value of the current shuffled element."))
+				)
+			);
+
+			return DefaultInterface;
+		}
+
+		static const FNodeClassMetadata& GetNodeInfo()
+		{
+			auto CreateNodeClassMetadata = []() -> FNodeClassMetadata
+			{
+				FName DataTypeName = GetMetasoundDataTypeName<ArrayType>();
+				FName OperatorName = TEXT("Shuffle");
+				FText NodeDisplayName = FText::Format(LOCTEXT("ArrayOpArrayShuffleDisplayNamePattern", "Shuffle ({0})"), FText::FromString(GetMetasoundDataTypeString<ArrayType>()));
+				FText NodeDescription = LOCTEXT("ArrayOpArrayShuffleDescription", "Output next element of a shuffled array on trigger.");
+				FVertexInterface NodeInterface = GetDefaultInterface();
+
+				return MetasoundArrayNodesPrivate::CreateArrayNodeClassMetadata(DataTypeName, OperatorName, NodeDisplayName, NodeDescription, NodeInterface);
+			};
+
+			static const FNodeClassMetadata Metadata = CreateNodeClassMetadata();
+
+			return Metadata;
+		}
+
+		static TUniquePtr<IOperator> CreateOperator(const FCreateOperatorParams& InParams, TArray<TUniquePtr<IOperatorBuildError>>& OutErrors)
+		{
+			using namespace ArrayNodeShuffleVertexNames;
+			using namespace MetasoundArrayNodesPrivate;
+
+			const FInputVertexInterface& Inputs = InParams.Node.GetVertexInterface().GetInputInterface();
+
+			TDataReadReference<FTrigger> InTriggerNext = GetDataReadReferenceOrConstructWithVertexDefault<FTrigger>(InParams.InputDataReferences, Inputs, GetInputTriggerNextName(), InParams.OperatorSettings);
+			TDataReadReference<FTrigger> InTriggerShuffle = GetDataReadReferenceOrConstructWithVertexDefault<FTrigger>(InParams.InputDataReferences, Inputs, GetInputTriggerShuffleName(), InParams.OperatorSettings);
+			TDataReadReference<FTrigger> InTriggerReset = GetDataReadReferenceOrConstructWithVertexDefault<FTrigger>(InParams.InputDataReferences, Inputs, GetInputTriggerResetName(), InParams.OperatorSettings);
+			FArrayDataReadReference InInputArray = GetDataReadReferenceOrConstructWithVertexDefault<ArrayType>(InParams.InputDataReferences, Inputs, GetInputShuffleArrayName(), InParams.OperatorSettings);
+			TDataReadReference<int32> InSeedValue = GetDataReadReferenceOrConstructWithVertexDefault<int32>(InParams.InputDataReferences, Inputs, GetInputSeedName(), InParams.OperatorSettings);
+			TDataReadReference<bool> bInAutoShuffle = GetDataReadReferenceOrConstructWithVertexDefault<bool>(InParams.InputDataReferences, Inputs, GetInputAutoShuffleName(), InParams.OperatorSettings);
+			TDataReadReference<bool> bInEnableSharedState = GetDataReadReferenceOrConstructWithVertexDefault<bool>(InParams.InputDataReferences, Inputs, GetInputEnableSharedStateName(), InParams.OperatorSettings);
+
+			return MakeUnique<TArrayShuffleOperator>(InParams, InTriggerNext, InTriggerShuffle, InTriggerReset, InInputArray, InSeedValue, bInAutoShuffle, bInEnableSharedState);
+		}
+
+		TArrayShuffleOperator(
+			const FCreateOperatorParams& InParams,
+			const TDataReadReference<FTrigger>& InTriggerNext, 
+			const TDataReadReference<FTrigger>& InTriggerShuffle,
+			const TDataReadReference<FTrigger>& InTriggerReset,
+			const FArrayDataReadReference& InInputArray,
+			const TDataReadReference<int32>& InSeedValue,
+			const TDataReadReference<bool>& bInAutoShuffle, 
+			const TDataReadReference<bool>& bInEnableSharedState)
+			: TriggerNext(InTriggerNext)
+			, TriggerShuffle(InTriggerShuffle)
+			, TriggerReset(InTriggerReset)
+			, InputArray(InInputArray)
+			, SeedValue(InSeedValue)
+			, bAutoShuffle(bInAutoShuffle)
+			, bEnableSharedState(bInEnableSharedState)
+			, TriggerOnNext(FTriggerWriteRef::CreateNew(InParams.OperatorSettings))
+			, TriggerOnShuffle(FTriggerWriteRef::CreateNew(InParams.OperatorSettings))
+			, TriggerOnReset(FTriggerWriteRef::CreateNew(InParams.OperatorSettings))
+			, OutValue(TDataWriteReferenceFactory<ElementType>::CreateAny(InParams.OperatorSettings))
+		{
+			// Check to see if this is a global shuffler or a local one. 
+			// Global shuffler will use a namespace to opt into it.
+			PrevSeedValue = *SeedValue;
+
+			const ArrayType& InputArrayRef = *InputArray;
+			int32 ArraySize = InputArrayRef.Num();
+
+			if (ArraySize > 0)
+			{
+				if (*bEnableSharedState)
+				{
+					// Get the environment variable for the unique ID of the sound
+					SharedStateUniqueId = InParams.Environment.GetValue<uint32>(TEXT("SoundUniqueId"));
+					check(SharedStateUniqueId != INDEX_NONE);
+
+					FSharedStateShuffleManager& SM = FSharedStateShuffleManager::Get();
+					SM.InitSharedState(SharedStateUniqueId, PrevSeedValue, ArraySize);
+				}
+				else
+				{
+					ArrayIndexShuffler = MakeUnique<FArrayIndexShuffler>(PrevSeedValue, ArraySize);
+				}
+			}
+			else 
+			{
+				UE_LOG(LogMetasound, Error, TEXT("Array Shuffle: Can't shuffle an empty array"));
+			}
+		}
+
+		virtual ~TArrayShuffleOperator() = default;
+
+		virtual FDataReferenceCollection GetInputs() const override
+		{
+			using namespace ArrayNodeShuffleVertexNames;
+
+			FDataReferenceCollection Inputs;
+			Inputs.AddDataReadReference(GetInputTriggerNextName(), TriggerNext);
+			Inputs.AddDataReadReference(GetInputTriggerShuffleName(), TriggerShuffle);
+			Inputs.AddDataReadReference(GetInputTriggerResetName(), TriggerReset);
+			Inputs.AddDataReadReference(GetInputShuffleArrayName(), InputArray);
+			Inputs.AddDataReadReference(GetInputSeedName(), SeedValue);
+			Inputs.AddDataReadReference(GetInputAutoShuffleName(), bAutoShuffle);
+			Inputs.AddDataReadReference(GetInputEnableSharedStateName(), bEnableSharedState);
+
+			return Inputs;
+		}
+
+		virtual FDataReferenceCollection GetOutputs() const override
+		{
+			using namespace ArrayNodeShuffleVertexNames;
+
+			FDataReferenceCollection Outputs;
+			Outputs.AddDataReadReference(GetOutputTriggerOnNextName(), TriggerOnNext);
+			Outputs.AddDataReadReference(GetOutputTriggerOnShuffleName(), TriggerOnShuffle);
+			Outputs.AddDataReadReference(GetOutputTriggerOnResetName(), TriggerOnReset);
+			Outputs.AddDataReadReference(GetOutputValueName(), OutValue);
+
+			return Outputs;
+		}
+
+		void Execute()
+		{
+			TriggerOnNext->AdvanceBlock();
+			TriggerOnShuffle->AdvanceBlock();
+			TriggerOnReset->AdvanceBlock();
+
+			const ArrayType& InputArrayRef = *InputArray;
+ 
+			// Check for a seed change
+			if (PrevSeedValue != *SeedValue)
+			{
+				PrevSeedValue = *SeedValue;
+
+				if (SharedStateUniqueId != INDEX_NONE)
+				{
+					FSharedStateShuffleManager& SM = FSharedStateShuffleManager::Get();
+					SM.SetSeed(SharedStateUniqueId, PrevSeedValue);
+				}
+				else
+				{
+					check(ArrayIndexShuffler.IsValid());
+					ArrayIndexShuffler->SetSeed(PrevSeedValue);
+				}
+			}
+
+			// Don't do anything if our array is empty
+			if (InputArrayRef.Num() == 0)
+			{
+				return;
+			}
+
+			TriggerReset->ExecuteBlock(
+				[&](int32 StartFrame, int32 EndFrame)
+				{
+				},
+				[this](int32 StartFrame, int32 EndFrame)
+				{
+					if (SharedStateUniqueId != INDEX_NONE)
+					{
+						FSharedStateShuffleManager& SM = FSharedStateShuffleManager::Get();
+						SM.ResetSeed(SharedStateUniqueId);
+					}
+					else
+					{
+						check(ArrayIndexShuffler.IsValid());
+						ArrayIndexShuffler->ResetSeed();
+					}
+					TriggerOnReset->TriggerFrame(StartFrame);
+				}
+			);
+
+			TriggerShuffle->ExecuteBlock(
+				[&](int32 StartFrame, int32 EndFrame)
+				{
+				},
+				[this](int32 StartFrame, int32 EndFrame)
+				{
+					if (SharedStateUniqueId != INDEX_NONE)
+					{
+						FSharedStateShuffleManager& SM = FSharedStateShuffleManager::Get();
+						SM.ShuffleArray(SharedStateUniqueId);
+					}
+					else
+					{
+						check(ArrayIndexShuffler.IsValid());
+						ArrayIndexShuffler->ShuffleArray();
+					}
+					TriggerOnShuffle->TriggerFrame(StartFrame);
+				}
+			);
+
+			TriggerNext->ExecuteBlock(
+				[&](int32 StartFrame, int32 EndFrame)
+				{
+				},
+				[this](int32 StartFrame, int32 EndFrame)
+				{
+					const ArrayType& InputArrayRef = *InputArray;
+
+					bool bShuffleTriggered = false;
+					int32 OutShuffleIndex = INDEX_NONE;
+
+					if (SharedStateUniqueId != INDEX_NONE)
+					{
+						FSharedStateShuffleManager& SM = FSharedStateShuffleManager::Get();
+						bShuffleTriggered = SM.NextValue(SharedStateUniqueId, *bAutoShuffle, OutShuffleIndex);
+					}
+					else
+					{
+						check(ArrayIndexShuffler.IsValid());
+						bShuffleTriggered = ArrayIndexShuffler->NextValue(*bAutoShuffle, OutShuffleIndex);
+					}
+
+					check(OutShuffleIndex != INDEX_NONE);
+
+					// The input array size may have changed, so make sure it's wrapped into range of the input array
+					*OutValue = InputArrayRef[OutShuffleIndex % InputArrayRef.Num()];
+
+					TriggerOnNext->TriggerFrame(StartFrame);
+
+					// Trigger out if the array was auto-shuffled
+					if (bShuffleTriggered)
+					{
+						TriggerOnShuffle->TriggerFrame(StartFrame);
+					}
+				}
+			);
+		}
+
+	private:
+
+		// Inputs
+		TDataReadReference<FTrigger> TriggerNext;
+		TDataReadReference<FTrigger> TriggerShuffle;
+		TDataReadReference<FTrigger> TriggerReset;
+		FArrayDataReadReference InputArray;
+		TDataReadReference<int32> SeedValue;
+		TDataReadReference<bool> bAutoShuffle;
+		TDataReadReference<bool> bEnableSharedState;
+
+		// Outputs
+		TDataWriteReference<FTrigger> TriggerOnNext;
+		TDataWriteReference<FTrigger> TriggerOnShuffle;
+		TDataWriteReference<FTrigger> TriggerOnReset;
+		TDataWriteReference<ElementType> OutValue;
+
+		// Data
+		TUniquePtr<FArrayIndexShuffler> ArrayIndexShuffler;
+		int32 PrevSeedValue = INDEX_NONE;
+		uint32 SharedStateUniqueId = INDEX_NONE;
+	};
+
+	template<typename ArrayType>
+	class TArrayShuffleNode : public FNodeFacade
+	{
+	public:
+		TArrayShuffleNode(const FNodeInitData& InInitData)
+			: FNodeFacade(InInitData.InstanceName, InInitData.InstanceID, TFacadeOperatorClass<TArrayShuffleOperator<ArrayType>>())
+		{
+		}
+
+		virtual ~TArrayShuffleNode() = default;
+	};
+}
+
+#undef LOCTEXT_NAMESPACE
diff --git a/Engine/Plugins/Runtime/Metasound/Source/MetasoundGraphCore/Public/MetasoundEnvironment.h b/Engine/Plugins/Runtime/Metasound/Source/MetasoundGraphCore/Public/MetasoundEnvironment.h
--- a/Engine/Plugins/Runtime/Metasound/Source/MetasoundGraphCore/Public/MetasoundEnvironment.h
+++ b/Engine/Plugins/Runtime/Metasound/Source/MetasoundGraphCore/Public/MetasoundEnvironment.h
@@ -6,9 +6,9 @@
 #include "CoreMinimal.h"
 #include <type_traits>
 
-#define DECLARE_METASOUND_ENVIRONMENT_VARIABLE_TYPE(VarType) \
+#define DECLARE_METASOUND_ENVIRONMENT_VARIABLE_TYPE(DLL_API, VarType) \
 	template<> \
-	struct ::Metasound::TEnvironmentVariableTypeInfo<VarType> \
+	struct DLL_API ::Metasound::TEnvironmentVariableTypeInfo<VarType> \
 	{ \
 		typedef VarType Type; \
 		static const FMetasoundEnvironmentVariableTypeId TypeId; \
@@ -38,6 +38,7 @@ namespace Metasound
 
 
 	template<typename VarType>
+	//struct METASOUNDGRAPHCORE_API TEnvironmentVariableTypeInfo
 	struct TEnvironmentVariableTypeInfo
 	{
 		typedef VarType Type;
@@ -251,18 +252,18 @@ namespace Metasound
 }
 
 /** Declare basic set of variable types. */
-DECLARE_METASOUND_ENVIRONMENT_VARIABLE_TYPE(void);
-DECLARE_METASOUND_ENVIRONMENT_VARIABLE_TYPE(bool);
-DECLARE_METASOUND_ENVIRONMENT_VARIABLE_TYPE(int8);
-DECLARE_METASOUND_ENVIRONMENT_VARIABLE_TYPE(uint8);
-DECLARE_METASOUND_ENVIRONMENT_VARIABLE_TYPE(int16);
-DECLARE_METASOUND_ENVIRONMENT_VARIABLE_TYPE(uint16);
-DECLARE_METASOUND_ENVIRONMENT_VARIABLE_TYPE(int32);
-DECLARE_METASOUND_ENVIRONMENT_VARIABLE_TYPE(uint32);
-DECLARE_METASOUND_ENVIRONMENT_VARIABLE_TYPE(int64);
-DECLARE_METASOUND_ENVIRONMENT_VARIABLE_TYPE(uint64);
-DECLARE_METASOUND_ENVIRONMENT_VARIABLE_TYPE(float);
-DECLARE_METASOUND_ENVIRONMENT_VARIABLE_TYPE(double);
-DECLARE_METASOUND_ENVIRONMENT_VARIABLE_TYPE(FString);
-DECLARE_METASOUND_ENVIRONMENT_VARIABLE_TYPE(FName);
+DECLARE_METASOUND_ENVIRONMENT_VARIABLE_TYPE(METASOUNDGRAPHCORE_API, void);
+DECLARE_METASOUND_ENVIRONMENT_VARIABLE_TYPE(METASOUNDGRAPHCORE_API, bool);
+DECLARE_METASOUND_ENVIRONMENT_VARIABLE_TYPE(METASOUNDGRAPHCORE_API, int8);
+DECLARE_METASOUND_ENVIRONMENT_VARIABLE_TYPE(METASOUNDGRAPHCORE_API, uint8);
+DECLARE_METASOUND_ENVIRONMENT_VARIABLE_TYPE(METASOUNDGRAPHCORE_API, int16);
+DECLARE_METASOUND_ENVIRONMENT_VARIABLE_TYPE(METASOUNDGRAPHCORE_API, uint16);
+DECLARE_METASOUND_ENVIRONMENT_VARIABLE_TYPE(METASOUNDGRAPHCORE_API, int32);
+DECLARE_METASOUND_ENVIRONMENT_VARIABLE_TYPE(METASOUNDGRAPHCORE_API, uint32);
+DECLARE_METASOUND_ENVIRONMENT_VARIABLE_TYPE(METASOUNDGRAPHCORE_API, int64);
+DECLARE_METASOUND_ENVIRONMENT_VARIABLE_TYPE(METASOUNDGRAPHCORE_API, uint64);
+DECLARE_METASOUND_ENVIRONMENT_VARIABLE_TYPE(METASOUNDGRAPHCORE_API, float);
+DECLARE_METASOUND_ENVIRONMENT_VARIABLE_TYPE(METASOUNDGRAPHCORE_API, double);
+DECLARE_METASOUND_ENVIRONMENT_VARIABLE_TYPE(METASOUNDGRAPHCORE_API, FString);
+DECLARE_METASOUND_ENVIRONMENT_VARIABLE_TYPE(METASOUNDGRAPHCORE_API, FName);
 
