MetasoundWaveShaper Node
- New metasound node that runs a waveshaper on incoming audio
- Adding Bias and more shapes to the DSP waveshaper class

diff --git a/Engine/Plugins/Runtime/Metasound/Source/MetasoundStandardNodes/Private/MetasoundWaveShaperNode.cpp b/Engine/Plugins/Runtime/Metasound/Source/MetasoundStandardNodes/Private/MetasoundWaveShaperNode.cpp
new file mode 100644
--- /dev/null
+++ b/Engine/Plugins/Runtime/Metasound/Source/MetasoundStandardNodes/Private/MetasoundWaveShaperNode.cpp
@@ -0,0 +1,195 @@
+// Copyright Epic Games, Inc. All Rights Reserved.
+
+#include "MetasoundWaveShaperNode.h"
+
+#include "Internationalization/Text.h"
+#include "MetasoundExecutableOperator.h"
+#include "MetasoundNodeRegistrationMacro.h"
+#include "MetasoundDataTypeRegistrationMacro.h"
+#include "MetasoundPrimitives.h"
+#include "MetasoundStandardNodesNames.h"
+#include "MetasoundAudioBuffer.h"
+#include "DSP/WaveShaper.h"
+
+#define LOCTEXT_NAMESPACE "MetasoundStandardNodes"
+
+namespace Metasound
+{
+	namespace WaveShaperNode
+	{
+		static const TCHAR* InParamNameAudioInput = TEXT("In");
+		static const TCHAR* InParamNameWaveShapeAmount = TEXT("Amount");
+		static const TCHAR* InParamNameWaveShapeBias = TEXT("Bias");
+		static const TCHAR* InParamNameOutputGain = TEXT("OutputGain");
+		static const TCHAR* InParamNameType = TEXT("Type");
+		static const TCHAR* OutParamNameAudio = TEXT("Out");
+	}
+
+	using namespace WaveShaperNode;
+
+	DECLARE_METASOUND_ENUM(Audio::EWaveShaperType, Audio::EWaveShaperType::ATan, METASOUNDSTANDARDNODES_API, FEnumEWaveShaperType, FEnumWaveShaperTypeInfo, FEnumWaveShaperReadRef, FEnumWaveShaperWriteRef);
+
+	DEFINE_METASOUND_ENUM_BEGIN(Audio::EWaveShaperType, FEnumEWaveShaperType, "WaveShaperType")
+		DEFINE_METASOUND_ENUM_ENTRY(Audio::EWaveShaperType::Sin, LOCTEXT("SinDescription", "Sine"), LOCTEXT("SinTT", "Sine WaveShaper")),
+		DEFINE_METASOUND_ENUM_ENTRY(Audio::EWaveShaperType::ATan, LOCTEXT("ATanDescription", "Inverse Tangent"), LOCTEXT("ATanTT", "Inverse Tangent WaveShaper")),
+		DEFINE_METASOUND_ENUM_ENTRY(Audio::EWaveShaperType::Tanh, LOCTEXT("TanHDescription", "Hyperbolic Tangent"), LOCTEXT("TanHTT", "Hyperbolic Tangent WaveShaper")),
+		DEFINE_METASOUND_ENUM_ENTRY(Audio::EWaveShaperType::Cubic, LOCTEXT("CubicDescription", "Cubic Polynomial"), LOCTEXT("CubicTT", "Cubic Polynomial WaveShaper")),
+		DEFINE_METASOUND_ENUM_ENTRY(Audio::EWaveShaperType::HardClip, LOCTEXT("HarClipDescription", "Hard Clip"), LOCTEXT("HardClipTT", "Hard Clipper"))
+		DEFINE_METASOUND_ENUM_END()
+
+	class FWaveShaperOperator : public TExecutableOperator<FWaveShaperOperator>
+	{
+	public:
+
+		static const FNodeClassMetadata& GetNodeInfo();
+		static const FVertexInterface& GetVertexInterface();
+		static TUniquePtr<IOperator> CreateOperator(const FCreateOperatorParams& InParams, FBuildErrorArray& OutErrors);
+
+		FWaveShaperOperator(const FOperatorSettings& InSettings, 
+			const FAudioBufferReadRef& InAudioInput, 
+			const FFloatReadRef& InWaveShapeAmount,
+			const FFloatReadRef& InBias,
+			const FFloatReadRef& InOutputGain,
+			const FEnumWaveShaperReadRef& InType);
+
+		virtual FDataReferenceCollection GetInputs() const override;
+		virtual FDataReferenceCollection GetOutputs() const override;
+		void Execute();
+
+	private:
+		// The input audio buffer
+		FAudioBufferReadRef AudioInput;
+
+		// The amount that the wave is shaped
+		FFloatReadRef Amount;
+
+		// DC offset to apply before gain
+		FFloatReadRef Bias;
+
+		// The amount of gain
+		FFloatReadRef OutputGain;
+
+		// What algorithm to use
+		FEnumWaveShaperReadRef Type;
+
+		// The audio output
+		FAudioBufferWriteRef AudioOutput;
+
+		// The internal WaveShaper
+		Audio::FWaveShaper WaveShaper;
+	};
+
+	FWaveShaperOperator::FWaveShaperOperator(const FOperatorSettings& InSettings, 
+		const FAudioBufferReadRef& InAudioInput, 
+		const FFloatReadRef& InAmount,
+		const FFloatReadRef& InBias,
+		const FFloatReadRef& InOutputGain,
+		const FEnumWaveShaperReadRef& InType)
+		: AudioInput(InAudioInput)
+		, Amount(InAmount)
+		, Bias(InBias)
+		, OutputGain(InOutputGain)
+		, Type(InType)
+		, AudioOutput(FAudioBufferWriteRef::CreateNew(InSettings))
+	{
+		WaveShaper.Init(InSettings.GetSampleRate());
+	}
+
+	FDataReferenceCollection FWaveShaperOperator::GetInputs() const
+	{
+		FDataReferenceCollection InputDataReferences;
+		InputDataReferences.AddDataReadReference(InParamNameAudioInput, AudioInput);
+		InputDataReferences.AddDataReadReference(InParamNameWaveShapeAmount, Amount);
+		InputDataReferences.AddDataReadReference(InParamNameWaveShapeBias, Bias);
+		InputDataReferences.AddDataReadReference(InParamNameOutputGain, OutputGain);
+		InputDataReferences.AddDataReadReference(InParamNameType, Type);
+
+		return InputDataReferences;
+	}
+
+	FDataReferenceCollection FWaveShaperOperator::GetOutputs() const
+	{
+		FDataReferenceCollection OutputDataReferences;
+		OutputDataReferences.AddDataReadReference(OutParamNameAudio, AudioOutput);
+		return OutputDataReferences;
+	}
+
+	void FWaveShaperOperator::Execute()
+	{
+		const float* InputAudio = AudioInput->GetData();
+
+		float* OutputAudio = AudioOutput->GetData();
+		int32 NumFrames = AudioInput->Num();
+
+		WaveShaper.SetAmount(*Amount);
+		WaveShaper.SetBias(*Bias);
+		WaveShaper.SetOutputGainLinear(*OutputGain);
+		WaveShaper.SetType(*Type);
+
+		WaveShaper.ProcessAudioBuffer(InputAudio, OutputAudio, NumFrames);
+	}
+
+	const FVertexInterface& FWaveShaperOperator::GetVertexInterface()
+	{
+		static const FVertexInterface Interface(
+			FInputVertexInterface(
+				TInputDataVertexModel<FAudioBuffer>(InParamNameAudioInput, LOCTEXT("AudioInputTT", "Audio input.")),
+				TInputDataVertexModel<float>(InParamNameWaveShapeAmount, LOCTEXT("WaveShapeAmountTT", "The amount of wave shaping to apply."), 1.0f),
+				TInputDataVertexModel<float>(InParamNameWaveShapeBias, LOCTEXT("WaveShaperBiasTT", "DC offset to apply before wave shaping."), 0.0f),
+				TInputDataVertexModel<float>(InParamNameOutputGain, LOCTEXT("WaveShaperGainTT", "The amount of gain to apply after processing."), 1.0f),
+				TInputDataVertexModel<FEnumEWaveShaperType>(InParamNameType, LOCTEXT("WaveShaperTypeTT", "Algorithm to use to process the audio."))
+			),
+			FOutputVertexInterface(
+				TOutputDataVertexModel<FAudioBuffer>(OutParamNameAudio, LOCTEXT("WaveShaperOutputTT", "Audio output."))
+			)
+		);
+
+		return Interface;
+	}
+
+	const FNodeClassMetadata& FWaveShaperOperator::GetNodeInfo()
+	{
+		auto InitNodeInfo = []() -> FNodeClassMetadata
+		{
+			FNodeClassMetadata Info;
+			Info.ClassName = { Metasound::StandardNodes::Namespace, TEXT("WaveShaper"), Metasound::StandardNodes::AudioVariant };
+			Info.MajorVersion = 1;
+			Info.MinorVersion = 0;
+			Info.DisplayName = LOCTEXT("Metasound_WaveShaperDisplayName", "WaveShaper");
+			Info.Description = LOCTEXT("Metasound_WaveShaperNodeDescription", "Applies non-linear shaping to the audio input.");
+			Info.Author = PluginAuthor;
+			Info.PromptIfMissing = PluginNodeMissingPrompt;
+			Info.DefaultInterface = GetVertexInterface();
+
+			return Info;
+		};
+
+		static const FNodeClassMetadata Info = InitNodeInfo();
+
+		return Info;
+	}
+
+	TUniquePtr<IOperator> FWaveShaperOperator::CreateOperator(const FCreateOperatorParams& InParams, FBuildErrorArray& OutErrors)
+	{
+		const FWaveShaperNode& WaveShaperNode = static_cast<const FWaveShaperNode&>(InParams.Node);
+		const FDataReferenceCollection& InputCollection = InParams.InputDataReferences;
+		const FInputVertexInterface& InputInterface = GetVertexInterface().GetInputInterface();
+
+		FAudioBufferReadRef AudioIn = InputCollection.GetDataReadReferenceOrConstruct<FAudioBuffer>(InParamNameAudioInput, InParams.OperatorSettings);
+		FFloatReadRef InAmount = InputCollection.GetDataReadReferenceOrConstructWithVertexDefault<float>(InputInterface, InParamNameWaveShapeAmount);
+		FFloatReadRef InBias = InputCollection.GetDataReadReferenceOrConstructWithVertexDefault<float>(InputInterface, InParamNameWaveShapeBias);
+		FFloatReadRef OutGain = InputCollection.GetDataReadReferenceOrConstructWithVertexDefault<float>(InputInterface, InParamNameOutputGain);
+		FEnumWaveShaperReadRef InType = InputCollection.GetDataReadReferenceOrConstruct<FEnumEWaveShaperType>(InParamNameType);
+
+		return MakeUnique<FWaveShaperOperator>(InParams.OperatorSettings, AudioIn, InAmount, InBias, OutGain, InType);
+	}
+
+	FWaveShaperNode::FWaveShaperNode(const FNodeInitData& InitData)
+		: FNodeFacade(InitData.InstanceName, InitData.InstanceID, TFacadeOperatorClass<FWaveShaperOperator>())
+	{
+	}
+
+	METASOUND_REGISTER_NODE(FWaveShaperNode)
+}
+
+#undef LOCTEXT_NAMESPACE
diff --git a/Engine/Plugins/Runtime/Metasound/Source/MetasoundStandardNodes/Public/MetasoundWaveShaperNode.h b/Engine/Plugins/Runtime/Metasound/Source/MetasoundStandardNodes/Public/MetasoundWaveShaperNode.h
new file mode 100644
--- /dev/null
+++ b/Engine/Plugins/Runtime/Metasound/Source/MetasoundStandardNodes/Public/MetasoundWaveShaperNode.h
@@ -0,0 +1,23 @@
+// Copyright Epic Games, Inc. All Rights Reserved.
+
+#pragma once
+
+#include "CoreMinimal.h"
+#include "MetasoundFacade.h"
+
+namespace Metasound
+{
+	/** FWaveShaperNode
+	 *
+	 *  Delays an audio buffer by a specified amount.
+	 */
+	class METASOUNDSTANDARDNODES_API FWaveShaperNode : public FNodeFacade
+	{
+	public:
+		/**
+		 * Constructor used by the Metasound Frontend.
+		 */
+		FWaveShaperNode(const FNodeInitData& InitData);
+	};
+} // namespace Metasound
+
diff --git a/Engine/Source/Runtime/SignalProcessing/Private/WaveShaper.cpp b/Engine/Source/Runtime/SignalProcessing/Private/WaveShaper.cpp
--- a/Engine/Source/Runtime/SignalProcessing/Private/WaveShaper.cpp
+++ b/Engine/Source/Runtime/SignalProcessing/Private/WaveShaper.cpp
@@ -7,8 +7,11 @@ namespace Audio
 {
 	FWaveShaper::FWaveShaper()
 		: Amount(1.0f)
-		, AtanAmount(FMath::Atan(Amount))
 		, OutputGain(1.0f)
+		, Bias(0.f)
+		, Type(EWaveShaperType::ATan)
+		, OneOverAtanAmount(1.f / (FMath::Atan(Amount)))
+		, OneOverTanhAmount(1.f / Audio::FastTanh(Amount))
 	{
 	}
 
@@ -22,8 +25,14 @@ namespace Audio
 
 	void FWaveShaper::SetAmount(const float InAmount)
 	{		
-		Amount = FMath::Max(InAmount, SMALL_NUMBER);
-		AtanAmount = FMath::Atan(Amount);
+		Amount = FMath::Max(InAmount, 1.f);
+		OneOverAtanAmount = 1.f / FMath::Atan(Amount);
+		OneOverTanhAmount = 1.f / Audio::FastTanh(Amount);
+	}
+
+	void FWaveShaper::SetBias(const float InBias)
+	{
+		Bias = FMath::Clamp(InBias, -1.f, 1.f);
 	}
 
 	void FWaveShaper::SetOutputGainDb(const float InGainDb)
@@ -31,9 +40,109 @@ namespace Audio
 		OutputGain = Audio::ConvertToLinear(InGainDb);
 	}
 
+	void FWaveShaper::SetOutputGainLinear(const float InGainLinear)
+	{
+		OutputGain = InGainLinear;
+	}
+
+	void FWaveShaper::SetType(const EWaveShaperType InType)
+	{
+		Type = InType;
+	}
+
 	void FWaveShaper::ProcessAudio(const float InSample, float& OutSample)
 	{
-		OutSample = OutputGain * FMath::Atan(InSample * Amount) / AtanAmount;
+		OutSample = OutputGain * FMath::Atan((InSample + Bias) * Amount) * OneOverAtanAmount;
+	}
+
+	void FWaveShaper::ProcessAudioBuffer(const float* InBuffer, float* OutBuffer, int32 NumFrames)
+	{
+		switch (Type)
+		{
+		case EWaveShaperType::ATan:
+			ProcessATan(InBuffer, OutBuffer, NumFrames);
+			break;
+		case EWaveShaperType::Cubic:
+			ProcessCubic(InBuffer, OutBuffer, NumFrames);
+			break;
+		case EWaveShaperType::Sin:
+			ProcessSin(InBuffer, OutBuffer, NumFrames);
+			break;
+		case EWaveShaperType::HardClip:
+			ProcessHardClip(InBuffer, OutBuffer, NumFrames);
+			break;
+		case EWaveShaperType::Tanh:
+		default:
+			ProcessTanh(InBuffer, OutBuffer, NumFrames);
+			break;
+		}
+	}
+
+	void FWaveShaper::ProcessHardClip(const float* InBuffer, float* OutBuffer, int32 NumSamples)
+	{
+		for (int32 Sample = 0; Sample < NumSamples; Sample++)
+		{
+			OutBuffer[Sample] = OutputGain * FMath::Clamp(Amount * (InBuffer[Sample] + Bias), -1.f, 1.f);
+		}
+	}
+
+	void FWaveShaper::ProcessATan(const float* InBuffer, float* OutBuffer, int32 NumSamples)
+	{
+		for (int32 Sample = 0; Sample < NumSamples; Sample++)
+		{
+			OutBuffer[Sample] = OutputGain * FMath::Atan((InBuffer[Sample] + Bias) * Amount) * OneOverAtanAmount;
+		}
+	}
+
+	void FWaveShaper::ProcessCubic(const float* InBuffer, float* OutBuffer, int32 NumSamples)
+	{
+		static const float CubicMax = 2.f / 3.f;
+		static const float OneThird = 1.f / 3.f;
+		float InSampleCopy = 0.f;
+
+		for (int32 Sample = 0; Sample < NumSamples; Sample++)
+		{
+			InSampleCopy = (InBuffer[Sample] + Bias) * Amount;
+
+			if (FMath::Abs(InSampleCopy) > 1.f)
+			{
+				OutBuffer[Sample] = FMath::Sign(InSampleCopy) * CubicMax;
+			}
+			else
+			{
+				OutBuffer[Sample] = InSampleCopy - (InSampleCopy * InSampleCopy * InSampleCopy * OneThird);
+			}
+
+			OutBuffer[Sample] *= OutputGain;
+		}
+	}
+
+	void FWaveShaper::ProcessSin(const float* InBuffer, float* OutBuffer, int32 NumSamples)
+	{
+		float InSampleCopy = 0.f;
+
+		for (int32 Sample = 0; Sample < NumSamples; Sample++)
+		{
+			InSampleCopy = (InBuffer[Sample] + Bias) * Amount;
+
+			if (FMath::Abs(InSampleCopy) > HALF_PI)
+			{
+				OutBuffer[Sample] = FMath::Sign(InSampleCopy);
+			}
+			else
+			{
+				OutBuffer[Sample] = Audio::FastSin3(InSampleCopy);
+			}
+
+			OutBuffer[Sample] *= OutputGain;
+		}
 	}
 
+	void FWaveShaper::ProcessTanh(const float* InBuffer, float* OutBuffer, int32 NumSamples)
+	{
+		for (int32 Sample = 0; Sample < NumSamples; Sample++)
+		{
+			OutBuffer[Sample] = OutputGain * Audio::FastTanh((InBuffer[Sample] + Bias) * Amount) * OneOverTanhAmount;
+		}
+	}
 }
diff --git a/Engine/Source/Runtime/SignalProcessing/Public/DSP/WaveShaper.h b/Engine/Source/Runtime/SignalProcessing/Public/DSP/WaveShaper.h
--- a/Engine/Source/Runtime/SignalProcessing/Public/DSP/WaveShaper.h
+++ b/Engine/Source/Runtime/SignalProcessing/Public/DSP/WaveShaper.h
@@ -6,6 +6,14 @@
 
 namespace Audio
 {
+	enum class SIGNALPROCESSING_API EWaveShaperType : int32
+	{
+		Sin,
+		ATan,
+		Tanh,
+		Cubic,
+		HardClip
+	};
 
 	// A digital wave shaping effect to cause audio distortion
 	// https://en.wikipedia.org/wiki/Waveshaper
@@ -24,17 +32,39 @@ namespace Audio
 		// Sets the amount of wave shapping. 0.0 is no effect.
 		void SetAmount(const float InAmount);
 
+		// Set DC offset before processing
+		void SetBias(const float InBias);
+
 		// Sets the output gain of the waveshaper
 		void SetOutputGainDb(const float InGainDb);
 
-		// Processes the audio frame (audio frame must have channels equal to that used during initialization)
+		// Sets the output gain of the waveshaper
+		void SetOutputGainLinear(const float InGainLinear);
+
+		void SetType(const EWaveShaperType InType);
+
+		// Processes one Sample, using ATan
+		// Deprecating in favor of ProcessAudioBuffer
 		void ProcessAudio(const float InSample, float& OutSample);
 
+		// Process an entire buffer of audio
+		void ProcessAudioBuffer(const float* InBuffer, float* OutBuffer, int32 NumFrames);
+
 	private:
 
 		float Amount;
-		float AtanAmount;
 		float OutputGain;
+		float Bias;
+		EWaveShaperType Type;
+
+		float OneOverAtanAmount;
+		float OneOverTanhAmount;
+
+		void ProcessHardClip(const float* InBuffer, float* OutBuffer, int32 NumSamples);
+		void ProcessTanh(const float* InBuffer, float* OutBuffer, int32 NumSamples);
+		void ProcessATan(const float* InBuffer, float* OutBuffer, int32 NumSamples);
+		void ProcessCubic(const float* InBuffer, float* OutBuffer, int32 NumSamples);
+		void ProcessSin(const float* InBuffer, float* OutBuffer, int32 NumSamples);
 	};
 
 }
