Adding seek and looping ability to metasound wave playe node.

diff --git a/Engine/Plugins/Runtime/Metasound/Source/MetasoundEngine/Private/MetasoundWave.cpp b/Engine/Plugins/Runtime/Metasound/Source/MetasoundEngine/Private/MetasoundWave.cpp
--- a/Engine/Plugins/Runtime/Metasound/Source/MetasoundEngine/Private/MetasoundWave.cpp
+++ b/Engine/Plugins/Runtime/Metasound/Source/MetasoundEngine/Private/MetasoundWave.cpp
@@ -32,6 +32,8 @@ namespace Audio
 		FsOutToInRatio = (OutputSampleRate / InputSampleRate);
 		MaxPitchShiftRatio = FMath::Pow(2.f, InInitParams.MaxPitchShiftMagnitudeAllowedInOctaves);
 		MaxPitchShiftCents = InInitParams.MaxPitchShiftMagnitudeAllowedInOctaves * 1200.f;
+		StartTimeSeconds = InInitParams.StartTimeSeconds;
+
 
 		NumChannels = InWave->GetNumChannels();
 		DecodeBlockSizeInFrames = 64;
@@ -53,7 +55,7 @@ namespace Audio
 		return bSuccessful;
 	}
 
-	uint32 FSimpleDecoderWrapper::GenerateAudio(float* OutputDest, int32 NumOutputFrames, float PitchShiftInCents)
+	uint32 FSimpleDecoderWrapper::GenerateAudio(float* OutputDest, int32 NumOutputFrames, float PitchShiftInCents, bool bIsLooping)
 	{
 		const uint32 NumOutputSamples = NumOutputFrames * NumChannels;
 
@@ -75,16 +77,19 @@ namespace Audio
 			Resampler.SetSampleRateRatio(SampleRateRatio);
 
 			// perform SRC and push to circular buffer until we have enough frames for the output
-			while (Decoder && !bDecoderIsDone && (OutputCircularBuffer.Num() < NumOutputSamples))
+			while (Decoder && !(bDecoderIsDone || bDecoderHasLooped) && (OutputCircularBuffer.Num() < NumOutputSamples))
 			{
 				// get more audio from the decoder
 				Audio::IDecoderOutput::FPushedAudioDetails Details;
-				bDecoderIsDone = (Decoder->Decode() == Audio::IDecoder::EDecodeResult::Finished);
+				const Audio::IDecoder::EDecodeResult  DecodeResult = Decoder->Decode(bIsLooping);
 				const int32 NumFramesDecoded = Output->PopAudio(PreSrcBuffer, Details) / NumChannels;
 				int32 NumResamplerOutputFrames = 0;
 				int32 Error = Resampler.ProcessAudio(PreSrcBuffer.GetData(), NumFramesDecoded, bDecoderIsDone, PostSrcBuffer.GetData(), MaxNumResamplerOutputFramesPerBlock, NumResamplerOutputFrames);
 				ensure(Error == 0);
 				OutputCircularBuffer.Push(PostSrcBuffer.GetData(), NumResamplerOutputFrames * NumChannels);
+
+				bDecoderIsDone = DecodeResult == Audio::IDecoder::EDecodeResult::Finished;
+				bDecoderHasLooped = DecodeResult == Audio::IDecoder::EDecodeResult::Looped;
 			}
 		}
 
@@ -92,8 +97,10 @@ namespace Audio
 		{
 			OutputCircularBuffer.Pop(OutputDest, NumOutputSamples);
 		}
-		else if (ensure(bDecoderIsDone))
+		else if (ensure(bDecoderHasLooped || bDecoderIsDone))
 		{
+			bDecoderHasLooped = false;
+
 			const int32 NumSamplesToPop = OutputCircularBuffer.Num();
 			const int32 NumSamplesRemaining = NumOutputSamples - NumSamplesToPop;
 			OutputCircularBuffer.Pop(OutputDest, OutputCircularBuffer.Num());
@@ -110,6 +117,14 @@ namespace Audio
 		return NumOutputSamples; // update once we are aware of partial decode on last buffer
 	}
 
+	void FSimpleDecoderWrapper::SeekToTime(const float InSeconds)
+	{
+		if (Input.IsValid())
+		{
+			Input->SeekToTime(InSeconds);
+		}
+	}
+
 	bool FSimpleDecoderWrapper::InitializeDecodersInternal(const FSoundWaveProxyPtr& Wave)
 	{
 		if (!ensure(Wave.IsValid()))
@@ -120,6 +135,7 @@ namespace Audio
 		// Input:
 		FName OldFormat = Wave->GetRuntimeFormat();
 		Input = MakeShareable(Audio::CreateBackCompatDecoderInput(OldFormat, Wave).Release());
+		Input->SeekToTime(StartTimeSeconds);
 
 		if (!Input)
 		{
@@ -147,6 +163,9 @@ namespace Audio
 		// Decoder:
 		Decoder = Codec->CreateDecoder(Input.Get(), Output.Get());
 
+		// TODO: Delete this
+		// END: Delete this
+
 		// return true if all the components were successfully create
 		return Input.IsValid() && Output.IsValid() && Decoder.IsValid();
 	}
diff --git a/Engine/Plugins/Runtime/Metasound/Source/MetasoundEngine/Private/MetasoundWavePlayerNode.cpp b/Engine/Plugins/Runtime/Metasound/Source/MetasoundEngine/Private/MetasoundWavePlayerNode.cpp
--- a/Engine/Plugins/Runtime/Metasound/Source/MetasoundEngine/Private/MetasoundWavePlayerNode.cpp
+++ b/Engine/Plugins/Runtime/Metasound/Source/MetasoundEngine/Private/MetasoundWavePlayerNode.cpp
@@ -52,14 +52,24 @@ namespace Metasound
 		FWavePlayerOperator(
 			const FOperatorSettings& InSettings,
 			const FWaveAssetReadRef& InWave,
-			const FTriggerReadRef& InTrigger,
-			const FFloatReadRef InPitchShiftCents)
+			const FTriggerReadRef& InPlayTrigger,
+			const FTriggerReadRef& InStopTrigger,
+			const FTriggerReadRef& InSeekTrigger,
+			const FFloatReadRef& InSeekTimeSeconds,
+			const FFloatReadRef& InPitchShiftSemiTones,
+			const FBoolReadRef& InLoop
+		)
 			: OperatorSettings(InSettings)
-			, TrigIn(InTrigger)
+			, PlayTrig(InPlayTrigger)
+			, StopTrig(InStopTrigger)
+			, SeekTrig(InSeekTrigger)
 			, Wave(InWave)
-			, PitchShiftCents(InPitchShiftCents)
+			, SeekTimeSeconds(InSeekTimeSeconds)
+			, PitchShiftSemiTones(InPitchShiftSemiTones)
+			, IsLooping(InLoop)
 			, AudioBufferL(FAudioBufferWriteRef::CreateNew(InSettings))
 			, AudioBufferR(FAudioBufferWriteRef::CreateNew(InSettings))
+			, TrigggerOnLooped(FTriggerWriteRef::CreateNew(InSettings))
 			, TrigggerOnDone(FTriggerWriteRef::CreateNew(InSettings))
 			, OutputSampleRate(InSettings.GetSampleRate())
 			, OutputBlockSizeInFrames(InSettings.GetNumFramesPerBlock())
@@ -77,8 +87,12 @@ namespace Metasound
 		{
 			FDataReferenceCollection InputDataReferences;
 			InputDataReferences.AddDataReadReference(TEXT("Audio"), FWaveAssetReadRef(Wave));
-			InputDataReferences.AddDataReadReference(TEXT("TrigIn"), FTriggerReadRef(TrigIn));
-			InputDataReferences.AddDataReadReference(TEXT("PitchShiftCents"), FFloatReadRef(PitchShiftCents));
+			InputDataReferences.AddDataReadReference(TEXT("Play"), FTriggerReadRef(PlayTrig));
+			InputDataReferences.AddDataReadReference(TEXT("Stop"), FTriggerReadRef(StopTrig));
+			InputDataReferences.AddDataReadReference(TEXT("Seek"), FTriggerReadRef(SeekTrig));
+			InputDataReferences.AddDataReadReference(TEXT("SeekTime"), FFloatReadRef(SeekTimeSeconds));
+			InputDataReferences.AddDataReadReference(TEXT("PitchShift"), FFloatReadRef(PitchShiftSemiTones));
+			InputDataReferences.AddDataReadReference(TEXT("Loop"), FBoolReadRef(IsLooping));
 			return InputDataReferences;
 		}
 
@@ -87,6 +101,7 @@ namespace Metasound
 			FDataReferenceCollection OutputDataReferences;
 			OutputDataReferences.AddDataReadReference(TEXT("AudioLeft"), FAudioBufferReadRef(AudioBufferL));
 			OutputDataReferences.AddDataReadReference(TEXT("AudioRight"), FAudioBufferReadRef(AudioBufferR));
+			OutputDataReferences.AddDataReadReference(TEXT("Looped"), FTriggerReadRef(TrigggerOnLooped));
 			OutputDataReferences.AddDataReadReference(TEXT("Done"), FTriggerReadRef(TrigggerOnDone));
 			return OutputDataReferences;
 		}
@@ -95,6 +110,7 @@ namespace Metasound
 		void Execute()
 		{
 			TrigggerOnDone->AdvanceBlock();
+			TrigggerOnLooped->AdvanceBlock();
 
 			// see if we have a new soundwave input
 			FName NewSoundWaveName = FName();
@@ -113,28 +129,123 @@ namespace Metasound
 			FMemory::Memzero(AudioBufferL->GetData(), OutputBlockSizeInFrames * sizeof(float));
 			FMemory::Memzero(AudioBufferR->GetData(), OutputBlockSizeInFrames * sizeof(float));
 
-			TrigIn->ExecuteBlock(
-				// OnPreTrigger
-				[&](int32 StartFrame, int32 EndFrame)
+			// Parse triggers and render audio
+			int32 PlayTrigIndex = 0;
+			int32 NextPlayFrame = 0;
+			const int32 NumPlayTrigs = PlayTrig->NumTriggeredInBlock();
+
+			int32 StopTrigIndex = 0;
+			int32 NextStopFrame = 0;
+			const int32 NumStopTrigs = StopTrig->NumTriggeredInBlock();
+
+			int32 SeekTrigIndex = 0;
+			int32 NextSeekFrame = 0;
+			const int32 NumSeekTrigs = SeekTrig->NumTriggeredInBlock();
+
+			int32 CurrAudioFrame = 0;
+			int32 NextAudioFrame = 0;
+
+			while (NextAudioFrame < (OutputBlockSizeInFrames -1))
+			{
+				const int32 NoTrigger = (OutputBlockSizeInFrames << 1);
+
+				// get the next Play and Seek indicies
+				// (play)
+				if (PlayTrigIndex < NumPlayTrigs)
 				{
- 					if (bIsPlaying)
- 					{
-						ExecuteInternal(StartFrame, EndFrame);
-					}
-				},
-				// OnTrigger
-				[&](int32 StartFrame, int32 EndFrame)
+					NextPlayFrame = (*PlayTrig)[PlayTrigIndex];
+				}
+				else
 				{
-					ResetDecoder();
+					NextPlayFrame = NoTrigger;
+				}
+
+				// (stop)
+				if (StopTrigIndex < NumStopTrigs)
+				{
+					NextStopFrame = (*StopTrig)[StopTrigIndex];
+				}
+				else
+				{
+					NextStopFrame = NoTrigger;
+				}
 
-					if (!bIsPlaying)
+				// (seek)
+				if (SeekTrigIndex < NumSeekTrigs)
+				{
+					NextSeekFrame = (*SeekTrig)[SeekTrigIndex];
+				}
+				else
+				{
+					NextSeekFrame = NoTrigger;
+				}
+
+				// determine the next audio frame we are going to render up to
+				NextAudioFrame = FMath::Min(NextPlayFrame, NextStopFrame);
+				NextAudioFrame = FMath::Min(NextAudioFrame, NextSeekFrame);
+
+				// no more triggers, rendering to the end of the block
+				if (NextAudioFrame == NoTrigger)
+				{
+					NextAudioFrame = OutputBlockSizeInFrames;
+				}
+
+				// render audio (while loop handles looping audio)
+				while (CurrAudioFrame != NextAudioFrame)
+				{
+					if (bIsPlaying)
+					{
+						CurrAudioFrame += ExecuteInternal(CurrAudioFrame, NextAudioFrame);
+					}
+					else
 					{
-						bIsPlaying = Decoder.CanGenerateAudio();
+						CurrAudioFrame = NextAudioFrame;
 					}
+				}
+
+				// execute the next trigger
+				if (CurrAudioFrame == NextSeekFrame)
+				{
+					ExecuteSeekRequest();
 
-					ExecuteInternal(StartFrame, EndFrame);
+					++SeekTrigIndex;
 				}
-			);
+
+				if (CurrAudioFrame == NextPlayFrame)
+				{
+					StartPlaying();
+					++PlayTrigIndex;
+				}
+
+				if (CurrAudioFrame == NextStopFrame)
+				{
+					bIsPlaying = false;
+					ResetDecoder();
+					++StopTrigIndex;
+				}
+			}
+		}
+
+		void StartPlaying()
+		{
+			ResetDecoder();
+
+			if (!bIsPlaying)
+			{
+				bIsPlaying = Decoder.CanGenerateAudio();
+			}
+		}
+
+		void ExecuteSeekRequest()
+		{
+			// TODO: get this to work w/o full decoder reset
+			// Decoder.SeekToTime(FMath::Max(0.f, *SeekTimeSeconds));
+
+			// in the mean time, using this instead:
+			if (!bIsPlaying)
+			{
+				StartPlaying();
+			}
 		}
 
 		bool ResetDecoder()
@@ -143,6 +254,7 @@ namespace Metasound
 			Params.OutputBlockSizeInFrames = OutputBlockSizeInFrames;
 			Params.OutputSampleRate = OutputSampleRate;
 			Params.MaxPitchShiftMagnitudeAllowedInOctaves = 4.f;
+			Params.StartTimeSeconds = FMath::Max(0.f, *SeekTimeSeconds);
 
 			if (false == Wave->IsSoundWaveValid())
 			{
@@ -153,7 +265,7 @@ namespace Metasound
 		}
 
 
-		void ExecuteInternal(int32 StartFrame, int32 EndFrame)
+		int32 ExecuteInternal(int32 StartFrame, int32 EndFrame)
 		{
 			bool bCanDecodeWave = Wave->IsSoundWaveValid() && ((*Wave)->GetNumChannels() <= 2); // only support mono or stereo inputs
 
@@ -161,12 +273,12 @@ namespace Metasound
 			float* FinalOutputLeft = AudioBufferL->GetData() + StartFrame;
 			float* FinalOutputRight = AudioBufferR->GetData() + StartFrame;
 			const int32 NumOutputFrames = (EndFrame - StartFrame);
+			int32 NumFramesDecoded = 0;
 
 			if (bCanDecodeWave)
 			{
 				ensure(Decoder.CanGenerateAudio());
 
-
 				const int32 NumInputChannels = (*Wave)->GetNumChannels();
 				const bool bNeedsUpmix = (NumInputChannels == 1);
 				const bool bNeedsDeinterleave = !bNeedsUpmix;
@@ -177,14 +289,20 @@ namespace Metasound
 				PostSrcBuffer.AddZeroed(NumSamplesToGenerate);
 				float* PostSrcBufferPtr = PostSrcBuffer.GetData();
 
-				int32 NumFramesDecoded = Decoder.GenerateAudio(PostSrcBufferPtr, NumOutputFrames, *PitchShiftCents) / NumInputChannels;
+				bool bIsLooping = *IsLooping;
+				NumFramesDecoded = Decoder.GenerateAudio(PostSrcBufferPtr, NumOutputFrames, (*PitchShiftSemiTones * 100.f), bIsLooping) / NumInputChannels;
 
 				// handle decoder having completed during it's decode
-				if (!Decoder.CanGenerateAudio() ||  (NumFramesDecoded < NumOutputFrames))
+				const bool bReachedEOF = !Decoder.CanGenerateAudio() || (NumFramesDecoded < NumOutputFrames);
+				if (bReachedEOF && !bIsLooping)
 				{
 					bIsPlaying = false;
 					TrigggerOnDone->TriggerFrame(StartFrame + NumFramesDecoded);
 				}
+				else if (bReachedEOF && bIsLooping)
+				{
+					TrigggerOnLooped->TriggerFrame(StartFrame + NumFramesDecoded);
+				}
 
 				if (bNeedsUpmix)
 				{
@@ -208,18 +326,25 @@ namespace Metasound
 				FMemory::Memzero(FinalOutputLeft, sizeof(float) * NumOutputFrames);
 				FMemory::Memzero(FinalOutputRight, sizeof(float) * NumOutputFrames);
 			}
+
+			return NumFramesDecoded;
 		}
 
 	private:
 		const FOperatorSettings OperatorSettings;
 
 		// i/o
-		FTriggerReadRef TrigIn;
+		FTriggerReadRef PlayTrig;
+		FTriggerReadRef StopTrig;
+		FTriggerReadRef SeekTrig;
 		FWaveAssetReadRef Wave;
-		FFloatReadRef PitchShiftCents;
+		FFloatReadRef SeekTimeSeconds;
+		FFloatReadRef PitchShiftSemiTones;
+		FBoolReadRef IsLooping;
 
 		FAudioBufferWriteRef AudioBufferL;
 		FAudioBufferWriteRef AudioBufferR;
+		FTriggerWriteRef TrigggerOnLooped;
 		FTriggerWriteRef TrigggerOnDone;
 
 		// source decode
@@ -244,10 +369,14 @@ namespace Metasound
 
 		const FDataReferenceCollection& InputDataRefs = InParams.InputDataReferences;
 
-		FTriggerReadRef TriggerPlay = InputDataRefs.GetDataReadReferenceOrConstruct<FTrigger>(TEXT("TrigIn"), InParams.OperatorSettings);
+		FTriggerReadRef TriggerPlay = InputDataRefs.GetDataReadReferenceOrConstruct<FTrigger>(TEXT("Play"), InParams.OperatorSettings);
+		FTriggerReadRef TriggerStop = InputDataRefs.GetDataReadReferenceOrConstruct<FTrigger>(TEXT("Stop"), InParams.OperatorSettings);
+		FTriggerReadRef TriggerSeek = InputDataRefs.GetDataReadReferenceOrConstruct<FTrigger>(TEXT("Seek"), InParams.OperatorSettings);
 		FWaveAssetReadRef Wave = InputDataRefs.GetDataReadReferenceOrConstruct<FWaveAsset>(TEXT("Wave"));
-		FFloatReadRef PitchShiftCents = InputDataRefs.GetDataReadReferenceOrConstruct<float>(TEXT("PitchShiftCents"));
-		
+		FFloatReadRef SeekTimeSeconds = InputDataRefs.GetDataReadReferenceOrConstruct<float>(TEXT("SeekTime"), 0.f);
+		FFloatReadRef PitchShiftSemiTones = InputDataRefs.GetDataReadReferenceOrConstruct<float>(TEXT("PitchShift"));
+		FBoolReadRef IsLooping = InputDataRefs.GetDataReadReferenceOrConstruct<bool>(TEXT("Loop"), false);
+
 		if (!Wave->IsSoundWaveValid())
 		{
 			AddBuildError<FWavePlayerError>(OutErrors, WaveNode, LOCTEXT("NoSoundWave", "No Sound Wave"));
@@ -255,14 +384,18 @@ namespace Metasound
 		}
 		else if ((*Wave)->GetNumChannels() > 2)
 		{
-			AddBuildError<FWavePlayerError>(OutErrors, WaveNode, LOCTEXT("WavePlayerCurrentlyOnlySuportsMonoAssets", "Wave Player Currently Only Supports Mono Or Stereo Assets"));
+			AddBuildError<FWavePlayerError>(OutErrors, WaveNode, LOCTEXT("WavePlayerCurrentlyOnlySuportsMonoOrStereoAssets", "Wave Player Currently Only Supports Mono Or Stereo Assets"));
 		}
 
 		return MakeUnique<FWavePlayerOperator>(
 			  InParams.OperatorSettings
 			, Wave
 			, TriggerPlay
-			, PitchShiftCents
+			, TriggerStop
+			, TriggerSeek
+			, SeekTimeSeconds
+			, PitchShiftSemiTones
+			, IsLooping
 			);
 	}
 
@@ -271,12 +404,17 @@ namespace Metasound
 		return FVertexInterface(
 			FInputVertexInterface(
 				TInputDataVertexModel<FWaveAsset>(TEXT("Wave"), LOCTEXT("WaveTooltip", "The Wave to be decoded")),
-				TInputDataVertexModel<FTrigger>(TEXT("TrigIn"), LOCTEXT("TrigInTooltip", "Trigger the playing of the input wave.")),
-				TInputDataVertexModel<float>(TEXT("PitchShiftCents"), LOCTEXT("PitchShiftCentsTooltip", "Pitch Shift in cents."))
+				TInputDataVertexModel<FTrigger>(TEXT("Play"), LOCTEXT("PlayTooltip", "Play the input wave.")),
+				TInputDataVertexModel<FTrigger>(TEXT("Stop"), LOCTEXT("StopTooltip", "Stop the input wave.")),
+				TInputDataVertexModel<FTrigger>(TEXT("Seek"), LOCTEXT("PlayTooltip", "Trigger the playing of the input wave.")),
+				TInputDataVertexModel<float>(TEXT("SeekTime"), LOCTEXT("SeekTimeTooltip", "Seek time in seconds.")),
+				TInputDataVertexModel<float>(TEXT("PitchShift"), LOCTEXT("PitchShiftTooltip", "Pitch Shift in semi-tones.")),
+				TInputDataVertexModel<bool>(TEXT("Loop"), LOCTEXT("LoopTooltip", "Wave will loop if true"))
 			),
 			FOutputVertexInterface(
 				TOutputDataVertexModel<FAudioBuffer>(TEXT("AudioLeft"), LOCTEXT("AudioTooltip", "The output audio")),
 				TOutputDataVertexModel<FAudioBuffer>(TEXT("AudioRight"), LOCTEXT("AudioTooltip", "The output audio")),
+				TOutputDataVertexModel<FTrigger>(TEXT("Looped"), LOCTEXT("TriggerToolTip", "Trigger that notifies when the sound is has looped")),
 				TOutputDataVertexModel<FTrigger>(TEXT("Done"), LOCTEXT("TriggerToolTip", "Trigger that notifies when the sound is done playing"))
 			)
 		);
@@ -321,8 +459,6 @@ namespace Metasound
 		return Factory;
 	}
 
-
-
 	const FVertexInterface& FWavePlayerNode::GetVertexInterface() const
 	{
 		return Interface;
diff --git a/Engine/Plugins/Runtime/Metasound/Source/MetasoundEngine/Public/MetasoundWave.h b/Engine/Plugins/Runtime/Metasound/Source/MetasoundEngine/Public/MetasoundWave.h
--- a/Engine/Plugins/Runtime/Metasound/Source/MetasoundEngine/Public/MetasoundWave.h
+++ b/Engine/Plugins/Runtime/Metasound/Source/MetasoundEngine/Public/MetasoundWave.h
@@ -81,6 +81,7 @@ namespace Audio
 			float OutputSampleRate;
 			uint32 OutputBlockSizeInFrames;
 			float MaxPitchShiftMagnitudeAllowedInOctaves = 4.f;
+			float StartTimeSeconds = 0.f;
 		};
 
 		bool CanGenerateAudio() const
@@ -91,7 +92,9 @@ namespace Audio
 		bool Initialize(const InitParams& InInitParams, const FSoundWaveProxyPtr& InWave);
 
 		// returns number of samples written.   
-		uint32 GenerateAudio(float* OutputDest, int32 NumOutputFrames, float PitchShiftInCents = 0.f);
+		uint32 GenerateAudio(float* OutputDest, int32 NumOutputFrames, float PitchShiftInCents = 0.f, bool bIsLooping = false);
+
+		void SeekToTime(const float InSeconds);
 
 	private:
 		float GetSampleRateRatio(float PitchShiftInCents) const;
@@ -124,11 +127,13 @@ namespace Audio
 		uint32 DecodeBlockSizeInSamples;
 		
 		// cached values
+		float StartTimeSeconds{ 0.f };
 		float LastPitchShiftCents{ 0.f };
 		int32 TotalNumFramesOutput{ 0 };
 		int32 TotalNumFramesDecoded{ 0 };
 
 		bool bDecoderIsDone{ true };
+		bool bDecoderHasLooped{ false };
 
 	}; // class FSimpleDecoderWrapper
 
diff --git a/Engine/Source/Runtime/AudioCodecEngine/Private/DecoderBackCompat.cpp b/Engine/Source/Runtime/AudioCodecEngine/Private/DecoderBackCompat.cpp
--- a/Engine/Source/Runtime/AudioCodecEngine/Private/DecoderBackCompat.cpp
+++ b/Engine/Source/Runtime/AudioCodecEngine/Private/DecoderBackCompat.cpp
@@ -49,7 +49,7 @@ namespace Audio
 		ResidualBuffer.SetNum(Reqs.NumSampleFramesWanted * Desc.NumChannels);
 	}
 
-	Audio::IDecoder::FDecodeReturn FBackCompat::Decode()
+	Audio::IDecoder::FDecodeReturn FBackCompat::Decode(bool bLoop)
 	{
 		uint32 NumFramesRemaining = Reqs.NumSampleFramesWanted;
 		FBackCompatInput& BackCompatSrc = static_cast<FBackCompatInput&>(*Src);
@@ -71,11 +71,11 @@ namespace Audio
 			int32 NumBytesStreamed = BuffSizeInBytes;
 			if (BackCompatSrc.Wave->IsStreaming())
 			{
-				bFinished = Info->StreamCompressedData(Buff, false, BuffSizeInBytes, NumBytesStreamed);
+				bFinished = Info->StreamCompressedData(Buff, bLoop, BuffSizeInBytes, NumBytesStreamed);
 			}
 			else
 			{
-				bFinished = Info->ReadCompressedData(Buff, false, BuffSizeInBytes);
+				bFinished = Info->ReadCompressedData(Buff, bLoop, BuffSizeInBytes);
 			}						
 			PushedDetails.SampleFramesStartOffset = FrameOffset;
 			ResidualBuffer.SetNum(NumBytesStreamed / sizeof(int16));
@@ -89,11 +89,18 @@ namespace Audio
 			NumFramesRemaining -= BuffSizeInFrames;
 		}
 		
-		if( bFinished )
+		if(!bFinished)
+		{
+			return EDecodeResult::MoreDataRemaining;
+		}
+		else if (bLoop) // (bFinished is true)
+		{
+			return EDecodeResult::Looped;
+		}
+		else // (bFinished is true and we aren't looping)
 		{
 			return EDecodeResult::Finished;
 		}
-		return EDecodeResult::MoreDataRemaining;
 	}
 
 }
\ No newline at end of file
diff --git a/Engine/Source/Runtime/AudioCodecEngine/Private/DecoderBackCompat.h b/Engine/Source/Runtime/AudioCodecEngine/Private/DecoderBackCompat.h
--- a/Engine/Source/Runtime/AudioCodecEngine/Private/DecoderBackCompat.h
+++ b/Engine/Source/Runtime/AudioCodecEngine/Private/DecoderBackCompat.h
@@ -40,6 +40,6 @@ namespace Audio
 			return false;
 		}
 
-		FDecodeReturn Decode() override;
+		FDecodeReturn Decode(bool bLoop = false) override;
 	};
 }
diff --git a/Engine/Source/Runtime/AudioCodecEngine/Private/DecoderInputBackCompat.cpp b/Engine/Source/Runtime/AudioCodecEngine/Private/DecoderInputBackCompat.cpp
--- a/Engine/Source/Runtime/AudioCodecEngine/Private/DecoderInputBackCompat.cpp
+++ b/Engine/Source/Runtime/AudioCodecEngine/Private/DecoderInputBackCompat.cpp
@@ -82,6 +82,14 @@ namespace Audio
 		return -1;
 	}
 
+	void FBackCompatInput::SeekToTime(const float InSeconds)
+	{
+		ICompressedAudioInfo* Info = GetInfo();
+		ensure(Info);
+
+		Info->SeekToTime(InSeconds);
+	}
+
 	bool FBackCompatInput::IsEndOfStream() const
 	{
 		// Not implemented.
diff --git a/Engine/Source/Runtime/AudioCodecEngine/Private/DecoderInputBackCompat.h b/Engine/Source/Runtime/AudioCodecEngine/Private/DecoderInputBackCompat.h
--- a/Engine/Source/Runtime/AudioCodecEngine/Private/DecoderInputBackCompat.h
+++ b/Engine/Source/Runtime/AudioCodecEngine/Private/DecoderInputBackCompat.h
@@ -33,6 +33,8 @@ namespace Audio
 
 		bool FindSection(FEncodedSectionBase& OutSection) override;
 		int64 Tell() const override;
+
+		void SeekToTime(const float InSeconds) override;
 		
 		TArrayView<const uint8> PeekNextPacket(
 			int32 InMaxPacketLength) const override;
diff --git a/Engine/Source/Runtime/AudioExtensions/Private/PcmCodec.cpp b/Engine/Source/Runtime/AudioExtensions/Private/PcmCodec.cpp
--- a/Engine/Source/Runtime/AudioExtensions/Private/PcmCodec.cpp
+++ b/Engine/Source/Runtime/AudioExtensions/Private/PcmCodec.cpp
@@ -170,7 +170,7 @@ namespace Audio
 				return bHasError;
 			}	
 
-			FDecodeReturn Decode() override
+			FDecodeReturn Decode(bool bIsLooping = true) override
 			{
 				// Reject bad input
 				if( !audio_ensure(Src) || !audio_ensure(Dst) || !audio_ensure(Header.NumChannels > 0))
diff --git a/Engine/Source/Runtime/AudioExtensions/Public/IAudioCodec.h b/Engine/Source/Runtime/AudioExtensions/Public/IAudioCodec.h
--- a/Engine/Source/Runtime/AudioExtensions/Public/IAudioCodec.h
+++ b/Engine/Source/Runtime/AudioExtensions/Public/IAudioCodec.h
@@ -116,6 +116,8 @@ namespace Audio
 
 		virtual int64 Tell() const = 0;
 
+		virtual void SeekToTime(const float InSeconds) { }
+
 		virtual TArrayView<const uint8> PeekNextPacket(
 			int32 InMaxPacketLength) const = 0;
 				
@@ -254,7 +256,7 @@ namespace Audio
 		virtual bool HasError() const = 0;
 
 		// Do the decode.
-		virtual FDecodeReturn Decode() = 0;	
+		virtual FDecodeReturn Decode(bool bIsLooping = false) = 0;	
 	};
 
 	struct AUDIOEXTENSIONS_API FEncodedSectionBase
