A Debug Log Metasound, for logging values on Trigger from within the MetaSound graph

diff --git a/Engine/Plugins/Runtime/Metasound/Source/MetasoundStandardNodes/Private/MetasoundDebugLogNode.cpp b/Engine/Plugins/Runtime/Metasound/Source/MetasoundStandardNodes/Private/MetasoundDebugLogNode.cpp
new file mode 100644
--- /dev/null
+++ b/Engine/Plugins/Runtime/Metasound/Source/MetasoundStandardNodes/Private/MetasoundDebugLogNode.cpp
@@ -0,0 +1,72 @@
+#include "MetasoundDebugLogNode.h"
+
+#include "MetasoundNodeRegistrationMacro.h"
+
+#define LOCTEXT_NAMESPACE "MetasoundStandardNodes"
+
+namespace Metasound
+{
+	namespace MetasoundDebugLogNodePrivate
+	{
+		FNodeClassMetadata CreateNodeClassMetadata(const FName& InDataTypeName, const FName& InOperatorName, const FText& InDisplayName, const FText& InDescription, const FVertexInterface& InDefaultInterface)
+		{
+			FNodeClassMetadata Metadata
+			{
+				FNodeClassName{FName("DebugLog"), InOperatorName, InDataTypeName},
+				1, // Major Version
+				0, // Minor Version
+				InDisplayName,
+				InDescription,
+				PluginAuthor,
+				PluginNodeMissingPrompt,
+				InDefaultInterface,
+				{LOCTEXT("DebugLogCategory", "Debug")},
+				{TEXT("DebugLog")},
+				FNodeDisplayStyle{}
+			};
+
+			return Metadata;
+		}
+	}
+
+	namespace DebugLogVertexNames
+	{
+		const FString& GetInputTriggerName()
+		{
+			static const FString Name = TEXT("Trigger");
+			return Name;
+		}
+
+		const FString& GetLabelDebugLogName()
+		{
+			static const FString Name = TEXT("Label");
+			return Name;
+		}
+
+		const FString& GetToLogDebugLogName()
+		{
+			static const FString Name = TEXT("Value To Log");
+			return Name;
+		}
+
+		const FString& GetOutputDebugLogName()
+		{
+			static const FString Name = TEXT("Was Successful");
+			return Name;
+		}
+	}
+
+	using FDebugLogNodeInt32 = TDebugLogNode<int32>;
+	METASOUND_REGISTER_NODE(FDebugLogNodeInt32)
+
+	using FDebugLogNodeFloat = TDebugLogNode<float>;
+	METASOUND_REGISTER_NODE(FDebugLogNodeFloat)
+
+	using FDebugLogNodeBool = TDebugLogNode<bool>;
+	METASOUND_REGISTER_NODE(FDebugLogNodeBool)
+
+	using FDebugLogNodeString = TDebugLogNode<FString>;
+	METASOUND_REGISTER_NODE(FDebugLogNodeString)
+}
+
+#undef LOCTEXT_NAMESPACE
\ No newline at end of file
diff --git a/Engine/Plugins/Runtime/Metasound/Source/MetasoundStandardNodes/Public/MetasoundDebugLogNode.h b/Engine/Plugins/Runtime/Metasound/Source/MetasoundStandardNodes/Public/MetasoundDebugLogNode.h
new file mode 100644
--- /dev/null
+++ b/Engine/Plugins/Runtime/Metasound/Source/MetasoundStandardNodes/Public/MetasoundDebugLogNode.h
@@ -0,0 +1,144 @@
+// Copyright Epic Games, Inc. All Rights Reserved.
+
+#pragma once
+
+#include "CoreMinimal.h"
+#include "MetasoundFacade.h"
+#include "MetasoundExecutableOperator.h"
+#include "MetasoundPrimitives.h"
+#include "MetasoundStandardNodesNames.h"
+#include "MetasoundTrigger.h"
+#include "Internationalization/Text.h"
+#include "MetasoundLog.h"
+
+#define LOCTEXT_NAMESPACE "MetasoundStandardNodes"
+
+namespace Metasound
+{	
+	namespace MetasoundDebugLogNodePrivate
+	{
+		METASOUNDSTANDARDNODES_API FNodeClassMetadata CreateNodeClassMetadata(const FName& InDataTypeName, const FName& InOperatorName, const FText& InDisplayName, const FText& InDescription, const FVertexInterface& InDefaultInterface);
+	}
+
+	namespace DebugLogVertexNames
+	{
+		METASOUNDSTANDARDNODES_API const FString& GetInputTriggerName();
+		METASOUNDSTANDARDNODES_API const FString& GetLabelDebugLogName();
+		METASOUNDSTANDARDNODES_API const FString& GetToLogDebugLogName();
+		METASOUNDSTANDARDNODES_API const FString& GetOutputDebugLogName();
+	}
+
+	template<typename DebugLogType>
+	class TDebugLogOperator : public TExecutableOperator<TDebugLogOperator<DebugLogType>>
+	{
+		public:
+			using FArrayDataReadReference = TDataReadReference<DebugLogType>;
+
+			
+			static const FVertexInterface& GetDefaultInterface()
+			{
+				static const FVertexInterface DefaultInterface(
+					FInputVertexInterface(
+
+						TInputDataVertexModel<FTrigger>(DebugLogVertexNames::GetInputTriggerName(), LOCTEXT("DebugLogTrigger", "Trigger to write the set value to the log.")),
+						TInputDataVertexModel<FString>(DebugLogVertexNames::GetLabelDebugLogName(), LOCTEXT("DebugLogLabel", "The label to attach to the value that will be logged")),
+						TInputDataVertexModel<DebugLogType>(DebugLogVertexNames::GetToLogDebugLogName(), LOCTEXT("DebugLogValueToLog", "The value to record to the log when triggered"))
+					),
+					FOutputVertexInterface(
+					)
+				);
+
+				return DefaultInterface;
+			}
+
+			static const FNodeClassMetadata& GetNodeInfo()
+			{
+				auto CreateNodeClassMetadata = []() -> FNodeClassMetadata
+				{
+					FName DataTypeName = GetMetasoundDataTypeName<DebugLogType>();
+					FName OperatorName = TEXT("DebugLog");
+					FText NodeDisplayName = FText::Format(LOCTEXT("DebugLogDisplayNamePattern", "DebugLog ({0})"), FText::FromString(GetMetasoundDataTypeString<DebugLogType>()));
+					FText NodeDescription = LOCTEXT("DebugLogOpDescription", "Used to record values to the log, on trigger");
+					FVertexInterface NodeInterface = GetDefaultInterface();
+
+					return MetasoundDebugLogNodePrivate::CreateNodeClassMetadata(DataTypeName, OperatorName, NodeDisplayName, NodeDescription, NodeInterface);
+				};
+
+				static const FNodeClassMetadata Metadata = CreateNodeClassMetadata();
+				return Metadata;
+			}
+
+			static TUniquePtr<IOperator> CreateOperator(const FCreateOperatorParams& InParams, TArray<TUniquePtr<IOperatorBuildError>>& OutErrors)
+			{
+				const FInputVertexInterface& InputInterface = InParams.Node.GetVertexInterface().GetInputInterface();
+				const FDataReferenceCollection& InputCollection = InParams.InputDataReferences;
+
+				FTriggerReadRef Trigger = InputCollection.GetDataReadReferenceOrConstruct<FTrigger>(DebugLogVertexNames::GetInputTriggerName(), InParams.OperatorSettings);
+				TDataReadReference<FString> Label = InputCollection.GetDataReadReferenceOrConstructWithVertexDefault<FString>(InputInterface, DebugLogVertexNames::GetLabelDebugLogName());
+				TDataReadReference<DebugLogType> ValueToLog = InputCollection.GetDataReadReferenceOrConstructWithVertexDefault<DebugLogType>(InputInterface, DebugLogVertexNames::GetToLogDebugLogName());
+
+				return MakeUnique<TDebugLogOperator<DebugLogType>>(InParams.OperatorSettings, Trigger, Label, ValueToLog);
+			}
+
+
+			TDebugLogOperator(const FOperatorSettings& InSettings, TDataReadReference<FTrigger> InTrigger, TDataReadReference<FString> InLabelDebugLog, TDataReadReference<DebugLogType> InValueToLogDebugLog)
+				: Trigger(InTrigger)
+				, Label(InLabelDebugLog)
+				, ValueToLog(InValueToLogDebugLog)
+			{
+			}
+
+			virtual ~TDebugLogOperator() = default;
+
+
+			virtual FDataReferenceCollection GetInputs() const override
+			{
+				FDataReferenceCollection Inputs;
+
+				Inputs.AddDataReadReference(DebugLogVertexNames::GetInputTriggerName(), Trigger);
+				Inputs.AddDataReadReference(DebugLogVertexNames::GetLabelDebugLogName(), Label);
+				Inputs.AddDataReadReference(DebugLogVertexNames::GetToLogDebugLogName(), ValueToLog);
+
+				return Inputs;
+			}
+
+			virtual FDataReferenceCollection GetOutputs() const override
+			{
+				FDataReferenceCollection Outputs;
+
+				return Outputs;
+			}
+
+			void Execute()
+			{
+				if (*Trigger)
+				{
+					UE_LOG(LogMetasound, Display, TEXT("[%s] was [%s]"), *(*Label), *LexToString(*ValueToLog));
+				}
+			}
+
+		private:
+
+			TDataReadReference<FTrigger> Trigger;
+			TDataReadReference<FString> Label;
+			TDataReadReference<DebugLogType> ValueToLog;
+	};
+
+	/** TDebugLogNode
+	 *
+	 *  Records a value to the log when triggered
+	 */
+	template<typename DebugLogType>
+	class METASOUNDSTANDARDNODES_API TDebugLogNode : public FNodeFacade
+	{
+		public:
+			/**
+			 * Constructor used by the Metasound Frontend.
+			 */
+			TDebugLogNode(const FNodeInitData& InInitData)
+				: FNodeFacade(InInitData.InstanceName, InInitData.InstanceID, TFacadeOperatorClass<TDebugLogOperator<DebugLogType>>())
+			{}
+
+			virtual ~TDebugLogNode() = default;
+	};
+} // namespace Metasound
\ No newline at end of file
