Fix for trigger accumulator logic which makes auto-reset not function

Removed 1-element for TriggerAny

diff --git a/Engine/Plugins/Runtime/Metasound/Source/MetasoundStandardNodes/Private/MetasoundTriggerAnyNode.cpp b/Engine/Plugins/Runtime/Metasound/Source/MetasoundStandardNodes/Private/MetasoundTriggerAnyNode.cpp
--- a/Engine/Plugins/Runtime/Metasound/Source/MetasoundStandardNodes/Private/MetasoundTriggerAnyNode.cpp
+++ b/Engine/Plugins/Runtime/Metasound/Source/MetasoundStandardNodes/Private/MetasoundTriggerAnyNode.cpp
@@ -63,7 +63,6 @@ namespace Metasound
 		}
 	}
 
-	REGISTER_TRIGGER_ANY_NODE(1)
 	REGISTER_TRIGGER_ANY_NODE(2)
 	REGISTER_TRIGGER_ANY_NODE(3)
 	REGISTER_TRIGGER_ANY_NODE(4)
diff --git a/Engine/Plugins/Runtime/Metasound/Source/MetasoundStandardNodes/Public/MetasoundTriggerAccumulatorNode.h b/Engine/Plugins/Runtime/Metasound/Source/MetasoundStandardNodes/Public/MetasoundTriggerAccumulatorNode.h
--- a/Engine/Plugins/Runtime/Metasound/Source/MetasoundStandardNodes/Public/MetasoundTriggerAccumulatorNode.h
+++ b/Engine/Plugins/Runtime/Metasound/Source/MetasoundStandardNodes/Public/MetasoundTriggerAccumulatorNode.h
@@ -145,11 +145,16 @@ namespace Metasound
 				ResetTriggerState();
 			}
 
-			OutputTrigger->TriggerFrame(InStartFrame);
+			if (!bOutputTriggered)
+			{
+				OutputTrigger->TriggerFrame(InStartFrame);
+				bOutputTriggered = true;
+			}
 		}
 
 		void ResetTriggerState()
 		{
+			bOutputTriggered = false;
 			bInputWasTriggered.Reset();
 			bInputWasTriggered.AddDefaulted(InputTriggers.Num());
 		}
@@ -158,6 +163,11 @@ namespace Metasound
 		{
 			OutputTrigger->AdvanceBlock();
 
+			if (bOutputTriggered && !*bAutoReset)
+			{
+				return;
+			}
+
 			for (uint32 i = 0; i < NumInputs; ++i)
 			{
 				InputTriggers[i]->ExecuteBlock(
@@ -180,6 +190,7 @@ namespace Metasound
 		FTriggerWriteRef OutputTrigger;
 
 		TArray<bool> bInputWasTriggered;
+		bool bOutputTriggered = false;
 	};
 
 	/** TTriggerAccumulatorNode
