Variety of small tasks

Renaming Period Trigger to Trigger Repeat.
Adding a default init value to Trigger Route Node
Change MIDIToFreq node to "MIDI To Frequency" and changing input pin from int to float

diff --git a/Engine/Plugins/Runtime/Metasound/Source/MetasoundStandardNodes/Private/MetasoundMIDIToFreqNode.cpp b/Engine/Plugins/Runtime/Metasound/Source/MetasoundStandardNodes/Private/MetasoundMIDIToFreqNode.cpp
--- a/Engine/Plugins/Runtime/Metasound/Source/MetasoundStandardNodes/Private/MetasoundMIDIToFreqNode.cpp
+++ b/Engine/Plugins/Runtime/Metasound/Source/MetasoundStandardNodes/Private/MetasoundMIDIToFreqNode.cpp
@@ -13,9 +13,32 @@
 
 namespace Metasound
 {
-	// Param list
-	static const TCHAR* InParamNameMIDI = TEXT("In MIDI");
-	static const TCHAR* OutParamNameFreq = TEXT("Out Freq");
+	namespace MIDIToFrequencyVertexNames
+	{
+		const FString& GetInputMIDIName()
+		{
+			static FString Name = TEXT("MIDI");
+			return Name;
+		}
+
+		const FText& GetInputMIDIDescription()
+		{
+			static FText Desc = LOCTEXT("MIDIToFrequencyInputMIDIName", "A value representing a MIDI note value.");
+			return Desc;
+		}
+
+		const FString& GetOutputFrequencyName()
+		{
+			static FString Name = TEXT("Frequency");
+			return Name;
+		}
+
+		const FText& GetOutputFrequencyDescription()
+		{
+			static FText Desc = LOCTEXT("MIDITOFrequencyNodeOutputFrequencyName", "Output frequency value in hertz that corresponds to the input MIDI note value.");
+			return Desc;
+		}
+	}
 
 	class FMIDIToFreqOperator : public TExecutableOperator<FMIDIToFreqOperator>
 	{
@@ -25,7 +48,7 @@ namespace Metasound
 		static const FVertexInterface& GetVertexInterface();
 		static TUniquePtr<IOperator> CreateOperator(const FCreateOperatorParams& InParams, FBuildErrorArray& OutErrors);
 
-		FMIDIToFreqOperator(const FOperatorSettings& InSettings, const FInt32ReadRef& InMIDINote);
+		FMIDIToFreqOperator(const FOperatorSettings& InSettings, const FFloatReadRef& InMIDINote);
 
 		virtual FDataReferenceCollection GetInputs() const override;
 		virtual FDataReferenceCollection GetOutputs() const override;
@@ -33,7 +56,7 @@ namespace Metasound
 
 	private:
 		// The input midi value
-		FInt32ReadRef MIDINote;
+		FFloatReadRef MIDINote;
 
 		// The output frequency
 		FFloatWriteRef FreqOutput;
@@ -42,7 +65,7 @@ namespace Metasound
 		int32 PrevMidiNote;
 	};
 
-	FMIDIToFreqOperator::FMIDIToFreqOperator(const FOperatorSettings& InSettings, const FInt32ReadRef& InMIDINote)
+	FMIDIToFreqOperator::FMIDIToFreqOperator(const FOperatorSettings& InSettings, const FFloatReadRef& InMIDINote)
 		: MIDINote(InMIDINote)
 		, FreqOutput(FFloatWriteRef::CreateNew(Audio::GetFrequencyFromMidi(*InMIDINote)))
 		, PrevMidiNote(*InMIDINote)
@@ -51,16 +74,20 @@ namespace Metasound
 
 	FDataReferenceCollection FMIDIToFreqOperator::GetInputs() const
 	{
+		using namespace MIDIToFrequencyVertexNames;
+
 		FDataReferenceCollection InputDataReferences;
-		InputDataReferences.AddDataReadReference(InParamNameMIDI, FInt32ReadRef(MIDINote));
+		InputDataReferences.AddDataReadReference(GetInputMIDIName(), MIDINote);
 
 		return InputDataReferences;
 	}
 
 	FDataReferenceCollection FMIDIToFreqOperator::GetOutputs() const
 	{
+		using namespace MIDIToFrequencyVertexNames;
+
 		FDataReferenceCollection OutputDataReferences;
-		OutputDataReferences.AddDataReadReference(OutParamNameFreq, FFloatReadRef(FreqOutput));
+		OutputDataReferences.AddDataReadReference(GetOutputFrequencyName(), FreqOutput);
 		return OutputDataReferences;
 	}
 
@@ -71,18 +98,20 @@ namespace Metasound
 		{
 			PrevMidiNote = *MIDINote;
 
-			*FreqOutput = Audio::GetFrequencyFromMidi(FMath::Clamp(*MIDINote, 0, 127));
+			*FreqOutput = Audio::GetFrequencyFromMidi(FMath::Clamp(*MIDINote, 0.0f, 127.0f));
 		}
 	}
 
 	const FVertexInterface& FMIDIToFreqOperator::GetVertexInterface()
 	{
+		using namespace MIDIToFrequencyVertexNames;
+
 		static const FVertexInterface Interface(
 			FInputVertexInterface(
-				TInputDataVertexModel<int32>(InParamNameMIDI, LOCTEXT("MIDIInToolTip", "Input MIDI note value."), 60)
+				TInputDataVertexModel<float>(GetInputMIDIName(), GetInputMIDIDescription(), 60)
 			),
 			FOutputVertexInterface(
-				TOutputDataVertexModel<float>(OutParamNameFreq, LOCTEXT("OutFreqTooltip", "The output frequency (hz) value corresponding to the MIDI note."))
+				TOutputDataVertexModel<float>(GetOutputFrequencyName(), GetOutputFrequencyDescription())
 			)
 		);
 
@@ -94,10 +123,10 @@ namespace Metasound
 		auto InitNodeInfo = []() -> FNodeClassMetadata
 		{
 			FNodeClassMetadata Info;
-			Info.ClassName = { Metasound::StandardNodes::Namespace, TEXT("MIDIToFreq"), Metasound::StandardNodes::AudioVariant };
+			Info.ClassName = { Metasound::StandardNodes::Namespace, TEXT("MIDI To Frequency"), TEXT("") };
 			Info.MajorVersion = 1;
 			Info.MinorVersion = 0;
-			Info.DisplayName = LOCTEXT("Metasound_MidiToFreqDisplayName", "MIDIToFreq");
+			Info.DisplayName = LOCTEXT("Metasound_MidiToFreqDisplayName", "MIDI To Frequency");
 			Info.Description = LOCTEXT("Metasound_MidiToFreqNodeDescription", "Converts a MIDI note value to a frequency (hz) value.");
 			Info.Author = PluginAuthor;
 			Info.PromptIfMissing = PluginNodeMissingPrompt;
@@ -113,11 +142,13 @@ namespace Metasound
 
 	TUniquePtr<IOperator> FMIDIToFreqOperator::CreateOperator(const FCreateOperatorParams& InParams, FBuildErrorArray& OutErrors)
 	{
+		using namespace MIDIToFrequencyVertexNames;
+
 		const FMIDIToFreqNode& MIDIToFreqNode = static_cast<const FMIDIToFreqNode&>(InParams.Node);
 		const FDataReferenceCollection& InputCollection = InParams.InputDataReferences;
 		const FInputVertexInterface& InputInterface = GetVertexInterface().GetInputInterface();
 
-		FInt32ReadRef InMIDINote = InputCollection.GetDataReadReferenceOrConstructWithVertexDefault<int32>(InputInterface, InParamNameMIDI);
+		FFloatReadRef InMIDINote = InputCollection.GetDataReadReferenceOrConstructWithVertexDefault<float>(InputInterface, GetInputMIDIName());
 
 		return MakeUnique<FMIDIToFreqOperator>(InParams.OperatorSettings, InMIDINote);
 	}
diff --git a/Engine/Plugins/Runtime/Metasound/Source/MetasoundStandardNodes/Private/MetasoundPeriodicTriggerNode.cpp b/Engine/Plugins/Runtime/Metasound/Source/MetasoundStandardNodes/Private/MetasoundPeriodicTriggerNode.cpp
deleted file mode 100644
diff --git a/Engine/Plugins/Runtime/Metasound/Source/MetasoundStandardNodes/Private/MetasoundTriggerRepeatNode.cpp b/Engine/Plugins/Runtime/Metasound/Source/MetasoundStandardNodes/Private/MetasoundTriggerRepeatNode.cpp
new file mode 100644
--- /dev/null
+++ b/Engine/Plugins/Runtime/Metasound/Source/MetasoundStandardNodes/Private/MetasoundTriggerRepeatNode.cpp
@@ -0,0 +1,237 @@
+// Copyright Epic Games, Inc. All Rights Reserved.
+
+#include "MetasoundTriggerRepeatNode.h"
+
+#include "CoreMinimal.h"
+
+#include "MetasoundBuilderInterface.h"
+#include "MetasoundDataReferenceCollection.h"
+#include "MetasoundExecutableOperator.h"
+#include "MetasoundFacade.h"
+#include "MetasoundNode.h"
+#include "MetasoundNodeInterface.h"
+#include "MetasoundNodeRegistrationMacro.h"
+#include "MetasoundOperatorInterface.h"
+#include "MetasoundPrimitives.h"
+#include "MetasoundSampleCounter.h"
+#include "MetasoundStandardNodesNames.h"
+
+#define LOCTEXT_NAMESPACE "MetasoundStandardNodes"
+
+namespace Metasound
+{
+	METASOUND_REGISTER_NODE(FTriggerRepeatNode)
+
+	namespace TriggerRepeatVertexNames
+	{
+		const FString& GetInputStartName()
+		{
+			static FString Name = TEXT("Start");
+			return Name;
+		}
+
+		const FText& GetInputStartDescription()
+		{
+			static FText Desc = LOCTEXT("TriggerEnableTooltip", "Starts executing periodic output triggers.");
+			return Desc;
+		}
+
+		const FString& GetInputStopName()
+		{
+			static FString Name = TEXT("Stop");
+			return Name;
+		}
+
+		const FText& GetInputStopDescription()
+		{
+			static FText Desc = LOCTEXT("TriggerDisableTooltip", "Stops executing periodic output triggers.");
+			return Desc;
+		}
+
+		const FString& GetInputPeriodName()
+		{
+			static FString Name = TEXT("Period");
+			return Name;
+		}
+
+		const FText& GetInputPeriodDescription()
+		{
+			static FText Desc = LOCTEXT("PeriodTooltip", "The period to trigger in seconds.");
+			return Desc;
+		}
+
+		const FString& GetOutputTriggerName()
+		{
+			static FString Name = TEXT("Out");
+			return Name;
+		}
+
+		const FText& GetOutputTriggerDescription()
+		{
+			static FText Desc = LOCTEXT("TriggerOutTooltip", "The periodically generated output trigger");
+			return Desc;
+		}
+	}
+
+
+	class FTriggerRepeatOperator : public TExecutableOperator<FTriggerRepeatOperator>
+	{
+		public:
+			static const FNodeClassMetadata& GetNodeInfo();
+			static const FVertexInterface& GetVertexInterface();
+			static TUniquePtr<IOperator> CreateOperator(const FCreateOperatorParams& InParams, FBuildErrorArray& OutErrors);
+
+			FTriggerRepeatOperator(const FOperatorSettings& InSettings, const FTriggerReadRef& InTriggerEnable, const FTriggerReadRef& InTriggerDisable, const FTimeReadRef& InPeriod);
+
+			virtual FDataReferenceCollection GetInputs() const override;
+
+			virtual FDataReferenceCollection GetOutputs() const override;
+
+			void Execute();
+
+		private:
+			bool bEnabled;
+
+			FTriggerWriteRef TriggerOut;
+
+			FTriggerReadRef TriggerEnable;
+			FTriggerReadRef TriggerDisable;
+
+			FTimeReadRef Period;
+
+			FSampleCount BlockSize;
+			FSampleCounter SampleCounter;
+	};
+
+	FTriggerRepeatOperator::FTriggerRepeatOperator(const FOperatorSettings& InSettings, const FTriggerReadRef& InTriggerEnable, const FTriggerReadRef& InTriggerDisable, const FTimeReadRef& InPeriod)
+	:	bEnabled(false)
+	,	TriggerOut(FTriggerWriteRef::CreateNew(InSettings))
+	,	TriggerEnable(InTriggerEnable)
+	,	TriggerDisable(InTriggerDisable)
+	,	Period(InPeriod)
+	,	BlockSize(InSettings.GetNumFramesPerBlock())
+	,	SampleCounter(0, InSettings.GetSampleRate())
+	{
+	}
+
+	FDataReferenceCollection FTriggerRepeatOperator::GetInputs() const
+	{
+		using namespace TriggerRepeatVertexNames;
+
+		FDataReferenceCollection InputDataReferences;
+		InputDataReferences.AddDataReadReference(GetInputStartName(), FTriggerReadRef(TriggerEnable));
+		InputDataReferences.AddDataReadReference(GetInputStopName(), FTriggerReadRef(TriggerDisable));
+		InputDataReferences.AddDataReadReference(GetInputPeriodName(), FTimeReadRef(Period));
+		return InputDataReferences;
+	}
+
+	FDataReferenceCollection FTriggerRepeatOperator::GetOutputs() const
+	{
+		using namespace TriggerRepeatVertexNames;
+
+		FDataReferenceCollection OutputDataReferences;
+		OutputDataReferences.AddDataReadReference(GetOutputTriggerName(), FTriggerReadRef(TriggerOut));
+
+		return OutputDataReferences;
+	}
+
+	void FTriggerRepeatOperator::Execute()
+	{
+		TriggerEnable->ExecuteBlock([](int32, int32) { },
+			[this](int32 StartFrame, int32 EndFrame)
+			{
+				bEnabled = true;
+				SampleCounter.SetNumSamples(0);
+			}
+		);
+
+		TriggerDisable->ExecuteBlock([](int32, int32) { },
+			[this](int32 StartFrame, int32 EndFrame)
+			{
+				bEnabled = false;
+			}
+		);
+
+		// Advance internal counter to get rid of old triggers.
+		TriggerOut->AdvanceBlock();
+
+		if (bEnabled)
+		{
+			FSampleCount PeriodInSamples = FSampleCounter::FromTime(*Period, SampleCounter.GetSampleRate()).GetNumSamples();
+
+			// Time must march on, can't stay in the now forever.
+			PeriodInSamples = FMath::Max(static_cast<FSampleCount>(1), PeriodInSamples);
+
+			while ((SampleCounter - BlockSize).GetNumSamples() <= 0)
+			{
+				const int32 StartOffset = static_cast<int32>(SampleCounter.GetNumSamples());
+				TriggerOut->TriggerFrame(StartOffset);
+				SampleCounter += PeriodInSamples;
+			}
+
+			SampleCounter -= BlockSize;
+		}
+	}
+
+
+	TUniquePtr<IOperator> FTriggerRepeatOperator::CreateOperator(const FCreateOperatorParams& InParams, FBuildErrorArray& OutErrors)
+	{
+		using namespace TriggerRepeatVertexNames;
+
+		const FTriggerRepeatNode& PeriodicTriggerNode = static_cast<const FTriggerRepeatNode&>(InParams.Node);
+		const FInputVertexInterface& InputInterface = GetVertexInterface().GetInputInterface();
+
+		FTriggerReadRef TriggerEnable = InParams.InputDataReferences.GetDataReadReferenceOrConstruct<FTrigger>(GetInputStartName(), InParams.OperatorSettings);
+		FTriggerReadRef TriggerDisable = InParams.InputDataReferences.GetDataReadReferenceOrConstruct<FTrigger>(GetInputStopName(), InParams.OperatorSettings);
+		FTimeReadRef Period = InParams.InputDataReferences.GetDataReadReferenceOrConstructWithVertexDefault<FTime, float>(InputInterface, GetInputPeriodName());
+
+		return MakeUnique<FTriggerRepeatOperator>(InParams.OperatorSettings, TriggerEnable, TriggerDisable, Period);
+	}
+
+	const FVertexInterface& FTriggerRepeatOperator::GetVertexInterface()
+	{
+		using namespace TriggerRepeatVertexNames;
+
+		static const FVertexInterface Interface(
+			FInputVertexInterface(
+				TInputDataVertexModel<FTrigger>(GetInputStartName(), GetInputStartDescription()),
+				TInputDataVertexModel<FTrigger>(GetInputStopName(), GetInputStopDescription()),
+				TInputDataVertexModel<FTime>(GetInputPeriodName(), GetInputPeriodDescription(), 0.2f)
+			),
+			FOutputVertexInterface(
+				TOutputDataVertexModel<FTrigger>(GetOutputTriggerName(), GetOutputTriggerDescription())
+			)
+		);
+
+		return Interface;
+	}
+
+	const FNodeClassMetadata& FTriggerRepeatOperator::GetNodeInfo()
+	{
+		auto InitNodeInfo = []() -> FNodeClassMetadata
+		{
+			FNodeClassMetadata Info;
+			Info.ClassName = { Metasound::StandardNodes::Namespace, TEXT("TriggerRepeat"), TEXT("") };
+			Info.MajorVersion = 1;
+			Info.MinorVersion = 0;
+			Info.DisplayName = LOCTEXT("Metasound_PeriodicTriggerNodeDisplayName", "Trigger Repeat");
+			Info.Description = LOCTEXT("Metasound_PeriodicTriggerNodeDescription", "Emits a trigger periodically based on the period duration given.");
+			Info.Author = PluginAuthor;
+			Info.PromptIfMissing = PluginNodeMissingPrompt;
+			Info.DefaultInterface = GetVertexInterface();
+
+			return Info;
+		};
+
+		static const FNodeClassMetadata Info = InitNodeInfo();
+
+		return Info;
+	}
+
+	FTriggerRepeatNode::FTriggerRepeatNode(const FNodeInitData& InInitData)
+	:	FNodeFacade(InInitData.InstanceName, InInitData.InstanceID, TFacadeOperatorClass<FTriggerRepeatOperator>())
+	{
+	}
+}
+
+#undef LOCTEXT_NAMESPACE //MetasoundPeriodicTriggerNode
diff --git a/Engine/Plugins/Runtime/Metasound/Source/MetasoundStandardNodes/Private/MetasoundTriggerRouteNode.cpp b/Engine/Plugins/Runtime/Metasound/Source/MetasoundStandardNodes/Private/MetasoundTriggerRouteNode.cpp
--- a/Engine/Plugins/Runtime/Metasound/Source/MetasoundStandardNodes/Private/MetasoundTriggerRouteNode.cpp
+++ b/Engine/Plugins/Runtime/Metasound/Source/MetasoundStandardNodes/Private/MetasoundTriggerRouteNode.cpp
@@ -49,6 +49,18 @@ namespace Metasound
 
 	namespace TriggerRouteVertexNames
 	{
+		const FString& GetInputInitValueName()
+		{
+			static const FString Name = TEXT("Init");
+			return Name;
+		}
+
+		const FText& GetInputInitValueDescription()
+		{
+			static const FText Desc = LOCTEXT("TriggerRouteInputInitValueDesc", "The value to initialize the output to before any inputs are triggered.");
+			return Desc;
+		}
+		
 		const FString GetInputTriggerName(uint32 InIndex)
 		{
 			return FText::Format(LOCTEXT("TriggerRouteTriggerInputName", "Set {0}"), InIndex).ToString();
diff --git a/Engine/Plugins/Runtime/Metasound/Source/MetasoundStandardNodes/Public/MetasoundPeriodicTriggerNode.h b/Engine/Plugins/Runtime/Metasound/Source/MetasoundStandardNodes/Public/MetasoundPeriodicTriggerNode.h
deleted file mode 100644
diff --git a/Engine/Plugins/Runtime/Metasound/Source/MetasoundStandardNodes/Public/MetasoundTriggerRepeatNode.h b/Engine/Plugins/Runtime/Metasound/Source/MetasoundStandardNodes/Public/MetasoundTriggerRepeatNode.h
new file mode 100644
--- /dev/null
+++ b/Engine/Plugins/Runtime/Metasound/Source/MetasoundStandardNodes/Public/MetasoundTriggerRepeatNode.h
@@ -0,0 +1,29 @@
+// Copyright Epic Games, Inc. All Rights Reserved.
+
+#pragma once
+
+#include "CoreMinimal.h"
+#include "MetasoundDataReferenceCollection.h"
+#include "MetasoundNode.h"
+#include "MetasoundBuilderInterface.h"
+#include "MetasoundExecutableOperator.h"
+#include "MetasoundFacade.h"
+#include "MetasoundOperatorInterface.h"
+#include "MetasoundNodeInterface.h"
+#include "MetasoundTime.h"
+#include "MetasoundTrigger.h"
+
+
+namespace Metasound
+{
+	class METASOUNDSTANDARDNODES_API FTriggerRepeatNode : public FNodeFacade
+	{
+		public:
+			FTriggerRepeatNode(const FNodeInitData& InInitData);
+
+			virtual ~FTriggerRepeatNode() = default;
+
+		private:
+			float DefaultPeriod = 1.0f;
+	};
+}
diff --git a/Engine/Plugins/Runtime/Metasound/Source/MetasoundStandardNodes/Public/MetasoundTriggerRouteNode.h b/Engine/Plugins/Runtime/Metasound/Source/MetasoundStandardNodes/Public/MetasoundTriggerRouteNode.h
--- a/Engine/Plugins/Runtime/Metasound/Source/MetasoundStandardNodes/Public/MetasoundTriggerRouteNode.h
+++ b/Engine/Plugins/Runtime/Metasound/Source/MetasoundStandardNodes/Public/MetasoundTriggerRouteNode.h
@@ -22,6 +22,8 @@ namespace Metasound
 
 	namespace TriggerRouteVertexNames
 	{
+		METASOUNDSTANDARDNODES_API const FString& GetInputInitValueName();
+		METASOUNDSTANDARDNODES_API const FText& GetInputInitValueDescription();
 		METASOUNDSTANDARDNODES_API const FString GetInputTriggerName(uint32 InIndex);
 		METASOUNDSTANDARDNODES_API const FText GetInputTriggerDescription(uint32 InIndex);
 		METASOUNDSTANDARDNODES_API const FString GetInputValueName(uint32 InIndex);
@@ -44,6 +46,8 @@ namespace Metasound
 			auto CreateDefaultInterface = []() -> FVertexInterface
 			{
 				FInputVertexInterface InputInterface;
+				InputInterface.Add(TInputDataVertexModel<ValueType>(GetInputInitValueName(), GetInputInitValueDescription()));
+
 				for (uint32 i = 0; i < NumInputs; ++i)
 				{
 					InputInterface.Add(TInputDataVertexModel<FTrigger>(GetInputTriggerName(i), GetInputTriggerDescription(i)));
@@ -84,6 +88,8 @@ namespace Metasound
 			const FInputVertexInterface& InputInterface = InParams.Node.GetVertexInterface().GetInputInterface();
 			const FDataReferenceCollection& InputCollection = InParams.InputDataReferences;
 
+			TDataReadReference<ValueType> InputInitValue = InputCollection.GetDataReadReferenceOrConstructWithVertexDefault<ValueType>(InputInterface, GetInputInitValueName());
+
 			TArray<FTriggerReadRef> InputTriggers;
 			TArray<TDataReadReference<ValueType>> InputValues;
 
@@ -93,16 +99,17 @@ namespace Metasound
 				InputValues.Add(InputCollection.GetDataReadReferenceOrConstructWithVertexDefault<ValueType>(InputInterface, GetInputValueName(i)));
 			}
 
-			return MakeUnique<TTriggerRouteOperator<ValueType, NumInputs>>(InParams.OperatorSettings, MoveTemp(InputTriggers), MoveTemp(InputValues));
+			return MakeUnique<TTriggerRouteOperator<ValueType, NumInputs>>(InParams.OperatorSettings, InputInitValue, MoveTemp(InputTriggers), MoveTemp(InputValues));
 		}
 
-		TTriggerRouteOperator(const FOperatorSettings& InSettings, const TArray<FTriggerReadRef>&& InInputTriggers, TArray<TDataReadReference<ValueType>>&& InInputValues)
-			: InputTriggers(InInputTriggers)
+		TTriggerRouteOperator(const FOperatorSettings& InSettings, const TDataReadReference<ValueType>& InInputInitValue, const TArray<FTriggerReadRef>&& InInputTriggers, const TArray<TDataReadReference<ValueType>>&& InInputValues)
+			: InputInitValue(InInputInitValue)
+			, InputTriggers(InInputTriggers)
 			, InputValues(InInputValues)
 			, OutputTrigger(FTriggerWriteRef::CreateNew(InSettings))
 			, OutputValue(TDataWriteReferenceFactory<ValueType>::CreateAny(InSettings))
 		{
-			*OutputValue = *InputValues[0];
+			*OutputValue = *InputInitValue;
 		}
 
 		virtual ~TTriggerRouteOperator() = default;
@@ -154,6 +161,7 @@ namespace Metasound
 
 	private:
 
+		TDataReadReference<ValueType> InputInitValue;
 		TArray<FTriggerReadRef> InputTriggers;
 		TArray<TDataReadReference<ValueType>> InputValues;
 
