Some panning/spatialization metasound nodes

diff --git a/Engine/Plugins/Runtime/Metasound/Source/MetasoundStandardNodes/Private/MetasoundITDPannerNode.cpp b/Engine/Plugins/Runtime/Metasound/Source/MetasoundStandardNodes/Private/MetasoundITDPannerNode.cpp
new file mode 100644
--- /dev/null
+++ b/Engine/Plugins/Runtime/Metasound/Source/MetasoundStandardNodes/Private/MetasoundITDPannerNode.cpp
@@ -0,0 +1,353 @@
+// Copyright Epic Games, Inc. All Rights Reserved.
+
+#include "MetasoundITDPannerNode.h"
+
+#include "Internationalization/Text.h"
+#include "MetasoundExecutableOperator.h"
+#include "MetasoundNodeRegistrationMacro.h"
+#include "MetasoundDataTypeRegistrationMacro.h"
+#include "MetasoundPrimitives.h"
+#include "MetasoundStandardNodesNames.h"
+#include "MetasoundTrigger.h"
+#include "MetasoundTime.h"
+#include "MetasoundAudioBuffer.h"
+#include "DSP/BufferVectorOperations.h"
+#include "DSP/Delay.h"
+#include "DSP/Dsp.h"
+
+#define LOCTEXT_NAMESPACE "MetasoundStandardNodes"
+
+namespace Metasound
+{
+	namespace ITDPannerVertexNames
+	{
+		const FString& GetInputAudioName()
+		{
+			static FString Name = TEXT("In");
+			return Name;
+		}
+
+		const FText& GetInputAudioDescription()
+		{
+			static FText Desc = LOCTEXT("ITDPannerNodeInDesc", "The input audio to spatialize.");
+			return Desc;
+		}
+
+		const FString& GetInputPanAngleName()
+		{
+			static FString Name = TEXT("Angle");
+			return Name;
+		}
+
+		const FText& GetInputPanAngleDescription()
+		{
+			static FText Desc = LOCTEXT("ITDPannerNodePanAngleDesc", "The sound source angle in degrees. 90 degrees is in front, 0 degrees is to the right, 270 degrees is behind, 180 degrees is to the left.");
+			return Desc;
+		}
+
+		const FString& GetInputDistanceFactorName()
+		{
+			static FString Name = TEXT("Distance Factor");
+			return Name;
+		}
+
+		const FText& GetInputDistanceFactorDescription()
+		{
+			static FText Desc = LOCTEXT("ITDPannerNodeDistanceFactorDesc", "The normalized distance factor (0.0 to 1.0) to use for ILD (Inter-aural level difference) calculations. 0.0 is near, 1.0f is far. The further away something is the less there is a difference in levels (gain) between the ears.");
+			return Desc;
+		}
+
+		const FString& GetInputHeadWidthName()
+		{
+			static FString Name = TEXT("Head Width");
+			return Name;
+		}
+
+		const FText& GetInputHeadWidthDescription()
+		{
+			static FText Desc = LOCTEXT("ITDPannerNodeHeadWidthDesc", "The width of the listener head to use for ITD calculations in centimeters.");
+			return Desc;
+		}
+
+		const FString& GetOutputAudioLeftName()
+		{
+			static FString Name = TEXT("Out Left");
+			return Name;
+		}
+
+		const FText& GetOutputAudioLeftDescription()
+		{
+			static FText Desc = LOCTEXT("ITDPannerNodeOutputLeftDescription", "Left channel audio output.");
+			return Desc;
+		}
+
+		const FString& GetOutputAudioRightName()
+		{
+			static FString Name = TEXT("Out Right");
+			return Name;
+		}
+
+		const FText& GetOutputAudioRightDescription()
+		{
+			static FText Desc = LOCTEXT("ITDPannerNodeOutputRightDescription", "Right channel audio output.");
+			return Desc;
+		}
+	}
+
+	class FITDPannerOperator : public TExecutableOperator<FITDPannerOperator>
+	{
+	public:
+
+		static const FNodeClassMetadata& GetNodeInfo();
+		static const FVertexInterface& GetVertexInterface();
+		static TUniquePtr<IOperator> CreateOperator(const FCreateOperatorParams& InParams, FBuildErrorArray& OutErrors);
+
+		FITDPannerOperator(const FOperatorSettings& InSettings,
+			const FAudioBufferReadRef& InAudioInput, 
+			const FFloatReadRef& InPanningAngle,
+			const FFloatReadRef& InDistanceFactor,
+			const FFloatReadRef& InHeadWidth);
+
+		virtual FDataReferenceCollection GetInputs() const override;
+		virtual FDataReferenceCollection GetOutputs() const override;
+		void Execute();
+
+	private:
+		void UpdateParams();
+		void UpdateXY();
+		void UpdateGains();
+		void UpdateDelays();
+
+		FAudioBufferReadRef AudioInput;
+		FFloatReadRef PanningAngle;
+		FFloatReadRef DistanceFactor;
+		FFloatReadRef HeadWidth;
+
+		FAudioBufferWriteRef AudioLeftOutput;
+		FAudioBufferWriteRef AudioRightOutput;
+
+		float CurrAngle = 0.0f;
+		float CurrX = 0.0f;
+		float CurrY = 0.0f;
+		float CurrDistanceFactor = 0.0f;
+		float CurrHeadWidth = 0.0f;
+		float CurrLeftGain = 0.0f;
+		float CurrRightGain = 0.0f;
+		float CurrLeftDelay = 0.0f;
+		float CurrRightDelay = 0.0f;
+
+		float PrevLeftGain = 0.0f;
+		float PrevRightGain = 0.0f;
+
+		Audio::FDelay LeftDelay;
+		Audio::FDelay RightDelay;
+	};
+
+	FITDPannerOperator::FITDPannerOperator(const FOperatorSettings& InSettings,
+		const FAudioBufferReadRef& InAudioInput, 
+		const FFloatReadRef& InPanningAngle,
+		const FFloatReadRef& InDistanceFactor,
+		const FFloatReadRef& InHeadWidth)
+		: AudioInput(InAudioInput)
+		, PanningAngle(InPanningAngle)
+		, DistanceFactor(InDistanceFactor)
+		, HeadWidth(InHeadWidth)
+		, AudioLeftOutput(FAudioBufferWriteRef::CreateNew(InSettings))
+		, AudioRightOutput(FAudioBufferWriteRef::CreateNew(InSettings))
+	{
+		LeftDelay.Init(InSettings.GetSampleRate(), 0.5f);
+		RightDelay.Init(InSettings.GetSampleRate(), 0.5f);
+
+		const float EaseFactor = Audio::FExponentialEase::GetFactorForTau(0.1f, InSettings.GetSampleRate());
+		LeftDelay.SetEaseFactor(EaseFactor);
+		RightDelay.SetEaseFactor(EaseFactor);
+
+		CurrAngle = FMath::Clamp(*PanningAngle, 0.0f, 360.0f);
+		CurrDistanceFactor = FMath::Clamp(*DistanceFactor, 0.0f, 1.0f);
+		CurrHeadWidth = FMath::Max(*InHeadWidth, 0.0f);
+
+		UpdateParams();
+
+		PrevLeftGain = CurrLeftGain;
+		PrevRightGain = CurrRightGain;
+	}
+
+	FDataReferenceCollection FITDPannerOperator::GetInputs() const
+	{
+		using namespace ITDPannerVertexNames;
+
+		FDataReferenceCollection InputDataReferences;
+		InputDataReferences.AddDataReadReference(GetInputAudioName(), AudioInput);
+		InputDataReferences.AddDataReadReference(GetInputPanAngleName(), PanningAngle);
+		InputDataReferences.AddDataReadReference(GetInputDistanceFactorName(), DistanceFactor);
+		InputDataReferences.AddDataReadReference(GetInputHeadWidthName(), HeadWidth);
+
+		return InputDataReferences;
+	}
+
+	FDataReferenceCollection FITDPannerOperator::GetOutputs() const
+	{
+		using namespace ITDPannerVertexNames;
+
+		FDataReferenceCollection OutputDataReferences;
+		OutputDataReferences.AddDataReadReference(GetOutputAudioLeftName(), AudioLeftOutput);
+		OutputDataReferences.AddDataReadReference(GetOutputAudioRightName(), AudioRightOutput);
+
+		return OutputDataReferences;
+	}
+
+	void FITDPannerOperator::UpdateParams()
+	{
+		// ****************
+		// Update the x-y values
+		const float CurrRadians = (CurrAngle / 360.0f) * 2.0f * PI;
+		FMath::SinCos(&CurrY, &CurrX, CurrRadians);
+
+		// ****************
+		// Update ILD gains
+		const float HeadRadiusMeters = 0.005f * CurrHeadWidth; // (InHeadWidth / 100.0f) / 2.0f;
+
+		// InX is -1.0 to 1.0, so get it in 0.0 to 1.0 (i.e. hard left, hard right)
+		const float Fraction = (CurrX + 1.0f) * 0.5f;
+
+		// Feed the linear pan value into a equal power equation
+		float PanLeft;
+		float PanRight;
+		FMath::SinCos(&PanRight, &PanLeft, 0.5f * PI * Fraction);
+
+		// If distance factor is 1.0 (i.e. far away) this will have equal gain, if distance factor is 0.0 it will be normal equal power pan.
+		CurrLeftGain = FMath::Lerp(PanLeft, 0.5f, CurrDistanceFactor);
+		CurrRightGain = FMath::Lerp(PanRight, 0.5f, CurrDistanceFactor);
+
+		// *********************
+		// Update the ITD delays
+
+		// Use pythagorean theorem to get distances
+		const float DistToLeftEar = FMath::Sqrt((CurrY * CurrY) + FMath::Square(HeadRadiusMeters + CurrX));
+		const float DistToRightEar = FMath::Sqrt((CurrY * CurrY) + FMath::Square(HeadRadiusMeters - CurrX));
+
+		// Compute delta time based on speed of sound 
+		constexpr float SpeedOfSound = 343.0f;
+		const float DeltaTimeSeconds = (DistToLeftEar - DistToRightEar) / SpeedOfSound;
+
+		if (DeltaTimeSeconds > 0.0f)
+		{
+			LeftDelay.SetEasedDelayMsec(1000.0f * DeltaTimeSeconds);
+			RightDelay.SetEasedDelayMsec(0.0f);
+		}
+		else
+		{
+			LeftDelay.SetEasedDelayMsec(0.0f);
+			RightDelay.SetEasedDelayMsec(-1000.0f * DeltaTimeSeconds);
+		}
+	}
+	
+
+	void FITDPannerOperator::Execute()
+	{
+		if (!FMath::IsNearlyEqual(*HeadWidth, CurrHeadWidth) || 
+			!FMath::IsNearlyEqual(*PanningAngle, CurrAngle) ||
+			!FMath::IsNearlyEqual(*DistanceFactor, CurrDistanceFactor))
+		{
+			CurrHeadWidth = *HeadWidth;
+			CurrAngle = *PanningAngle;
+			CurrDistanceFactor = *DistanceFactor;
+
+			UpdateParams();
+		}
+
+		const float* InputBufferPtr = AudioInput->GetData();
+		int32 InputSampleCount = AudioInput->Num();
+		float* OutputLeftBufferPtr = AudioLeftOutput->GetData();
+		float* OutputRightBufferPtr = AudioRightOutput->GetData();
+
+		// Feed the input audio into the left and right delays
+		for (int32 i = 0; i < InputSampleCount; ++i)
+		{
+			OutputLeftBufferPtr[i] = LeftDelay.ProcessAudioSample(InputBufferPtr[i]);
+			OutputRightBufferPtr[i] = RightDelay.ProcessAudioSample(InputBufferPtr[i]);
+		}
+
+		// Now apply the panning
+		if (FMath::IsNearlyEqual(PrevLeftGain, CurrLeftDelay))
+		{
+			Audio::MultiplyBufferByConstantInPlace(OutputLeftBufferPtr, InputSampleCount, PrevLeftGain);
+			Audio::MultiplyBufferByConstantInPlace(OutputRightBufferPtr, InputSampleCount, PrevRightGain);
+		}
+		else
+		{
+			Audio::FadeBufferFast(OutputLeftBufferPtr, InputSampleCount, PrevLeftGain, CurrLeftGain);
+			Audio::FadeBufferFast(OutputRightBufferPtr, InputSampleCount, PrevRightGain, CurrRightGain);
+
+			PrevLeftGain = CurrLeftGain;
+			PrevRightGain = CurrRightGain;
+		}
+	}
+
+	const FVertexInterface& FITDPannerOperator::GetVertexInterface()
+	{
+		using namespace ITDPannerVertexNames;
+
+		static const FVertexInterface Interface(
+			FInputVertexInterface(
+				TInputDataVertexModel<FAudioBuffer>(GetInputAudioName(), GetInputAudioDescription()),
+				TInputDataVertexModel<float>(GetInputPanAngleName(), GetInputPanAngleDescription(), 0.0f),
+				TInputDataVertexModel<float>(GetInputDistanceFactorName(), GetInputDistanceFactorDescription(), 0.0f),
+				TInputDataVertexModel<float>(GetInputHeadWidthName(), GetInputHeadWidthDescription(), 34.0f)
+			),
+			FOutputVertexInterface(
+				TOutputDataVertexModel<FAudioBuffer>(GetOutputAudioLeftName(), GetOutputAudioLeftDescription()),
+				TOutputDataVertexModel<FAudioBuffer>(GetOutputAudioRightName(), GetOutputAudioRightDescription())
+			)
+		);
+
+		return Interface;
+	}
+
+	const FNodeClassMetadata& FITDPannerOperator::GetNodeInfo()
+	{
+		auto InitNodeInfo = []() -> FNodeClassMetadata
+		{
+			FNodeClassMetadata Info;
+			Info.ClassName = { Metasound::StandardNodes::Namespace, TEXT("ITD Panner"), TEXT("") };
+			Info.MajorVersion = 1;
+			Info.MinorVersion = 0;
+			Info.DisplayName = LOCTEXT("Metasound_ITDPannerDisplayName", "ITD Panner");
+			Info.Description = LOCTEXT("Metasound_ITDPannerNodeDescription", "Pans an input audio signal using an inter-aural time delay method.");
+			Info.Author = PluginAuthor;
+			Info.PromptIfMissing = PluginNodeMissingPrompt;
+			Info.DefaultInterface = GetVertexInterface();
+
+			return Info;
+		};
+
+		static const FNodeClassMetadata Info = InitNodeInfo();
+
+		return Info;
+	}
+
+	TUniquePtr<IOperator> FITDPannerOperator::CreateOperator(const FCreateOperatorParams& InParams, FBuildErrorArray& OutErrors)
+	{
+		const FITDPannerNode& ITDPannerNode = static_cast<const FITDPannerNode&>(InParams.Node);
+		const FDataReferenceCollection& InputCollection = InParams.InputDataReferences;
+		const FInputVertexInterface& InputInterface = GetVertexInterface().GetInputInterface();
+
+		using namespace ITDPannerVertexNames;
+
+		FAudioBufferReadRef AudioIn = InputCollection.GetDataReadReferenceOrConstruct<FAudioBuffer>(GetInputAudioName(), InParams.OperatorSettings);
+		FFloatReadRef PanningAngle = InputCollection.GetDataReadReferenceOrConstructWithVertexDefault<float>(InputInterface, GetInputPanAngleName());
+		FFloatReadRef DistanceFactor = InputCollection.GetDataReadReferenceOrConstructWithVertexDefault<float>(InputInterface, GetInputDistanceFactorName());
+		FFloatReadRef HeadWidth = InputCollection.GetDataReadReferenceOrConstructWithVertexDefault<float>(InputInterface, GetInputHeadWidthName());
+
+		return MakeUnique<FITDPannerOperator>(InParams.OperatorSettings, AudioIn, PanningAngle, DistanceFactor, HeadWidth);
+	}
+
+	FITDPannerNode::FITDPannerNode(const FNodeInitData& InitData)
+		: FNodeFacade(InitData.InstanceName, InitData.InstanceID, TFacadeOperatorClass<FITDPannerOperator>())
+	{
+	}
+
+	METASOUND_REGISTER_NODE(FITDPannerNode)
+}
+
+#undef LOCTEXT_NAMESPACE
diff --git a/Engine/Plugins/Runtime/Metasound/Source/MetasoundStandardNodes/Private/MetasoundStereoPannerNode.cpp b/Engine/Plugins/Runtime/Metasound/Source/MetasoundStandardNodes/Private/MetasoundStereoPannerNode.cpp
new file mode 100644
--- /dev/null
+++ b/Engine/Plugins/Runtime/Metasound/Source/MetasoundStandardNodes/Private/MetasoundStereoPannerNode.cpp
@@ -0,0 +1,293 @@
+// Copyright Epic Games, Inc. All Rights Reserved.
+
+#include "MetasoundStereoPannerNode.h"
+
+#include "Internationalization/Text.h"
+#include "MetasoundExecutableOperator.h"
+#include "MetasoundNodeRegistrationMacro.h"
+#include "MetasoundDataTypeRegistrationMacro.h"
+#include "MetasoundPrimitives.h"
+#include "MetasoundStandardNodesNames.h"
+#include "MetasoundTrigger.h"
+#include "MetasoundTime.h"
+#include "MetasoundAudioBuffer.h"
+#include "DSP/BufferVectorOperations.h"
+
+#define LOCTEXT_NAMESPACE "MetasoundStandardNodes"
+
+namespace Metasound
+{
+	namespace StereoPannerVertexNames
+	{
+		const FString& GetInputAudioName()
+		{
+			static FString Name = TEXT("In");
+			return Name;
+		}
+
+		const FText& GetInputAudioDescription()
+		{
+			static FText Desc = LOCTEXT("StereoPannerNodeInDesc", "The input audio to pan.");
+			return Desc;
+		}
+
+		const FString& GetInputPanAmountName()
+		{
+			static FString Name = TEXT("Pan Amount");
+			return Name;
+		}
+
+		const FText& GetInputPanAmountDescription()
+		{
+			static FText Desc = LOCTEXT("StereoPannerNodePanAmountDesc", "The amount of pan. -1.0 is full left, 1.0 is full right.");
+			return Desc;
+		}
+
+		const FString& GetInputPanningLawName()
+		{
+			static FString Name = TEXT("Panning Law");
+			return Name;
+		}
+
+		const FText& GetInputPanningLawDescription()
+		{
+			static FText Desc = LOCTEXT("StereoPannerNodePanningLawDescription", "Which panning law should be used for the stereo panner.");
+			return Desc;
+		}
+
+		const FString& GetOutputAudioLeftName()
+		{
+			static FString Name = TEXT("Out Left");
+			return Name;
+		}
+
+		const FText& GetOutputAudioLeftDescription()
+		{
+			static FText Desc = LOCTEXT("StereoPannerNodeOutputLeftDescription", "Left channel audio output.");
+			return Desc;
+		}
+
+		const FString& GetOutputAudioRightName()
+		{
+			static FString Name = TEXT("Out Right");
+			return Name;
+		}
+
+		const FText& GetOutputAudioRightDescription()
+		{
+			static FText Desc = LOCTEXT("StereoPannerNodeOutputRightDescription", "Right channel audio output.");
+			return Desc;
+		}
+	}
+
+	enum class EPanningLaw
+	{
+		EqualPower = 0,
+		Linear
+	};
+
+	DECLARE_METASOUND_ENUM(EPanningLaw, EPanningLaw::EqualPower, METASOUNDSTANDARDNODES_API,
+	FEnumPanningLaw, FEnumPanningLawInfo, FPanningLawReadRef, FPanningLawWriteRef);
+
+	DEFINE_METASOUND_ENUM_BEGIN(EPanningLaw, FEnumPanningLaw, "PanningLaw")
+		DEFINE_METASOUND_ENUM_ENTRY(EPanningLaw::EqualPower, LOCTEXT("PanningLawEqualPowerName", "Equal Power"), LOCTEXT("PanningLawEqualPowerTT", "The power of the audio signal is constant while panning.")),
+		DEFINE_METASOUND_ENUM_ENTRY(EPanningLaw::Linear, LOCTEXT("PanningLawLinearName", "Linear"), LOCTEXT("PanningLawLinearTT", "The amplitude of the audio signal is constant while panning.")),
+	DEFINE_METASOUND_ENUM_END()
+
+	class FStereoPannerOperator : public TExecutableOperator<FStereoPannerOperator>
+	{
+	public:
+
+		static const FNodeClassMetadata& GetNodeInfo();
+		static const FVertexInterface& GetVertexInterface();
+		static TUniquePtr<IOperator> CreateOperator(const FCreateOperatorParams& InParams, FBuildErrorArray& OutErrors);
+
+		FStereoPannerOperator(const FOperatorSettings& InSettings,
+			const FAudioBufferReadRef& InAudioInput, 
+			const FFloatReadRef& InPanningAmount,
+			const FPanningLawReadRef& InPanningLaw);
+
+		virtual FDataReferenceCollection GetInputs() const override;
+		virtual FDataReferenceCollection GetOutputs() const override;
+		void Execute();
+
+	private:
+		float GetInputDelayTimeMsec() const;
+		void ComputePanGains(float InPanningAmmount, float& OutLeftGain, float& OutRightGain) const;
+
+		// The input audio buffer
+		FAudioBufferReadRef AudioInput;
+
+		// The amount of delay time
+		FFloatReadRef PanningAmount;
+
+		// The the dry level
+		FPanningLawReadRef PanningLaw;
+
+		// The audio output
+		FAudioBufferWriteRef AudioLeftOutput;
+		FAudioBufferWriteRef AudioRightOutput;
+
+		float PrevPanningAmount = 0.0f;
+		float PrevLeftPan = 0.0f;
+		float PrevRightPan = 0.0f;
+	};
+
+	FStereoPannerOperator::FStereoPannerOperator(const FOperatorSettings& InSettings,
+		const FAudioBufferReadRef& InAudioInput, 
+		const FFloatReadRef& InPanningAmount,
+		const FPanningLawReadRef& InPanningLaw)
+
+		: AudioInput(InAudioInput)
+		, PanningAmount(InPanningAmount)
+		, PanningLaw(InPanningLaw)
+		, AudioLeftOutput(FAudioBufferWriteRef::CreateNew(InSettings))
+		, AudioRightOutput(FAudioBufferWriteRef::CreateNew(InSettings))
+	{
+		PrevPanningAmount = *PanningAmount;
+
+		ComputePanGains(PrevPanningAmount, PrevLeftPan, PrevRightPan);
+	}
+
+	FDataReferenceCollection FStereoPannerOperator::GetInputs() const
+	{
+		using namespace StereoPannerVertexNames;
+
+		FDataReferenceCollection InputDataReferences;
+		InputDataReferences.AddDataReadReference(GetInputAudioName(), AudioInput);
+		InputDataReferences.AddDataReadReference(GetInputPanAmountName(), PanningAmount);
+		InputDataReferences.AddDataReadReference(GetInputPanningLawName(), PanningLaw);
+
+		return InputDataReferences;
+	}
+
+	FDataReferenceCollection FStereoPannerOperator::GetOutputs() const
+	{
+		using namespace StereoPannerVertexNames;
+
+		FDataReferenceCollection OutputDataReferences;
+		OutputDataReferences.AddDataReadReference(GetOutputAudioLeftName(), AudioLeftOutput);
+		OutputDataReferences.AddDataReadReference(GetOutputAudioRightName(), AudioRightOutput);
+
+		return OutputDataReferences;
+	}
+
+	void FStereoPannerOperator::ComputePanGains(float InPanningAmmount, float& OutLeftGain, float& OutRightGain) const
+	{
+		// Convert [-1.0, 1.0] to [0.0, 1.0]
+		float Fraction = 0.5f * (InPanningAmmount + 1.0f);
+
+		if (*PanningLaw == EPanningLaw::EqualPower)
+		{
+			// Compute the left and right amount with one math call
+			FMath::SinCos(&OutRightGain, &OutLeftGain, 0.5f * PI * Fraction);
+		}
+		else
+		{
+			OutLeftGain = Fraction;
+			OutRightGain = 1.0f - Fraction;
+		}
+	}
+
+
+	void FStereoPannerOperator::Execute()
+	{
+		float CurrentPanningAmount = FMath::Clamp(*PanningAmount, -1.0f, 1.0f);
+
+		const float* InputBufferPtr = AudioInput->GetData();
+		int32 InputSampleCount = AudioInput->Num();
+		float* OutputLeftBufferPtr = AudioLeftOutput->GetData();
+		float* OutputRightBufferPtr = AudioRightOutput->GetData();
+
+		if (FMath::IsNearlyEqual(PrevPanningAmount, CurrentPanningAmount))
+		{
+			Audio::BufferMultiplyByConstant(InputBufferPtr, PrevLeftPan, OutputLeftBufferPtr, InputSampleCount);
+			Audio::BufferMultiplyByConstant(InputBufferPtr, PrevRightPan, OutputRightBufferPtr, InputSampleCount);
+		}
+		else 
+		{
+			// The pan amount has changed so recompute it
+			float CurrentLeftPan;
+			float CurrentRightPan;
+			ComputePanGains(CurrentPanningAmount, CurrentLeftPan, CurrentRightPan);
+
+			// Copy the input to the output buffers
+			FMemory::Memcpy(OutputLeftBufferPtr, InputBufferPtr, InputSampleCount * sizeof(float));
+			FMemory::Memcpy(OutputRightBufferPtr, InputBufferPtr, InputSampleCount * sizeof(float));
+
+			// Do a fast fade on the buffers from the prev left/right gains to current left/right gains
+			Audio::FadeBufferFast(OutputLeftBufferPtr, InputSampleCount, PrevLeftPan, CurrentLeftPan);
+			Audio::FadeBufferFast(OutputRightBufferPtr, InputSampleCount, PrevRightPan, CurrentRightPan);
+
+			// lerp through the buffer to the target panning amount
+			PrevPanningAmount = *PanningAmount;
+			PrevLeftPan = CurrentLeftPan;
+			PrevRightPan = CurrentRightPan;
+		}
+	}
+
+	const FVertexInterface& FStereoPannerOperator::GetVertexInterface()
+	{
+		using namespace StereoPannerVertexNames;
+
+		static const FVertexInterface Interface(
+			FInputVertexInterface(
+				TInputDataVertexModel<FAudioBuffer>(GetInputAudioName(), GetInputAudioDescription()),
+				TInputDataVertexModel<float>(GetInputPanAmountName(), GetInputPanAmountDescription(), 0.0f),
+				TInputDataVertexModel<FEnumPanningLaw>(GetInputPanningLawName(), GetInputPanningLawDescription())
+			),
+			FOutputVertexInterface(
+				TOutputDataVertexModel<FAudioBuffer>(GetOutputAudioLeftName(), GetOutputAudioLeftDescription()),
+				TOutputDataVertexModel<FAudioBuffer>(GetOutputAudioRightName(), GetOutputAudioRightDescription())
+			)
+		);
+
+		return Interface;
+	}
+
+	const FNodeClassMetadata& FStereoPannerOperator::GetNodeInfo()
+	{
+		auto InitNodeInfo = []() -> FNodeClassMetadata
+		{
+			FNodeClassMetadata Info;
+			Info.ClassName = { Metasound::StandardNodes::Namespace, TEXT("Stereo Panner"), TEXT("") };
+			Info.MajorVersion = 1;
+			Info.MinorVersion = 0;
+			Info.DisplayName = LOCTEXT("Metasound_StereoPannerDisplayName", "Stereo Panner");
+			Info.Description = LOCTEXT("Metasound_StereoPannerNodeDescription", "Pans an input audio signal to left and right outputs.");
+			Info.Author = PluginAuthor;
+			Info.PromptIfMissing = PluginNodeMissingPrompt;
+			Info.DefaultInterface = GetVertexInterface();
+
+			return Info;
+		};
+
+		static const FNodeClassMetadata Info = InitNodeInfo();
+
+		return Info;
+	}
+
+	TUniquePtr<IOperator> FStereoPannerOperator::CreateOperator(const FCreateOperatorParams& InParams, FBuildErrorArray& OutErrors)
+	{
+		const FStereoPannerNode& StereoPannerNode = static_cast<const FStereoPannerNode&>(InParams.Node);
+		const FDataReferenceCollection& InputCollection = InParams.InputDataReferences;
+		const FInputVertexInterface& InputInterface = GetVertexInterface().GetInputInterface();
+
+		using namespace StereoPannerVertexNames;
+
+		FAudioBufferReadRef AudioIn = InputCollection.GetDataReadReferenceOrConstruct<FAudioBuffer>(GetInputAudioName(), InParams.OperatorSettings);
+		FFloatReadRef PanningAmount = InputCollection.GetDataReadReferenceOrConstructWithVertexDefault<float>(InputInterface, GetInputPanAmountName());
+		FPanningLawReadRef PanningLaw = InputCollection.GetDataReadReferenceOrConstruct<FEnumPanningLaw>(GetInputPanningLawName());
+
+		return MakeUnique<FStereoPannerOperator>(InParams.OperatorSettings, AudioIn, PanningAmount, PanningLaw);
+	}
+
+	FStereoPannerNode::FStereoPannerNode(const FNodeInitData& InitData)
+		: FNodeFacade(InitData.InstanceName, InitData.InstanceID, TFacadeOperatorClass<FStereoPannerOperator>())
+	{
+	}
+
+	METASOUND_REGISTER_NODE(FStereoPannerNode)
+}
+
+#undef LOCTEXT_NAMESPACE
diff --git a/Engine/Plugins/Runtime/Metasound/Source/MetasoundStandardNodes/Public/MetasoundITDPannerNode.h b/Engine/Plugins/Runtime/Metasound/Source/MetasoundStandardNodes/Public/MetasoundITDPannerNode.h
new file mode 100644
--- /dev/null
+++ b/Engine/Plugins/Runtime/Metasound/Source/MetasoundStandardNodes/Public/MetasoundITDPannerNode.h
@@ -0,0 +1,23 @@
+// Copyright Epic Games, Inc. All Rights Reserved.
+
+#pragma once
+
+#include "CoreMinimal.h"
+#include "MetasoundFacade.h"
+
+namespace Metasound
+{
+	/** FITDPannerNode
+	 *
+	 *  Pans an audio buffer using an ITD method (Interaural time delay)
+	 */
+	class METASOUNDSTANDARDNODES_API FITDPannerNode  : public FNodeFacade
+	{
+	public:
+		/**
+		 * Constructor used by the Metasound Frontend.
+		 */
+		FITDPannerNode(const FNodeInitData& InitData);
+	};
+} // namespace Metasound
+
diff --git a/Engine/Plugins/Runtime/Metasound/Source/MetasoundStandardNodes/Public/MetasoundStereoPannerNode.h b/Engine/Plugins/Runtime/Metasound/Source/MetasoundStandardNodes/Public/MetasoundStereoPannerNode.h
new file mode 100644
--- /dev/null
+++ b/Engine/Plugins/Runtime/Metasound/Source/MetasoundStandardNodes/Public/MetasoundStereoPannerNode.h
@@ -0,0 +1,23 @@
+// Copyright Epic Games, Inc. All Rights Reserved.
+
+#pragma once
+
+#include "CoreMinimal.h"
+#include "MetasoundFacade.h"
+
+namespace Metasound
+{
+	/** FStereoPannerNode
+	 *
+	 *  Pans an input audio signal to stereo output
+	 */
+	class METASOUNDSTANDARDNODES_API FStereoPannerNode : public FNodeFacade
+	{
+	public:
+		/**
+		 * Constructor used by the Metasound Frontend.
+		 */
+		FStereoPannerNode(const FNodeInitData& InitData);
+	};
+} // namespace Metasound
+
