Metasound UI/UX Updates
- Input/Output Reference nodes (Allows more than one input node per input, & hiding output nodes if not desired)
- Implement hotswapping input/output types
- Update MS Graph context menu to dynamically add actions for available inputs/outputs
- Add copy/paste for input/output nodes (but disallow from other graphs)
- Various utility functions for frontend for iterating node inputs/outputs
- Consolidation of GetDisplayName/GetDescription node logic into Frontend
- 1st pass on fixing up sync code to support input/output ed reference nodes
- Migrate edgraph input literals to new UObject class to allow for input references & hotswapping type

diff --git a/Engine/Plugins/Runtime/Metasound/Source/MetasoundEditor/Private/MetasoundDetailCustomization.cpp b/Engine/Plugins/Runtime/Metasound/Source/MetasoundEditor/Private/MetasoundDetailCustomization.cpp
--- a/Engine/Plugins/Runtime/Metasound/Source/MetasoundEditor/Private/MetasoundDetailCustomization.cpp
+++ b/Engine/Plugins/Runtime/Metasound/Source/MetasoundEditor/Private/MetasoundDetailCustomization.cpp
@@ -9,6 +9,7 @@
 #include "DetailWidgetRow.h"
 #include "Framework/Notifications/NotificationManager.h"
 #include "IDetailGroup.h"
+#include "Input/Events.h"
 #include "MetasoundAssetBase.h"
 #include "MetasoundFrontend.h"
 #include "MetasoundFrontendController.h"
@@ -16,6 +17,7 @@
 #include "PropertyEditorDelegates.h"
 #include "PropertyHandle.h"
 #include "PropertyRestriction.h"
+#include "SGraphPalette.h"
 #include "SlateCore/Public/Styling/SlateColor.h"
 #include "Templates/Casts.h"
 #include "Templates/SharedPointer.h"
@@ -39,14 +41,6 @@ namespace Metasound
 {
 	namespace Editor
 	{
-		/* Enums to use when grouping the blueprint members in the list panel. The order here will determine the order in the list */
-		enum class ENodeSection : uint8
-		{
-			NONE = 0,
-			INPUTS,
-			OUTPUTS
-		};
-
 		FName BuildChildPath(const FString& InBasePath, FName InPropertyName)
 		{
 			return FName(InBasePath + TEXT(".") + InPropertyName.ToString());
@@ -200,31 +194,6 @@ namespace Metasound
 			TArray<TWeakObjectPtr<UObject>> Objects;
 			DetailLayout.GetObjectsBeingCustomized(Objects);
 
-			const FName InterfacePropertyPath = BuildChildPath(GetMetadataRootClassPath(), GET_MEMBER_NAME_CHECKED(FMetasoundFrontendClass, Interface));
-			const FName InputsPropertyPath = BuildChildPath(InterfacePropertyPath, GET_MEMBER_NAME_CHECKED(FMetasoundFrontendClassInterface, Inputs));
-			const FName OutputsPropertyPath = BuildChildPath(InterfacePropertyPath, GET_MEMBER_NAME_CHECKED(FMetasoundFrontendClassInterface, Outputs));
-
-			// TODO: Move Fixed Arrays to GraphActionMenu for selection, ordering, and
-			// eventually ability to create input nodes independent of registered inputs.
-// 			SAssignNew(GraphActionMenu, SGraphActionMenu, false)
-// 				.OnGetFilterText(this, &FMetasoundDetailCustomization::GetFilterText)
-// 				.OnCreateWidgetForAction(this, &FMetasoundDetailCustomization::OnCreateWidgetForAction)
-// 				.OnCollectAllActions(this, &FMetasoundDetailCustomization::CollectAllActions)
-// 				.OnCollectStaticSections(this, &FMetasoundDetailCustomization::CollectStaticSections)
-// 				.OnActionDragged(this, &FMetasoundDetailCustomization::OnActionDragged)
-// 				.OnActionSelected(this, &FMetasoundDetailCustomization::OnGlobalActionSelected)
-// 				.OnActionDoubleClicked(this, &FMetasoundDetailCustomization::OnActionDoubleClicked)
-// 				.OnContextMenuOpening(this, &FMetasoundDetailCustomization::OnContextMenuOpening)
-// 				.OnCategoryTextCommitted(this, &FMetasoundDetailCustomization::OnCategoryNameCommitted)
-// 				.OnCanRenameSelectedAction(this, &FMetasoundDetailCustomization::CanRequestRenameOnActionNode)
-// 				.OnGetSectionTitle(this, &FMetasoundDetailCustomization::OnGetSectionTitle)
-// 				.OnGetSectionWidget(this, &FMetasoundDetailCustomization::OnGetSectionWidget)
-// 				.OnActionMatchesName(this, &FMetasoundDetailCustomization::HandleActionMatchesName)
-// 				.AlphaSortItems(false)
-// 				.UseSectionStyling(true);
-			
-			BuildIOFixedArray<FMetasoundFrontendClassInput>(DetailLayout, "Inputs", InputsPropertyPath);
-			BuildIOFixedArray<FMetasoundFrontendClassOutput>(DetailLayout, "Outputs", OutputsPropertyPath);
 
 			// Hack to hide parent structs for nested metadata properties
 			DetailLayout.HideCategory("CustomView");
@@ -247,34 +216,6 @@ namespace Metasound
 			DetailLayout.HideCategory("Subtitles");
 			DetailLayout.HideCategory("Voice Management");
 		}
-
-		FText FMetasoundDetailCustomization::OnGetSectionTitle(int32 InSectionID)
-		{
-			FText SeperatorTitle;
-			/* Setup an appropriate name for the section for this node */
-			switch (static_cast<ENodeSection>(InSectionID))
-			{
-				case ENodeSection::INPUTS:
-				{
-					SeperatorTitle = LOCTEXT("Inputs_Title", "Inputs");
-				}
-				break;
-
-				case ENodeSection::OUTPUTS:
-				{
-					SeperatorTitle = LOCTEXT("Outputs_Title", "Outputs");
-				}
-				break;
-
-				default:
-				{
-					SeperatorTitle = LOCTEXT("Missing_Title", "UNSET");
-				}
-				break;
-			}
-
-			return SeperatorTitle;
-		}
 	} // namespace Editor
 } // namespace Metasound
 #undef LOCTEXT_NAMESPACE
diff --git a/Engine/Plugins/Runtime/Metasound/Source/MetasoundEditor/Private/MetasoundDetailCustomization.h b/Engine/Plugins/Runtime/Metasound/Source/MetasoundEditor/Private/MetasoundDetailCustomization.h
--- a/Engine/Plugins/Runtime/Metasound/Source/MetasoundEditor/Private/MetasoundDetailCustomization.h
+++ b/Engine/Plugins/Runtime/Metasound/Source/MetasoundEditor/Private/MetasoundDetailCustomization.h
@@ -3,11 +3,13 @@
 
 #include "IDetailCustomization.h"
 #include "SGraphActionMenu.h"
+#include "Types/SlateEnums.h"
 #include "UObject/NameTypes.h"
 
 // Forward Declarations
 class FPropertyRestriction;
 class IDetailLayoutBuilder;
+struct FPointerEvent;
 
 namespace Metasound
 {
@@ -26,11 +28,9 @@ namespace Metasound
 			FName GetMetadataRootClassPath() const;
 			FName GetMetadataPropertyPath() const;
 
-			FText OnGetSectionTitle(int32 InSectionID);
-
 			FName DocumentPropertyName;
 
-// 			TSharedPtr<SGraphActionMenu> GraphActionMenu;
+			TSharedPtr<SGraphActionMenu> GraphActionMenu;
 		};
 	} // namespace Editor
 } // namespace Metasound
\ No newline at end of file
diff --git a/Engine/Plugins/Runtime/Metasound/Source/MetasoundEditor/Private/MetasoundEditor.cpp b/Engine/Plugins/Runtime/Metasound/Source/MetasoundEditor/Private/MetasoundEditor.cpp
--- a/Engine/Plugins/Runtime/Metasound/Source/MetasoundEditor/Private/MetasoundEditor.cpp
+++ b/Engine/Plugins/Runtime/Metasound/Source/MetasoundEditor/Private/MetasoundEditor.cpp
@@ -3,6 +3,7 @@
 #include "MetasoundEditor.h"
 
 #include "AudioDevice.h"
+#include "AudioMeterStyle.h"
 #include "Components/AudioComponent.h"
 #include "EdGraph/EdGraphNode.h"
 #include "EdGraphUtilities.h"
@@ -34,7 +35,10 @@
 #include "Widgets/Docking/SDockTab.h"
 #include "Widgets/Notifications/SNotificationList.h"
 #include "Widgets/SBoxPanel.h"
-#include "AudioMeterStyle.h"
+#include "Templates/SharedPointer.h"
+
+// This needs to be moved to public directory
+#include "../../GraphEditor/Private/GraphActionNode.h"
 
 #define LOCTEXT_NAMESPACE "MetasoundEditor"
 
@@ -42,6 +46,87 @@ namespace Metasound
 {
 	namespace Editor
 	{
+		static const TArray<FText> NodeSectionNames
+		{
+			LOCTEXT("NodeSectionName_Invalid", "INVALID"),
+			LOCTEXT("NodeSectionName_Inputs", "Inputs"),
+			LOCTEXT("NodeSectionName_Outputs", "Outputs")
+		};
+
+		class FMetasoundGraphNodeSchemaAction : public FEdGraphSchemaAction
+		{
+		public:
+			UEdGraph* Graph = nullptr;
+			FGuid NodeID;
+
+			FMetasoundGraphNodeSchemaAction()
+				: FEdGraphSchemaAction()
+			{
+			}
+
+			FMetasoundGraphNodeSchemaAction(FText InNodeCategory, FText InMenuDesc, FText InToolTip, const int32 InGrouping, const ENodeSection InSectionID)
+				: FEdGraphSchemaAction(MoveTemp(InNodeCategory), MoveTemp(InMenuDesc), MoveTemp(InToolTip), InGrouping, FText(), static_cast<int32>(InSectionID))
+			{
+			}
+
+			FMetasoundAssetBase& GetMetasoundAssetChecked() const
+			{
+				UObject* Object = CastChecked<UMetasoundEditorGraph>(Graph)->GetMetasound();
+				FMetasoundAssetBase* MetasoundAsset = IMetasoundUObjectRegistry::Get().GetObjectAsAssetBase(Object);
+				check(MetasoundAsset);
+				return *MetasoundAsset;
+			}
+
+			UMetasoundEditorGraphVariable* GetVariable() const
+			{
+				UMetasoundEditorGraph* MetasoundGraph = CastChecked<UMetasoundEditorGraph>(Graph);
+				return MetasoundGraph->FindVariable(NodeID);
+			}
+
+			Frontend::FGraphHandle GetGraphHandle() const
+			{
+				return GetMetasoundAssetChecked().GetRootGraphHandle();
+			}
+
+			Frontend::FNodeHandle GetNodeHandle() const
+			{
+				return GetGraphHandle()->GetNodeWithID(NodeID);
+			}
+
+			bool IsRequired() const
+			{
+				return GetNodeHandle()->IsRequired();
+			}
+
+			// FEdGraphSchemaAction interface
+			virtual bool IsParentable() const override
+			{
+				return true;
+			}
+
+			virtual void MovePersistentItemToCategory(const FText& NewCategoryName) override
+			{
+				checkNoEntry();
+			}
+
+			virtual int32 GetReorderIndexInContainer() const override
+			{
+				TArray<Frontend::FConstNodeHandle> InputHandles = GetGraphHandle()->GetConstInputNodes();
+				return InputHandles.IndexOfByPredicate([=](const Frontend::FConstNodeHandle& NodeHandle)
+				{
+					return NodeHandle->GetID() == NodeID;
+				});
+			}
+
+			virtual bool ReorderToBeforeAction(TSharedRef<FEdGraphSchemaAction> OtherAction) override
+			{
+				// TODO: Implement reordering
+				checkNoEntry();
+
+				return false;
+			}
+		};
+
 		void FEditor::RegisterTabSpawners(const TSharedRef<FTabManager>& InTabManager)
 		{
 			WorkspaceMenuCategory = InTabManager->AddLocalWorkspaceMenuCategory(LOCTEXT("WorkspaceMenu_MetasoundEditor", "Metasound Editor"));
@@ -65,9 +150,9 @@ namespace Metasound
 			.SetGroup(WorkspaceMenuCategoryRef)
 			.SetIcon(FSlateIcon(FEditorStyle::GetStyleSetName(), "LevelEditor.Tabs.Details"));
 
-			InTabManager->RegisterTabSpawner(TabFactory::Names::Metasound, FOnSpawnTab::CreateLambda([InMetasoundGeneral = MetasoundGeneral](const FSpawnTabArgs& Args)
+			InTabManager->RegisterTabSpawner(TabFactory::Names::Metasound, FOnSpawnTab::CreateLambda([InMetasoundMenu = MetasoundGeneralMenu](const FSpawnTabArgs& Args)
 			{
-				return TabFactory::CreateMetasoundTab(InMetasoundGeneral, Args);
+				return TabFactory::CreateMetasoundTab(InMetasoundMenu, Args);
 			}))
 			.SetDisplayName(LOCTEXT("MetasoundTab", "Metasound"))
 			.SetGroup(WorkspaceMenuCategoryRef)
@@ -104,6 +189,28 @@ namespace Metasound
 				Stop();
 			}
 
+			if (Metasound)
+			{
+				FMetasoundAssetBase* MetasoundAsset = IMetasoundUObjectRegistry::Get().GetObjectAsAssetBase(Metasound);
+				check(MetasoundAsset);
+
+				if(UMetasoundEditorGraph* Graph = Cast<UMetasoundEditorGraph>(MetasoundAsset->GetGraph()))
+				{
+					for (TPair<FGuid, FDelegateHandle>& Pair : NameChangeDelegateHandles)
+					{
+						if (UMetasoundEditorGraphInput* Input = Graph->FindInput(Pair.Key))
+						{
+							Input->NameChanged.Remove(Pair.Value);
+						}
+						if (UMetasoundEditorGraphOutput* Output = Graph->FindOutput(Pair.Key))
+						{
+							Output->NameChanged.Remove(Pair.Value);
+						}
+					}
+				}
+				NameChangeDelegateHandles.Reset();
+			}
+
 			if (MetasoundMeterAnalyzer.IsValid() && ResultsDelegateHandle.IsValid())
 			{
 				MetasoundMeterAnalyzer->OnLatestPerChannelMeterResultsNative.Remove(ResultsDelegateHandle);
@@ -219,7 +326,7 @@ namespace Metasound
 			if (MetasoundDetails.IsValid())
 			{
 				MetasoundDetails->SetObjects(SelectedObjects);
-				MetasoundDetails->HideFilterArea(SelectedObjects.IsEmpty() || SelectedObjects.Num() > 1);
+				MetasoundDetails->HideFilterArea(true);
 			}
 		}
 
@@ -261,6 +368,11 @@ namespace Metasound
 				MetasoundGraphEditor->NotifyGraphChanged();
 			}
 
+			if (MetasoundGeneralMenu.IsValid())
+			{
+				MetasoundGeneralMenu->RefreshAllActions(true /* bPreserveExpansion */);
+			}
+
 			FSlateApplication::Get().DismissAllMenus();
 
 			if (Metasound)
@@ -284,6 +396,7 @@ namespace Metasound
 					{
 						FGraphBuilder::SynchronizeGraph(*MetasoundObj);
 						MetasoundGraphEditor->NotifyGraphChanged();
+						MetasoundGeneralMenu->RefreshAllActions(true /* bPreserveExpansion */);
 					}
 				}
 			}
@@ -299,14 +412,34 @@ namespace Metasound
 
 			FPropertyEditorModule& PropertyModule = FModuleManager::LoadModuleChecked<FPropertyEditorModule>("PropertyEditor");
 
-			MetasoundGeneral = PropertyModule.CreateDetailView(Args);
-			MetasoundGeneral->SetObject(Metasound);
+			SAssignNew(MetasoundGeneralMenu, SGraphActionMenu, false)
+				.AlphaSortItems(false)
+// 				.OnActionDoubleClicked(this, &FEditor::OnActionDoubleClicked)
+				.OnActionDragged(this, &FEditor::OnActionDragged)
+				.OnActionMatchesName(this, &FEditor::HandleActionMatchesName)
+				.OnActionSelected(this, &FEditor::OnActionSelected)
+// 				.OnCategoryTextCommitted(this, &FEditor::OnCategoryNameCommitted)
+				.OnCollectAllActions(this, &FEditor::CollectAllActions)
+				.OnCollectStaticSections(this, &FEditor::CollectStaticSections)
+// 				.OnContextMenuOpening(this, &FEditor::OnContextMenuOpening)
+				.OnCreateWidgetForAction(this, &FEditor::OnCreateWidgetForAction)
+  				.OnCanRenameSelectedAction(this, &FEditor::CanRenameOnActionNode)
+				.OnGetFilterText(this, &FEditor::GetFilterText)
+				.OnGetSectionTitle(this, &FEditor::OnGetSectionTitle)
+				.OnGetSectionWidget(this, &FEditor::OnGetMenuSectionWidget)
+				.UseSectionStyling(true);
 
 			MetasoundDetails = PropertyModule.CreateDetailView(Args);
 
 			Palette = SNew(SMetasoundPalette);
 		}
 
+		// TODO: Tie in rename on GraphActionMenu.  For now, just renameable via field in inspector
+		bool FEditor::CanRenameOnActionNode(TWeakPtr<FGraphActionNode> InSelectedNode) const
+		{
+			return false;
+		}
+
 		void FEditor::OnMeterOutput(UMeterAnalyzer* InMeterAnalyzer, int32 ChannelIndex, const FMeterResults& InMeterResults)
 		{
 			if (InMeterAnalyzer && MetasoundMeter.IsValid())
@@ -405,6 +538,19 @@ namespace Metasound
 						ToolbarBuilder.AddToolBarButton(FEditorCommands::Get().Stop);
 					}
 					ToolbarBuilder.EndSection();
+
+					ToolbarBuilder.BeginSection("Utilities");
+					{
+						ToolbarBuilder.AddToolBarButton(
+							FEditorCommands::Get().EditMetasoundSettings,
+								NAME_None,
+								TAttribute<FText>(),
+								TAttribute<FText>(),
+								TAttribute<FSlateIcon>::Create([this]() { return GetSettingsImage(); }),
+								"EditMetasoundSettings"
+						);
+					}
+					ToolbarBuilder.EndSection();
 				})
 			);
 
@@ -417,6 +563,12 @@ namespace Metasound
 			return FSlateIcon("MetasoundStyle", IconName);
 		}
 
+		FSlateIcon FEditor::GetSettingsImage() const
+		{
+			const FName IconName = "MetasoundEditor.Settings";
+			return FSlateIcon("MetasoundStyle", IconName);
+		}
+
 		FSlateIcon FEditor::GetExportStatusImage() const
 		{
 			FName IconName = "MetasoundEditor.Export";
@@ -459,6 +611,14 @@ namespace Metasound
 			ToolkitCommands->MapAction(
 				FGenericCommands::Get().Redo,
 				FExecuteAction::CreateSP(this, &FEditor::RedoGraphAction));
+
+			ToolkitCommands->MapAction(
+				Commands.EditMetasoundSettings,
+				FExecuteAction::CreateSP(this, &FEditor::EditMetasoundSettings));
+
+			ToolkitCommands->MapAction(
+				FGenericCommands::Get().Delete,
+				FExecuteAction::CreateSP(this, &FEditor::DeleteSelected));
 		}
 
 		void FEditor::Import()
@@ -616,6 +776,23 @@ namespace Metasound
 			// TODO: Implement? Will we support single node playback?
 		}
 
+		void FEditor::EditMetasoundSettings()
+		{
+			if (MetasoundGeneralMenu.IsValid())
+			{
+				MetasoundGeneralMenu->SelectItemByName(FName());
+			}
+
+			if (MetasoundGraphEditor.IsValid())
+			{
+				bManuallyClearingGraphSelection = true;
+				MetasoundGraphEditor->ClearSelectionSet();
+				bManuallyClearingGraphSelection = false;
+			}
+
+			SetSelection({ Metasound });
+		}
+
 		void FEditor::SyncInBrowser()
 		{
 			TArray<UObject*> ObjectsToSync;
@@ -666,6 +843,9 @@ namespace Metasound
 				GraphEditorCommands->MapAction(FEditorCommands::Get().BrowserSync,
 					FExecuteAction::CreateSP(this, &FEditor::SyncInBrowser));
 
+				GraphEditorCommands->MapAction(FEditorCommands::Get().EditMetasoundSettings,
+					FExecuteAction::CreateSP(this, &FEditor::EditMetasoundSettings));
+
 				GraphEditorCommands->MapAction(FEditorCommands::Get().AddInput,
 					FExecuteAction::CreateSP(this, &FEditor::AddInput),
 					FCanExecuteAction::CreateSP(this, &FEditor::CanAddInput));
@@ -682,10 +862,6 @@ namespace Metasound
 				GraphEditorCommands->MapAction(FGenericCommands::Get().SelectAll,
 					FExecuteAction::CreateLambda([this]() { MetasoundGraphEditor->SelectAllNodes(); }));
 
-				GraphEditorCommands->MapAction(FGenericCommands::Get().Delete,
-					FExecuteAction::CreateSP(this, &FEditor::DeleteSelectedNodes),
-					FCanExecuteAction::CreateLambda([this]() { return CanDeleteNodes(); }));
-
 				GraphEditorCommands->MapAction(FGenericCommands::Get().Copy,
 					FExecuteAction::CreateSP(this, &FEditor::CopySelectedNodes),
 					FCanExecuteAction::CreateSP(this, &FEditor::CanCopyNodes));
@@ -753,18 +929,29 @@ namespace Metasound
 				.ShowGraphStateOverlay(false);
 		}
 
-		void FEditor::OnSelectedNodesChanged(const TSet<UObject*>& NewSelection)
+		void FEditor::OnSelectedNodesChanged(const TSet<UObject*>& InSelectedNodes)
 		{
 			TArray<UObject*> Selection;
-
-			if (!NewSelection.IsEmpty())
+			for (UObject* NodeObject : InSelectedNodes)
 			{
-				for (TSet<UObject*>::TConstIterator SetIt(NewSelection); SetIt; ++SetIt)
+				if (UMetasoundEditorGraphInputNode* InputNode = Cast<UMetasoundEditorGraphInputNode>(NodeObject))
 				{
-					Selection.Add(*SetIt);
+					Selection.Add(InputNode->Input);
+				}
+				else if (UMetasoundEditorGraphOutputNode* OutputNode = Cast<UMetasoundEditorGraphOutputNode>(NodeObject))
+				{
+					Selection.Add(OutputNode->Output);
+				}
+				else
+				{
+					Selection.Add(NodeObject);
 				}
 			}
 
+			if (MetasoundGeneralMenu.IsValid() && !bManuallyClearingGraphSelection)
+			{
+				MetasoundGeneralMenu->SelectItemByName(FName());
+			}
 			SetSelection(Selection);
 		}
 
@@ -772,12 +959,126 @@ namespace Metasound
 		{
 			if (NodeBeingChanged)
 			{
-				const FScopedTransaction Transaction(LOCTEXT("RenameNode", "Rename Node"));
+				const FScopedTransaction Transaction(TEXT(""), LOCTEXT("RenameNode", "Rename Node"), NodeBeingChanged);
 				NodeBeingChanged->Modify();
 				NodeBeingChanged->OnRenameNode(NewText.ToString());
 			}
 		}
 
+		void FEditor::DeleteSelected()
+		{
+			using namespace Frontend;
+
+			if (MetasoundGeneralMenu.IsValid())
+			{
+				TArray<TSharedPtr<FEdGraphSchemaAction>> Actions;
+				MetasoundGeneralMenu->GetSelectedActions(Actions);
+
+				if (!Actions.IsEmpty())
+				{
+					const FScopedTransaction Transaction(TEXT(""), LOCTEXT("MetasoundEditorDeleteSelectedNode", "Delete Metasound Variable"), Metasound);
+
+					UMetasoundEditorGraph* Graph = CastChecked<UMetasoundEditorGraph>(MetasoundGraphEditor->GetCurrentGraph());
+					Graph->Modify();
+
+					TSharedPtr<FMetasoundGraphNodeSchemaAction> ActionToDelete;
+					for (const TSharedPtr<FEdGraphSchemaAction>& Action : Actions)
+					{
+						TSharedPtr<FMetasoundGraphNodeSchemaAction> MetasoundAction = StaticCastSharedPtr<FMetasoundGraphNodeSchemaAction>(Action);
+						if (MetasoundAction.IsValid())
+						{
+							Frontend::FNodeHandle NodeHandle = MetasoundAction->GetNodeHandle();
+							if (ensure(NodeHandle->IsValid()))
+							{
+								if (NodeHandle->IsRequired())
+								{
+									if (MetasoundGraphEditor.IsValid())
+									{
+										FNotificationInfo Info(LOCTEXT("CannotDelete_RequiredVariable", "Delete failed: Input/Output is required."));
+										Info.bFireAndForget = true;
+										Info.bUseSuccessFailIcons = false;
+										Info.ExpireDuration = 5.0f;
+
+										MetasoundGraphEditor->AddNotification(Info, false /* bSuccess */);
+									}
+								}
+								else
+								{
+									ActionToDelete = MetasoundAction;
+									break;
+								}
+							}
+						}
+					}
+
+					if (ActionToDelete.IsValid())
+					{
+						UMetasoundEditorGraphVariable* Variable = ActionToDelete->GetVariable();
+
+						FNodeHandle VariableHandle = Variable->GetNodeHandle();
+						const FGuid IDToDelete = VariableHandle->GetID();
+
+						auto GetNameToSelect = [IDToDelete](const TArray<FConstNodeHandle>& Handles)
+						{
+							int32 IndexToDelete = -1;
+							bool bGetNextValidName = false;
+							for (int32 i = 0; i < Handles.Num(); ++i)
+							{
+								if (Handles[i]->GetID() == IDToDelete)
+								{
+									IndexToDelete = i;
+									break;
+								}
+							}
+
+							if (IndexToDelete >= 0)
+							{
+								for (int32 i = IndexToDelete + 1; i < Handles.Num(); ++i)
+								{
+									if (Handles[i]->GetNodeStyle().Display.Visibility != EMetasoundFrontendNodeStyleDisplayVisibility::Hidden)
+									{
+										return FName(*(Handles[i]->GetDisplayName().ToString()));
+									}
+								}
+
+								for (int32 i = IndexToDelete - 1; i >= 0; --i)
+								{
+									if (Handles[i]->GetNodeStyle().Display.Visibility != EMetasoundFrontendNodeStyleDisplayVisibility::Hidden)
+									{
+										return FName(*(Handles[i]->GetDisplayName().ToString()));
+									}
+								}
+							}
+
+							return FName();
+						};
+
+						const TArray<FConstNodeHandle> InputHandles = VariableHandle->GetOwningGraph()->GetConstInputNodes();
+						FName NameToSelect = GetNameToSelect(InputHandles);
+						int32 SectionId = static_cast<int32>(ENodeSection::Inputs);
+						if (NameToSelect.IsNone())
+						{
+							SectionId = static_cast<int32>(ENodeSection::Outputs);
+							const TArray<FConstNodeHandle> OutputHandles = VariableHandle->GetOwningGraph()->GetConstOutputNodes();
+							NameToSelect = GetNameToSelect(OutputHandles);
+						}
+
+						FGraphBuilder::DeleteVariableNodeHandle(*Variable);
+						Graph->RemoveVariable(*Variable);
+						MetasoundGeneralMenu->RefreshAllActions(true /* bPreserveExpansion */);
+
+						if (!NameToSelect.IsNone())
+						{
+							MetasoundGeneralMenu->SelectItemByName(NameToSelect, ESelectInfo::Direct, SectionId);
+						}
+					}
+					return;
+				}
+			}
+
+			DeleteSelectedNodes();
+		}
+
 		void FEditor::DeleteSelectedNodes()
 		{
 			using namespace Metasound::Frontend;
@@ -785,22 +1086,31 @@ namespace Metasound
 			const FScopedTransaction Transaction(LOCTEXT("MetasoundEditorDeleteSelectedNode", "Delete Selected Metasound Node(s)"));
 
 			UMetasoundEditorGraph* Graph = CastChecked<UMetasoundEditorGraph>(MetasoundGraphEditor->GetCurrentGraph());
-			Graph->Modify();
 
 			const FGraphPanelSelectionSet SelectedNodes = MetasoundGraphEditor->GetSelectedNodes();
 			MetasoundGraphEditor->ClearSelectionSet();
 
-			for (FGraphPanelSelectionSet::TConstIterator NodeIt(SelectedNodes); NodeIt; ++NodeIt)
+			bool bModified = false;
+			for (UObject* NodeObj : SelectedNodes)
 			{
 				// Some nodes may not be metasound nodes (ex. comments and perhaps aliases eventually), but can be safely deleted.
-				if (UEdGraphNode* Node = Cast<UEdGraphNode>(*NodeIt))
+				if (UEdGraphNode* Node = Cast<UEdGraphNode>(NodeObj))
 				{
-					if (!FGraphBuilder::DeleteNode(*Node))
+					if (FGraphBuilder::DeleteNode(*Node))
+					{
+						bModified = true;
+					}
+					else
 					{
 						MetasoundGraphEditor->SetNodeSelection(Node, true /* bSelect */);
 					}
 				}
 			}
+
+			if (bModified)
+			{
+				Graph->Modify();
+			}
 		}
 
 		void FEditor::CutSelectedNodes()
@@ -920,9 +1230,8 @@ namespace Metasound
 				return;
 			}
 
-			const FScopedTransaction Transaction(NSLOCTEXT("UnrealEd", "MetasoundEditorPaste", "Paste Metasound Node(s)"));
+			const FScopedTransaction Transaction(LOCTEXT("MetasoundEditorPaste", "Paste Metasound Node(s)"));
 			Graph->Modify();
-			Metasound->Modify();
 
 			// Clear the selection set (newly pasted stuff will be selected)
 			MetasoundGraphEditor->ClearSelectionSet();
@@ -932,24 +1241,53 @@ namespace Metasound
 
 			ClipboardContent.Empty();
 
+			UMetasoundEditorGraph* MetasoundGraph = CastChecked<UMetasoundEditorGraph>(MetasoundAsset->GetGraph());
+
 			TArray<UEdGraphNode*> NodesToRemove;
 			for (UEdGraphNode* GraphNode : PastedGraphNodes)
 			{
 				GraphNode->CreateNewGuid();
-				if (UMetasoundEditorGraphNode* MetasoundGraphNode = Cast<UMetasoundEditorGraphNode>(GraphNode))
+				if (UMetasoundEditorGraphExternalNode* ExternalNode = Cast<UMetasoundEditorGraphExternalNode>(GraphNode))
+				{
+					FNodeHandle NewHandle = FGraphBuilder::AddNodeHandle(*Metasound, *ExternalNode);
+					if (!NewHandle->IsValid())
+					{
+						NodesToRemove.Add(GraphNode);
+					}
+				}
+				else if (UMetasoundEditorGraphInputNode* InputNode = Cast<UMetasoundEditorGraphInputNode>(GraphNode))
+				{
+					if (!MetasoundGraph->Inputs.Contains(InputNode->Input))
+					{
+						NodesToRemove.Add(GraphNode);
+					}
+				}
+				else if (UMetasoundEditorGraphOutputNode* OutputNode = Cast<UMetasoundEditorGraphOutputNode>(GraphNode))
 				{
-					// TODO: Support copying inputs after adding replacing InputGraphNodes with InputReferenceGraphNodes
-					FNodeHandle NewHandle = INodeController::GetInvalidHandle();
-					if (MetasoundGraphNode->IsA(UMetasoundEditorGraphExternalNode::StaticClass()))
+					if (!MetasoundGraph->Outputs.Contains(OutputNode->Output))
 					{
-						NewHandle = FGraphBuilder::AddNodeHandle(*Metasound, *MetasoundGraphNode);
+						NodesToRemove.Add(GraphNode);
 					}
 
-					if (!NewHandle->IsValid())
+					auto NodeMatches = [OutputNodeID = OutputNode->GetNodeID()](const TObjectPtr<UEdGraphNode>& EdNode)
+					{
+						if (UMetasoundEditorGraphOutputNode* OutputNode = Cast<UMetasoundEditorGraphOutputNode>(EdNode))
+						{
+							return OutputNodeID == OutputNode->GetNodeID();
+						}
+						return false;
+					};
+
+					// Can only have one output reference node
+					if (MetasoundGraph->Nodes.ContainsByPredicate(NodeMatches))
 					{
 						NodesToRemove.Add(GraphNode);
 					}
 				}
+				else
+				{
+					checkNoEntry();
+				}
 			}
 
 			// Remove nodes failed to import before attempting to connect/place
@@ -988,7 +1326,7 @@ namespace Metasound
 					{
 						const FVector2D NewNodeLocation = FVector2D(GraphNode->NodePosX, GraphNode->NodePosY);
 						FMetasoundFrontendNodeStyle NodeStyle = NodeHandle->GetNodeStyle();
-						NodeStyle.Display.Location = NewNodeLocation;
+						NodeStyle.Display.Locations.FindOrAdd(NodeHandle->GetID()) = NewNodeLocation;
 						NodeHandle->SetNodeStyle(NodeStyle);
 					}
 				}
@@ -1071,6 +1409,343 @@ namespace Metasound
 			check(GEditor);
 			GEditor->RedoTransaction();
 		}
+
+		void FEditor::OnInputNameChanged(FGuid InNodeID)
+		{
+			if (!MetasoundGeneralMenu.IsValid())
+			{
+				return;
+			}
+
+			TArray<TSharedPtr<FEdGraphSchemaAction>> SelectedActions;
+			MetasoundGeneralMenu->GetSelectedActions(SelectedActions);
+			MetasoundGeneralMenu->RefreshAllActions(/* bPreserveExpansion */ true);
+
+			for(const TSharedPtr<FEdGraphSchemaAction>& Action : SelectedActions)
+			{
+				TSharedPtr<FMetasoundGraphNodeSchemaAction> MetasoundAction = StaticCastSharedPtr<FMetasoundGraphNodeSchemaAction>(Action);
+				if (MetasoundAction.IsValid())
+				{
+					Frontend::FNodeHandle NodeHandle = MetasoundAction->GetNodeHandle();
+					if (InNodeID == NodeHandle->GetID())
+					{
+						TArray<Frontend::FConstOutputHandle> OutputHandles = NodeHandle->GetConstOutputs();
+						if (ensure(OutputHandles.Num() == 1))
+						{
+							const FName ActionName = *OutputHandles[0]->GetDisplayName().ToString();
+							MetasoundGeneralMenu->SelectItemByName(ActionName, ESelectInfo::Direct, Action->GetSectionID());
+						}
+						break;
+					}
+				}
+			}
+		}
+
+		void FEditor::OnOutputNameChanged(FGuid InNodeID)
+		{
+			if (!MetasoundGeneralMenu.IsValid())
+			{
+				return;
+			}
+
+			TArray<TSharedPtr<FEdGraphSchemaAction>> SelectedActions;
+			MetasoundGeneralMenu->GetSelectedActions(SelectedActions);
+			MetasoundGeneralMenu->RefreshAllActions(/* bPreserveExpansion */ true);
+
+			for (const TSharedPtr<FEdGraphSchemaAction>& Action : SelectedActions)
+			{
+				TSharedPtr<FMetasoundGraphNodeSchemaAction> MetasoundAction = StaticCastSharedPtr<FMetasoundGraphNodeSchemaAction>(Action);
+				if (MetasoundAction.IsValid())
+				{
+					Frontend::FNodeHandle NodeHandle = MetasoundAction->GetNodeHandle();
+					if (InNodeID == NodeHandle->GetID())
+					{
+						TArray<Frontend::FConstInputHandle> InputHandles = NodeHandle->GetConstInputs();
+						if (ensure(InputHandles.Num() == 1))
+						{
+							const FName ActionName = *InputHandles[0]->GetDisplayName().ToString();
+							MetasoundGeneralMenu->SelectItemByName(ActionName, ESelectInfo::Direct, Action->GetSectionID());
+						}
+						break;
+					}
+				}
+			}
+		}
+
+		void FEditor::CollectAllActions(FGraphActionListBuilderBase& OutAllActions)
+		{
+			FMetasoundAssetBase* MetasoundAsset = IMetasoundUObjectRegistry::Get().GetObjectAsAssetBase(Metasound);
+			check(MetasoundAsset);
+
+			UMetasoundEditorGraph* EditorGraph = Cast<UMetasoundEditorGraph>(MetasoundAsset->GetGraph());
+
+			MetasoundAsset->GetRootGraphHandle()->IterateConstNodes([this, EditorGraph, MetasoundAsset, ActionList = &OutAllActions](const Frontend::FConstNodeHandle& Input)
+			{
+				if (Input->GetNodeStyle().Display.Visibility == EMetasoundFrontendNodeStyleDisplayVisibility::Hidden)
+				{
+					return;
+				}
+
+				const FText Tooltip = Input->GetDescription();
+				const FText InputName = Input->GetDisplayName();
+				TSharedPtr<FMetasoundGraphNodeSchemaAction> NewFuncAction = MakeShared<FMetasoundGraphNodeSchemaAction>(FText::GetEmpty(), InputName, Tooltip, 1, ENodeSection::Inputs);
+				NewFuncAction->Graph = MetasoundAsset->GetGraph();
+				ActionList->AddAction(NewFuncAction);
+
+				const FGuid NodeID = Input->GetID();
+				NewFuncAction->NodeID = NodeID;
+				UMetasoundEditorGraphInput * EdGraphInput = EditorGraph->FindInput(NodeID);
+				if (ensure(EdGraphInput))
+				{
+					if (FDelegateHandle* NameChangeDelegate = NameChangeDelegateHandles.Find(NodeID))
+					{
+						EdGraphInput->NameChanged.Remove(*NameChangeDelegate);
+					}
+					NameChangeDelegateHandles.FindOrAdd(NodeID) = EdGraphInput->NameChanged.AddSP(this, &FEditor::OnInputNameChanged);
+				}
+			}, EMetasoundFrontendClassType::Input);
+
+			MetasoundAsset->GetRootGraphHandle()->IterateConstNodes([this, EditorGraph, MetasoundAsset, ActionList = &OutAllActions](const Frontend::FConstNodeHandle& Output)
+			{
+				const FText Tooltip = Output->GetDescription();
+				const FText OutputName = Output->GetDisplayName();
+
+				TSharedPtr<FMetasoundGraphNodeSchemaAction> NewFuncAction = MakeShared<FMetasoundGraphNodeSchemaAction>(FText::GetEmpty(), OutputName, Tooltip, 1, ENodeSection::Outputs);
+				NewFuncAction->Graph = MetasoundAsset->GetGraph();
+				NewFuncAction->NodeID = Output->GetID();
+				ActionList->AddAction(NewFuncAction);
+
+				const FGuid NodeID = Output->GetID();
+				NewFuncAction->NodeID = NodeID;
+				UMetasoundEditorGraphOutput* EdGraphOutput = EditorGraph->FindOutput(NodeID);
+				if (ensure(EdGraphOutput))
+				{
+					if (FDelegateHandle* NameChangeDelegate = NameChangeDelegateHandles.Find(NodeID))
+					{
+						EdGraphOutput->NameChanged.Remove(*NameChangeDelegate);
+					}
+					NameChangeDelegateHandles.FindOrAdd(NodeID) = EdGraphOutput->NameChanged.AddSP(this, &FEditor::OnOutputNameChanged);
+				}
+			}, EMetasoundFrontendClassType::Output);
+		}
+
+		void FEditor::CollectStaticSections(TArray<int32>& StaticSectionIDs)
+		{
+			for (int32 i = 0; i < static_cast<int32>(ENodeSection::COUNT); ++i)
+			{
+				if (static_cast<ENodeSection>(i) != ENodeSection::None)
+				{
+					StaticSectionIDs.Add(i);
+				}
+			}
+		}
+
+		bool FEditor::HandleActionMatchesName(FEdGraphSchemaAction* InAction, const FName& InName) const
+		{
+			return false;
+		}
+
+		// TODO: Implement
+		FReply FEditor::OnActionDragged(const TArray<TSharedPtr<FEdGraphSchemaAction>>& InActions, const FPointerEvent& MouseEvent)
+		{
+			return FReply::Unhandled();
+		}
+
+		void FEditor::OnActionSelected(const TArray<TSharedPtr<FEdGraphSchemaAction>>& InActions, ESelectInfo::Type InSelectionType)
+		{
+			if (InSelectionType == ESelectInfo::OnMouseClick || InSelectionType == ESelectInfo::OnKeyPress || InSelectionType == ESelectInfo::OnNavigation || InActions.IsEmpty())
+			{
+				TArray<UObject*> SelectedObjects;
+				for (const TSharedPtr<FEdGraphSchemaAction>& Action : InActions)
+				{
+					TSharedPtr<FMetasoundGraphNodeSchemaAction> MetasoundNodeAction = StaticCastSharedPtr<FMetasoundGraphNodeSchemaAction>(Action);
+					if (MetasoundNodeAction.IsValid())
+					{
+						SelectedObjects.Add(MetasoundNodeAction->GetVariable());
+					}
+				}
+
+				if (InSelectionType != ESelectInfo::Direct && !InActions.IsEmpty())
+				{
+					if (MetasoundGraphEditor.IsValid())
+					{
+						bManuallyClearingGraphSelection = true;
+						MetasoundGraphEditor->ClearSelectionSet();
+						bManuallyClearingGraphSelection = false;
+					}
+					SetSelection(SelectedObjects);
+				}
+			}
+		}
+
+		// TODO: Add ability to filter inputs/outputs in "Metasound" Tab
+		FText FEditor::GetFilterText() const
+		{
+			return FText::GetEmpty();
+		}
+
+		TSharedRef<SWidget> FEditor::OnCreateWidgetForAction(FCreateWidgetForActionData* const InCreateData)
+		{
+			return SNew(SGraphPaletteItem, InCreateData);
+		}
+
+		FText FEditor::GetSectionTitle(ENodeSection InSection) const
+		{
+			const int32 SectionIndex = static_cast<int32>(InSection);
+			if (ensure(NodeSectionNames.IsValidIndex(SectionIndex)))
+			{
+				return NodeSectionNames[SectionIndex];
+			}
+
+			return FText::GetEmpty();
+		}
+
+		FText FEditor::OnGetSectionTitle(int32 InSectionID)
+		{
+			if (ensure(NodeSectionNames.IsValidIndex(InSectionID)))
+			{
+				return NodeSectionNames[InSectionID];
+			}
+
+			return FText::GetEmpty();
+		}
+
+		TSharedRef<SWidget> FEditor::OnGetMenuSectionWidget(TSharedRef<SWidget> RowWidget, int32 InSectionID)
+		{
+			TWeakPtr<SWidget> WeakRowWidget = RowWidget;
+
+			FText AddNewText;
+			FName MetaDataTag;
+
+			switch (static_cast<ENodeSection>(InSectionID))
+			{
+				case ENodeSection::Inputs:
+				{
+					AddNewText = LOCTEXT("AddNewInput", "Input");
+					MetaDataTag = "AddNewInput";
+				}
+				break;
+
+				case ENodeSection::Outputs:
+				{
+					AddNewText = LOCTEXT("AddNewOutput", "Output");
+					MetaDataTag = "AddNewOutput";
+				}
+				break;
+
+				default:
+					return SNullWidget::NullWidget;
+			}
+
+			return CreateAddInputButton(InSectionID, AddNewText, MetaDataTag);
+		}
+
+		bool FEditor::CanAddNewElementToSection(int32 InSectionID) const
+		{
+			return true;
+		}
+
+		FReply FEditor::OnAddButtonClickedOnSection(int32 InSectionID)
+		{
+			if (!Metasound)
+			{
+				return FReply::Unhandled();
+			}
+
+			FMetasoundAssetBase* MetasoundAsset = IMetasoundUObjectRegistry::Get().GetObjectAsAssetBase(Metasound);
+			check(MetasoundAsset);
+			UMetasoundEditorGraph* Graph = Cast<UMetasoundEditorGraph>(MetasoundAsset->GetGraph());
+			if (!ensure(Graph))
+			{
+				return FReply::Unhandled();
+			}
+
+			FName NameToSelect;
+			switch (static_cast<ENodeSection>(InSectionID))
+			{
+				case ENodeSection::Inputs:
+				{
+					const FScopedTransaction Transaction(LOCTEXT("AddInputNode", "Add Metasound Input"));
+					Metasound->Modify();
+
+					const FName DataTypeName = "Bool";
+					FString NewNodeName = FGraphBuilder::GenerateUniqueInputName(*Metasound);
+					NameToSelect = FName(*NewNodeName);
+
+					Frontend::FNodeHandle NodeHandle = FGraphBuilder::AddInputNodeHandle(*Metasound, NewNodeName, DataTypeName, FText::GetEmpty());
+					if (ensure(NodeHandle->IsValid()))
+					{
+						TObjectPtr<UMetasoundEditorGraphInput> Input = Graph->FindOrAddInput(NodeHandle);
+						if (ensure(Input))
+						{
+							FGuid NodeID = NodeHandle->GetID();
+							if (FDelegateHandle* NameChangeDelegate = NameChangeDelegateHandles.Find(NodeID))
+							{
+								Input->NameChanged.Remove(*NameChangeDelegate);
+							}
+							NameChangeDelegateHandles.FindOrAdd(NodeID) = Input->NameChanged.AddSP(this, &FEditor::OnInputNameChanged);
+						}
+					}
+				}
+				break;
+
+				case ENodeSection::Outputs:
+				{
+					const FScopedTransaction Transaction(TEXT(""), LOCTEXT("AddOutputNode", "Add Metasound Output"), Metasound);
+					Metasound->Modify();
+
+					const FName DataTypeName = "Bool";
+					FString NewNodeName = FGraphBuilder::GenerateUniqueOutputName(*Metasound);
+					NameToSelect = FName(*NewNodeName);
+
+					Frontend::FNodeHandle NodeHandle = FGraphBuilder::AddOutputNodeHandle(*Metasound, NewNodeName, DataTypeName, FText::GetEmpty());
+					if (ensure(NodeHandle->IsValid()))
+					{
+						TObjectPtr<UMetasoundEditorGraphOutput> Output = Graph->FindOrAddOutput(NodeHandle);
+						if (ensure(Output))
+						{
+							FGuid NodeID = NodeHandle->GetID();
+							if (FDelegateHandle* NameChangeDelegate = NameChangeDelegateHandles.Find(NodeID))
+							{
+								Output->NameChanged.Remove(*NameChangeDelegate);
+							}
+							NameChangeDelegateHandles.FindOrAdd(NodeID) = Output->NameChanged.AddSP(this, &FEditor::OnInputNameChanged);
+						}
+					}
+				}
+				break;
+
+				default:
+				return FReply::Unhandled();
+			}
+
+			if (MetasoundGeneralMenu.IsValid())
+			{
+				MetasoundGeneralMenu->RefreshAllActions(/* bPreserveExpansion */ true);
+				if (NameToSelect.IsNone())
+				{
+					MetasoundGeneralMenu->SelectItemByName(NameToSelect);
+				}
+			}
+			return FReply::Handled();
+		}
+
+		TSharedRef<SWidget> FEditor::CreateAddInputButton(int32 InSectionID, FText AddNewText, FName MetaDataTag)
+		{
+			return
+				SNew(SButton)
+				.ButtonStyle(FEditorStyle::Get(), "SimpleButton")
+				.OnClicked(this, &FEditor::OnAddButtonClickedOnSection, InSectionID)
+				.IsEnabled(this, &FEditor::CanAddNewElementToSection, InSectionID)
+				.ContentPadding(FMargin(1, 0))
+				.AddMetaData<FTagMetaData>(FTagMetaData(MetaDataTag))
+				.ToolTipText(AddNewText)
+				[
+					SNew(SImage)
+					.Image(FAppStyle::Get().GetBrush("Icons.PlusCircle"))
+				.ColorAndOpacity(FSlateColor::UseForeground())
+				];
+		}
 	}
 }
 #undef LOCTEXT_NAMESPACE
diff --git a/Engine/Plugins/Runtime/Metasound/Source/MetasoundEditor/Private/MetasoundEditor.h b/Engine/Plugins/Runtime/Metasound/Source/MetasoundEditor/Private/MetasoundEditor.h
--- a/Engine/Plugins/Runtime/Metasound/Source/MetasoundEditor/Private/MetasoundEditor.h
+++ b/Engine/Plugins/Runtime/Metasound/Source/MetasoundEditor/Private/MetasoundEditor.h
@@ -22,6 +22,9 @@
 #include "Toolkits/AssetEditorToolkit.h"
 #include "UObject/StrongObjectPtr.h"
 #include "Widgets/SPanel.h"
+#include "WorkflowOrientedApp/WorkflowTabFactory.h"
+#include "WorkflowOrientedApp/WorkflowTabManager.h"
+
 
 // Forward Declarations
 class IDetailsView;
@@ -35,6 +38,7 @@ class SVerticalBox;
 class UEdGraphNode;
 class UMetasound;
 
+struct FGraphActionNode;
 struct FMeterResults;
 struct FPropertyChangedEvent;
 
@@ -42,6 +46,17 @@ namespace Metasound
 {
 	namespace Editor
 	{
+		/* Enums to use when grouping the members in the list panel. Enum order dictates visible order. */
+		enum class ENodeSection : uint8
+		{
+			None,
+			Inputs,
+			Outputs,
+
+			COUNT
+		};
+
+
 		class FEditor : public IMetasoundEditor, public FGCObject, public FNotifyHook, public FEditorUndoClient
 		{
 		public:
@@ -88,9 +103,28 @@ namespace Metasound
 				return MetasoundGraphEditor->GetSelectedNodes().Num();
 			}
 
+			void OnInputNameChanged(FGuid InNodeID);
+			void OnOutputNameChanged(FGuid InNodeID);
+
 			void OnMeterOutput(UMeterAnalyzer* InMeterAnalyzer, int32 ChannelIndex, const FMeterResults& MeterResults);
 
 		protected:
+			// Callbacks for action tree
+			bool CanRenameOnActionNode(TWeakPtr<FGraphActionNode> InSelectedNode) const;
+			bool CanAddNewElementToSection(int32 InSectionID) const;
+			void CollectAllActions(FGraphActionListBuilderBase& OutAllActions);
+			void CollectStaticSections(TArray<int32>& StaticSectionIDs);
+			TSharedRef<SWidget> CreateAddInputButton(int32 InSectionID, FText AddNewText, FName MetaDataTag);
+			FText GetFilterText() const;
+			bool HandleActionMatchesName(FEdGraphSchemaAction* InAction, const FName& InName) const;
+			FReply OnActionDragged(const TArray<TSharedPtr<FEdGraphSchemaAction>>& InActions, const FPointerEvent& MouseEvent);
+			void OnActionSelected(const TArray<TSharedPtr<FEdGraphSchemaAction>>& InActions, ESelectInfo::Type InSelectionType);
+			FReply OnAddButtonClickedOnSection(int32 InSectionID);
+			TSharedRef<SWidget> OnGetMenuSectionWidget(TSharedRef<SWidget> RowWidget, int32 InSectionID);
+			FText GetSectionTitle(ENodeSection InSection) const;
+			FText OnGetSectionTitle(int32 InSectionID);
+			TSharedRef<SWidget> OnCreateWidgetForAction(struct FCreateWidgetForActionData* const InCreateData);
+
 			/** Called when the selection changes in the GraphEditor */
 			void OnSelectedNodesChanged(const TSet<UObject*>& NewSelection);
 
@@ -103,6 +137,9 @@ namespace Metasound
 			 */
 			void OnNodeTitleCommitted(const FText& NewText, ETextCommit::Type CommitInfo, UEdGraphNode* NodeBeingChanged);
 
+			/** Deletes from the Metasound Menu (i.e. input or output) if in focus, or the currently selected nodes if the graph editor is in focus. */
+			void DeleteSelected();
+
 			/** Delete the currently selected nodes */
 			void DeleteSelectedNodes();
 
@@ -147,6 +184,8 @@ namespace Metasound
 
 			FSlateIcon GetExportStatusImage() const;
 
+			FSlateIcon GetSettingsImage() const;
+
 			/** Toolbar command methods */
 			void Import();
 			void Export();
@@ -166,8 +205,12 @@ namespace Metasound
 			/** Sync the content browser to the current selection of nodes */
 			void SyncInBrowser();
 
+			/** Show the metasound object's settings in the Inspector */
+			void EditMetasoundSettings();
+
 			/** Add an input to the currently selected node */
 			void AddInput();
+
 			/** Whether we can add an input to the currently selected node */
 			bool CanAddInput() const;
 
@@ -193,8 +236,8 @@ namespace Metasound
 			/** Details tab */
 			TSharedPtr<IDetailsView> MetasoundDetails;
 
-			/** General tab */
-			TSharedPtr<IDetailsView> MetasoundGeneral;
+			/** Metasound General info widget */
+			TSharedPtr<SGraphActionMenu> MetasoundGeneralMenu;
 
 			/** Metasound Output Meter widget */
 			TSharedPtr<SAudioMeter> MetasoundMeter;
@@ -222,11 +265,17 @@ namespace Metasound
 			/** Handle for results delegate for metasound meter analyzer. */
 			FDelegateHandle ResultsDelegateHandle;
 
+			TMap<FGuid, FDelegateHandle> NameChangeDelegateHandles;
+
 			/** Whether or not metasound being edited is valid */
 			bool bPassedValidation = true;
 
 			/** Clipboard content used when pasting from clipboard to avoid deserializing twice */
 			FString ClipboardContent;
+
+			/** Boolean state for when selection change handle should not respond due to selection state
+			  * being manually applied in code */
+			bool bManuallyClearingGraphSelection = false;
 		};
 	} // namespace Editor
 } // namespace Metasound
diff --git a/Engine/Plugins/Runtime/Metasound/Source/MetasoundEditor/Private/MetasoundEditorCommands.cpp b/Engine/Plugins/Runtime/Metasound/Source/MetasoundEditor/Private/MetasoundEditorCommands.cpp
--- a/Engine/Plugins/Runtime/Metasound/Source/MetasoundEditor/Private/MetasoundEditorCommands.cpp
+++ b/Engine/Plugins/Runtime/Metasound/Source/MetasoundEditor/Private/MetasoundEditorCommands.cpp
@@ -19,6 +19,8 @@ namespace Metasound
 			UI_COMMAND(BrowserSync, "Browse", "Selects the Metasound in the content browser. If referencing Metasound nodes are selected, selects referenced assets instead.", EUserInterfaceActionType::Button, FInputChord());
 			UI_COMMAND(AddInput, "Add Input", "Adds an input to the node", EUserInterfaceActionType::Button, FInputChord());
 			UI_COMMAND(DeleteInput, "Delete Input", "Removes an input from the node", EUserInterfaceActionType::Button, FInputChord());
+
+			UI_COMMAND(EditMetasoundSettings, "Settings", "Edit Metasound Settings", EUserInterfaceActionType::Button, FInputChord());
 		}
 	} // namespace Editor
 } // namespace Metasound
diff --git a/Engine/Plugins/Runtime/Metasound/Source/MetasoundEditor/Private/MetasoundEditorGraph.cpp b/Engine/Plugins/Runtime/Metasound/Source/MetasoundEditor/Private/MetasoundEditorGraph.cpp
--- a/Engine/Plugins/Runtime/Metasound/Source/MetasoundEditor/Private/MetasoundEditorGraph.cpp
+++ b/Engine/Plugins/Runtime/Metasound/Source/MetasoundEditor/Private/MetasoundEditorGraph.cpp
@@ -6,35 +6,250 @@
 #include "MetasoundEditorGraphNode.h"
 #include "MetasoundEditorGraphInputNodes.h"
 #include "MetasoundEditorModule.h"
+#include "MetasoundUObjectRegistry.h"
 
+#define LOCTEXT_NAMESPACE "MetasoundEditor"
 
-UMetasoundEditorGraphInputNode* UMetasoundEditorGraph::CreateInputNode(EMetasoundFrontendLiteralType LiteralType, UClass* InLiteralObjectClass, bool bInSelectNewNode)
+
+TArray<UMetasoundEditorGraphNode*> UMetasoundEditorGraphVariable::GetNodes() const
+{
+	TArray<UMetasoundEditorGraphNode*> Nodes;
+
+	UMetasoundEditorGraph* Graph = Cast<UMetasoundEditorGraph>(GetOuter());
+	if (ensure(Graph))
+	{
+		Graph->GetNodesOfClassEx<UMetasoundEditorGraphNode>(Nodes);
+		for (int32 i = Nodes.Num() -1; i >= 0; --i)
+		{
+			UMetasoundEditorGraphNode* Node = Nodes[i];
+			if (Node && Node->GetNodeID() != NodeID)
+			{
+				Nodes.RemoveAtSwap(i, 1, false /* bAllowShrinking */);
+			}
+		}
+	}
+
+	return Nodes;
+}
+
+void UMetasoundEditorGraphVariable::SetDataType(FName InNewType)
 {
 	using namespace Metasound::Editor;
+	using namespace Metasound::Frontend;
+
+	UMetasoundEditorGraph* Graph = Cast<UMetasoundEditorGraph>(GetOuter());
+	if (!ensure(Graph))
+	{
+		return;
+	}
+
+	const FScopedTransaction Transaction(LOCTEXT("SetVariableDataType", "Set Metasound Variable Type"));
+	Graph->Modify();
+
+	// 1. Cache current editor input node reference positions & delete nodes.
+	TArray<UMetasoundEditorGraphNode*> InputNodes = GetNodes();
+	TArray<FVector2D> NodeLocations;
+	for (UMetasoundEditorGraphNode* Node : InputNodes)
+	{
+		if (ensure(Node))
+		{
+			NodeLocations.Add(FVector2D(Node->NodePosX, Node->NodePosY));
+		}
+	}
+
+	// 2. Cache the old version's Frontend data.
+	FNodeHandle NodeHandle = GetNodeHandle();
+	const FString NodeName = NodeHandle->GetNodeName();
+	const FText NodeDisplayName = NodeHandle->GetDisplayName();
+
+	// 3. Delete the Frontend variable
+	FGraphBuilder::DeleteVariableNodeHandle(*this);
+
+	// 4. Add the new input node with the same identifier data but new datatype.
+	UObject& Metasound = Graph->GetMetasoundChecked();
+	FNodeHandle NewNodeHandle = AddNodeHandle(NodeName, NodeDisplayName, InNewType);
+	if (!ensure(NewNodeHandle->IsValid()))
+	{
+		return;
+	}
+
+	ClassName = NewNodeHandle->GetClassName();
+	NodeID = NewNodeHandle->GetID();
+	TypeName = InNewType;
+
+	// 5. Report data type changed immediately after assignment to child
+	// class(es) so underlying data can be fixed-up prior to recreating
+	// referencing nodes.
+	OnDataTypeChanged();
+
+	// 6. Create new node references in the same locations as the old locations
+	for (FVector2D Location : NodeLocations)
+	{
+		FGraphBuilder::AddNode(Metasound, NewNodeHandle, Location, false /* bInSelectNewNode */);
+	}
+
+	// Notify now that the node has a new ID (doing so before creating & syncing Frontend Node &
+	// EdGraph variable can result in refreshing editors while in a desync'ed state)
+	NameChanged.Broadcast(NodeID);
+}
+
+Metasound::Frontend::FNodeHandle UMetasoundEditorGraphVariable::GetNodeHandle() const
+{
+	using namespace Metasound;
+
+	UObject* Object = CastChecked<UMetasoundEditorGraph>(GetOuter())->GetMetasound();
+	if (!ensure(Object))
+	{
+		return Frontend::INodeController::GetInvalidHandle();
+	}
+
+	FMetasoundAssetBase* MetasoundAsset = IMetasoundUObjectRegistry::Get().GetObjectAsAssetBase(Object);
+	check(MetasoundAsset);
+
+	return MetasoundAsset->GetRootGraphHandle()->GetNodeWithID(NodeID);
+}
+
+Metasound::Frontend::FConstNodeHandle UMetasoundEditorGraphVariable::GetConstNodeHandle() const
+{
+	return GetNodeHandle();
+}
 
-	IMetasoundEditorModule & EditorModule = FModuleManager::GetModuleChecked<IMetasoundEditorModule>("MetasoundEditor");
+void UMetasoundEditorGraphInput::UpdateDocumentInput() const
+{
+	using namespace Metasound;
+	using namespace Metasound::Frontend;
+
+	UObject* Object = CastChecked<UMetasoundEditorGraph>(GetOuter())->GetMetasound();
+	if (!ensure(Object))
+	{
+		return;
+	}
 
-	TSubclassOf<UMetasoundEditorGraphInputNode> NodeClass = EditorModule.FindNodeInputClass(LiteralType);
+	if (!ensure(Literal))
+	{
+		return;
+	}
 
-	if (ensure(NodeClass))
+	FMetasoundAssetBase* MetasoundAsset = IMetasoundUObjectRegistry::Get().GetObjectAsAssetBase(Object);
+	check(MetasoundAsset);
+
+	FGraphHandle GraphHandle = MetasoundAsset->GetRootGraphHandle();
+	FNodeHandle NodeHandle = GraphHandle->GetNodeWithID(NodeID);
+	const FString& NodeName = NodeHandle->GetNodeName();
+
+	const FGuid VertexID = GraphHandle->GetVertexIDForInputVertex(NodeName);
+	GraphHandle->SetDefaultInput(VertexID, Literal->GetDefault());
+}
+
+void UMetasoundEditorGraphInput::UpdatePreviewInstance(const Metasound::FVertexKey& InParameterName, IAudioInstanceTransmitter& InInstanceTransmitter) const
+{
+	if (ensure(Literal))
 	{
-		UMetasoundEditorGraphInputNode* NewGraphNode = CastChecked<UMetasoundEditorGraphInputNode>(CreateNode(NodeClass, bInSelectNewNode));
+		Literal->UpdatePreviewInstance(InParameterName, InInstanceTransmitter);
+	}
+}
+
+Metasound::Frontend::FNodeHandle UMetasoundEditorGraphInput::AddNodeHandle(const FString& InNodeName, const FText& InNodeDisplayName, FName InDataType)
+{
+	using namespace Metasound::Editor;
+	using namespace Metasound::Frontend;
+
+	UMetasoundEditorGraph* Graph = Cast<UMetasoundEditorGraph>(GetOuter());
+	if (!ensure(Graph))
+	{
+		return Super::AddNodeHandle(InNodeName, InNodeDisplayName, InDataType);
+	}
+
+	UObject& Metasound = Graph->GetMetasoundChecked();
+	Metasound::Frontend::FNodeHandle NewNodeHandle = FGraphBuilder::AddInputNodeHandle(Metasound, InNodeName, InDataType, FText::GetEmpty());
+	NewNodeHandle->SetDisplayName(InNodeDisplayName);
+	return NewNodeHandle;
+}
 
-		NewGraphNode->CreateNewGuid();
-		NewGraphNode->PostPlacedNewNode();
+void UMetasoundEditorGraphInput::OnLiteralChanged()
+{
+	UpdateDocumentInput();
 
-		if (NewGraphNode->Pins.IsEmpty())
+	if (UMetasoundEditorGraph* MetasoundGraph = Cast<UMetasoundEditorGraph>(GetOuter()))
+	{
+		if (IAudioInstanceTransmitter* Transmitter = MetasoundGraph->GetMetasoundInstanceTransmitter())
 		{
-			NewGraphNode->AllocateDefaultPins();
+			// TODO: fix how identifying the parameter to update is determined. It should not be done
+			// with a "DisplayName" but rather the vertex Guid.
+			Metasound::Frontend::FConstNodeHandle NodeHandle = GetConstNodeHandle();
+			Metasound::FVertexKey VertexKey = Metasound::FVertexKey(NodeHandle->GetDisplayName().ToString());
+			UpdatePreviewInstance(VertexKey, *Transmitter);
+		}
+	}
+}
+
+void UMetasoundEditorGraphInput::OnDataTypeChanged()
+{
+	using namespace Metasound;
+	using namespace Metasound::Editor;
+	using namespace Metasound::Frontend;
+
+	IMetasoundEditorModule& EditorModule = FModuleManager::GetModuleChecked<IMetasoundEditorModule>("MetasoundEditor");
+	const FEditorDataType& EditorDataType = EditorModule.FindDataType(TypeName);
+	const EMetasoundFrontendLiteralType LiteralType = static_cast<EMetasoundFrontendLiteralType>(EditorDataType.RegistryInfo.PreferredLiteralType);
+
+	TSubclassOf<UMetasoundEditorGraphInputLiteral> InputLiteralClass = EditorModule.FindInputLiteralClass(LiteralType);
+	if (!InputLiteralClass)
+	{
+		InputLiteralClass = UMetasoundEditorGraphInputLiteral::StaticClass();
+	}
+	Literal = NewObject<UMetasoundEditorGraphInputLiteral>(this, InputLiteralClass);
+}
+
+Metasound::Frontend::FNodeHandle UMetasoundEditorGraphOutput::AddNodeHandle(const FString& InNodeName, const FText& InNodeDisplayName, FName InDataType)
+{
+	using namespace Metasound::Editor;
+	using namespace Metasound::Frontend;
+
+	UMetasoundEditorGraph* Graph = Cast<UMetasoundEditorGraph>(GetOuter());
+	if (!ensure(Graph))
+	{
+		return Super::AddNodeHandle(InNodeName, InNodeDisplayName, InDataType);
+	}
+
+	UObject& Metasound = Graph->GetMetasoundChecked();
+	Metasound::Frontend::FNodeHandle NewNodeHandle = FGraphBuilder::AddOutputNodeHandle(Metasound, InNodeName, InDataType, FText::GetEmpty());
+	NewNodeHandle->SetDisplayName(InNodeDisplayName);
+	return NewNodeHandle;
+}
+
+UMetasoundEditorGraphInputNode* UMetasoundEditorGraph::CreateInputNode(Metasound::Frontend::FNodeHandle InNodeHandle, bool bInSelectNewNode)
+{
+	using namespace Metasound::Editor;
+	using namespace Metasound::Frontend;
+
+	TArray<FConstOutputHandle> NodeOutputs = InNodeHandle->GetConstOutputs();
+	if (!ensure(!NodeOutputs.IsEmpty()))
+	{
+		return nullptr;
+	}
+
+	UEdGraphNode* NewEdGraphNode = CreateNode(UMetasoundEditorGraphInputNode::StaticClass(), bInSelectNewNode);
+	UMetasoundEditorGraphInputNode* NewInputNode = CastChecked<UMetasoundEditorGraphInputNode>(NewEdGraphNode);
+	if (ensure(NewInputNode))
+	{
+		NewInputNode->CreateNewGuid();
+		NewInputNode->PostPlacedNewNode();
+
+		NewInputNode->Input = FindOrAddInput(InNodeHandle);
+
+		if (NewInputNode->Pins.IsEmpty())
+		{
+			NewInputNode->AllocateDefaultPins();
 		}
 
 		PostEditChange();
 		MarkPackageDirty();
 
-		return NewGraphNode;
+		return NewInputNode;
 	}
-
-	return nullptr;
+	
+	return nullptr;	
 }
 
 UObject* UMetasoundEditorGraph::GetMetasound() const
@@ -71,3 +286,149 @@ void UMetasoundEditorGraph::Synchronize()
 	}
 }
 
+UMetasoundEditorGraphInput* UMetasoundEditorGraph::FindInput(FGuid InNodeID) const
+{
+	const TObjectPtr<UMetasoundEditorGraphInput>* Input = Inputs.FindByPredicate([InNodeID](const TObjectPtr<UMetasoundEditorGraphInput>& Input)
+	{
+		return Input->NodeID == InNodeID;
+	});
+	return Input ? Input->Get() : nullptr;
+}
+
+UMetasoundEditorGraphInput* UMetasoundEditorGraph::FindOrAddInput(Metasound::Frontend::FNodeHandle InNodeHandle)
+{
+	using namespace Metasound::Editor;
+	using namespace Metasound::Frontend;
+
+	FGraphHandle Graph = InNodeHandle->GetOwningGraph();
+
+	FName TypeName;
+	FGuid VertexID;
+
+	ensure(InNodeHandle->GetNumInputs() == 1);
+	InNodeHandle->IterateConstInputs([InGraph = &Graph, InTypeName = &TypeName, InVertexID = &VertexID](FConstInputHandle InputHandle)
+		{
+			*InTypeName = InputHandle->GetDataType();
+			*InVertexID = (*InGraph)->GetVertexIDForInputVertex(InputHandle->GetName());
+		});
+
+	const FGuid NodeID = InNodeHandle->GetID();
+	if (TObjectPtr<UMetasoundEditorGraphInput> Input = FindInput(NodeID))
+	{
+		ensure(Input->TypeName == TypeName);
+		return Input;
+	}
+
+	if (UMetasoundEditorGraphInput* NewInput = NewObject<UMetasoundEditorGraphInput>(this))
+	{
+		NewInput->NodeID = NodeID;
+		NewInput->ClassName = InNodeHandle->GetClassName();
+		NewInput->TypeName = TypeName;
+
+		FMetasoundFrontendLiteral DefaultLiteral = Graph->GetDefaultInput(VertexID);
+		EMetasoundFrontendLiteralType LiteralType = DefaultLiteral.GetType();
+		IMetasoundEditorModule& EditorModule = FModuleManager::GetModuleChecked<IMetasoundEditorModule>("MetasoundEditor");
+		TSubclassOf<UMetasoundEditorGraphInputLiteral> InputLiteralClass = EditorModule.FindInputLiteralClass(LiteralType);
+		NewInput->Literal = NewObject<UMetasoundEditorGraphInputLiteral>(this, InputLiteralClass);
+
+		Inputs.Add(NewInput);
+		return NewInput;
+	}
+
+	checkNoEntry();
+	return nullptr;
+}
+
+UMetasoundEditorGraphOutput* UMetasoundEditorGraph::FindOutput(FGuid InNodeID) const
+{
+	const TObjectPtr<UMetasoundEditorGraphOutput>* Output = Outputs.FindByPredicate([InNodeID](const TObjectPtr<UMetasoundEditorGraphOutput>& Output)
+	{
+		return Output->NodeID == InNodeID;
+	});
+	return Output ? Output->Get() : nullptr;
+}
+
+UMetasoundEditorGraphOutput* UMetasoundEditorGraph::FindOrAddOutput(Metasound::Frontend::FNodeHandle InNodeHandle)
+{
+	using namespace Metasound::Editor;
+	using namespace Metasound::Frontend;
+
+	FGraphHandle Graph = InNodeHandle->GetOwningGraph();
+
+	FName TypeName;
+	FGuid VertexID;
+
+	ensure(InNodeHandle->GetNumOutputs() == 1);
+	InNodeHandle->IterateConstOutputs([InGraph = &Graph, InTypeName = &TypeName, InVertexID = &VertexID](FConstOutputHandle OutputHandle)
+	{
+		*InTypeName = OutputHandle->GetDataType();
+		*InVertexID = (*InGraph)->GetVertexIDForInputVertex(OutputHandle->GetName());
+	});
+
+	const FGuid NodeID = InNodeHandle->GetID();
+	if (TObjectPtr<UMetasoundEditorGraphOutput> Output = FindOutput(NodeID))
+	{
+		ensure(Output->TypeName == TypeName);
+		return Output;
+	}
+
+	if (UMetasoundEditorGraphOutput* NewOutput = NewObject<UMetasoundEditorGraphOutput>(this))
+	{
+		NewOutput->NodeID = NodeID;
+		NewOutput->ClassName = InNodeHandle->GetClassName();
+		NewOutput->TypeName = TypeName;
+		Outputs.Add(NewOutput);
+		return NewOutput;
+	}
+
+	checkNoEntry();
+	return nullptr;
+}
+
+UMetasoundEditorGraphVariable* UMetasoundEditorGraph::FindVariable(FGuid InNodeID) const
+{
+	if (UMetasoundEditorGraphOutput* Output = FindOutput(InNodeID))
+	{
+		return Output;
+	}
+
+	return FindInput(InNodeID);
+}
+
+void UMetasoundEditorGraph::IterateInputs(TUniqueFunction<void(UMetasoundEditorGraphInput&)> InFunction) const
+{
+	for (UMetasoundEditorGraphInput* Input : Inputs)
+	{
+		if (ensure(Input))
+		{
+			InFunction(*Input);
+		}
+	}
+}
+
+void UMetasoundEditorGraph::IterateOutputs(TUniqueFunction<void(UMetasoundEditorGraphOutput&)> InFunction) const
+{
+	for (UMetasoundEditorGraphOutput* Output : Outputs)
+	{
+		if (ensure(Output))
+		{
+			InFunction(*Output);
+		}
+	}
+}
+
+bool UMetasoundEditorGraph::RemoveVariable(UMetasoundEditorGraphVariable& InVariable)
+{
+	if (UMetasoundEditorGraphInput* Input = Cast<UMetasoundEditorGraphInput>(&InVariable))
+	{
+		return Inputs.Remove(Input) > 0;
+	}
+
+	if (UMetasoundEditorGraphOutput* Output = Cast<UMetasoundEditorGraphOutput>(&InVariable))
+	{
+		return Outputs.Remove(Output) > 0;
+	}
+
+	return false;
+}
+#undef LOCTEXT_NAMESPACE
diff --git a/Engine/Plugins/Runtime/Metasound/Source/MetasoundEditor/Private/MetasoundEditorGraphBuilder.cpp b/Engine/Plugins/Runtime/Metasound/Source/MetasoundEditor/Private/MetasoundEditorGraphBuilder.cpp
--- a/Engine/Plugins/Runtime/Metasound/Source/MetasoundEditor/Private/MetasoundEditorGraphBuilder.cpp
+++ b/Engine/Plugins/Runtime/Metasound/Source/MetasoundEditor/Private/MetasoundEditorGraphBuilder.cpp
@@ -62,7 +62,7 @@ namespace Metasound
 			}
 		}
 
-		UMetasoundEditorGraphExternalNode* FGraphBuilder::AddExternalNode(UObject& InMetasound, Frontend::FNodeHandle& InNodeHandle, bool bInSelectNewNode)
+		UMetasoundEditorGraphExternalNode* FGraphBuilder::AddExternalNode(UObject& InMetasound, Frontend::FNodeHandle& InNodeHandle, FVector2D InLocation, bool bInSelectNewNode)
 		{
 			UMetasoundEditorGraphExternalNode* NewGraphNode = nullptr;
 			if (!ensure(InNodeHandle->GetClassType() == EMetasoundFrontendClassType::External))
@@ -81,10 +81,18 @@ namespace Metasound
 			NodeCreator.Finalize();
 			InitGraphNode(InNodeHandle, NewGraphNode, InMetasound);
 
+			SynchronizeNodeLocation(InLocation, InNodeHandle, *NewGraphNode);
+
 			return NewGraphNode;
 		}
 
-		UMetasoundEditorGraphOutputNode* FGraphBuilder::AddOutputNode(UObject& InMetasound, Frontend::FNodeHandle& InNodeHandle, bool bInSelectNewNode)
+		UMetasoundEditorGraphExternalNode* FGraphBuilder::AddExternalNode(UObject& InMetasound, const Frontend::FNodeClassInfo& InClassInfo, FVector2D InLocation, bool bInSelectNewNode)
+		{
+			Frontend::FNodeHandle NodeHandle = AddExternalNodeHandle(InMetasound, InClassInfo);
+			return AddExternalNode(InMetasound, NodeHandle, InLocation, bInSelectNewNode);
+		}
+
+		UMetasoundEditorGraphOutputNode* FGraphBuilder::AddOutputNode(UObject& InMetasound, Frontend::FNodeHandle& InNodeHandle, FVector2D InLocation, bool bInSelectNewNode)
 		{
 			using namespace Frontend;
 
@@ -94,23 +102,20 @@ namespace Metasound
 				return nullptr;
 			}
 
-			TArray<FConstInputHandle> Inputs = InNodeHandle->GetConstInputs();
-			if (!ensure(!Inputs.IsEmpty()))
-			{
-				return nullptr;
-			}
-
 			FMetasoundAssetBase* MetasoundAsset = IMetasoundUObjectRegistry::Get().GetObjectAsAssetBase(&InMetasound);
 			check(MetasoundAsset);
 			UEdGraph& Graph = MetasoundAsset->GetGraphChecked();
 			FGraphNodeCreator<UMetasoundEditorGraphOutputNode> NodeCreator(Graph);
+
 			NewGraphNode = NodeCreator.CreateNode(bInSelectNewNode);
+			UMetasoundEditorGraph* MetasoundGraph = CastChecked<UMetasoundEditorGraph>(&Graph);
+			NewGraphNode->Output = MetasoundGraph->FindOrAddOutput(InNodeHandle);
 
 			NodeCreator.Finalize();
 			InitGraphNode(InNodeHandle, NewGraphNode, InMetasound);
 
-			FConstInputHandle InputHandle = Inputs[0];
-			NewGraphNode->OutputTypeName = InputHandle->GetDataType();
+			SynchronizeNodeLocation(InLocation, InNodeHandle, *NewGraphNode);
+
 			return NewGraphNode;
 		}
 
@@ -119,10 +124,7 @@ namespace Metasound
 			const FMetasoundFrontendNodeStyle& Style = InNodeHandle->GetNodeStyle();
 
 			NewGraphNode->CreateNewGuid();
-
-			NewGraphNode->NodeID = InNodeHandle->GetID();
-			NewGraphNode->NodePosX = Style.Display.Location.X;
-			NewGraphNode->NodePosY = Style.Display.Location.Y;
+			NewGraphNode->SetNodeID(InNodeHandle->GetID());
 
 			RebuildNodePins(*NewGraphNode);
 
@@ -130,29 +132,12 @@ namespace Metasound
 			InMetasound.MarkPackageDirty();
 		}
 
-		UMetasoundEditorGraphExternalNode* FGraphBuilder::AddExternalNode(UObject& InMetasound, const Frontend::FNodeClassInfo& InClassInfo, const FMetasoundFrontendNodeStyle& InNodeStyle, bool bInSelectNewNode)
-		{
-			Frontend::FNodeHandle NodeHandle = AddExternalNodeHandle(InMetasound, InClassInfo, InNodeStyle);
-			return AddExternalNode(InMetasound, NodeHandle, bInSelectNewNode);
-		}
-
-		Frontend::FNodeHandle FGraphBuilder::AddExternalNodeHandle(UObject& InMetasound, const Frontend::FNodeClassInfo& InClassInfo, const FMetasoundFrontendNodeStyle& InNodeStyle)
+		Frontend::FNodeHandle FGraphBuilder::AddExternalNodeHandle(UObject& InMetasound, const Frontend::FNodeClassInfo& InClassInfo)
 		{
 			FMetasoundAssetBase* MetasoundAsset = IMetasoundUObjectRegistry::Get().GetObjectAsAssetBase(&InMetasound);
 			check(MetasoundAsset);
 
-			Frontend::FNodeHandle NewNode = MetasoundAsset->GetRootGraphHandle()->AddNode(InClassInfo);
-			NewNode->SetNodeStyle(InNodeStyle);
-			return NewNode;
-		}
-
-		FString FGraphBuilder::GetDataTypeDisplayName(const FName& InDataTypeName)
-		{
-			FString CategoryString = InDataTypeName.ToString();
-			int32 Index = 0;
-			CategoryString.FindLastChar(':', Index);
-
-			return CategoryString.RightChop(Index + 1);
+			return MetasoundAsset->GetRootGraphHandle()->AddNode(InClassInfo);
 		}
 
 		TArray<FString> FGraphBuilder::GetDataTypeNameCategories(const FName& InDataTypeName)
@@ -171,9 +156,9 @@ namespace Metasound
 			return Categories;
 		}
 
-		FString FGraphBuilder::GenerateUniqueInputName(const UObject& InMetasound, const FName InBaseName)
+		FString FGraphBuilder::GenerateUniqueInputName(const UObject& InMetasound)
 		{
-			FString NameBase = GetDataTypeDisplayName(InBaseName);
+			FString NameBase = TEXT("Input");
 			const FMetasoundAssetBase* MetasoundAsset = IMetasoundUObjectRegistry::Get().GetObjectAsAssetBase(&InMetasound);
 			check(MetasoundAsset);
 
@@ -189,9 +174,9 @@ namespace Metasound
 			return NewInputName;
 		}
 
-		FString FGraphBuilder::GenerateUniqueOutputName(const UObject& InMetasound, const FName InBaseName)
+		FString FGraphBuilder::GenerateUniqueOutputName(const UObject& InMetasound)
 		{
-			FString NameBase = GetDataTypeDisplayName(InBaseName);
+			FString NameBase = TEXT("Output");
 			const FMetasoundAssetBase* MetasoundAsset = IMetasoundUObjectRegistry::Get().GetObjectAsAssetBase(&InMetasound);
 			check(MetasoundAsset);
 
@@ -263,7 +248,17 @@ namespace Metasound
 			return GetOutputHandleFromPin(InPin);
 		}
 
-		UMetasoundEditorGraphInputNode* FGraphBuilder::AddInputNode(UObject& InMetasound, Frontend::FNodeHandle InNodeHandle, bool bInSelectNewNode)
+		void FGraphBuilder::SynchronizeNodeLocation(FVector2D InLocation, Frontend::FNodeHandle InNodeHandle, UMetasoundEditorGraphNode& InNode)
+		{
+			InNode.NodePosX = InLocation.X;
+			InNode.NodePosY = InLocation.Y;
+
+			FMetasoundFrontendNodeStyle Style = InNodeHandle->GetNodeStyle();
+			Style.Display.Locations.FindOrAdd(InNode.NodeGuid) = InLocation;
+			InNodeHandle->SetNodeStyle(Style);
+		}
+
+		UMetasoundEditorGraphInputNode* FGraphBuilder::AddInputNode(UObject& InMetasound, Frontend::FNodeHandle InNodeHandle, FVector2D InLocation, bool bInSelectNewNode)
 		{
 			using namespace Frontend;
 
@@ -276,24 +271,10 @@ namespace Metasound
 				return nullptr;
 			}
 
-
-			TArray<FConstOutputHandle> Outputs = InNodeHandle->GetConstOutputs();
-			if (!ensure(!Outputs.IsEmpty()))
-			{
-				return nullptr;
-			}
-
-			FConstOutputHandle OutputHandle = Outputs[0];
-			FGraphHandle Graph = InNodeHandle->GetOwningGraph();
-			FGuid VertexID = Graph->GetVertexIDForInputVertex(OutputHandle->GetName());
-			FMetasoundFrontendLiteral DefaultLiteral = Graph->GetDefaultInput(VertexID);
-			EMetasoundFrontendLiteralType LiteralType = DefaultLiteral.GetType();
-			UMetasoundEditorGraphInputNode* NewGraphNode = MetasoundGraph->CreateInputNode(LiteralType, nullptr, bInSelectNewNode);
+			UMetasoundEditorGraphInputNode* NewGraphNode = MetasoundGraph->CreateInputNode(InNodeHandle, bInSelectNewNode);
 			if (ensure(NewGraphNode))
 			{
-				NewGraphNode->ClassName = InNodeHandle->GetClassName();
-				NewGraphNode->NodeID = InNodeHandle->GetID();
-				NewGraphNode->InputTypeName = OutputHandle->GetDataType();
+				SynchronizeNodeLocation(InLocation, InNodeHandle, *NewGraphNode);
 
 				RebuildNodePins(*NewGraphNode);
 				return NewGraphNode;
@@ -302,17 +283,14 @@ namespace Metasound
 			return nullptr;
 		}
 
-		UMetasoundEditorGraphInputNode* FGraphBuilder::AddInputNode(UObject& InMetasound, const FString& InName, const FName InTypeName, const FMetasoundFrontendNodeStyle& InNodeStyle, const FText& InToolTip, bool bInSelectNewNode)
+		UMetasoundEditorGraphInputNode* FGraphBuilder::AddInputNode(UObject& InMetasound, const FString& InName, const FName InTypeName, const FText& InToolTip, FVector2D InLocation, EMetasoundFrontendNodeStyleDisplayVisibility Visibility, bool bInSelectNewNode)
 		{
 			using namespace Frontend;
 
-			FNodeHandle NodeHandle = AddInputNodeHandle(InMetasound, InName, InTypeName, InNodeStyle, InToolTip);
-			UMetasoundEditorGraphInputNode* NewGraphNode = AddInputNode(InMetasound, NodeHandle, bInSelectNewNode);
+			FNodeHandle NodeHandle = AddInputNodeHandle(InMetasound, InName, InTypeName, InToolTip, Visibility);
+			UMetasoundEditorGraphInputNode* NewGraphNode = AddInputNode(InMetasound, NodeHandle, InLocation, bInSelectNewNode);
 			if (ensure(NewGraphNode))
 			{
-				const FMetasoundFrontendNodeStyle& Style = NodeHandle->GetNodeStyle();
-				NewGraphNode->NodePosX = Style.Display.Location.X;
-				NewGraphNode->NodePosY = Style.Display.Location.Y;
 				return NewGraphNode;
 			}
 			else
@@ -344,13 +322,9 @@ namespace Metasound
 			FOutputHandle OutputHandle = InputHandle->GetCurrentlyConnectedOutput();
 			if (!OutputHandle->IsValid())
 			{
-				FMetasoundFrontendNodeStyle Style;
-				Style.Display.Visibility = EMetasoundFrontendNodeStyleDisplayVisibility::Hidden;
-
 				const FName TypeName = InputHandle->GetDataType();
-				const FString NewInputName = GenerateUniqueInputName(InMetasound, "LiteralInput");
-				Frontend::FNodeHandle NewInputNode = AddInputNodeHandle(InMetasound, NewInputName, TypeName, Style, FText::GetEmpty(), &PinDataTypeDefaultLiteral);
-				NewInputNode->SetNodeStyle(Style);
+				const FString NewInputName = GenerateUniqueInputName(InMetasound);
+				Frontend::FNodeHandle NewInputNode = AddInputNodeHandle(InMetasound, NewInputName, TypeName, FText::GetEmpty(), EMetasoundFrontendNodeStyleDisplayVisibility::Hidden, &PinDataTypeDefaultLiteral);
 
 				const FMetasoundFrontendLiteral* DefaultLiteral = InputHandle->GetDefaultLiteral();
 				Frontend::FGraphHandle GraphHandle = InNodeHandle->GetOwningGraph();
@@ -482,30 +456,38 @@ namespace Metasound
 
 		Frontend::FNodeHandle FGraphBuilder::AddNodeHandle(UObject& InMetasound, UMetasoundEditorGraphNode& InGraphNode)
 		{
-			FMetasoundFrontendNodeStyle Style;
-			Style.Display.Location = FVector2D(InGraphNode.NodePosX, InGraphNode.NodePosY);
-
- 			if (UMetasoundEditorGraphInputNode* InputNode = Cast<UMetasoundEditorGraphInputNode>(&InGraphNode))
- 			{
- 				const TArray<UEdGraphPin*>& Pins = InGraphNode.GetAllPins();
- 				const UEdGraphPin* Pin = Pins.IsEmpty() ? nullptr : Pins[0];
- 				if (ensure(Pin) && ensure(Pin->Direction == EGPD_Output))
- 				{
- 					return FGraphBuilder::AddInputNodeHandle(InMetasound, Pin->GetName(), InputNode->InputTypeName, Style, InGraphNode.GetTooltipText());
- 				}
- 			}
-
- 			if (UMetasoundEditorGraphOutputNode* OutputNode = Cast<UMetasoundEditorGraphOutputNode>(&InGraphNode))
- 			{
- 				const TArray<UEdGraphPin*>& Pins = InGraphNode.GetAllPins();
- 				const UEdGraphPin* Pin = Pins.IsEmpty() ? nullptr : Pins[0];
- 				if (ensure(Pin) && ensure(Pin->Direction == EGPD_Input))
- 				{
- 					return FGraphBuilder::AddOutputNodeHandle(InMetasound, Pin->GetName(), OutputNode->OutputTypeName, Style, InGraphNode.GetTooltipText());
- 				}
- 			}
-
-			if (UMetasoundEditorGraphExternalNode* ExternalNode = Cast<UMetasoundEditorGraphExternalNode>(&InGraphNode))
+			using namespace Frontend;
+
+			FNodeHandle NodeHandle = INodeController::GetInvalidHandle();
+			if (UMetasoundEditorGraphInputNode* InputNode = Cast<UMetasoundEditorGraphInputNode>(&InGraphNode))
+			{
+				const TArray<UEdGraphPin*>& Pins = InGraphNode.GetAllPins();
+				const UEdGraphPin* Pin = Pins.IsEmpty() ? nullptr : Pins[0];
+				if (ensure(Pin) && ensure(Pin->Direction == EGPD_Output))
+				{
+					UMetasoundEditorGraphInput* Input = InputNode->Input;
+					if (ensure(Input))
+					{
+						NodeHandle = FGraphBuilder::AddInputNodeHandle(InMetasound, Pin->GetName(), Input->TypeName, InGraphNode.GetTooltipText());
+					}
+				}
+			}
+
+			else if (UMetasoundEditorGraphOutputNode* OutputNode = Cast<UMetasoundEditorGraphOutputNode>(&InGraphNode))
+			{
+				const TArray<UEdGraphPin*>& Pins = InGraphNode.GetAllPins();
+				const UEdGraphPin* Pin = Pins.IsEmpty() ? nullptr : Pins[0];
+				if (ensure(Pin) && ensure(Pin->Direction == EGPD_Input))
+				{
+					UMetasoundEditorGraphOutput* Output = OutputNode->Output;
+					if (ensure(Output))
+					{
+						NodeHandle = FGraphBuilder::AddOutputNodeHandle(InMetasound, Pin->GetName(), Output->TypeName, InGraphNode.GetTooltipText());
+					}
+				}
+			}
+
+			else if (UMetasoundEditorGraphExternalNode* ExternalNode = Cast<UMetasoundEditorGraphExternalNode>(&InGraphNode))
 			{
 				FFrontendQuery Query;
 				const FFrontendQuerySelectionView Result = Query
@@ -523,20 +505,24 @@ namespace Metasound
 						check(MetasoundAsset);
 
 						Frontend::FNodeHandle NewNode = MetasoundAsset->GetRootGraphHandle()->AddNode(FrontendClass->Metadata);
-						NewNode->SetNodeStyle(Style);
+						ExternalNode->SetNodeID(NewNode->GetID());
 
-						InGraphNode.NodeID = NewNode->GetID();
-
-						return NewNode;
+						NodeHandle = NewNode;
 					}
 				}
 			}
 
-			// Cast failed so no handle created
-			return Frontend::INodeController::GetInvalidHandle();
+			if (NodeHandle->IsValid())
+			{
+				FMetasoundFrontendNodeStyle Style = NodeHandle->GetNodeStyle();
+				Style.Display.Locations.Add(InGraphNode.NodeGuid, FVector2D(InGraphNode.NodePosX, InGraphNode.NodePosY));
+				NodeHandle->SetNodeStyle(Style);
+			}
+
+			return NodeHandle;
 		}
 
-		Frontend::FNodeHandle FGraphBuilder::AddInputNodeHandle(UObject& InMetasound, const FString& InName, const FName InTypeName, const FMetasoundFrontendNodeStyle& InNodeStyle, const FText& InToolTip, const FMetasoundFrontendLiteral* InDefaultValue)
+		Frontend::FNodeHandle FGraphBuilder::AddInputNodeHandle(UObject& InMetasound, const FString& InName, const FName InTypeName, const FText& InToolTip, EMetasoundFrontendNodeStyleDisplayVisibility InVisibility, const FMetasoundFrontendLiteral* InDefaultValue)
 		{
 			FMetasoundAssetBase* MetasoundAsset = IMetasoundUObjectRegistry::Get().GetObjectAsAssetBase(&InMetasound);
 			check(MetasoundAsset);
@@ -551,7 +537,10 @@ namespace Metasound
 			Description.VertexID = VertexID;
 
 			Frontend::FNodeHandle NodeHandle = GraphHandle->AddInputVertex(Description);
-			NodeHandle->SetNodeStyle(InNodeStyle);
+
+			FMetasoundFrontendNodeStyle Style = NodeHandle->GetNodeStyle();
+			Style.Display.Visibility = InVisibility;
+			NodeHandle->SetNodeStyle(Style);
 
 			if (ensure(NodeHandle->IsValid()))
 			{
@@ -573,13 +562,13 @@ namespace Metasound
 			return NodeHandle;
 		}
 
-		UMetasoundEditorGraphOutputNode* FGraphBuilder::AddOutputNode(UObject& InMetasound, const FString& InName, const FName InTypeName, const FMetasoundFrontendNodeStyle& InNodeStyle, const FText& InToolTip, bool bInSelectNewNode)
+		UMetasoundEditorGraphOutputNode* FGraphBuilder::AddOutputNode(UObject& InMetasound, const FString& InName, const FName InTypeName, const FText& InToolTip, FVector2D InLocation, bool bInSelectNewNode)
 		{
-			Frontend::FNodeHandle NodeHandle = AddOutputNodeHandle(InMetasound, InName, InTypeName, InNodeStyle, InToolTip);
-			return AddOutputNode(InMetasound, NodeHandle, bInSelectNewNode);
+			Frontend::FNodeHandle NodeHandle = AddOutputNodeHandle(InMetasound, InName, InTypeName, InToolTip);
+			return AddOutputNode(InMetasound, NodeHandle, InLocation, bInSelectNewNode);
 		}
 
-		Frontend::FNodeHandle FGraphBuilder::AddOutputNodeHandle(UObject& InMetasound, const FString& InName, const FName InTypeName, const FMetasoundFrontendNodeStyle& InNodeStyle, const FText& InToolTip)
+		Frontend::FNodeHandle FGraphBuilder::AddOutputNodeHandle(UObject& InMetasound, const FString& InName, const FName InTypeName, const FText& InToolTip)
 		{
 			FMetasoundAssetBase* MetasoundAsset = IMetasoundUObjectRegistry::Get().GetObjectAsAssetBase(&InMetasound);
 			check(MetasoundAsset);
@@ -593,32 +582,30 @@ namespace Metasound
 			Description.VertexID = GraphHandle->GetNewVertexID();
 
 			Frontend::FNodeHandle NodeHandle = GraphHandle->AddOutputVertex(Description);
-			NodeHandle->SetNodeStyle(InNodeStyle);
-
 			GraphHandle->SetOutputDisplayName(InName, FText::FromString(InName));
 
 			return NodeHandle;
 		}
 
-		UEdGraphNode* FGraphBuilder::AddNode(UObject& InMetasound, Frontend::FNodeHandle InNodeHandle, bool bInSelectNewNode)
+		UMetasoundEditorGraphNode* FGraphBuilder::AddNode(UObject& InMetasound, Frontend::FNodeHandle InNodeHandle, FVector2D InLocation, bool bInSelectNewNode)
 		{
 			switch (InNodeHandle->GetClassType())
 			{
 				case EMetasoundFrontendClassType::Input:
 				{
-					return CastChecked<UEdGraphNode>(AddInputNode(InMetasound, InNodeHandle, bInSelectNewNode));
+					return CastChecked<UMetasoundEditorGraphNode>(AddInputNode(InMetasound, InNodeHandle, InLocation, bInSelectNewNode));
 				}
 				break;
 
 				case EMetasoundFrontendClassType::External:
 				{
-					return CastChecked<UEdGraphNode>(AddExternalNode(InMetasound, InNodeHandle, bInSelectNewNode));
+					return CastChecked<UMetasoundEditorGraphNode>(AddExternalNode(InMetasound, InNodeHandle, InLocation, bInSelectNewNode));
 				}
 				break;
 
 				case EMetasoundFrontendClassType::Output:
 				{
-					return CastChecked<UEdGraphNode>(AddOutputNode(InMetasound, InNodeHandle, bInSelectNewNode));
+					return CastChecked<UMetasoundEditorGraphNode>(AddOutputNode(InMetasound, InNodeHandle, InLocation, bInSelectNewNode));
 				}
 				break;
 
@@ -685,6 +672,55 @@ namespace Metasound
 			return true;
 		}
 
+		void FGraphBuilder::DisconnectPin(UEdGraphPin& InPin, bool bAddLiteralInputs)
+		{
+			using namespace Editor;
+			using namespace Frontend;
+
+			TArray<FInputHandle> InputHandles;
+			TArray<UEdGraphPin*> InputPins;
+
+			UObject& Metasound = CastChecked<UMetasoundEditorGraphNode>(InPin.GetOwningNode())->GetMetasoundChecked();
+
+			if (InPin.Direction == EGPD_Input)
+			{
+				FNodeHandle NodeHandle = CastChecked<UMetasoundEditorGraphNode>(InPin.GetOwningNode())->GetNodeHandle();
+				InputHandles = NodeHandle->GetInputsWithVertexName(InPin.GetName());
+				InputPins.Add(&InPin);
+			}
+			else
+			{
+				check(InPin.Direction == EGPD_Output);
+				for (UEdGraphPin* Pin : InPin.LinkedTo)
+				{
+					FNodeHandle NodeHandle = CastChecked<UMetasoundEditorGraphNode>(Pin->GetOwningNode())->GetNodeHandle();
+					InputHandles.Append(NodeHandle->GetInputsWithVertexName(Pin->GetName()));
+					InputPins.Add(Pin);
+				}
+			}
+
+			IMetasoundEditorModule& EditorModule = FModuleManager::GetModuleChecked<IMetasoundEditorModule>("MetasoundEditor");
+			for (int32 i = 0; i < InputHandles.Num(); ++i)
+			{
+				FInputHandle InputHandle = InputHandles[i];
+				FConstOutputHandle OutputHandle = InputHandle->GetCurrentlyConnectedOutput();
+				const FMetasoundFrontendNodeStyle& Style = OutputHandle->GetOwningNode()->GetNodeStyle();
+
+				// Hidden nodes are not "connected" from the perspective of EdGraph,
+				// and therefore should be ignored when disconnecting links.
+				if (Style.Display.Visibility != EMetasoundFrontendNodeStyleDisplayVisibility::Hidden)
+				{
+					InputHandle->Disconnect();
+
+					if (bAddLiteralInputs)
+					{
+						FNodeHandle NodeHandle = InputHandle->GetOwningNode();
+						FGraphBuilder::AddOrUpdateLiteralInput(Metasound, NodeHandle, *InputPins[i]);
+					}
+				}
+			}
+		}
+
 		void FGraphBuilder::ConstructGraph(UObject& InMetasound)
 		{
 			using namespace Frontend;
@@ -713,26 +749,28 @@ namespace Metasound
 			for (FNodeHandle& NodeHandle : NodeHandles)
 			{
 				const EMetasoundFrontendClassType NodeType = NodeHandle->GetClassType();
-				FMetasoundFrontendNodeStyle Style;
+				FVector2D NewLocation;
 				if (NodeType == EMetasoundFrontendClassType::Input)
 				{
-					Style.Display.Location = InputNodeLocation;
+					NewLocation = InputNodeLocation;
 					InputNodeLocation.Y += 100.0f;
 				}
 				else if (NodeType == EMetasoundFrontendClassType::Output)
 				{
-					Style.Display.Location = OutputNodeLocation;
+					NewLocation = OutputNodeLocation;
 					OutputNodeLocation.Y += 100.0f;
 				}
 				else
 				{
-					Style.Display.Location = OpNodeLocation;
+					NewLocation = OpNodeLocation;
 					OpNodeLocation.Y += 100.0f;
 				}
-				NodeHandle->SetNodeStyle(Style);
 
-				UEdGraphNode* NewNode = AddNode(InMetasound, NodeHandle, false /* bInSelectNewNode */);
-				NewIdNodeMap.Add(NodeHandle->GetID(), FNodePair { NodeHandle, NewNode });
+				UEdGraphNode* NewNode = AddNode(InMetasound, NodeHandle, NewLocation, false /* bInSelectNewNode */);
+				if (ensure(NewNode))
+				{
+					NewIdNodeMap.Add(NodeHandle->GetID(), FNodePair{ NodeHandle, NewNode });
+				}
 			}
 
 			for (const TPair<FGuid, FNodePair>& IdNodePair : NewIdNodeMap)
@@ -779,27 +817,48 @@ namespace Metasound
 			InMetasound.MarkPackageDirty();
 		}
 
-		void FGraphBuilder::DeleteLiteralInputs(UEdGraphNode& InNode)
+		void FGraphBuilder::DeleteVariableNodeHandle(UMetasoundEditorGraphVariable& InVariable)
 		{
 			using namespace Frontend;
 
-			FNodeHandle NodeHandle = Frontend::INodeController::GetInvalidHandle();
-			if (UMetasoundEditorGraphNode* Node = Cast<UMetasoundEditorGraphNode>(&InNode))
+			FNodeHandle NodeHandle = InVariable.GetNodeHandle();
+			TArray<UMetasoundEditorGraphNode*> Nodes = InVariable.GetNodes();
+			for (UMetasoundEditorGraphNode* Node : Nodes)
 			{
-				NodeHandle = Node->GetNodeHandle();
+				if (ensure(Node))
+				{
+					// Remove the give node's location from the frontend node
+					FMetasoundFrontendNodeStyle Style = NodeHandle->GetNodeStyle();
+					Style.Display.Locations.Remove(Node->NodeGuid);
+					NodeHandle->SetNodeStyle(Style);
+
+					FGraphBuilder::DeleteNode(*Node);
+				}
+			}
+
+			const FString NodeName = NodeHandle->GetNodeName();
+			const FText NodeDisplayName = NodeHandle->GetDisplayName();
+			FGraphHandle GraphHandle = NodeHandle->GetOwningGraph();
+			GraphHandle->RemoveNode(*NodeHandle);
+		}
 
-				TArray<FInputHandle> Inputs = NodeHandle->GetInputs();
-				for (FInputHandle& Input : Inputs)
+		void FGraphBuilder::DeleteLiteralInputs(UEdGraphNode& InNode)
+		{
+			using namespace Frontend;
+
+			UMetasoundEditorGraphNode* Node = CastChecked<UMetasoundEditorGraphNode>(&InNode);
+			FNodeHandle NodeHandle = Node->GetNodeHandle();
+			TArray<FInputHandle> Inputs = NodeHandle->GetInputs();
+			for (FInputHandle& Input : Inputs)
+			{
+				FOutputHandle Output = Input->GetCurrentlyConnectedOutput();
+				if (Output->IsValid())
 				{
-					FOutputHandle Output = Input->GetCurrentlyConnectedOutput();
-					if (Output->IsValid())
+					FNodeHandle InputHandle = Output->GetOwningNode();
+					if (InputHandle->GetNodeStyle().Display.Visibility == EMetasoundFrontendNodeStyleDisplayVisibility::Hidden)
 					{
-						FNodeHandle InputHandle = Output->GetOwningNode();
-						if (InputHandle->GetNodeStyle().Display.Visibility == EMetasoundFrontendNodeStyleDisplayVisibility::Hidden)
-						{
-							FGraphHandle Graph = InputHandle->GetOwningGraph();
-							Graph->RemoveNode(*InputHandle);
-						}
+						FGraphHandle Graph = InputHandle->GetOwningGraph();
+						Graph->RemoveNode(*InputHandle);
 					}
 				}
 			}
@@ -809,75 +868,77 @@ namespace Metasound
 		{
 			using namespace Frontend;
 
-			FNodeHandle NodeHandle = Frontend::INodeController::GetInvalidHandle();
-			if (UMetasoundEditorGraphNode* Node = Cast<UMetasoundEditorGraphNode>(&InNode))
+			if (!InNode.CanUserDeleteNode())
 			{
-				NodeHandle = Node->GetNodeHandle();
+				return false;
+			}
 
-				if (NodeHandle->GetClassType() == EMetasoundFrontendClassType::Input)
-				{
-					FConstDocumentHandle DocumentHandle = NodeHandle->GetOwningGraph()->GetOwningDocument();
-					auto IsRequiredInput = [&](const Frontend::FConstInputHandle& InputHandle)
-					{
-						return DocumentHandle->IsRequiredInput(InputHandle->GetName());
-					};
-					TArray<Frontend::FConstInputHandle> NodeInputs = NodeHandle->GetConstInputs();
+			InNode.MarkPackageDirty();
 
-					if (Frontend::FConstInputHandle* InputHandle = NodeInputs.FindByPredicate(IsRequiredInput))
-					{
-						FNotificationInfo Info(FText::Format(LOCTEXT("Metasounds_CannotDeleteRequiredInput",
-							"'Required Input '{0}' cannot be deleted."), (*InputHandle)->GetDisplayName()));
-						Info.bFireAndForget = true;
-						Info.ExpireDuration = 2.0f;
-						Info.bUseThrobber = true;
-						FSlateNotificationManager::Get().AddNotification(Info);
-						return false;
-					}
-				}
+			// If node isn't a MetasoundEditorGraphNode, just remove and return (ex. comment nodes)
+			UMetasoundEditorGraphNode* Node = Cast<UMetasoundEditorGraphNode>(&InNode);
+			UMetasoundEditorGraph* Graph = CastChecked<UMetasoundEditorGraph>(InNode.GetGraph());
+			if (!Node)
+			{
+				Graph->RemoveNode(Node);
+				return true;
+			}
 
-				if (NodeHandle->GetClassType() == EMetasoundFrontendClassType::Output)
-				{
-					FConstDocumentHandle DocumentHandle = NodeHandle->GetOwningGraph()->GetOwningDocument();
-					auto IsRequiredOutput = [&](const Frontend::FConstOutputHandle& OutputHandle)
-					{
-						return DocumentHandle->IsRequiredOutput(OutputHandle->GetName());
-					};
-					TArray<Frontend::FConstOutputHandle> NodeOutputs = NodeHandle->GetConstOutputs();
 
-					if (Frontend::FConstOutputHandle* OutputHandle = NodeOutputs.FindByPredicate(IsRequiredOutput))
-					{
-						FNotificationInfo Info(FText::Format(LOCTEXT("Metasounds_CannotDeleteRequiredOutput",
-							"'Required Output '{0}' cannot be deleted."), (*OutputHandle)->GetDisplayName()));
-						Info.bFireAndForget = true;
-						Info.ExpireDuration = 2.0f;
-						Info.bUseThrobber = true;
-						FSlateNotificationManager::Get().AddNotification(Info);
-						return false;
-					}
-				}
+			FNodeHandle NodeHandle = Node->GetNodeHandle();
+			if (!ensure(NodeHandle->IsValid()))
+			{
+				return false;
 			}
 
+			// Delete literal inputs prior to disconnecting/removing Frontend node
+			// representation to avoid orphaning them, causing them to not get deleted.
 			DeleteLiteralInputs(InNode);
 
-			UMetasoundEditorGraph* Graph = CastChecked<UMetasoundEditorGraph>(InNode.GetGraph());
-			if (InNode.CanUserDeleteNode() && Graph->RemoveNode(&InNode))
+			// Remove connects only to pins associated with this EdGraph node
+			// only (Iterate pins and not Frontend representation to preserve
+			// other input/output EditorGraph reference node associations)
+			Node->IteratePins([](UEdGraphPin* Pin, int32 Index)
 			{
-				Graph->PostEditChange();
-				Graph->MarkPackageDirty();
-			}
+				if (Pin)
+				{
+					// Only add literal inputs for output pins as adding when disconnecting
+					// inputs would immediately orphan them on EditorGraph node removal below.
+					const bool bAddLiteralInputs = Pin->Direction == EGPD_Output;
+					FGraphBuilder::DisconnectPin(*Pin, bAddLiteralInputs);
+				}
+			});
 
-			if (NodeHandle->IsValid())
+			Frontend::FGraphHandle GraphHandle = NodeHandle->GetOwningGraph();
+			if (GraphHandle->IsValid())
 			{
-				Frontend::FGraphHandle GraphHandle = NodeHandle->GetOwningGraph();
-				if (GraphHandle->IsValid())
+				switch (NodeHandle->GetClassType())
 				{
-					GraphHandle->RemoveNode(*NodeHandle);
+					case EMetasoundFrontendClassType::Output:
+					case EMetasoundFrontendClassType::Input:
+					{
+						// NodeHandle does not get removed in these cases as EdGraph Inputs/Outputs
+						// merely reference their respective types set on the MetasoundGraph. It must
+						// be removed from the location display data however for graph sync reasons.
+						FMetasoundFrontendNodeStyle Style = NodeHandle->GetNodeStyle();
+						Style.Display.Locations.Remove(InNode.NodeGuid);
+						NodeHandle->SetNodeStyle(Style);
+					}
+					break;
+
+					case EMetasoundFrontendClassType::Graph:
+					case EMetasoundFrontendClassType::External:
+					default:
+					{
+						static_assert(static_cast<int32>(EMetasoundFrontendClassType::Invalid) == 4, "Possible missing MetasoundFrontendClassType switch case coverage.");
+
+						ensure(GraphHandle->RemoveNode(*NodeHandle));
+					}
+					break;
 				}
 			}
 
-			InNode.PostEditChange();
-			InNode.MarkPackageDirty();
-			return true;
+			return ensure(Graph->RemoveNode(&InNode));
 		}
 
 		void FGraphBuilder::RebuildNodePins(UMetasoundEditorGraphNode& InGraphNode)
@@ -969,34 +1030,6 @@ namespace Metasound
 			return false;
 		}
 
-		bool FGraphBuilder::IsRequiredInput(Frontend::FNodeHandle InNodeHandle)
-		{
-			if (InNodeHandle->GetClassType() != EMetasoundFrontendClassType::Input)
-			{
-				return false;
-			}
-
-			Frontend::FConstDocumentHandle DocumentHandle = InNodeHandle->GetOwningGraph()->GetOwningDocument();
-			return InNodeHandle->GetConstOutputs().ContainsByPredicate([&](const Frontend::FConstOutputHandle& OutputHandle)
-			{
-				return DocumentHandle->IsRequiredInput(OutputHandle->GetName());
-			});
-		}
-
-		bool FGraphBuilder::IsRequiredOutput(Frontend::FNodeHandle InNodeHandle)
-		{
-			if (InNodeHandle->GetClassType() != EMetasoundFrontendClassType::Output)
-			{
-				return false;
-			}
-
-			Frontend::FConstDocumentHandle DocumentHandle = InNodeHandle->GetOwningGraph()->GetOwningDocument();
-			return InNodeHandle->GetConstInputs().ContainsByPredicate([&](const Frontend::FConstInputHandle& InputHandle)
-			{
-				return DocumentHandle->IsRequiredOutput(InputHandle->GetName());
-			});
-		}
-
 		UEdGraphPin* FGraphBuilder::AddPinToNode(UMetasoundEditorGraphNode& InEditorNode, Frontend::FInputHandle InInputHandle)
 		{
 			using namespace Frontend;
@@ -1032,16 +1065,167 @@ namespace Metasound
 			return NewPin;
 		}
 
+		bool FGraphBuilder::SynchronizeConnections(UObject& InMetasound)
+		{
+			using namespace Frontend;
+
+			bool bIsGraphDirty = false;
+
+			FMetasoundAssetBase* MetasoundAsset = IMetasoundUObjectRegistry::Get().GetObjectAsAssetBase(&InMetasound);
+			check(MetasoundAsset);
+
+			FGraphHandle GraphHandle = MetasoundAsset->GetRootGraphHandle();
+
+			UMetasoundEditorGraph* EditorGraph = CastChecked<UMetasoundEditorGraph>(MetasoundAsset->GetGraph());
+
+
+			TArray<UMetasoundEditorGraphNode*> EditorNodes;
+			EditorGraph->GetNodesOfClass(EditorNodes);
+
+			TMap<FGuid, TArray<UMetasoundEditorGraphNode*>> EditorNodesByFrontendID;
+			for (UMetasoundEditorGraphNode* EditorNode : EditorNodes)
+			{
+				EditorNodesByFrontendID.FindOrAdd(EditorNode->GetNodeID()).Add(EditorNode);
+			}
+
+			// Iterate through all nodes in metasound editor graph and synchronize connections.
+			for (UMetasoundEditorGraphNode* EditorNode : EditorNodes)
+			{
+				bool bIsNodeDirty = false;
+
+				FNodeHandle Node = EditorNode->GetNodeHandle();
+
+				if (EMetasoundFrontendClassType::Input == Node->GetClassType())
+				{
+					// Skip this node if it is an input node. Input pins on input 
+					// nodes are not connected internal to the graph.
+					continue;
+				}
+
+				TArray<UEdGraphPin*> Pins = EditorNode->GetAllPins();
+				TArray<FInputHandle> NodeInputs = Node->GetInputs();
+
+				for (FInputHandle& NodeInput : NodeInputs)
+				{
+					auto IsMatchingInputPin = [&](const UEdGraphPin* Pin) -> bool
+					{
+						return IsMatchingInputHandleAndPin(NodeInput, *Pin);
+					};
+
+					UEdGraphPin* MatchingPin = nullptr;
+					if (UEdGraphPin** DoublePointer = Pins.FindByPredicate(IsMatchingInputPin))
+					{
+						MatchingPin = *DoublePointer;
+					}
+
+					if (!ensure(MatchingPin))
+					{
+						continue;
+					}
+
+					// Remove pin so it isn't used twice.
+					Pins.Remove(MatchingPin);
+
+					bool bShowConnectionInEditor = true;
+					bool bIsConnectedToLiteralInput = false;
+
+					FOutputHandle OutputHandle = NodeInput->GetCurrentlyConnectedOutput();
+					if (OutputHandle->IsValid())
+					{
+						FNodeHandle InputNodeHandle = OutputHandle->GetOwningNode();
+						const FMetasoundFrontendNodeStyle& InputNodeStyle = InputNodeHandle->GetNodeStyle();
+						bShowConnectionInEditor = InputNodeStyle.Display.Visibility == EMetasoundFrontendNodeStyleDisplayVisibility::Visible;
+						bIsConnectedToLiteralInput = !bShowConnectionInEditor;
+					}
+					else
+					{
+						bShowConnectionInEditor = false;
+					}
+						
+					if (bShowConnectionInEditor)
+					{
+						bool bAddLink = false;
+
+						if (MatchingPin->LinkedTo.IsEmpty())
+						{
+							// No link currently exists. Add the appropriate link.
+							bAddLink = true;
+						}
+						else if (!IsMatchingOutputHandleAndPin(OutputHandle, *MatchingPin->LinkedTo[0]))
+						{
+							// The wrong link exists.
+							MatchingPin->BreakAllPinLinks();
+							bAddLink = true;
+						}
+
+						if (bAddLink)
+						{
+							const FGuid NodeID = OutputHandle->GetOwningNodeID();
+							TArray<UMetasoundEditorGraphNode*>* OutputEditorNode = EditorNodesByFrontendID.Find(NodeID);
+							if (ensure(OutputEditorNode))
+							{
+								if (ensure(!OutputEditorNode->IsEmpty()))
+								{
+									UEdGraphPin* OutputPin = (*OutputEditorNode)[0]->FindPinChecked(OutputHandle->GetName(), EEdGraphPinDirection::EGPD_Output);
+									MatchingPin->MakeLinkTo(OutputPin);
+									bIsNodeDirty = true;
+								}
+							}
+						}
+					}
+					else
+					{
+						// No link should exist.
+						if (!MatchingPin->LinkedTo.IsEmpty())
+						{
+							MatchingPin->BreakAllPinLinks();
+							bIsNodeDirty = true;
+						}
+
+						// Hidden literal input Link exists, so match pin default value to literal input value
+						if (bIsConnectedToLiteralInput)
+						{
+							FNodeHandle InputNode = OutputHandle->GetOwningNode();
+							const FMetasoundFrontendNodeStyle& Style = InputNode->GetNodeStyle();
+							const FGuid VertexID = GraphHandle->GetVertexIDForInputVertex(InputNode->GetNodeName());
+							const FMetasoundFrontendLiteral DefaultLiteral = GraphHandle->GetDefaultInput(VertexID);
+							const FString DefaultLiteralString = DefaultLiteral.ToString();
+							if (MatchingPin->DefaultValue != DefaultLiteralString)
+							{
+								MatchingPin->DefaultValue = DefaultLiteralString;
+								bIsNodeDirty = true;
+							}
+						}
+					}
+				}
+
+				if (bIsNodeDirty)
+				{
+					EditorNode->MarkPackageDirty();
+				}
+
+				bIsGraphDirty |= bIsNodeDirty;
+			}
+
+			if (bIsGraphDirty)
+			{
+				EditorGraph->MarkPackageDirty();
+			}
+
+			return bIsGraphDirty;
+		}
+
 		bool FGraphBuilder::SynchronizeGraph(UObject& InMetasound)
 		{
 			using namespace Frontend;
 
-			bool bIsEditorGraphDirty = false;
+			bool bIsEditorGraphDirty = SynchronizeVariables(InMetasound);
 
 			FMetasoundAssetBase* MetasoundAsset = IMetasoundUObjectRegistry::Get().GetObjectAsAssetBase(&InMetasound);
 			check(MetasoundAsset);
 
-			// Get all nodes from frontend graph
+			// Get all external nodes from Frontend graph.  Input and output references will only be added/synchronized
+			// if required when synchronizing connections (as they are not required to inhabit editor graph).
 			FGraphHandle GraphHandle = MetasoundAsset->GetRootGraphHandle();
 			TArray<FNodeHandle> FrontendNodes = GraphHandle->GetNodes();
 
@@ -1051,13 +1235,19 @@ namespace Metasound
 			TArray<UMetasoundEditorGraphNode*> EditorNodes;
 			EditorGraph->GetNodesOfClass(EditorNodes);
 
-			// Find existing pairs of nodes and editor nodes
-			struct FNodePair
+			TMap<FGuid, UMetasoundEditorGraphNode*> EditorNodesByEdNodeGuid;
+			for (UMetasoundEditorGraphNode* Node : EditorNodes)
 			{
-				UMetasoundEditorGraphNode* EditorNode = nullptr;
-				FNodeHandle Node = INodeController::GetInvalidHandle();
+				EditorNodesByEdNodeGuid.Add(Node->NodeGuid, Node);
+			}
+
+			// Find existing array of editor nodes associated with Frontend node
+			struct FAssociatedNodes
+			{
+				TArray<UMetasoundEditorGraphNode*> EditorNodes;
+				FNodeHandle Node = Metasound::Frontend::INodeController::GetInvalidHandle();
 			};
-			TMap<FGuid, FNodePair> PairedNodes;
+			TMap<FGuid, FAssociatedNodes> AssociatedNodes;
 
 			// Reverse iterate so paired nodes can safely be removed from the array.
 			for (int32 i = FrontendNodes.Num() - 1; i >= 0; i--)
@@ -1068,23 +1258,37 @@ namespace Metasound
 					return InEditorNode->GetNodeID() == Node->GetID();
 				};
 
-				UMetasoundEditorGraphNode* EditorNode = nullptr;
-				if (UMetasoundEditorGraphNode** PointerEditorNode = EditorNodes.FindByPredicate(IsEditorNodeWithSameNodeID))
+				bool bFoundEditorNode = false;
+				for (int32 j = EditorNodes.Num() - 1; j >= 0; --j)
 				{
-					EditorNode = *PointerEditorNode;
+					UMetasoundEditorGraphNode* EditorNode = EditorNodes[j];
+					if (EditorNode->GetNodeID() == Node->GetID())
+					{
+						bFoundEditorNode = true;
+						FAssociatedNodes& AssociatedNodeData = AssociatedNodes.FindOrAdd(Node->GetID());
+						if (AssociatedNodeData.Node->IsValid())
+						{
+							ensure(AssociatedNodeData.Node == Node);
+						}
+						else
+						{
+							AssociatedNodeData.Node = Node;
+						}
+
+						AssociatedNodeData.EditorNodes.Add(EditorNode);
+						EditorNodes.RemoveAtSwap(j, 1, false /* bAllowShrinking */);
+					}
 				}
 
-				if (EditorNode)
+				if (bFoundEditorNode)
 				{
-					PairedNodes.Add(Node->GetID(), { EditorNode, Node });
-					EditorNodes.RemoveSwap(EditorNode);
-					FrontendNodes.RemoveAtSwap(i);
+					FrontendNodes.RemoveAtSwap(i, 1, false /* bAllowShrinking */);
 				}
 			}
 
-			// FrontendNodes contains nodes which need to be added.
-			// EditorNodes contains nodes that need to be removed.
-			// PairedNodes contains pairs which we have to check have synchronized pins
+			// FrontendNodes contains nodes which need to be added to the editor graph.
+			// EditorNodes contains nodes that need to be removed from the editor graph.
+			// AssociatedNodes contains pairs which we have to check have synchronized pins
 
 			// Add and remove nodes first in order to make sure correct editor nodes
 			// exist before attempting to synchronize connections.
@@ -1096,23 +1300,49 @@ namespace Metasound
 			// Add missing editor nodes marked as visible.
 			for (FNodeHandle Node : FrontendNodes)
 			{
-				const FMetasoundFrontendNodeStyle& Style = Node->GetNodeStyle();
+				FMetasoundFrontendNodeStyle Style = Node->GetNodeStyle();
 				if (Style.Display.Visibility == EMetasoundFrontendNodeStyleDisplayVisibility::Visible)
 				{
 					bIsEditorGraphDirty = true;
-					UEdGraphNode* NewNode = AddNode(InMetasound, Node, false /* bInSelectNewNode */);
-					if (ensure(NewNode))
+
+					TArray<UMetasoundEditorGraphNode*> AddedNodes;
+					for (const TPair<FGuid, FVector2D>& Location : Style.Display.Locations)
+					{
+						FVector2D NewLocation = FVector2D::ZeroVector;
+						UMetasoundEditorGraphNode* NewNode = AddNode(InMetasound, Node, Location.Value, false /* bInSelectNewNode */);
+						if (ensure(NewNode))
+						{
+							FAssociatedNodes& AssociatedNodeData = AssociatedNodes.FindOrAdd(Node->GetID());
+							if (AssociatedNodeData.Node->IsValid())
+							{
+								ensure(AssociatedNodeData.Node == Node);
+							}
+							else
+							{
+								AssociatedNodeData.Node = Node;
+							}
+
+							AddedNodes.Add(NewNode);
+							AssociatedNodeData.EditorNodes.Add(NewNode);
+						}
+					}
+
+					Style.Display.Locations.Reset();
+					for (UMetasoundEditorGraphNode* EditorNode : AddedNodes)
 					{
-						PairedNodes.Add(Node->GetID(), { Cast<UMetasoundEditorGraphNode>(NewNode), Node });
+						Style.Display.Locations.Add(EditorNode->NodeGuid, FVector2D(EditorNode->NodePosX, EditorNode->NodePosY));
 					}
+					Node->SetNodeStyle(Style);
 				}
 			}
 
-			// Synchronize pins on node pairs.
-			for (const TPair<FGuid, FNodePair>& IdNodePair : PairedNodes)
+			// Synchronize pins on node associations.
+			for (const TPair<FGuid, FAssociatedNodes>& IdNodePair : AssociatedNodes)
 			{
-				UMetasoundEditorGraphNode* EditorNode = IdNodePair.Value.EditorNode;
-				bIsEditorGraphDirty |= SynchronizeNodePins(*IdNodePair.Value.EditorNode, IdNodePair.Value.Node);
+				for (UMetasoundEditorGraphNode* EditorNode : IdNodePair.Value.EditorNodes)
+				{
+					bIsEditorGraphDirty |= SynchronizeNodePins(*EditorNode, IdNodePair.Value.Node);
+				}
 			}
 
 			// Synchronize connections.
@@ -1221,150 +1451,77 @@ namespace Metasound
 			return bIsNodeDirty;
 		}
 
-		bool FGraphBuilder::SynchronizeConnections(UObject& InMetasound)
+		bool FGraphBuilder::SynchronizeVariables(UObject& InMetasound)
 		{
 			using namespace Frontend;
 
-			bool bIsGraphDirty = false;
+			bool bIsEditorGraphDirty = false;
 
 			FMetasoundAssetBase* MetasoundAsset = IMetasoundUObjectRegistry::Get().GetObjectAsAssetBase(&InMetasound);
 			check(MetasoundAsset);
+			UMetasoundEditorGraph* Graph = CastChecked<UMetasoundEditorGraph>(MetasoundAsset->GetGraph());
 
-			FGraphHandle GraphHandle = MetasoundAsset->GetRootGraphHandle();
+			IMetasoundEditorModule& EditorModule = FModuleManager::GetModuleChecked<IMetasoundEditorModule>("MetasoundEditor");
 
-			UMetasoundEditorGraph* EditorGraph = CastChecked<UMetasoundEditorGraph>(MetasoundAsset->GetGraph());
-			TArray<UMetasoundEditorGraphNode*> EditorNodes;
+			TSet<UMetasoundEditorGraphInput*> Inputs;
+			TSet<UMetasoundEditorGraphOutput*> Outputs;
 
-			// Cache map of editor nodes indexed by node id.
-			TMap<FGuid, UMetasoundEditorGraphNode*> EditorNodesByID;
-			for (UEdGraphNode* EdGraphNode : EditorGraph->Nodes)
+			FGraphHandle GraphHandle = MetasoundAsset->GetRootGraphHandle();
+			GraphHandle->IterateNodes([&](FNodeHandle NodeHandle)
 			{
-				if (UMetasoundEditorGraphNode* MetasoundEditorNode = Cast<UMetasoundEditorGraphNode>(EdGraphNode))
+				if (UMetasoundEditorGraphInput* Input = Graph->FindInput(NodeHandle->GetID()))
 				{
-					EditorNodes.Add(MetasoundEditorNode);
-					EditorNodesByID.Add(MetasoundEditorNode->GetNodeID(), MetasoundEditorNode);
+					Inputs.Add(Input);
+					return;
 				}
-			}
-
-			// Iterate through all nodes in metasound editor graph and synchronize connections.
-			for (UMetasoundEditorGraphNode* MetasoundEditorNode : EditorNodes)
-			{
-				bool bIsNodeDirty = false;
-
-				FNodeHandle Node = MetasoundEditorNode->GetNodeHandle();
-
-				if (EMetasoundFrontendClassType::Input == Node->GetClassType())
+				if (!ensure(NodeHandle->GetNumInputs() == 1))
 				{
-					// Skip this node if it is an input node. Input pins on input 
-					// nodes are not connected internal to the graph.
-					continue;
+					return;
 				}
 
-				TArray<UEdGraphPin*> Pins = MetasoundEditorNode->GetAllPins();
-				TArray<FInputHandle> NodeInputs = Node->GetInputs();
+				Inputs.Add(Graph->FindOrAddInput(NodeHandle));
+				bIsEditorGraphDirty = true;
+			}, EMetasoundFrontendClassType::Input);
 
-				for (FInputHandle& NodeInput : NodeInputs)
+			GraphHandle->IterateNodes([&](FNodeHandle NodeHandle)
+			{
+				if (UMetasoundEditorGraphOutput* Output = Graph->FindOutput(NodeHandle->GetID()))
 				{
-					auto IsMatchingInputPin = [&](const UEdGraphPin* Pin) -> bool
-					{
-						return IsMatchingInputHandleAndPin(NodeInput, *Pin);
-					};
-
-					UEdGraphPin* MatchingPin = nullptr;
-					if (UEdGraphPin** DoublePointer = Pins.FindByPredicate(IsMatchingInputPin))
-					{
-						MatchingPin = *DoublePointer;
-					}
-
-					if (!ensure(MatchingPin))
-					{
-						continue;
-					}
-
-					// Remove pin so it isn't used twice.
-					Pins.Remove(MatchingPin);
-
-					bool bShowConnectionInEditor = true;
-					bool bIsConnectedToLiteralInput = false;
-
-					FOutputHandle OutputHandle = NodeInput->GetCurrentlyConnectedOutput();
-					if (OutputHandle->IsValid())
-					{
-						FNodeHandle InputNodeHandle = OutputHandle->GetOwningNode();
-						const FMetasoundFrontendNodeStyle& InputNodeStyle = InputNodeHandle->GetNodeStyle();
-						bShowConnectionInEditor = InputNodeStyle.Display.Visibility == EMetasoundFrontendNodeStyleDisplayVisibility::Visible;
-						bIsConnectedToLiteralInput = !bShowConnectionInEditor;
-					}
-					else
-					{
-						bShowConnectionInEditor = false;
-					}
-						
-					if (bShowConnectionInEditor)
-					{
-						bool bAddLink = false;
-
-						if (MatchingPin->LinkedTo.IsEmpty())
-						{
-							// No link currently exists. Add the appropriate link.
-							bAddLink = true;
-						}
-						else if (!IsMatchingOutputHandleAndPin(OutputHandle, *MatchingPin->LinkedTo[0]))
-						{
-							// The wrong link exists.
-							MatchingPin->BreakAllPinLinks();
-							bAddLink = true;
-						}
+					Outputs.Add(Output);
+					return;
+				}
+				if (!ensure(NodeHandle->GetNumOutputs() == 1))
+				{
+					return;
+				}
 
-						if (bAddLink)
-						{
-							const FGuid InputNodeID = OutputHandle->GetOwningNodeID();
-							UMetasoundEditorGraphNode* OutputEditorNode = EditorNodesByID[InputNodeID];
-							UEdGraphPin* OutputPin = OutputEditorNode->FindPinChecked(OutputHandle->GetName(), EEdGraphPinDirection::EGPD_Output);
-							MatchingPin->MakeLinkTo(OutputPin);
-							bIsNodeDirty = true;
-						}
-					}
-					else
-					{
-						// No link should exist.
-						if (!MatchingPin->LinkedTo.IsEmpty())
-						{
-							MatchingPin->BreakAllPinLinks();
-							bIsNodeDirty = true;
-						}
+				Outputs.Add(Graph->FindOrAddOutput(NodeHandle));
+				bIsEditorGraphDirty = true;
+			}, EMetasoundFrontendClassType::Output);
 
-						// Hidden literal input Link exists, so match pin default value to literal input value
-						if (bIsConnectedToLiteralInput)
-						{
-							FNodeHandle InputNode = OutputHandle->GetOwningNode();
-							const FMetasoundFrontendNodeStyle& Style = InputNode->GetNodeStyle();
-							const FGuid VertexID = GraphHandle->GetVertexIDForInputVertex(InputNode->GetNodeName());
-							const FMetasoundFrontendLiteral DefaultLiteral = GraphHandle->GetDefaultInput(VertexID);
-							const FString DefaultLiteralString = DefaultLiteral.ToString();
-							if (MatchingPin->DefaultValue != DefaultLiteralString)
-							{
-								MatchingPin->DefaultValue = DefaultLiteralString;
-								bIsNodeDirty = true;
-							}
-						}
-					}
+			TArray<UMetasoundEditorGraphVariable*> ToRemove;
+			Graph->IterateInputs([&](UMetasoundEditorGraphInput& Input)
+			{
+				if (!Inputs.Contains(&Input))
+				{
+					ToRemove.Add(&Input);
 				}
-
-				if (bIsNodeDirty)
+			});
+			Graph->IterateOutputs([&](UMetasoundEditorGraphOutput& Output)
+			{
+				if (!Outputs.Contains(&Output))
 				{
-					MetasoundEditorNode->MarkPackageDirty();
+					ToRemove.Add(&Output);
 				}
+			});
 
-				bIsGraphDirty |= bIsNodeDirty;
-			}
-
-			if (bIsGraphDirty)
+			bIsEditorGraphDirty |= !ToRemove.IsEmpty();
+			for (UMetasoundEditorGraphVariable* Variable : ToRemove)
 			{
-				EditorGraph->MarkPackageDirty();
+				Graph->RemoveVariable(*Variable);
 			}
 
-			return bIsGraphDirty;
+			return bIsEditorGraphDirty;
 		}
 	} // namespace Editor
 } // namespace Metasound
diff --git a/Engine/Plugins/Runtime/Metasound/Source/MetasoundEditor/Private/MetasoundEditorGraphBuilder.h b/Engine/Plugins/Runtime/Metasound/Source/MetasoundEditor/Private/MetasoundEditorGraphBuilder.h
--- a/Engine/Plugins/Runtime/Metasound/Source/MetasoundEditor/Private/MetasoundEditorGraphBuilder.h
+++ b/Engine/Plugins/Runtime/Metasound/Source/MetasoundEditor/Private/MetasoundEditorGraphBuilder.h
@@ -49,52 +49,72 @@ namespace Metasound
 			static const FText FunctionMenuName;
 
 			// Adds an EdGraph node to mirror the provided FNodeHandle.
-			static UEdGraphNode* AddNode(UObject& InMetasound, Frontend::FNodeHandle InNodeHandle, bool bInSelectNewNode = true);
+			static UMetasoundEditorGraphNode* AddNode(UObject& InMetasound, Frontend::FNodeHandle InNodeHandle, FVector2D InLocation, bool bInSelectNewNode = true);
 
 			// Adds a node handle to mirror the provided graph node and binds to it.  Does *NOT* mirror existing EdGraph connections
 			// nor does it remove existing bound Frontend Node (if set) from associated Frontend Graph.
 			static Frontend::FNodeHandle AddNodeHandle(UObject& InMetasound, UMetasoundEditorGraphNode& InGraphNode);
 
-			// Adds an output node to the editor graph that corresponds to the provided node handle.
+			// Adds an input node to the editor graph that corresponds to the provided node handle.
 			// Generates analogous FNodeHandle.
-			static UMetasoundEditorGraphInputNode* AddInputNode(UObject& InMetasound, const FString& InName, const FName InTypeName, const FMetasoundFrontendNodeStyle& InNodeStyle, const FText& InToolTip, bool bInSelectNewNode = true);
+			static UMetasoundEditorGraphInputNode* AddInputNode(
+				UObject& InMetasound,
+				const FString& InName,
+				const FName InTypeName,
+				const FText& InToolTip,
+				FVector2D InLocation,
+				EMetasoundFrontendNodeStyleDisplayVisibility Visibility = EMetasoundFrontendNodeStyleDisplayVisibility::Visible,
+				bool bInSelectNewNode = true);
 
 			// Adds a corresponding UMetasoundEditorGraphInputNode for the provided node handle.
-			static UMetasoundEditorGraphInputNode* AddInputNode(UObject& InMetasound, Frontend::FNodeHandle InNodeHandle, bool bInSelectNewNode = true);
+			static UMetasoundEditorGraphInputNode* AddInputNode(UObject& InMetasound, Frontend::FNodeHandle InNodeHandle, FVector2D InLocation, bool bInSelectNewNode = true);
 
 			// Generates FNodeHandle for the given external node data. Does not bind or create EdGraph representation of given node.
-			static Frontend::FNodeHandle AddInputNodeHandle(UObject& InMetasound, const FString& InName, const FName InTypeName, const FMetasoundFrontendNodeStyle& InNodeStyle, const FText& InToolTip, const FMetasoundFrontendLiteral* InDefaultValue = nullptr);
+			static Frontend::FNodeHandle AddInputNodeHandle(
+				UObject& InMetasound,
+				const FString& InName,
+				const FName InTypeName,
+				const FText& InToolTip,
+				EMetasoundFrontendNodeStyleDisplayVisibility Visibility = EMetasoundFrontendNodeStyleDisplayVisibility::Visible,
+				const FMetasoundFrontendLiteral* InDefaultValue = nullptr);
 
 			// Adds a corresponding UMetasoundEditorGraphExternalNode for the provided node handle.
-			static UMetasoundEditorGraphExternalNode* AddExternalNode(UObject& InMetasound, Frontend::FNodeHandle& InNodeHandle, bool bInSelectNewNode = true);
+			static UMetasoundEditorGraphExternalNode* AddExternalNode(UObject& InMetasound, Frontend::FNodeHandle& InNodeHandle, FVector2D InLocation, bool bInSelectNewNode = true);
 
 			// Adds an externally-defined node with the given class info to both the editor and document graphs.
 			// Generates analogous FNodeHandle.
-			static UMetasoundEditorGraphExternalNode* AddExternalNode(UObject& InMetasound, const Frontend::FNodeClassInfo& InClassInfo, const FMetasoundFrontendNodeStyle& InNodeStyle, bool bInSelectNewNode = true);
+			static UMetasoundEditorGraphExternalNode* AddExternalNode(UObject& InMetasound, const Frontend::FNodeClassInfo& InClassInfo, FVector2D InLocation, bool bInSelectNewNode = true);
+
+			// Synchronizes node location data
+			static void SynchronizeNodeLocation(FVector2D InLocation, Frontend::FNodeHandle InNodeHandle, UMetasoundEditorGraphNode& InNode);
 
 			// Generates FNodeHandle for the given external node data. Does not bind or create EdGraph representation of given node.
-			static Frontend::FNodeHandle AddExternalNodeHandle(UObject& InMetasound, const Frontend::FNodeClassInfo& InClassInfo, const FMetasoundFrontendNodeStyle& InNodeStyle);
+			static Frontend::FNodeHandle AddExternalNodeHandle(UObject& InMetasound, const Frontend::FNodeClassInfo& InClassInfo);
 
 			// Adds an output node to the editor graph that corresponds to the provided node handle.
-			static UMetasoundEditorGraphOutputNode* AddOutputNode(UObject& InMetasound, Frontend::FNodeHandle& InNodeHandle, bool bInSelectNewNode = true);
+			static UMetasoundEditorGraphOutputNode* AddOutputNode(UObject& InMetasound, Frontend::FNodeHandle& InNodeHandle, FVector2D InLocation, bool bInSelectNewNode = true);
 
 			// Adds an output node to the editor graph that corresponds to the provided node handle.
 			// Generates analogous FNodeHandle.
-			static UMetasoundEditorGraphOutputNode* AddOutputNode(UObject& InMetasound, const FString& InName, const FName InTypeName, const FMetasoundFrontendNodeStyle& InNodeStyle, const FText& InToolTip, bool bInSelectNewNode = true);
+			static UMetasoundEditorGraphOutputNode* AddOutputNode(UObject& InMetasound, const FString& InName, const FName InTypeName, const FText& InToolTip, FVector2D InLocation, bool bInSelectNewNode = true);
 
 			// Generates analogous FNodeHandle for the given internal node data. Does not bind nor create EdGraph representation of given node.
-			static Frontend::FNodeHandle AddOutputNodeHandle(UObject& InMetasound, const FString& InName, const FName InTypeName, const FMetasoundFrontendNodeStyle& InNodeStyle, const FText& InToolTip);
+			static Frontend::FNodeHandle AddOutputNodeHandle(UObject& InMetasound, const FString& InName, const FName InTypeName, const FText& InToolTip);
 
 			// Attempts to connect frontend node counterparts together for provided pines.  Returns true if succeeded,
 			// and breaks pin link and returns false if failed.  If bConnectEdPins is set, will attempt to connect
 			// the Editor Graph representation of the pins.
 			static bool ConnectNodes(UEdGraphPin& InInputPin, UEdGraphPin& InOutputPin, bool bInConnectEdPins);
 
-			static FString GenerateUniqueInputName(const UObject& InMetasound, const FName InBaseName);
+			// Disconnects pin from any linked input or output nodes, and reflects change
+			// in the Frontend graph.  (Handles generation of literal inputs where needed).
+			// If bAddLiteralInputs - If true, will attempt to connect literal inputs where
+			// applicable post disconnection.
+			static void DisconnectPin(UEdGraphPin& InPin, bool bAddLiteralInputs = true);
 
-			static FString GenerateUniqueOutputName(const UObject& InMetasound, const FName InBaseName);
+			static FString GenerateUniqueInputName(const UObject& InMetasound);
 
-			static FString GetDataTypeDisplayName(const FName& InDataTypeName);
+			static FString GenerateUniqueOutputName(const UObject& InMetasound);
 
 			static TArray<FString> GetDataTypeNameCategories(const FName& InDataTypeName);
 
@@ -115,8 +135,9 @@ namespace Metasound
 
 			static bool GetPinDefaultLiteral(UEdGraphPin& InInputPin, Frontend::FInputHandle InputHandle, FMetasoundFrontendLiteral& OutLiteralDefault);
 
-			static bool IsRequiredInput(Frontend::FNodeHandle InNodeHandle);
-			static bool IsRequiredOutput(Frontend::FNodeHandle InNodeHandle);
+			// Deletes Editor Graph Variable's associated Frontend node, as well as any
+			// Editor Graph nodes referencing the given variable.
+			static void DeleteVariableNodeHandle(UMetasoundEditorGraphVariable& InVariable);
 
 			// Constructs graph with default inputs & outputs.
 			static void ConstructGraph(UObject& InMetasound);
@@ -149,13 +170,18 @@ namespace Metasound
 			// Adds and removes connections so that the UEdGraph of the metasound has the same
 			// connections as the FMetasoundFrontendDocument graph.
 			//
-			// @return True if the UEdGraph was altered. False otherwise. 
+			// @return True if the UEdGraph was altered. False otherwise.
 			static bool SynchronizeConnections(UObject& InMetasound);
 
-			// Returns true if the FInputHandle and UEdGraphPin match each other. 
+			// Synchronizes inputs and outputs for the given Metasound.
+			//
+			// @return True if the UEdGraph was altered. False otherwise.
+			static bool SynchronizeVariables(UObject& InMetasound);
+
+			// Returns true if the FInputHandle and UEdGraphPin match each other.
 			static bool IsMatchingInputHandleAndPin(const Frontend::FInputHandle& InInputHandle, const UEdGraphPin& InEditorPin);
 
-			// Returns true if the FOutputHandle and UEdGraphPin match each other. 
+			// Returns true if the FOutputHandle and UEdGraphPin match each other.
 			static bool IsMatchingOutputHandleAndPin(const Frontend::FOutputHandle& InOutputHandle, const UEdGraphPin& InEditorPin);
 		};
 	} // namespace Editor
diff --git a/Engine/Plugins/Runtime/Metasound/Source/MetasoundEditor/Private/MetasoundEditorGraphNode.cpp b/Engine/Plugins/Runtime/Metasound/Source/MetasoundEditor/Private/MetasoundEditorGraphNode.cpp
--- a/Engine/Plugins/Runtime/Metasound/Source/MetasoundEditor/Private/MetasoundEditorGraphNode.cpp
+++ b/Engine/Plugins/Runtime/Metasound/Source/MetasoundEditor/Private/MetasoundEditorGraphNode.cpp
@@ -113,11 +113,13 @@ Metasound::Frontend::FGraphHandle UMetasoundEditorGraphNode::GetRootGraphHandle(
 
 Metasound::Frontend::FConstNodeHandle UMetasoundEditorGraphNode::GetConstNodeHandle() const
 {
+	const FGuid NodeID = GetNodeID();
 	return GetConstRootGraphHandle()->GetNodeWithID(NodeID);
 }
 
 Metasound::Frontend::FNodeHandle UMetasoundEditorGraphNode::GetNodeHandle() const
 {
+	const FGuid NodeID = GetNodeID();
 	return GetRootGraphHandle()->GetNodeWithID(NodeID);
 }
 
@@ -199,26 +201,6 @@ bool UMetasoundEditorGraphNode::CanCreateUnderSpecifiedSchema(const UEdGraphSche
 
 bool UMetasoundEditorGraphNode::CanUserDeleteNode() const
 {
-	const FMetasoundAssetBase* ConstMetasoundAsset = Metasound::IMetasoundUObjectRegistry::Get().GetObjectAsAssetBase(&GetMetasoundChecked());
-
-	// TODO: create FConstGraphHandle so that this const cast is unnecessary.
-	FMetasoundAssetBase* MetasoundAsset = const_cast<FMetasoundAssetBase*>(ConstMetasoundAsset);
-	check(MetasoundAsset);
-
-	Metasound::Frontend::FDocumentHandle DocumentHandle = MetasoundAsset->GetDocumentHandle();
-	Metasound::Frontend::FNodeHandle NodeHandle = DocumentHandle->GetRootGraph()->GetNodeWithID(NodeID);
-
-	const FString& NodeName = GetNodeHandle()->GetNodeName();
-	if (DocumentHandle->IsRequiredInput(NodeName))
-	{
-		return false;
-	}
-
-	if (DocumentHandle->IsRequiredOutput(NodeName))
-	{
-		return false;
-	}
-
 	return true;
 }
 
@@ -227,49 +209,12 @@ FString UMetasoundEditorGraphNode::GetDocumentationLink() const
 	return TEXT("Shared/GraphNodes/Metasound");
 }
 
-FGuid UMetasoundEditorGraphNode::GetNodeID() const
-{
-	return NodeID;
-}
-
 FText UMetasoundEditorGraphNode::GetNodeTitle(ENodeTitleType::Type TitleType) const
 {
 	using namespace Metasound::Frontend;
 
 	FConstNodeHandle NodeHandle = GetNodeHandle();
-	switch (NodeHandle->GetClassType())
-	{
-		case EMetasoundFrontendClassType::Input:
-		{
-			// TODO: have INodeController determine display name
-			TArray<FConstInputHandle> InputHandles = NodeHandle->GetConstInputs();
-			if (InputHandles.Num() > 0)
-			{
-				return FText::Format(LOCTEXT("MetasoundGraphNode_InputTitleFormat", "Input {0}"), InputHandles[0]->GetDisplayName());
-			}
-			return LOCTEXT("MetasoundGraphNode_DefaultInputTitle", "Input");
-		}
-		break;
-
-		case EMetasoundFrontendClassType::Output:
-		{
-			TArray<FConstOutputHandle> OutputHandles = NodeHandle->GetConstOutputs();
-			if (OutputHandles.Num() > 0)
-			{
-				return FText::Format(LOCTEXT("MetasoundGraphNode_OutputTitleFormat", "Output {0}"), OutputHandles[0]->GetDisplayName());
-			}
-			return LOCTEXT("MetasoundGraphNode_DefaultOutputTitle", "Output");
-		}
-		break;
-
-		default:
-		case EMetasoundFrontendClassType::External:
-		case EMetasoundFrontendClassType::Graph:
-		{
-			return NodeHandle->GetClassDisplayName();
-		}
-		break;
-	}
+	return NodeHandle->GetDisplayTitle();
 }
 
 void UMetasoundEditorGraphNode::PinDefaultValueChanged(UEdGraphPin* Pin)
@@ -301,7 +246,7 @@ void UMetasoundEditorGraphNode::PostEditChangeProperty(FPropertyChangedEvent& In
 	{
 		Metasound::Frontend::FNodeHandle NodeHandle = GetNodeHandle();
 		FMetasoundFrontendNodeStyle Style = NodeHandle->GetNodeStyle();
-		Style.Display.Location = FVector2D(NodePosX, NodePosY);
+		Style.Display.Locations.FindOrAdd(NodeGuid) = FVector2D(NodePosX, NodePosY);
 		GetNodeHandle()->SetNodeStyle(Style);
 	}
 
@@ -396,7 +341,7 @@ void UMetasoundEditorGraphNode::GetNodeContextMenuActions(UToolMenu* Menu, UGrap
 
 FText UMetasoundEditorGraphNode::GetTooltipText() const
 {
-	return GetRootGraphHandle()->GetNodeWithID(NodeID)->GetClassDescription();
+	return GetConstNodeHandle()->GetDescription();
 }
 
 FString UMetasoundEditorGraphNode::GetDocumentationExcerptName() const
@@ -405,4 +350,30 @@ FString UMetasoundEditorGraphNode::GetDocumentationExcerptName() const
 	// This is done so that the excerpt name in the doc file can be found by find-in-files when searching for the full class name.
 	return FString::Printf(TEXT("%s%s"), UMetasound::StaticClass()->GetPrefixCPP(), *UMetasound::StaticClass()->GetName());
 }
+
+FMetasoundFrontendClassName UMetasoundEditorGraphOutputNode::GetClassName() const
+{
+	if (ensure(Output))
+	{
+		return Output->ClassName;
+	}
+	return FMetasoundFrontendClassName();
+}
+
+FGuid UMetasoundEditorGraphOutputNode::GetNodeID() const
+{
+	if (ensure(Output))
+	{
+		return Output->NodeID;
+	}
+	return FGuid();
+}
+
+void UMetasoundEditorGraphOutputNode::SetNodeID(FGuid InNodeID)
+{
+	if (ensure(Output))
+	{
+		Output->NodeID = InNodeID;
+	}
+}
 #undef LOCTEXT_NAMESPACE
diff --git a/Engine/Plugins/Runtime/Metasound/Source/MetasoundEditor/Private/MetasoundEditorGraphSchema.cpp b/Engine/Plugins/Runtime/Metasound/Source/MetasoundEditor/Private/MetasoundEditorGraphSchema.cpp
--- a/Engine/Plugins/Runtime/Metasound/Source/MetasoundEditor/Private/MetasoundEditorGraphSchema.cpp
+++ b/Engine/Plugins/Runtime/Metasound/Source/MetasoundEditor/Private/MetasoundEditorGraphSchema.cpp
@@ -157,37 +157,23 @@ namespace Metasound
 		}
 
 		template <typename TAction>
-		void GetDataTypeActions(FGraphActionMenuBuilder& ActionMenuBuilder, FDataTypeFilterFunction InFilter, const FText& InCategory, const FText& InTooltipFormat, bool bShowSelectedActions)
+		void GetDataTypeActions(FGraphActionMenuBuilder& ActionMenuBuilder, const TArray<Frontend::FConstNodeHandle>& InNodeHandles, FInterfaceNodeFilterFunction InFilter, const FText& InCategory, const FText& InTooltipFormat, bool bShowSelectedActions)
 		{
-			using namespace Metasound::Editor;
-			using namespace Metasound::Frontend;
-
-			static const FText MenuJoinFormat = LOCTEXT("MetasoundActionsFormatSubCategory", "{0}|{1}");
+			using namespace Editor;
+			using namespace Frontend;
 
-			IMetasoundEditorModule& EditorModule = FModuleManager::GetModuleChecked<IMetasoundEditorModule>("MetasoundEditor");
-			EditorModule.IterateDataTypes([&](const FEditorDataType& DataType)
+			for (FConstNodeHandle NodeHandle : InNodeHandles)
 			{
-				if (InFilter && !InFilter(DataType))
+				if (InFilter && !InFilter(NodeHandle))
 				{
-					return;
+					continue;
 				}
 
-				const FName DataTypeName = DataType.RegistryInfo.DataTypeName;
-				const FText DataTypeDisplayName = FText::FromString(FGraphBuilder::GetDataTypeDisplayName(DataTypeName));
-				const FText DataTypeTextName = FText::FromName(DataTypeName);
-				const TArray<FString> Categories = FGraphBuilder::GetDataTypeNameCategories(DataTypeName);
-				const FText CategoriesText = FText::FromString(FString::Join(Categories, TEXT("|")));
-				TSharedPtr<TAction> NewNodeAction = MakeShared<TAction>
-				(
-					FText::Format(MenuJoinFormat, InCategory, CategoriesText),
-					DataTypeDisplayName,
-					DataTypeName,
-					FText::Format(InTooltipFormat, DataTypeTextName),
-					0
-				);
-
+				const FText NodeDisplayName = NodeHandle->GetDisplayName();
+				const FText Tooltip = FText::Format(InTooltipFormat, NodeDisplayName);
+				TSharedPtr<TAction> NewNodeAction = MakeShared<TAction>(InCategory, NodeDisplayName, NodeHandle->GetID(), Tooltip, 0);
 				ActionMenuBuilder.AddAction(NewNodeAction);
-			});
+			}
 		}
 
 		FConnectionDrawingPolicy* FGraphConnectionDrawingPolicyFactory::CreateConnectionPolicy(
@@ -266,15 +252,10 @@ UEdGraphNode* FMetasoundGraphSchemaAction_NewNode::PerformAction(UEdGraph* Paren
 	using namespace Metasound::Frontend;
 
 	const FScopedTransaction Transaction(LOCTEXT("MetasoundEditorNewNode", "Add New Metasound Node"));
-
 	ParentGraph->Modify();
 
 	UObject& ParentMetasound = CastChecked<UMetasoundEditorGraph>(ParentGraph)->GetMetasoundChecked();
-	ParentMetasound.Modify();
-
-	FMetasoundFrontendNodeStyle Style;
-	Style.Display.Location = Location;
-	if (UEdGraphNode* NewGraphNode = FGraphBuilder::AddExternalNode(ParentMetasound, NodeClassInfo, Style, bSelectNewNode))
+	if (UMetasoundEditorGraphExternalNode* NewGraphNode = FGraphBuilder::AddExternalNode(ParentMetasound, NodeClassInfo, Location, bSelectNewNode))
 	{
 		TryConnectNewNodeToPin(*NewGraphNode, FromPin);
 		return NewGraphNode;
@@ -283,29 +264,38 @@ UEdGraphNode* FMetasoundGraphSchemaAction_NewNode::PerformAction(UEdGraph* Paren
 	return nullptr;
 }
 
-FMetasoundGraphSchemaAction_NewInput::FMetasoundGraphSchemaAction_NewInput(FText InNodeCategory, FText InDisplayName, FName InTypeName, FText InToolTip, const int32 InGrouping)
+FMetasoundGraphSchemaAction_NewInput::FMetasoundGraphSchemaAction_NewInput(FText InNodeCategory, FText InDisplayName, FGuid InNodeID, FText InToolTip, const int32 InGrouping)
 	: FEdGraphSchemaAction(MoveTemp(InNodeCategory), MoveTemp(InDisplayName), MoveTemp(InToolTip), InGrouping)
-	, NodeTypeName(InTypeName)
+	, NodeID(InNodeID)
 {
 }
 
-UEdGraphNode* FMetasoundGraphSchemaAction_NewInput::PerformAction(UEdGraph* ParentGraph, UEdGraphPin* FromPin, const FVector2D Location, bool bSelectNewNode /* = true */)
+UEdGraphNode* FMetasoundGraphSchemaAction_NewInput::PerformAction(UEdGraph* ParentGraph, UEdGraphPin* FromPin, const FVector2D InLocation, bool bSelectNewNode /* = true */)
 {
 	using namespace Metasound::Editor;
+	using namespace Metasound::Frontend;
 
-	const FScopedTransaction Transaction(LOCTEXT("MetasoundEditorNewInput", "Add New Metasound Input"));
+	const FScopedTransaction Transaction(LOCTEXT("MetasoundEditorNewInputNode", "Add New Metasound Input Node"));
 
 	ParentGraph->Modify();
 
-	UObject& ParentMetasound = CastChecked<UMetasoundEditorGraph>(ParentGraph)->GetMetasoundChecked();
+	UMetasoundEditorGraph* MetasoundGraph = CastChecked<UMetasoundEditorGraph>(ParentGraph);
+	UObject& ParentMetasound = MetasoundGraph->GetMetasoundChecked();
 	ParentMetasound.Modify();
 
-	FString NewNodeName = FGraphBuilder::GenerateUniqueInputName(ParentMetasound, NodeTypeName);
-
-	FMetasoundFrontendNodeStyle Style;
-	Style.Display.Location = Location;
+	UMetasoundEditorGraphInput* Input = MetasoundGraph->FindInput(NodeID);
+	if (!ensure(Input))
+	{
+		return nullptr;
+	}
+	
+	FNodeHandle NodeHandle = Input->GetNodeHandle();
+	if (!ensure(NodeHandle->IsValid()))
+	{
+		return nullptr;
+	}
 
-	if (UMetasoundEditorGraphInputNode* NewGraphNode = FGraphBuilder::AddInputNode(ParentMetasound, NewNodeName, NodeTypeName, Style, FText::GetEmpty()))
+	if (UMetasoundEditorGraphInputNode* NewGraphNode = FGraphBuilder::AddInputNode(ParentMetasound, NodeHandle, InLocation))
 	{
 		UEdGraphNode* EdGraphNode = CastChecked<UEdGraphNode>(NewGraphNode);
 		TryConnectNewNodeToPin(*EdGraphNode, FromPin);
@@ -315,29 +305,38 @@ UEdGraphNode* FMetasoundGraphSchemaAction_NewInput::PerformAction(UEdGraph* Pare
 	return nullptr;
 }
 
-FMetasoundGraphSchemaAction_NewOutput::FMetasoundGraphSchemaAction_NewOutput(FText InNodeCategory, FText InDisplayName, FName InTypeName, FText InToolTip, const int32 InGrouping)
+FMetasoundGraphSchemaAction_NewOutput::FMetasoundGraphSchemaAction_NewOutput(FText InNodeCategory, FText InDisplayName, FGuid InOutputNodeID, FText InToolTip, const int32 InGrouping)
 	: FEdGraphSchemaAction(MoveTemp(InNodeCategory), MoveTemp(InDisplayName), MoveTemp(InToolTip), InGrouping)
-	, NodeTypeName(InTypeName)
+	, NodeID(InOutputNodeID)
 {
 }
 
 UEdGraphNode* FMetasoundGraphSchemaAction_NewOutput::PerformAction(UEdGraph* ParentGraph, UEdGraphPin* FromPin, const FVector2D Location, bool bSelectNewNode /* = true */)
 {
 	using namespace Metasound::Editor;
+	using namespace Metasound::Frontend;
 
-	const FScopedTransaction Transaction(LOCTEXT("MetasoundEditorNewOutput", "Add New Metasound Output"));
+	const FScopedTransaction Transaction(LOCTEXT("MetasoundEditorNewInputNode", "Add New Metasound Output Node"));
 
 	ParentGraph->Modify();
 
-	UObject& ParentMetasound = CastChecked<UMetasoundEditorGraph>(ParentGraph)->GetMetasoundChecked();
+	UMetasoundEditorGraph* MetasoundGraph = CastChecked<UMetasoundEditorGraph>(ParentGraph);
+	UObject& ParentMetasound = MetasoundGraph->GetMetasoundChecked();
 	ParentMetasound.Modify();
 
-	FString NewNodeName = FGraphBuilder::GenerateUniqueInputName(ParentMetasound, NodeTypeName);
+	UMetasoundEditorGraphOutput* Output = MetasoundGraph->FindOutput(NodeID);
+	if (!ensure(Output))
+	{
+		return nullptr;
+	}
 
-	FMetasoundFrontendNodeStyle Style;
-	Style.Display.Location = Location;
+	FNodeHandle NodeHandle = Output->GetNodeHandle();
+	if (!ensure(NodeHandle->IsValid()))
+	{
+		return nullptr;
+	}
 
-	if (UEdGraphNode* NewGraphNode = FGraphBuilder::AddOutputNode(ParentMetasound, NewNodeName, NodeTypeName, Style, FText::GetEmpty()))
+	if (UMetasoundEditorGraphOutputNode* NewGraphNode = FGraphBuilder::AddOutputNode(ParentMetasound, NodeHandle, Location, bSelectNewNode))
 	{
 		TryConnectNewNodeToPin(*NewGraphNode, FromPin);
 		return NewGraphNode;
@@ -400,30 +399,35 @@ void UMetasoundEditorGraphSchema::GetPaletteActions(FGraphActionMenuBuilder& Act
 {
 	GetCommentAction(ActionMenuBuilder);
 	GetFunctionActions(ActionMenuBuilder);
-	GetDataTypeInputNodeActions(ActionMenuBuilder);
-	GetDataTypeOutputNodeActions(ActionMenuBuilder);
 }
 
 void UMetasoundEditorGraphSchema::GetGraphContextActions(FGraphContextMenuBuilder& ContextMenuBuilder) const
 {
 	using namespace Metasound;
 	using namespace Metasound::Editor;
+	using namespace Metasound::Frontend;
 
 	FActionClassFilters ClassFilters;
 	if (ContextMenuBuilder.FromPin)
 	{
 		if (ContextMenuBuilder.FromPin->Direction == EGPD_Input)
 		{
-			Frontend::FConstInputHandle InputHandle = Editor::FGraphBuilder::GetConstInputHandleFromPin(ContextMenuBuilder.FromPin);
+			FConstInputHandle InputHandle = FGraphBuilder::GetConstInputHandleFromPin(ContextMenuBuilder.FromPin);
 			ClassFilters.OutputFilterFunction = [InputHandle](const FMetasoundFrontendClassOutput& InOutput)
 			{
 				return InOutput.TypeName == InputHandle->GetDataType();
 			};
 
 			// Show only input nodes as output nodes can only connected if FromPin is input
-			GetDataTypeInputNodeActions(ContextMenuBuilder, [InputHandle](const Editor::FEditorDataType& InDataType)
+			FConstGraphHandle GraphHandle = InputHandle->GetOwningNode()->GetOwningGraph();
+			GetDataTypeInputNodeActions(ContextMenuBuilder, GraphHandle, [InputHandle](FConstNodeHandle NodeHandle)
 			{
-				return InDataType.RegistryInfo.DataTypeName == InputHandle->GetDataType();
+				bool bHasOutputOfType = false;
+				NodeHandle->IterateConstOutputs([&](FConstOutputHandle PotentialOutputHandle)
+				{
+					bHasOutputOfType |= PotentialOutputHandle->GetDataType() == InputHandle->GetDataType();
+				});
+				return bHasOutputOfType;
 			});
 		}
 
@@ -435,13 +439,17 @@ void UMetasoundEditorGraphSchema::GetGraphContextActions(FGraphContextMenuBuilde
 				return InInput.TypeName == OutputHandle->GetDataType();
 			};
 
-			FDataTypeFilterFunction DataTypeFilter = [OutputHandle](const Editor::FEditorDataType& InDataType)
-			{
-				return InDataType.RegistryInfo.DataTypeName == OutputHandle->GetDataType();
-			};
-
 			// Show only output nodes as input nodes can only connected if FromPin is output
-			GetDataTypeOutputNodeActions(ContextMenuBuilder, DataTypeFilter);
+			FConstGraphHandle GraphHandle = OutputHandle->GetOwningNode()->GetOwningGraph();
+			GetDataTypeOutputNodeActions(ContextMenuBuilder, GraphHandle, [OutputHandle](FConstNodeHandle NodeHandle)
+			{
+				bool bHasInputOfType = false;
+				NodeHandle->IterateConstInputs([&](FConstInputHandle PotentialInputHandle)
+				{
+					bHasInputOfType |= PotentialInputHandle->GetDataType() == OutputHandle->GetDataType();
+				});
+				return bHasInputOfType;
+			});
 		}
 	}
 	else
@@ -454,8 +462,15 @@ void UMetasoundEditorGraphSchema::GetGraphContextActions(FGraphContextMenuBuilde
 		}
 
 		GetCommentAction(ContextMenuBuilder, ContextMenuBuilder.CurrentGraph);
-		GetDataTypeInputNodeActions(ContextMenuBuilder);
-		GetDataTypeOutputNodeActions(ContextMenuBuilder);
+		if (UObject* Metasound = MetasoundEditor->GetMetasoundObject())
+		{
+			FMetasoundAssetBase* MetasoundAsset = IMetasoundUObjectRegistry::Get().GetObjectAsAssetBase(Metasound);
+			check(MetasoundAsset);
+			FGraphHandle GraphHandle = MetasoundAsset->GetRootGraphHandle();
+
+			GetDataTypeInputNodeActions(ContextMenuBuilder, GraphHandle);
+			GetDataTypeOutputNodeActions(ContextMenuBuilder, GraphHandle);
+		}
 	}
 
 	GetFunctionActions(ContextMenuBuilder, ClassFilters);
@@ -686,45 +701,7 @@ void UMetasoundEditorGraphSchema::BreakPinLinks(UEdGraphPin& TargetPin, bool bSe
 
 	const FScopedTransaction Transaction(NSLOCTEXT("UnrealEd", "GraphEd_BreakPinLinks", "Break Pin Links"));
 
-	TArray<FInputHandle> InputHandles;
-	TArray<UEdGraphPin*> InputPins;
-
-	UObject& Metasound = CastChecked<UMetasoundEditorGraphNode>(TargetPin.GetOwningNode())->GetMetasoundChecked();
-
-	if (TargetPin.Direction == EGPD_Input)
-	{
-		FNodeHandle NodeHandle = CastChecked<UMetasoundEditorGraphNode>(TargetPin.GetOwningNode())->GetNodeHandle();
-		InputHandles = NodeHandle->GetInputsWithVertexName(TargetPin.GetName());
-		InputPins.Add(&TargetPin);
-	}
-	else
-	{
-		check(TargetPin.Direction == EGPD_Output);
-		for (UEdGraphPin* Pin : TargetPin.LinkedTo)
-		{
-			FNodeHandle NodeHandle = CastChecked<UMetasoundEditorGraphNode>(Pin->GetOwningNode())->GetNodeHandle();
-			InputHandles.Append(NodeHandle->GetInputsWithVertexName(Pin->GetName()));
-			InputPins.Add(Pin);
-		}
-	}
-
-	IMetasoundEditorModule& EditorModule = FModuleManager::GetModuleChecked<IMetasoundEditorModule>("MetasoundEditor");
-	for (int32 i = 0; i < InputHandles.Num(); ++i)
-	{
-		FInputHandle InputHandle = InputHandles[i];
-		FConstOutputHandle OutputHandle = InputHandle->GetCurrentlyConnectedOutput();
-		const FMetasoundFrontendNodeStyle& Style = OutputHandle->GetOwningNode()->GetNodeStyle();
-
-		// Hidden nodes are not "connected" from the perspective of EdGraph,
-		// and therefore should be ignored when breaking links.
-		if (Style.Display.Visibility != EMetasoundFrontendNodeStyleDisplayVisibility::Hidden)
-		{
-			InputHandle->Disconnect();
-
-			FNodeHandle NodeHandle = InputHandle->GetOwningNode();
-			FGraphBuilder::AddOrUpdateLiteralInput(Metasound, NodeHandle, *InputPins[i]);
-		}
-	}
+	FGraphBuilder::DisconnectPin(TargetPin);
 
 	Super::BreakPinLinks(TargetPin, bSendsNodeNotifcation);
 }
@@ -787,18 +764,54 @@ void UMetasoundEditorGraphSchema::GetConversionActions(FGraphActionMenuBuilder&
 	}
 }
 
-void UMetasoundEditorGraphSchema::GetDataTypeInputNodeActions(FGraphActionMenuBuilder& ActionMenuBuilder, Metasound::Editor::FDataTypeFilterFunction InFilter, bool bShowSelectedActions) const
+void UMetasoundEditorGraphSchema::GetDataTypeInputNodeActions(FGraphContextMenuBuilder& ActionMenuBuilder, Metasound::Frontend::FConstGraphHandle InGraphHandle, Metasound::Editor::FInterfaceNodeFilterFunction InFilter, bool bShowSelectedActions) const
 {
+	using namespace Metasound::Editor;
+	using namespace Metasound::Frontend;
+
 	const FText InputMenuName = LOCTEXT("MetasoundActionsInputsMenu", "Inputs");
-	const FText InputTooltipFormat = LOCTEXT("MetasoundTooltipAddInputFormat", "Adds an input of type {0} to the Metasound");
-	Metasound::Editor::GetDataTypeActions<FMetasoundGraphSchemaAction_NewInput>(ActionMenuBuilder, InFilter, InputMenuName, InputTooltipFormat, bShowSelectedActions);
+	const FText InputTooltipFormat = LOCTEXT("MetasoundTooltipAddInputFormat", "Adds an input node referencing '{0}' to the graph.");
+
+	TArray<FConstNodeHandle> Inputs = InGraphHandle->GetConstInputNodes();
+	for (int32 i = Inputs.Num() - 1; i >= 0; --i)
+	{
+		const FConstNodeHandle& NodeHandle = Inputs[i];
+		if (NodeHandle->GetNodeStyle().Display.Visibility == EMetasoundFrontendNodeStyleDisplayVisibility::Hidden)
+		{
+			Inputs.RemoveAtSwap(i, 1, false);
+		}
+	}
+	GetDataTypeActions<FMetasoundGraphSchemaAction_NewInput>(ActionMenuBuilder, Inputs, InFilter, InputMenuName, InputTooltipFormat, bShowSelectedActions);
 }
 
-void UMetasoundEditorGraphSchema::GetDataTypeOutputNodeActions(FGraphActionMenuBuilder& ActionMenuBuilder, Metasound::Editor::FDataTypeFilterFunction InFilter, bool bShowSelectedActions) const
+void UMetasoundEditorGraphSchema::GetDataTypeOutputNodeActions(FGraphContextMenuBuilder& ActionMenuBuilder, Metasound::Frontend::FConstGraphHandle InGraphHandle, Metasound::Editor::FInterfaceNodeFilterFunction InFilter, bool bShowSelectedActions) const
 {
+	using namespace Metasound::Editor;
+	using namespace Metasound::Frontend;
+
 	const FText OutputMenuName = LOCTEXT("MetasoundActionsOutputsMenu", "Outputs");
-	const FText OutputTooltipFormat = LOCTEXT("MetasoundTooltipAddOutputFormat", "Adds an output of type {0} to the Metasound");
-	Metasound::Editor::GetDataTypeActions<FMetasoundGraphSchemaAction_NewOutput>(ActionMenuBuilder, InFilter, OutputMenuName, OutputTooltipFormat, bShowSelectedActions);
+	const FText OutputTooltipFormat = LOCTEXT("MetasoundTooltipAddOutputFormat", "Adds an output node referencing '{0}' to the graph.");
+
+	TArray<FConstNodeHandle> Outputs = InGraphHandle->GetConstOutputNodes();
+
+	// Prune and only add actions for outputs that are not already represented in the graph
+	// (as there should only be one output reference node ever to avoid confusion with which
+	// is handling active input)
+	if (const UMetasoundEditorGraph* Graph = Cast<UMetasoundEditorGraph>(ActionMenuBuilder.CurrentGraph))
+	{
+		for (int32 i = Outputs.Num() - 1; i >= 0; --i)
+		{
+			if (UMetasoundEditorGraphOutput* Output = Graph->FindOutput(Outputs[i]->GetID()))
+			{
+				if (!Output->GetNodes().IsEmpty())
+				{
+					Outputs.RemoveAtSwap(i, 1, false /* bAllowShrinking */);
+				}
+			}
+		}
+	}
+
+	GetDataTypeActions<FMetasoundGraphSchemaAction_NewOutput>(ActionMenuBuilder, Outputs, InFilter, OutputMenuName, OutputTooltipFormat, bShowSelectedActions);
 }
 
 void UMetasoundEditorGraphSchema::GetFunctionActions(FGraphActionMenuBuilder& ActionMenuBuilder, Metasound::Editor::FActionClassFilters InFilters, bool bShowSelectedActions) const
diff --git a/Engine/Plugins/Runtime/Metasound/Source/MetasoundEditor/Private/MetasoundEditorModule.cpp b/Engine/Plugins/Runtime/Metasound/Source/MetasoundEditor/Private/MetasoundEditorModule.cpp
--- a/Engine/Plugins/Runtime/Metasound/Source/MetasoundEditor/Private/MetasoundEditorModule.cpp
+++ b/Engine/Plugins/Runtime/Metasound/Source/MetasoundEditor/Private/MetasoundEditorModule.cpp
@@ -14,12 +14,13 @@
 #include "MetasoundSource.h"
 #include "MetasoundAssetTypeActions.h"
 #include "MetasoundDetailCustomization.h"
-#include "MetasoundInputNodeDetailCustomization.h"
+#include "MetasoundEditorGraph.h"
 #include "MetasoundEditorGraphBuilder.h"
 #include "MetasoundEditorGraphNodeFactory.h"
 #include "MetasoundEditorGraphSchema.h"
 #include "MetasoundEditorSettings.h"
 #include "MetasoundFrontendRegistries.h"
+#include "MetasoundNodeDetailCustomization.h"
 #include "Modules/ModuleInterface.h"
 #include "Modules/ModuleManager.h"
 #include "PropertyEditorDelegates.h"
@@ -76,6 +77,7 @@ namespace Metasound
 					Set("MetasoundEditor.Export.Small", new FSlateImageBrush(RootToContentDir(TEXT("/Icons/build_40x.png")), Icon20x20));
 					Set("MetasoundEditor.ExportError", new FSlateImageBrush(RootToContentDir(TEXT("/Icons/build_error_40x.png")), Icon40x40));
 					Set("MetasoundEditor.ExportError.Small", new FSlateImageBrush(RootToContentDir(TEXT("/Icons/build_error_40x.png")), Icon20x20));
+					Set("MetasoundEditor.Settings", new FSlateImageBrush(RootToContentDir(TEXT("/Icons/settings_40x.png")), Icon20x20));
 
 					// Graph Editor
 					Set("MetasoundEditor.Graph.Node.Body.Input", new FSlateImageBrush(RootToContentDir(TEXT("/Graph/node_input_body_64x.png")), FVector2D(114.0f, 64.0f)));
@@ -102,7 +104,7 @@ namespace Metasound
 		{
 			void RegisterNodeInputClasses()
 			{
-				TSubclassOf<UMetasoundEditorGraphInputNode> NodeClass;
+				TSubclassOf<UMetasoundEditorGraphInputLiteral> NodeClass;
 				for (TObjectIterator<UClass> ClassIt; ClassIt; ++ClassIt)
 				{
 					UClass* Class = *ClassIt;
@@ -116,19 +118,19 @@ namespace Metasound
 						continue;
 					}
 
-					if (!ClassIt->IsChildOf(UMetasoundEditorGraphInputNode::StaticClass()))
+					if (!ClassIt->IsChildOf(UMetasoundEditorGraphInputLiteral::StaticClass()))
 					{
 						continue;
 					}
 
-					if (const UMetasoundEditorGraphInputNode* InputNode = Class->GetDefaultObject<UMetasoundEditorGraphInputNode>())
+					if (const UMetasoundEditorGraphInputLiteral* InputCDO = Class->GetDefaultObject<UMetasoundEditorGraphInputLiteral>())
 					{
-						NodeInputClassRegistry.Add(InputNode->GetLiteralType(), InputNode->GetClass());
+						NodeInputClassRegistry.Add(InputCDO->GetLiteralType(), InputCDO->GetClass());
 					}
 				}
 			}
 
-			const TSubclassOf<UMetasoundEditorGraphInputNode> FindNodeInputClass(EMetasoundFrontendLiteralType InLiteralType) const override
+			const TSubclassOf<UMetasoundEditorGraphInputLiteral> FindInputLiteralClass(EMetasoundFrontendLiteralType InLiteralType) const override
 			{
 				return NodeInputClassRegistry.FindRef(InLiteralType);
 			}
@@ -183,12 +185,14 @@ namespace Metasound
 						switch (RegistryInfo.PreferredLiteralType)
 						{
 							case ELiteralType::Boolean:
+							case ELiteralType::BooleanArray:
 							{
 								PinCategory = FGraphBuilder::PinCategoryBoolean;
 							}
 							break;
 
 							case ELiteralType::Float:
+							case ELiteralType::FloatArray:
 							{
 								PinCategory = FGraphBuilder::PinCategoryFloat;
 
@@ -212,38 +216,28 @@ namespace Metasound
 							break;
 
 							case ELiteralType::Integer:
+							case ELiteralType::IntegerArray:
 							{
 								PinCategory = FGraphBuilder::PinCategoryInt32;
 							}
 							break;
 
 							case ELiteralType::String:
+							case ELiteralType::StringArray:
 							{
 								PinCategory = FGraphBuilder::PinCategoryString;
 							}
 							break;
 
 							case ELiteralType::UObjectProxy:
-							{
-								PinCategory = FGraphBuilder::PinCategoryObject;
-							}
-							break;
-
 							case ELiteralType::UObjectProxyArray:
 							{
-								// TODO: Implement, or will be nuked in favor of general array support for all types
+								PinCategory = FGraphBuilder::PinCategoryObject;
 							}
 							break;
 
-							
 							case ELiteralType::None:
 							case ELiteralType::Invalid:
-							// TODO: handle array types.
-							case ELiteralType::NoneArray:
-							case ELiteralType::BooleanArray:
-							case ELiteralType::IntegerArray:
-							case ELiteralType::StringArray:
-							case ELiteralType::FloatArray:
 							default:
 							{
 								// Audio types are ubiquitous, so added as subcategory
@@ -284,16 +278,20 @@ namespace Metasound
 					FOnGetDetailCustomizationInstance::CreateLambda([]() { return MakeShared<FMetasoundDetailCustomization>(UMetasoundSource::GetDocumentPropertyName()); }));
 
 				PropertyModule.RegisterCustomClassLayout(
-					UMetasoundEditorGraphInputNode::StaticClass()->GetFName(),
-					FOnGetDetailCustomizationInstance::CreateLambda([]() { return MakeShared<FMetasoundInputNodeDetailCustomization>(); }));
+					UMetasoundEditorGraphInput::StaticClass()->GetFName(),
+					FOnGetDetailCustomizationInstance::CreateLambda([]() { return MakeShared<FMetasoundInputDetailCustomization>(); }));
+
+				PropertyModule.RegisterCustomClassLayout(
+					UMetasoundEditorGraphOutput::StaticClass()->GetFName(),
+					FOnGetDetailCustomizationInstance::CreateLambda([]() { return MakeShared<FMetasoundOutputDetailCustomization>(); }));
 
 				PropertyModule.RegisterCustomPropertyTypeLayout(
-					"MetasoundEditorGraphInputInt",
-					FOnGetPropertyTypeCustomizationInstance::CreateLambda([]() { return MakeShared<FMetasoundInputNodeIntDetailCustomization>(); }));
+					"MetasoundEditorGraphInputIntRef",
+					FOnGetPropertyTypeCustomizationInstance::CreateLambda([]() { return MakeShared<FMetasoundInputIntDetailCustomization>(); }));
 
 				PropertyModule.RegisterCustomPropertyTypeLayout(
-					"MetasoundEditorGraphInputObject",
-					FOnGetPropertyTypeCustomizationInstance::CreateLambda([]() { return MakeShared<FMetasoundInputNodeObjectDetailCustomization>(); }));
+					"MetasoundEditorGraphInputObjectRef",
+					FOnGetPropertyTypeCustomizationInstance::CreateLambda([]() { return MakeShared<FMetasoundInputObjectDetailCustomization>(); }));
 
 				StyleSet = MakeShared<FSlateStyle>();
 
@@ -357,7 +355,7 @@ namespace Metasound
 
 			TArray<TSharedPtr<FAssetTypeActions_Base>> AssetActions;
 			TMap<FName, FEditorDataType> DataTypeInfo;
-			TMap<EMetasoundFrontendLiteralType, const TSubclassOf<UMetasoundEditorGraphInputNode>> NodeInputClassRegistry;
+			TMap<EMetasoundFrontendLiteralType, const TSubclassOf<UMetasoundEditorGraphInputLiteral>> NodeInputClassRegistry;
 
 			TSharedPtr<FMetasoundGraphNodeFactory> GraphNodeFactory;
 			TSharedPtr<FGraphPanelPinConnectionFactory> GraphConnectionFactory;
diff --git a/Engine/Plugins/Runtime/Metasound/Source/MetasoundEditor/Private/MetasoundEditorModule.h b/Engine/Plugins/Runtime/Metasound/Source/MetasoundEditor/Private/MetasoundEditorModule.h
--- a/Engine/Plugins/Runtime/Metasound/Source/MetasoundEditor/Private/MetasoundEditorModule.h
+++ b/Engine/Plugins/Runtime/Metasound/Source/MetasoundEditor/Private/MetasoundEditorModule.h
@@ -9,8 +9,9 @@
 #include "Modules/ModuleInterface.h"
 #include "Templates/Function.h"
 
+
 class UMetasoundEditorGraph;
-class UMetasoundEditorGraphInputNode;
+class UMetasoundEditorGraphInputLiteral;
 
 DECLARE_LOG_CATEGORY_EXTERN(LogMetasoundEditor, Log, All);
 
@@ -38,7 +39,7 @@ namespace Metasound
 			virtual const FEditorDataType& FindDataType(FName InDataTypeName) const = 0;
 			virtual void IterateDataTypes(TUniqueFunction<void(const FEditorDataType&)> InDataTypeFunction) const = 0;
 
-			virtual const TSubclassOf<UMetasoundEditorGraphInputNode> FindNodeInputClass(EMetasoundFrontendLiteralType InLiteralType) const = 0;
+			virtual const TSubclassOf<UMetasoundEditorGraphInputLiteral> FindInputLiteralClass(EMetasoundFrontendLiteralType InLiteralType) const = 0;
 		};
 	} // namespace Editor
 } // namespace Metasound
diff --git a/Engine/Plugins/Runtime/Metasound/Source/MetasoundEditor/Private/MetasoundEditorTabFactory.h b/Engine/Plugins/Runtime/Metasound/Source/MetasoundEditor/Private/MetasoundEditorTabFactory.h
--- a/Engine/Plugins/Runtime/Metasound/Source/MetasoundEditor/Private/MetasoundEditorTabFactory.h
+++ b/Engine/Plugins/Runtime/Metasound/Source/MetasoundEditor/Private/MetasoundEditorTabFactory.h
@@ -65,7 +65,7 @@ namespace Metasound
 					];
 			}
 
-			TSharedRef<SDockTab> CreateMetasoundTab(TSharedPtr<IDetailsView> DetailsView, const FSpawnTabArgs& Args)
+			TSharedRef<SDockTab> CreateMetasoundTab(TSharedPtr<SGraphActionMenu> GraphActionMenu, const FSpawnTabArgs& Args)
 			{
 				check(Args.GetTabId() == Names::Metasound);
 
@@ -73,7 +73,7 @@ namespace Metasound
 					.Icon(FAppStyle::Get().GetBrush("Icons.Settings"))
 					.Label(LOCTEXT("MetasoundGeneralTitle", "Metasound"))
 					[
-						DetailsView.ToSharedRef()
+						GraphActionMenu.ToSharedRef()
 					];
 			}
 
diff --git a/Engine/Plugins/Runtime/Metasound/Source/MetasoundEditor/Private/MetasoundInputNodeDetailCustomization.cpp b/Engine/Plugins/Runtime/Metasound/Source/MetasoundEditor/Private/MetasoundInputNodeDetailCustomization.cpp
deleted file mode 100644
diff --git a/Engine/Plugins/Runtime/Metasound/Source/MetasoundEditor/Private/MetasoundInputNodeDetailCustomization.h b/Engine/Plugins/Runtime/Metasound/Source/MetasoundEditor/Private/MetasoundInputNodeDetailCustomization.h
deleted file mode 100644
diff --git a/Engine/Plugins/Runtime/Metasound/Source/MetasoundEditor/Private/MetasoundNodeDetailCustomization.cpp b/Engine/Plugins/Runtime/Metasound/Source/MetasoundEditor/Private/MetasoundNodeDetailCustomization.cpp
new file mode 100644
--- /dev/null
+++ b/Engine/Plugins/Runtime/Metasound/Source/MetasoundEditor/Private/MetasoundNodeDetailCustomization.cpp
@@ -0,0 +1,951 @@
+// Copyright Epic Games, Inc. All Rights Reserved.
+#include "MetasoundNodeDetailCustomization.h"
+
+#include "Containers/Set.h"
+#include "CoreMinimal.h"
+#include "Delegates/Delegate.h"
+#include "DetailCategoryBuilder.h"
+#include "DetailLayoutBuilder.h"
+#include "DetailWidgetRow.h"
+#include "Framework/Notifications/NotificationManager.h"
+#include "IDetailChildrenBuilder.h"
+#include "IDetailGroup.h"
+#include "Internationalization/Text.h"
+#include "MetasoundAssetBase.h"
+#include "MetasoundEditorGraphBuilder.h"
+#include "MetasoundEditorGraphNode.h"
+#include "MetasoundEditorGraphInputNodes.h"
+#include "MetasoundFrontend.h"
+#include "MetasoundFrontendController.h"
+#include "MetasoundFrontendRegistries.h"
+#include "MetasoundUObjectRegistry.h"
+#include "PropertyCustomizationHelpers.h"
+#include "PropertyEditorDelegates.h"
+#include "PropertyHandle.h"
+#include "PropertyRestriction.h"
+#include "SlateCore/Public/Styling/SlateColor.h"
+#include "Templates/Casts.h"
+#include "Templates/SharedPointer.h"
+#include "UObject/WeakObjectPtr.h"
+#include "UObject/WeakObjectPtrTemplates.h"
+#include "Widgets/Input/SCheckBox.h"
+#include "Widgets/Input/SMultiLineEditableTextBox.h"
+#include "Widgets/Input/STextComboBox.h"
+#include "Widgets/Notifications/SNotificationList.h"
+#include "Widgets/SToolTip.h"
+#include "Widgets/SWidget.h"
+#include "Widgets/Text/STextBlock.h"
+#include "MetasoundEditorModule.h"
+
+
+#define LOCTEXT_NAMESPACE "MetasoundEditor"
+
+namespace Metasound
+{
+	namespace Editor
+	{
+		namespace VariableCustomizationPrivate
+		{
+			/** Minimum size of the details title panel */
+			static const float DetailsTitleMinWidth = 125.f;
+			/** Maximum size of the details title panel */
+			static const float DetailsTitleMaxWidth = 300.f;
+			/** magic number retrieved from SGraphNodeComment::GetWrapAt() */
+			static const float DetailsTitleWrapPadding = 32.0f;
+
+			static const FString ArrayIdentifier = TEXT(":Array");
+
+			static const FText DataTypeNameText = LOCTEXT("Node_DataTypeName", "Type");
+			static const FText DefaultPropertyText = LOCTEXT("Node_DefaultPropertyName", "Default Value");
+			static const FText NodeTooltipText = LOCTEXT("Node_Tooltip", "Tooltip");
+
+			static const FText InputNameText = LOCTEXT("Input_Name", "Input Name");
+			static const FText OutputNameText = LOCTEXT("Output_Name", "Output Name");
+
+			static const FName DataTypeNameIdentifier = "DataTypeName";
+			static const FName ProxyGeneratorClassNameIdentifier = "GeneratorClass";
+		}
+
+		void FMetasoundInputIntDetailCustomization::CacheProxyData(TSharedPtr<IPropertyHandle> ProxyHandle)
+		{
+			DataTypeName = FName();
+
+			const FString* MetadataDataTypeName = ProxyHandle->GetInstanceMetaData(VariableCustomizationPrivate::DataTypeNameIdentifier);
+			if (ensure(MetadataDataTypeName))
+			{
+				DataTypeName = **MetadataDataTypeName;
+			}
+		}
+
+		TSharedRef<SWidget> FMetasoundInputIntDetailCustomization::CreateStructureWidget(TSharedPtr<IPropertyHandle>& StructPropertyHandle) const
+		{
+			using namespace Frontend;
+
+			if (FMetasoundFrontendRegistryContainer* Registry = FMetasoundFrontendRegistryContainer::Get())
+			{
+				TSharedPtr<IPropertyHandle> ValueProperty = StructPropertyHandle->GetChildHandle(GET_MEMBER_NAME_CHECKED(FMetasoundEditorGraphInputIntRef, Value));
+				if (ValueProperty.IsValid())
+				{
+					TSharedPtr<const IEnumDataTypeInterface> EnumInterface = Registry->GetEnumInterfaceForDataType(DataTypeName);
+
+					// Not an enum, so just display as underlying type (int32)
+					if (!EnumInterface.IsValid())
+					{
+						return ValueProperty->CreatePropertyValueWidget();
+					}
+
+					auto GetAll = [Interface = EnumInterface](TArray<TSharedPtr<FString>>& OutStrings, TArray<TSharedPtr<SToolTip>>& OutTooltips, TArray<bool>&)
+					{
+						for (const IEnumDataTypeInterface::FGenericInt32Entry& i : Interface->GetAllEntries())
+						{
+							OutTooltips.Emplace(SNew(SToolTip).Text(i.Tooltip));
+							OutStrings.Emplace(MakeShared<FString>(i.DisplayName.ToString()));
+						}
+					};
+					auto GetValue = [Interface = EnumInterface, Prop = ValueProperty]()
+					{
+						int32 IntValue;
+						if (Prop->GetValue(IntValue) != FPropertyAccess::Success)
+						{
+							IntValue = Interface->GetDefaultValue();
+							UE_LOG(LogTemp, Warning, TEXT("Failed to read int Property '%s', defaulting."), *GetNameSafe(Prop->GetProperty()));
+						}
+						if (TOptional<IEnumDataTypeInterface::FGenericInt32Entry> Result = Interface->FindByValue(IntValue))
+						{
+							return Result->DisplayName.ToString();
+						}
+						UE_LOG(LogTemp, Warning, TEXT("Failed to resolve int value '%d' to a valid enum value for enum '%s'"),
+							IntValue, *Interface->GetNamespace().ToString());
+
+						// Return default (should always succeed as we can't have empty Enums and we must have a default).
+						return Interface->FindByValue(Interface->GetDefaultValue())->DisplayName.ToString();
+					};
+					auto SelectedValue = [Interface = EnumInterface, Prop = ValueProperty](const FString& InSelected)
+					{
+						TOptional<IEnumDataTypeInterface::FGenericInt32Entry> Found =
+							Interface->FindEntryBy([TextSelected = FText::FromString(InSelected)](const IEnumDataTypeInterface::FGenericInt32Entry& i)
+						{
+							return i.DisplayName.EqualTo(TextSelected);
+						});
+
+						if (Found)
+						{
+							// Only save the changes if its different and we can read the old value to check that.
+							int32 CurrentValue;
+							bool bReadCurrentValue = Prop->GetValue(CurrentValue) == FPropertyAccess::Success;
+							if ((bReadCurrentValue && CurrentValue != Found->Value) || !bReadCurrentValue)
+							{
+								ensure(Prop->SetValue(Found->Value) == FPropertyAccess::Success);
+							}
+						}
+						else
+						{
+							UE_LOG(LogTemp, Warning, TEXT("Failed to Set Valid Value for Property '%s' with Value of '%s', writing default."),
+								*GetNameSafe(Prop->GetProperty()), *InSelected);
+
+							ensure(Prop->SetValue(Interface->GetDefaultValue()) == FPropertyAccess::Success);
+						}
+					};
+
+					return PropertyCustomizationHelpers::MakePropertyComboBox(
+						nullptr,
+						FOnGetPropertyComboBoxStrings::CreateLambda(GetAll),
+						FOnGetPropertyComboBoxValue::CreateLambda(GetValue),
+						FOnPropertyComboBoxValueSelected::CreateLambda(SelectedValue)
+					);
+				}
+			}
+
+			return SNullWidget::NullWidget;
+		}
+
+		void FMetasoundInputObjectDetailCustomization::CacheProxyData(TSharedPtr<IPropertyHandle> ProxyHandle)
+		{
+			ProxyGenClass.Reset();
+
+			const FString* MetadataProxyGenClass = ProxyHandle->GetInstanceMetaData(VariableCustomizationPrivate::ProxyGeneratorClassNameIdentifier);
+			TSharedPtr<IPropertyHandle> MetadataHandle = ProxyHandle->GetParentHandle();
+			if (!ensure(MetadataProxyGenClass))
+			{
+				return;
+			}
+
+			const FName ClassName = FName(*MetadataProxyGenClass);
+			for (TObjectIterator<UClass> ClassIt; ClassIt; ++ClassIt)
+			{
+				UClass* Class = *ClassIt;
+				if (!Class->IsNative())
+				{
+					continue;
+				}
+				
+				if (Class->HasAnyClassFlags(CLASS_Abstract | CLASS_Deprecated | CLASS_NewerVersionExists))
+				{
+					continue;
+				}
+
+				if (ClassIt->GetFName() != ClassName)
+				{
+					continue;
+				}
+
+				ProxyGenClass = *ClassIt;
+				return;
+			}
+
+			ensureMsgf(false, TEXT("Failed to find ProxyGeneratorClass. Class not set "));
+		}
+
+		TSharedRef<SWidget> FMetasoundInputObjectDetailCustomization::CreateStructureWidget(TSharedPtr<IPropertyHandle>& StructPropertyHandle) const
+		{
+			TSharedPtr<IPropertyHandle> PropertyHandle = StructPropertyHandle->GetChildHandle(GET_MEMBER_NAME_CHECKED(FMetasoundEditorGraphInputObjectRef, Object));
+
+			auto ValidateAsset = [InProxyGenClass = ProxyGenClass](const FAssetData& InAsset)
+			{
+				if (!InProxyGenClass.IsValid())
+				{
+					return false;
+				}
+
+				if (UObject* Object = InAsset.GetAsset())
+				{
+					if (UClass* Class = Object->GetClass())
+					{
+						return Class->IsChildOf(InProxyGenClass.Get());
+					}
+				}
+
+				return false;
+			};
+
+			auto CommitAsset = [InPropertyHandle = PropertyHandle, InProxyGenClass = ProxyGenClass](const FAssetData& InAssetData)
+			{
+				// if we've hit this code, the presumption is that the datatype for this parameter has already defined a corresponding UClass that can be used to set it.
+				ensureAlways(InProxyGenClass.IsValid());
+
+				UObject* InObject = InAssetData.GetAsset();
+
+				InPropertyHandle->SetValue(InObject);
+			};
+
+			auto GetAssetPath = [PropertyHandle = PropertyHandle]()
+			{
+				UObject* Object = nullptr;
+				if (PropertyHandle->GetValue(Object) == FPropertyAccess::Success)
+				{
+					return Object->GetPathName();
+				}
+				return FString();
+			};
+
+			TArray<const UClass*> AllowedClasses;
+			AllowedClasses.Add(ProxyGenClass.Get());
+
+			return SNew(SObjectPropertyEntryBox)
+				.ObjectPath_Lambda(GetAssetPath)
+				.AllowedClass(ProxyGenClass.Get())
+				.OnShouldSetAsset_Lambda(ValidateAsset)
+				.OnObjectChanged_Lambda(CommitAsset)
+				.AllowClear(false)
+				.DisplayUseSelected(true)
+				.DisplayBrowse(true)
+				.DisplayThumbnail(true)
+				.NewAssetFactories(PropertyCustomizationHelpers::GetNewAssetFactoriesForClasses(AllowedClasses));
+		}
+
+		TSharedRef<SWidget> FMetasoundInputArrayDetailCustomizationBase::CreateValueWidget(TSharedPtr<IPropertyHandleArray> ParentArrayProperty, TSharedPtr<IPropertyHandle> StructPropertyHandle, bool bIsInArray) const
+		{
+			TSharedRef<SWidget> ObjectPickerWidget = CreateStructureWidget(StructPropertyHandle);
+			if (!bIsInArray)
+			{
+				return ObjectPickerWidget;
+			}
+
+			TSharedPtr<IPropertyHandle> StructPropertyPtr = StructPropertyHandle;
+			FExecuteAction InsertAction = FExecuteAction::CreateLambda([ParentArrayProperty, StructPropertyPtr]
+			{
+				const int32 ArrayIndex = StructPropertyPtr.IsValid() ? StructPropertyPtr->GetIndexInArray() : INDEX_NONE;
+				if (ParentArrayProperty.IsValid() && ArrayIndex >= 0)
+				{
+					ParentArrayProperty->Insert(ArrayIndex);
+				}
+			});
+
+			FExecuteAction DeleteAction = FExecuteAction::CreateLambda([ParentArrayProperty, StructPropertyPtr]
+			{
+				const int32 ArrayIndex = StructPropertyPtr.IsValid() ? StructPropertyPtr->GetIndexInArray() : INDEX_NONE;
+				if (ParentArrayProperty.IsValid() && ArrayIndex >= 0)
+				{
+					ParentArrayProperty->DeleteItem(ArrayIndex);
+				}
+			});
+
+			FExecuteAction DuplicateAction = FExecuteAction::CreateLambda([ParentArrayProperty, StructPropertyPtr]
+			{
+				const int32 ArrayIndex = StructPropertyPtr.IsValid() ? StructPropertyPtr->GetIndexInArray() : INDEX_NONE;
+				if (ParentArrayProperty.IsValid() && ArrayIndex >= 0)
+				{
+					ParentArrayProperty->DuplicateItem(ArrayIndex);
+				}
+			});
+
+			return SNew(SHorizontalBox)
+				+ SHorizontalBox::Slot()
+				.FillWidth(0.95f)
+				.Padding(1.0f, 0.0f, 0.0f, 0.0f)
+				.VAlign(VAlign_Center)
+				[
+					ObjectPickerWidget
+				]
+				+ SHorizontalBox::Slot()
+				.FillWidth(0.05f)
+				.Padding(1.0f, 0.0f, 0.0f, 0.0f)
+				.VAlign(VAlign_Center)
+				[
+					PropertyCustomizationHelpers::MakeInsertDeleteDuplicateButton(InsertAction, DeleteAction, DuplicateAction)
+				];
+		}
+
+		void FMetasoundInputArrayDetailCustomizationBase::CustomizeChildren(TSharedRef<IPropertyHandle> StructPropertyHandle, IDetailChildrenBuilder& ChildBuilder, IPropertyTypeCustomizationUtils& StructCustomizationUtils)
+		{
+			bool bIsInArray = false;
+			TSharedPtr<IPropertyHandleArray> ParentArrayProperty;
+			TSharedPtr<IPropertyHandle> ProxyProperty = StructPropertyHandle;
+			{
+				TSharedPtr<IPropertyHandle> ParentProperty = ProxyProperty->GetParentHandle();
+				if (ProxyProperty.IsValid() && ParentProperty.IsValid())
+				{
+					ParentArrayProperty = ParentProperty->AsArray();
+					if (ParentArrayProperty.IsValid())
+					{
+						ProxyProperty = ParentProperty;
+						bIsInArray = true;
+					}
+				}
+			}
+
+			CacheProxyData(ProxyProperty);
+
+			TSharedRef<SWidget> ValueWidget = CreateValueWidget(ParentArrayProperty, StructPropertyHandle, bIsInArray);
+
+			FDetailWidgetRow& ValueRow = ChildBuilder.AddCustomRow(VariableCustomizationPrivate::DefaultPropertyText);
+			if (bIsInArray)
+			{
+				ValueRow.NameContent()
+				[
+					StructPropertyHandle->CreatePropertyNameWidget()
+				];
+			}
+			else
+			{
+				ValueRow.NameContent()
+				[
+					SNew(STextBlock)
+					.Text(VariableCustomizationPrivate::DefaultPropertyText)
+					.Font(IDetailLayoutBuilder::GetDetailFont())
+				];
+			}
+
+			TArray<UObject*> OuterObjects;
+			StructPropertyHandle->GetOuterObjects(OuterObjects);
+			TArray<TWeakObjectPtr<UMetasoundEditorGraphInput>> Inputs;
+			for (UObject* Object : OuterObjects)
+			{
+				if (UMetasoundEditorGraphInput* Input = Cast<UMetasoundEditorGraphInput>(Object))
+				{
+					Inputs.Add(Input);
+				}
+			}
+
+			FSimpleDelegate OnLiteralChanged = FSimpleDelegate::CreateLambda([InInputs = Inputs]()
+			{
+				for (const TWeakObjectPtr<UMetasoundEditorGraphInput>& GraphInput : InInputs)
+				{
+					if (GraphInput.IsValid())
+					{
+						GraphInput->OnLiteralChanged();
+					}
+				}
+			});
+			StructPropertyHandle->SetOnChildPropertyValueChanged(OnLiteralChanged);
+
+			ValueRow.ValueContent()
+			[
+				ValueWidget
+			];
+		}
+
+		void FMetasoundInputArrayDetailCustomizationBase::CustomizeHeader(TSharedRef<IPropertyHandle> StructPropertyHandle, FDetailWidgetRow& HeaderRow, IPropertyTypeCustomizationUtils& StructCustomizationUtils)
+		{
+		}
+
+		void FMetasoundVariableDataTypeSelector::AddDataTypeSelector(IDetailLayoutBuilder& InDetailLayout, const FText& InRowName, TWeakObjectPtr<UMetasoundEditorGraphVariable> InGraphVariable, bool bIsRequired)
+		{
+			DetailLayoutBuilder = &InDetailLayout;
+
+			IDetailCategoryBuilder& CategoryBuilder = InDetailLayout.EditCategory("General");
+
+			TSharedPtr<FString> CurrentTypeString;
+			FString CurrentTypeName = InGraphVariable->TypeName.ToString();
+			bool bCurrentTypeIsArray = CurrentTypeName.EndsWith(VariableCustomizationPrivate::ArrayIdentifier);
+			if (bCurrentTypeIsArray)
+			{
+				CurrentTypeName.LeftChopInline(VariableCustomizationPrivate::ArrayIdentifier.Len());
+			}
+
+			DataTypeNames.Reset();
+			IMetasoundEditorModule& EditorModule = FModuleManager::GetModuleChecked<IMetasoundEditorModule>("MetasoundEditor");
+			EditorModule.IterateDataTypes([&](const FEditorDataType& EditorDataType)
+			{
+				const FString TypeName = EditorDataType.RegistryInfo.DataTypeName.ToString();
+
+				// Array types are handles separately via checkbox
+				if (TypeName.EndsWith(VariableCustomizationPrivate::ArrayIdentifier))
+				{
+					return;
+				}
+
+				TSharedPtr<FString> TypeStrPtr = MakeShared<FString>(TypeName);
+				if (TypeName == CurrentTypeName)
+				{
+					CurrentTypeString = TypeStrPtr;
+				}
+				DataTypeNames.Add(TypeStrPtr);
+			});
+
+			if (!ensure(CurrentTypeString.IsValid()))
+			{
+				return;
+			}
+
+			DataTypeNames.Sort([](const TSharedPtr<FString>& DataTypeNameL, const TSharedPtr<FString>& DataTypeNameR)
+			{
+				if (DataTypeNameL.IsValid() && DataTypeNameR.IsValid())
+				{
+					return DataTypeNameR->Compare(*DataTypeNameL.Get()) > 0;
+				}
+				return false;
+			});
+
+			CategoryBuilder.AddCustomRow(InRowName)
+			.IsEnabled(!bIsRequired)
+			.NameContent()
+			[
+				SNew(STextBlock)
+				.Text(InRowName)
+				.Font(IDetailLayoutBuilder::GetDetailFont())
+			]
+			.ValueContent()
+			[
+				SNew(SHorizontalBox)
+				+ SHorizontalBox::Slot()
+				.FillWidth(0.60f)
+				.Padding(1.0f, 0.0f, 0.0f, 0.0f)
+				.VAlign(VAlign_Center)
+				[
+					SAssignNew(DataTypeComboBox, STextComboBox)
+					.OptionsSource(&DataTypeNames)
+					.InitiallySelectedItem(CurrentTypeString)
+					.OnSelectionChanged_Lambda([this, InGraphVariable](TSharedPtr<FString> ItemSelected, ESelectInfo::Type SelectInfo)
+					{
+						OnBaseDataTypeChanged(InGraphVariable, ItemSelected, SelectInfo);
+					})
+					.IsEnabled(!bIsRequired)
+				]
+				+ SHorizontalBox::Slot()
+				.FillWidth(0.40f)
+				.Padding(1.0f, 0.0f, 0.0f, 0.0f)
+				.VAlign(VAlign_Center)
+				[
+					SAssignNew(DataTypeArrayCheckbox, SCheckBox)
+					.IsChecked_Lambda([this, InGraphVariable]()
+					{
+						return OnGetDataTypeArrayCheckState(InGraphVariable);
+					})
+					.OnCheckStateChanged_Lambda([this, InGraphVariable](ECheckBoxState InNewState)
+					{
+						OnDataTypeArrayChanged(InGraphVariable, InNewState);
+					})
+					[
+						SNew(STextBlock)
+						.Text(LOCTEXT("Node_IsArray", "Is Array"))
+						.Font(IDetailLayoutBuilder::GetDetailFont())
+					]
+				]
+			];
+		}
+
+		ECheckBoxState FMetasoundVariableDataTypeSelector::OnGetDataTypeArrayCheckState(TWeakObjectPtr<UMetasoundEditorGraphVariable> InGraphVariable) const
+		{
+			if (InGraphVariable.IsValid())
+			{
+				FString CurrentTypeName = InGraphVariable->TypeName.ToString();
+				bool bCurrentTypeIsArray = CurrentTypeName.EndsWith(VariableCustomizationPrivate::ArrayIdentifier);
+				return bCurrentTypeIsArray ? ECheckBoxState::Checked : ECheckBoxState::Unchecked;
+			}
+
+			return ECheckBoxState::Undetermined;
+		}
+
+		void FMetasoundInputDetailCustomization::CustomizeDetails(IDetailLayoutBuilder& DetailLayout)
+		{
+			using namespace Frontend;
+
+			TMetasoundVariableDetailCustomization<UMetasoundEditorGraphInput>::CustomizeDetails(DetailLayout);
+
+			if (!GraphVariable.IsValid())
+			{
+				return;
+			}
+
+			IDetailCategoryBuilder& CategoryBuilder = DetailLayout.EditCategory("General");
+
+			const bool bIsRequired = IsRequired();
+			DisplayNameEditableTextBox = SNew(SEditableTextBox)
+				.Text(this, &FMetasoundInputDetailCustomization::GetDisplayName)
+				.OnTextChanged(this, &FMetasoundInputDetailCustomization::OnDisplayNameChanged)
+				.OnTextCommitted(this, &FMetasoundInputDetailCustomization::OnDisplayNameCommitted)
+				.IsReadOnly(bIsRequired)
+				.Font(IDetailLayoutBuilder::GetDetailFont());
+
+			CategoryBuilder.AddCustomRow(VariableCustomizationPrivate::InputNameText)
+			.EditCondition(!bIsRequired, nullptr)
+			.NameContent()
+			[
+				SNew(STextBlock)
+				.Font(IDetailLayoutBuilder::GetDetailFontBold())
+				.Text(VariableCustomizationPrivate::InputNameText)
+				.ToolTipText(TAttribute<FText>::Create([GraphVariable = this->GraphVariable]()
+				{
+					if (GraphVariable.IsValid())
+					{
+						FNodeHandle NodeHandle = GraphVariable->GetNodeHandle();
+						FMetasoundFrontendNodeStyle NodeStyle = NodeHandle->GetNodeStyle();
+						return NodeHandle->GetDescription();
+					}
+
+					return FText::GetEmpty();
+				}))
+			]
+			.ValueContent()
+			[
+				DisplayNameEditableTextBox.ToSharedRef()
+			];
+
+			CategoryBuilder.AddCustomRow(VariableCustomizationPrivate::NodeTooltipText)
+			.EditCondition(!bIsRequired, nullptr)
+			.NameContent()
+			[
+				SNew(STextBlock)
+				.Font(IDetailLayoutBuilder::GetDetailFontBold())
+				.Text(VariableCustomizationPrivate::NodeTooltipText)
+			]
+			.ValueContent()
+			[
+				SNew(SMultiLineEditableTextBox)
+				.Text(this, &FMetasoundInputDetailCustomization::GetTooltip)
+				.OnTextCommitted(this, &FMetasoundInputDetailCustomization::OnTooltipCommitted)
+				.IsReadOnly(bIsRequired)
+				.ModiferKeyForNewLine(EModifierKey::Shift)
+				.RevertTextOnEscape(true)
+				.WrapTextAt(VariableCustomizationPrivate::DetailsTitleMaxWidth - VariableCustomizationPrivate::DetailsTitleWrapPadding)
+				.Font(IDetailLayoutBuilder::GetDetailFont())
+			];
+
+			AddDataTypeSelector(DetailLayout, VariableCustomizationPrivate::DataTypeNameText, GraphVariable, bIsRequired);
+
+			CategoryBuilder.AddCustomRow(LOCTEXT("InputPrivate", "Private"))
+			.Visibility(TAttribute<EVisibility>(this, &FMetasoundInputDetailCustomization::ExposePrivateVisibility))
+			.NameContent()
+			[
+				SNew(STextBlock)
+				.Text(LOCTEXT("InputPrivate", "Private"))
+				.Font(IDetailLayoutBuilder::GetDetailFont())
+			]
+			.ValueContent()
+			[
+				SNew(SCheckBox)
+				.IsChecked(this, &FMetasoundInputDetailCustomization::OnGetPrivateCheckboxState)
+				.OnCheckStateChanged(this, &FMetasoundInputDetailCustomization::OnPrivateChanged)
+			];
+
+			FNodeHandle NodeHandle = GraphVariable->GetNodeHandle();
+			const TArray<FOutputHandle>& Outputs = NodeHandle->GetOutputs();
+			if (!ensure(!Outputs.IsEmpty()))
+			{
+				return;
+			}
+
+			IDetailCategoryBuilder& DefaultCategoryBuilder = DetailLayout.EditCategory("DefaultValue");
+
+			TSharedPtr<IPropertyHandle> LiteralHandle = DetailLayout.GetProperty(GET_MEMBER_NAME_CHECKED(UMetasoundEditorGraphInput, Literal));
+			if (ensure(GraphVariable.IsValid()) && ensure(LiteralHandle.IsValid()))
+			{
+				TSharedPtr<IPropertyHandle> DefaultValueHandle;
+				UObject* LiteralObject = nullptr;
+				if (LiteralHandle->GetValue(LiteralObject) == FPropertyAccess::Success)
+				{
+					if (ensure(LiteralObject))
+					{
+						LiteralHandle->MarkHiddenByCustomization();
+						if (IDetailPropertyRow* Row = DefaultCategoryBuilder.AddExternalObjectProperty(TArray<UObject*>({ LiteralObject }), "Default"))
+						{
+							DefaultValueHandle = Row->GetPropertyHandle();
+							if (DefaultValueHandle.IsValid())
+							{
+								SetDefaultPropertyMetaData(DefaultValueHandle.ToSharedRef());
+
+								FSimpleDelegate OnLiteralChanged = FSimpleDelegate::CreateLambda([GraphVariable = this->GraphVariable]()
+								{
+									if (GraphVariable.IsValid())
+									{
+										GraphVariable->OnLiteralChanged();
+									}
+								});
+
+								DefaultValueHandle->SetOnPropertyValueChanged(OnLiteralChanged);
+								DefaultValueHandle->SetOnChildPropertyValueChanged(OnLiteralChanged);
+
+								TSharedPtr<IPropertyHandleArray> DefaultValueArray = DefaultValueHandle->AsArray();
+								if (DefaultValueArray.IsValid())
+								{
+									DefaultValueArray->SetOnNumElementsChanged(OnLiteralChanged);
+								}
+							}
+						}
+					}
+					else
+					{
+						DefaultCategoryBuilder.AddProperty(LiteralHandle);
+					}
+				}
+			}
+		}
+
+		void FMetasoundVariableDataTypeSelector::OnDataTypeArrayChanged(TWeakObjectPtr<UMetasoundEditorGraphVariable> InGraphVariable, ECheckBoxState InNewState)
+		{
+			if (InGraphVariable.IsValid())
+			{
+				TSharedPtr<FString> DataTypeRoot = DataTypeComboBox->GetSelectedItem();
+				if (ensure(DataTypeRoot.IsValid()))
+				{
+					FString DataTypeString = *DataTypeRoot.Get();
+					if (InNewState == ECheckBoxState::Checked)
+					{
+						DataTypeString += VariableCustomizationPrivate::ArrayIdentifier;
+					}
+
+					InGraphVariable->SetDataType(FName(DataTypeString));
+
+					// Required to rebuild the literal details customization.
+					// This is seemingly dangerous (as the Builder's raw ptr is cached),
+					// but the builder cannot be accessed any other way and instances of
+					// this type are always built from and managed by the parent DetailLayoutBuilder.
+					check(DetailLayoutBuilder);
+					DetailLayoutBuilder->ForceRefreshDetails();
+				}
+			}
+		}
+
+		void FMetasoundVariableDataTypeSelector::OnBaseDataTypeChanged(TWeakObjectPtr<UMetasoundEditorGraphVariable> InGraphVariable, TSharedPtr<FString> ItemSelected, ESelectInfo::Type SelectInfo)
+		{
+			if (ItemSelected.IsValid() && !ItemSelected->IsEmpty() && InGraphVariable.IsValid())
+			{
+				FString DataTypeString = *ItemSelected.Get();
+				if (DataTypeArrayCheckbox->GetCheckedState() == ECheckBoxState::Checked)
+				{
+					DataTypeString += VariableCustomizationPrivate::ArrayIdentifier;
+				}
+
+				InGraphVariable->SetDataType(FName(DataTypeString));
+
+				// Required to rebuild the literal details customization.
+				// This is seemingly dangerous (as the Builder's raw ptr is cached),
+				// but the builder cannot be accessed any other way and instances of
+				// this type are always built from and managed by the parent DetailLayoutBuilder.
+				check(DetailLayoutBuilder);
+				DetailLayoutBuilder->ForceRefreshDetails();
+			}
+		}
+
+		void FMetasoundInputDetailCustomization::SetDefaultPropertyMetaData(TSharedRef<IPropertyHandle> InDefaultPropertyHandle) const
+		{
+			using namespace Frontend;
+
+			if (!GraphVariable.IsValid())
+			{
+				return;
+			}
+
+			FMetasoundFrontendRegistryContainer* Registry = FMetasoundFrontendRegistryContainer::Get();
+			if (!ensure(Registry))
+			{
+				return;
+			}
+
+			const FName TypeName = GetLiteralDataType();
+			if (TypeName.IsNone())
+			{
+				return;
+			}
+
+			FString TypeNameString = TypeName.ToString();
+			if (TypeNameString.EndsWith(VariableCustomizationPrivate::ArrayIdentifier))
+			{
+				TypeNameString = TypeNameString.LeftChop(VariableCustomizationPrivate::ArrayIdentifier.Len());
+			}
+			InDefaultPropertyHandle->SetInstanceMetaData(VariableCustomizationPrivate::DataTypeNameIdentifier, TypeNameString);
+
+			FDataTypeRegistryInfo DataTypeInfo;
+			if (!ensure(Registry->GetInfoForDataType(TypeName, DataTypeInfo)))
+			{
+				return;
+			}
+
+			const EMetasoundFrontendLiteralType LiteralType = GetMetasoundFrontendLiteralType(DataTypeInfo.PreferredLiteralType);
+			if (LiteralType != EMetasoundFrontendLiteralType::UObject && LiteralType != EMetasoundFrontendLiteralType::UObjectArray)
+			{
+				return;
+			}
+
+			UClass* ProxyGenClass = DataTypeInfo.ProxyGeneratorClass;
+			if (ProxyGenClass)
+			{
+				const FString ClassName = ProxyGenClass->GetName();
+				InDefaultPropertyHandle->SetInstanceMetaData(VariableCustomizationPrivate::ProxyGeneratorClassNameIdentifier, ClassName);
+			}
+		}
+
+		FName FMetasoundInputDetailCustomization::GetLiteralDataType() const
+		{
+			using namespace Frontend;
+
+			FNodeHandle NodeHandle = GraphVariable->GetNodeHandle();
+			const TArray<FOutputHandle> Outputs = NodeHandle->GetOutputs();
+
+			if (!ensure(Outputs.Num() == 1))
+			{
+				return FName();
+			}
+
+			return Outputs[0]->GetDataType();
+		}
+
+		void FMetasoundInputDetailCustomization::OnDisplayNameChanged(const FText& InNewName)
+		{
+			using namespace Frontend;
+
+			bIsNameInvalid = false;
+
+			if (!ensure(GraphVariable.IsValid()))
+			{
+				return;
+			}
+
+			if (InNewName.IsEmpty())
+			{
+				bIsNameInvalid = true;
+				DisplayNameEditableTextBox->SetError(FText::Format(LOCTEXT("InputRenameInvalid_NameEmpty", "{0} cannot be empty string."), InNewName));
+				return;
+			}
+
+			FConstNodeHandle NodeHandle = GraphVariable->GetConstNodeHandle();
+			FConstGraphHandle GraphHandle = NodeHandle->GetOwningGraph();
+			const FGuid NodeID = NodeHandle->GetID();
+
+			GraphHandle->IterateConstNodes([this, NodeID, InNewName](FConstNodeHandle NodeToCompare)
+			{
+				if (NodeID != NodeToCompare->GetID())
+				{
+					if (InNewName.CompareToCaseIgnored(NodeToCompare->GetDisplayName()) == 0)
+					{
+						bIsNameInvalid = true;
+						DisplayNameEditableTextBox->SetError(FText::Format(LOCTEXT("InputRenameInvalid_NameTaken", "{0} is already in use"), InNewName));
+					}
+				}
+			}, EMetasoundFrontendClassType::Input);
+		}
+
+		void FMetasoundOutputDetailCustomization::CustomizeDetails(IDetailLayoutBuilder& DetailLayout)
+		{
+			using namespace Frontend;
+
+			TMetasoundVariableDetailCustomization<UMetasoundEditorGraphOutput>::CustomizeDetails(DetailLayout);
+
+			if (!GraphVariable.IsValid())
+			{
+				return;
+			}
+
+			IDetailCategoryBuilder& CategoryBuilder = DetailLayout.EditCategory("General");
+
+			const bool bIsRequired = IsRequired();
+			DisplayNameEditableTextBox = SNew(SEditableTextBox)
+				.Text(this, &FMetasoundOutputDetailCustomization::GetDisplayName)
+				.OnTextChanged(this, &FMetasoundOutputDetailCustomization::OnDisplayNameChanged)
+				.OnTextCommitted(this, &FMetasoundOutputDetailCustomization::OnDisplayNameCommitted)
+				.IsReadOnly(bIsRequired)
+				.Font(IDetailLayoutBuilder::GetDetailFont());
+
+			CategoryBuilder.AddCustomRow(VariableCustomizationPrivate::OutputNameText)
+			.EditCondition(!bIsRequired, nullptr)
+			.NameContent()
+			[
+				SNew(STextBlock)
+				.Font(IDetailLayoutBuilder::GetDetailFontBold())
+				.Text(VariableCustomizationPrivate::OutputNameText)
+				.ToolTipText(TAttribute<FText>::Create([GraphVariable = this->GraphVariable]()
+				{
+					if (GraphVariable.IsValid())
+					{
+						FNodeHandle NodeHandle = GraphVariable->GetNodeHandle();
+						return NodeHandle->GetDescription();
+					}
+
+					return FText::GetEmpty();
+				}))
+			]
+			.ValueContent()
+			[
+				DisplayNameEditableTextBox.ToSharedRef()
+			];
+
+			CategoryBuilder.AddCustomRow(VariableCustomizationPrivate::NodeTooltipText)
+			.EditCondition(!bIsRequired, nullptr)
+			.NameContent()
+			[
+				SNew(STextBlock)
+				.Font(IDetailLayoutBuilder::GetDetailFontBold())
+				.Text(VariableCustomizationPrivate::NodeTooltipText)
+			]
+			.ValueContent()
+			[
+				SNew(SMultiLineEditableTextBox)
+				.Text(this, &FMetasoundOutputDetailCustomization::GetTooltip)
+				.OnTextCommitted(this, &FMetasoundOutputDetailCustomization::OnTooltipCommitted)
+				.IsReadOnly(bIsRequired)
+				.ModiferKeyForNewLine(EModifierKey::Shift)
+				.RevertTextOnEscape(true)
+				.WrapTextAt(VariableCustomizationPrivate::DetailsTitleMaxWidth - VariableCustomizationPrivate::DetailsTitleWrapPadding)
+				.Font(IDetailLayoutBuilder::GetDetailFont())
+			];
+
+			AddDataTypeSelector(DetailLayout, VariableCustomizationPrivate::DataTypeNameText, GraphVariable, bIsRequired);
+
+			CategoryBuilder.AddCustomRow(LOCTEXT("OutputPrivate", "Private"))
+			.Visibility(TAttribute<EVisibility>(this, &FMetasoundOutputDetailCustomization::ExposePrivateVisibility))
+			.NameContent()
+			[
+				SNew(STextBlock)
+				.Text(LOCTEXT("OutputPrivate", "Private"))
+				.Font(IDetailLayoutBuilder::GetDetailFont())
+			]
+			.ValueContent()
+			[
+				SNew(SCheckBox)
+				.IsChecked(this, &FMetasoundOutputDetailCustomization::OnGetPrivateCheckboxState)
+				.OnCheckStateChanged(this, &FMetasoundOutputDetailCustomization::OnPrivateChanged)
+			];
+		}
+
+		void FMetasoundOutputDetailCustomization::SetDefaultPropertyMetaData(TSharedRef<IPropertyHandle> InDefaultPropertyHandle) const
+		{
+			using namespace Frontend;
+
+			if (!GraphVariable.IsValid())
+			{
+				return;
+			}
+
+			FMetasoundFrontendRegistryContainer* Registry = FMetasoundFrontendRegistryContainer::Get();
+			if (!ensure(Registry))
+			{
+				return;
+			}
+
+			const FName TypeName = GetLiteralDataType();
+			if (TypeName.IsNone())
+			{
+				return;
+			}
+
+			static const FString ArrayIdentifier = TEXT(":Array");
+			FString TypeNameString = TypeName.ToString();
+			if (TypeNameString.EndsWith(ArrayIdentifier))
+			{
+				TypeNameString = TypeNameString.LeftChop(ArrayIdentifier.Len());
+			}
+			InDefaultPropertyHandle->SetInstanceMetaData(VariableCustomizationPrivate::DataTypeNameIdentifier, TypeNameString);
+
+			FDataTypeRegistryInfo DataTypeInfo;
+			if (!ensure(Registry->GetInfoForDataType(TypeName, DataTypeInfo)))
+			{
+				return;
+			}
+
+			const EMetasoundFrontendLiteralType LiteralType = GetMetasoundFrontendLiteralType(DataTypeInfo.PreferredLiteralType);
+			if (LiteralType != EMetasoundFrontendLiteralType::UObject && LiteralType != EMetasoundFrontendLiteralType::UObjectArray)
+			{
+				return;
+			}
+
+			UClass* ProxyGenClass = DataTypeInfo.ProxyGeneratorClass;
+			if (ProxyGenClass)
+			{
+				const FString ClassName = ProxyGenClass->GetName();
+				InDefaultPropertyHandle->SetInstanceMetaData(VariableCustomizationPrivate::ProxyGeneratorClassNameIdentifier, ClassName);
+			}
+		}
+
+		FName FMetasoundOutputDetailCustomization::GetLiteralDataType() const
+		{
+			using namespace Frontend;
+
+			FNodeHandle NodeHandle = GraphVariable->GetNodeHandle();
+			const TArray<FOutputHandle> Outputs = NodeHandle->GetOutputs();
+
+			if (!ensure(Outputs.Num() == 1))
+			{
+				return FName();
+			}
+
+			return Outputs[0]->GetDataType();
+		}
+
+		void FMetasoundOutputDetailCustomization::OnDisplayNameChanged(const FText& InNewName)
+		{
+			using namespace Frontend;
+
+			bIsNameInvalid = false;
+
+			if (!ensure(GraphVariable.IsValid()))
+			{
+				return;
+			}
+
+			if (InNewName.IsEmpty())
+			{
+				bIsNameInvalid = true;
+				DisplayNameEditableTextBox->SetError(FText::Format(LOCTEXT("OutputRenameInvalid_NameEmpty", "{0} cannot be empty string."), InNewName));
+				return;
+			}
+
+			FConstNodeHandle NodeHandle = GraphVariable->GetConstNodeHandle();
+			FConstGraphHandle GraphHandle = NodeHandle->GetOwningGraph();
+			const FGuid NodeID = NodeHandle->GetID();
+
+			GraphHandle->IterateConstNodes([this, NodeID, InNewName](FConstNodeHandle NodeToCompare)
+			{
+				if (NodeID != NodeToCompare->GetID())
+				{
+					if (InNewName.CompareToCaseIgnored(NodeToCompare->GetDisplayName()) == 0)
+					{
+						bIsNameInvalid = true;
+						DisplayNameEditableTextBox->SetError(FText::Format(LOCTEXT("OutputRenameInvalid_NameTaken", "{0} is already in use"), InNewName));
+					}
+				}
+			}, EMetasoundFrontendClassType::Output);
+		}
+	} // namespace Editor
+} // namespace Metasound
+#undef LOCTEXT_NAMESPACE
diff --git a/Engine/Plugins/Runtime/Metasound/Source/MetasoundEditor/Private/MetasoundNodeDetailCustomization.h b/Engine/Plugins/Runtime/Metasound/Source/MetasoundEditor/Private/MetasoundNodeDetailCustomization.h
new file mode 100644
--- /dev/null
+++ b/Engine/Plugins/Runtime/Metasound/Source/MetasoundEditor/Private/MetasoundNodeDetailCustomization.h
@@ -0,0 +1,263 @@
+// Copyright Epic Games, Inc. All Rights Reserved.
+#pragma once
+
+#include "Components/Widget.h"
+#include "CoreMinimal.h"
+#include "DetailLayoutBuilder.h"
+#include "IDetailCustomization.h"
+#include "IPropertyTypeCustomization.h"
+#include "Layout/Visibility.h"
+#include "MetasoundAssetBase.h"
+#include "MetasoundEditorGraphNode.h"
+#include "MetasoundEditorGraphInputNodes.h"
+#include "MetasoundUObjectRegistry.h"
+#include "ScopedTransaction.h"
+#include "Types/SlateEnums.h"
+#include "UObject/NameTypes.h"
+#include "Widgets/Input/SEditableTextBox.h"
+#include "WorkflowOrientedApp/SModeWidget.h"
+
+// Forward Declarations
+class FDetailWidgetRow;
+class FPropertyRestriction;
+class IDetailLayoutBuilder;
+class IPropertyHandle;
+class SCheckBox;
+class STextComboBox;
+
+#define LOCTEXT_NAMESPACE "MetasoundEditor"
+
+namespace Metasound
+{
+	namespace Editor
+	{
+		class FMetasoundInputArrayDetailCustomizationBase : public IPropertyTypeCustomization
+		{
+		public:
+			//~ Begin IPropertyTypeCustomization
+			virtual void CustomizeHeader(TSharedRef<IPropertyHandle> StructPropertyHandle, FDetailWidgetRow& HeaderRow, IPropertyTypeCustomizationUtils& StructCustomizationUtils) override;
+			virtual void CustomizeChildren(TSharedRef<IPropertyHandle> StructPropertyHandle, IDetailChildrenBuilder& ChildBuilder, IPropertyTypeCustomizationUtils& StructCustomizationUtils) override;
+			//~ End IPropertyTypeCustomization
+
+		protected:
+			virtual TSharedRef<SWidget> CreateStructureWidget(TSharedPtr<IPropertyHandle>& PropertyHandle) const = 0;
+			virtual void CacheProxyData(TSharedPtr<IPropertyHandle> ProxyHandle) { }
+
+		private:
+			TSharedRef<SWidget> CreateValueWidget(TSharedPtr<IPropertyHandleArray> ParentArrayProperty, TSharedPtr<IPropertyHandle> StructPropertyHandle, bool bIsInArray) const;
+		};
+
+		class FMetasoundInputIntDetailCustomization : public FMetasoundInputArrayDetailCustomizationBase
+		{
+		protected:
+			virtual TSharedRef<SWidget> CreateStructureWidget(TSharedPtr<IPropertyHandle>& StructPropertyHandle) const override;
+			virtual void CacheProxyData(TSharedPtr<IPropertyHandle> ProxyHandle) override;
+
+		private:
+			FName DataTypeName;
+		};
+
+		class FMetasoundInputObjectDetailCustomization : public FMetasoundInputArrayDetailCustomizationBase
+		{
+		protected:
+			virtual TSharedRef<SWidget> CreateStructureWidget(TSharedPtr<IPropertyHandle>& StructPropertyHandle) const override;
+			virtual void CacheProxyData(TSharedPtr<IPropertyHandle> ProxyHandle) override;
+
+		private:
+			TWeakObjectPtr<UClass> ProxyGenClass;
+		};
+
+		class FMetasoundVariableDataTypeSelector
+		{
+		public:
+			void AddDataTypeSelector(IDetailLayoutBuilder& InDetailLayoutBuilder, const FText& InRowName, TWeakObjectPtr<UMetasoundEditorGraphVariable> InGraphVariable, bool bIsRequired);
+
+			void OnDataTypeArrayChanged(TWeakObjectPtr<UMetasoundEditorGraphVariable> InGraphVariable, ECheckBoxState InNewState);
+			ECheckBoxState OnGetDataTypeArrayCheckState(TWeakObjectPtr<UMetasoundEditorGraphVariable> InGraphVariable) const;
+			void OnBaseDataTypeChanged(TWeakObjectPtr<UMetasoundEditorGraphVariable> InGraphVariable, TSharedPtr<FString> ItemSelected, ESelectInfo::Type SelectInfo);
+
+		protected:
+			TFunction<void()> OnDataTypeChanged;
+		
+		private:
+				TSharedPtr<SCheckBox> DataTypeArrayCheckbox;
+				TSharedPtr<STextComboBox> DataTypeComboBox;
+				TArray<TSharedPtr<FString>> DataTypeNames;
+
+				IDetailLayoutBuilder* DetailLayoutBuilder = nullptr;
+		};
+
+		template <typename VariableType>
+		class TMetasoundVariableDetailCustomization : public IDetailCustomization, public FMetasoundVariableDataTypeSelector
+		{
+		public:
+			TMetasoundVariableDetailCustomization(const FText& InVariableLabel)
+				: IDetailCustomization()
+				, VariableLabel(InVariableLabel)
+			{
+			}
+
+		protected:
+			FText VariableLabel;
+
+			TWeakObjectPtr<VariableType> GraphVariable;
+			TSharedPtr<SEditableTextBox> DisplayNameEditableTextBox;
+			bool bIsNameInvalid = false;
+
+			// IDetailCustomization interface
+			virtual void CustomizeDetails(IDetailLayoutBuilder& DetailLayout) override
+			{
+				TArray<TWeakObjectPtr<UObject>> Objects;
+				DetailLayout.GetObjectsBeingCustomized(Objects);
+				if (Objects.IsEmpty())
+				{
+					return;
+				}
+
+				GraphVariable = Cast<VariableType>(Objects[0].Get());
+			}
+			// End of IDetailCustomization interface
+
+			virtual void OnDisplayNameChanged(const FText& InNewText) = 0;
+
+			FText GetDisplayName() const
+			{
+				using namespace Frontend;
+
+				if (GraphVariable.IsValid())
+				{
+					return GraphVariable->GetNodeHandle()->GetDisplayName();
+				}
+
+				return FText::GetEmpty();
+			}
+
+			bool IsRequired() const
+			{
+				if (GraphVariable.IsValid())
+				{
+					Frontend::FConstNodeHandle NodeHandle = GraphVariable->GetConstNodeHandle();
+					return NodeHandle->IsRequired();
+				}
+
+				return true;
+			}
+
+			void OnTooltipCommitted(const FText& InNewText, ETextCommit::Type InTextCommit)
+			{
+				using namespace Frontend;
+
+				if (GraphVariable.IsValid())
+				{
+					const FText TransactionLabel = FText::Format(LOCTEXT("SetTooltip", "Set the Metasound {0}'s tooltip"), VariableLabel);
+					const FScopedTransaction Transaction(TransactionLabel);
+					FNodeHandle NodeHandle = GraphVariable->GetNodeHandle();
+					NodeHandle->SetDescription(InNewText);
+				}
+			}
+
+			FText GetTooltip() const
+			{
+				using namespace Frontend;
+				if (GraphVariable.IsValid())
+				{
+					FNodeHandle NodeHandle = GraphVariable->GetNodeHandle();
+					return NodeHandle->GetDescription();
+				}
+
+				return FText::GetEmpty();
+			}
+
+			void OnDisplayNameCommitted(const FText& InNewName, ETextCommit::Type InTextCommit)
+			{
+				using namespace Frontend;
+
+				if (!bIsNameInvalid && GraphVariable.IsValid())
+				{
+					const FText TransactionLabel = FText::Format(LOCTEXT("Rename Variable", "Rename Metasound {0}"), VariableLabel);
+					const FScopedTransaction Transaction(TransactionLabel);
+					FNodeHandle NodeHandle = GraphVariable->GetNodeHandle();
+					NodeHandle->SetDisplayName(InNewName);
+
+					GraphVariable->NameChanged.Broadcast(NodeHandle->GetID());
+					DisplayNameEditableTextBox->SetError(FText::GetEmpty());
+				}
+			}
+
+			ECheckBoxState OnGetPrivateCheckboxState() const
+			{
+				if (GraphVariable.IsValid())
+				{
+					return GraphVariable->GetNodeHandle()->GetNodeStyle().bIsPrivate ? ECheckBoxState::Checked : ECheckBoxState::Unchecked;
+				}
+
+				return ECheckBoxState::Unchecked;
+			}
+
+			void OnPrivateChanged(ECheckBoxState InNewState)
+			{
+				if (GraphVariable.IsValid())
+				{
+					const bool bIsChecked = InNewState == ECheckBoxState::Checked;
+					Frontend::FNodeHandle NodeHandle = GraphVariable->GetNodeHandle();
+					FMetasoundFrontendNodeStyle NodeStyle = NodeHandle->GetNodeStyle();
+					NodeStyle.bIsPrivate = bIsChecked;
+					NodeHandle->SetNodeStyle(NodeStyle);
+				}
+			}
+
+			EVisibility ExposePrivateVisibility() const
+			{
+				if (IsRequired())
+				{
+					return EVisibility::Collapsed;
+				}
+
+				return EVisibility::Visible;
+			}
+		};
+
+		class FMetasoundInputDetailCustomization : public TMetasoundVariableDetailCustomization<UMetasoundEditorGraphInput>
+		{
+		public:
+			FMetasoundInputDetailCustomization()
+				: TMetasoundVariableDetailCustomization<UMetasoundEditorGraphInput>(LOCTEXT("InputVariableLabel", "Input"))
+			{
+			}
+
+			// IDetailCustomization interface
+			virtual void CustomizeDetails(IDetailLayoutBuilder& DetailLayout) override;
+			// End of IDetailCustomization interface
+
+		protected:
+			virtual void OnDisplayNameChanged(const FText& InNewText) override;
+
+		private:
+			void SetDefaultPropertyMetaData(TSharedRef<IPropertyHandle> InDefaultPropertyHandle) const;
+
+			FName GetLiteralDataType() const;
+		};
+
+		class FMetasoundOutputDetailCustomization : public TMetasoundVariableDetailCustomization<UMetasoundEditorGraphOutput>
+		{
+		public:
+			FMetasoundOutputDetailCustomization()
+				: TMetasoundVariableDetailCustomization<UMetasoundEditorGraphOutput>(LOCTEXT("OutputVariableLabel", "Output"))
+			{
+			}
+
+			// IDetailCustomization interface
+			virtual void CustomizeDetails(IDetailLayoutBuilder& DetailLayout) override;
+			// End of IDetailCustomization interface
+
+		protected:
+			virtual void OnDisplayNameChanged(const FText& InNewText) override;
+
+		private:
+			void SetDefaultPropertyMetaData(TSharedRef<IPropertyHandle> InDefaultPropertyHandle) const;
+
+			FName GetLiteralDataType() const;
+		};
+	} // namespace Editor
+} // namespace Metasound
+#undef LOCTEXT_NAMESPACE
diff --git a/Engine/Plugins/Runtime/Metasound/Source/MetasoundEditor/Public/MetasoundEditorCommands.h b/Engine/Plugins/Runtime/Metasound/Source/MetasoundEditor/Public/MetasoundEditorCommands.h
--- a/Engine/Plugins/Runtime/Metasound/Source/MetasoundEditor/Public/MetasoundEditorCommands.h
+++ b/Engine/Plugins/Runtime/Metasound/Source/MetasoundEditor/Public/MetasoundEditorCommands.h
@@ -46,6 +46,9 @@ namespace Metasound
 			/** Removes an input from the node */
 			TSharedPtr<FUICommandInfo> DeleteInput;
 
+			/** Shows the Metasound's top-level object settings in the Inspector */
+			TSharedPtr<FUICommandInfo> EditMetasoundSettings;
+
 			/** Initialize commands */
 			virtual void RegisterCommands() override;
 		};
diff --git a/Engine/Plugins/Runtime/Metasound/Source/MetasoundEditor/Public/MetasoundEditorGraph.h b/Engine/Plugins/Runtime/Metasound/Source/MetasoundEditor/Public/MetasoundEditorGraph.h
--- a/Engine/Plugins/Runtime/Metasound/Source/MetasoundEditor/Public/MetasoundEditorGraph.h
+++ b/Engine/Plugins/Runtime/Metasound/Source/MetasoundEditor/Public/MetasoundEditorGraph.h
@@ -3,8 +3,11 @@
 
 #include "CoreMinimal.h"
 #include "MetasoundEditor.h"
+#include "MetasoundFrontendController.h"
+#include "MetasoundFrontendDocument.h"
 #include "MetasoundFrontendLiteral.h"
 #include "MetasoundSource.h"
+#include "MetasoundUObjectRegistry.h"
 #include "UObject/ObjectMacros.h"
 
 #include "MetasoundEditorGraph.generated.h"
@@ -12,19 +15,114 @@
 // Forward Declarations
 struct FMetasoundFrontendDocument;
 class UMetasoundEditorGraphInputNode;
+class UMetasoundEditorGraphNode;
 
+DECLARE_MULTICAST_DELEGATE_OneParam(FOnMetasoundNodeNameChanged, FGuid /* NodeID */);
 
-UCLASS(MinimalAPI)
-class UMetasoundEditorGraph : public UMetasoundEditorGraphBase
+UCLASS(Abstract)
+class METASOUNDEDITOR_API UMetasoundEditorGraphVariable : public UObject
 {
 	GENERATED_BODY()
 
+protected:
+	virtual Metasound::Frontend::FNodeHandle AddNodeHandle(const FString& InNodeName, const FText& InNodeDisplayName, FName InDataType)
+	{
+		return Metasound::Frontend::INodeController::GetInvalidHandle();
+	}
+
+public:
+	UPROPERTY()
+	FMetasoundFrontendClassName ClassName;
+
+	UPROPERTY()
+	FGuid NodeID;
+
+	UPROPERTY()
+	FName TypeName;
+
+	// Called when the name of the associated Frontend Node is changed
+	FOnMetasoundNodeNameChanged NameChanged;
+
+	Metasound::Frontend::FNodeHandle GetNodeHandle() const;
+	Metasound::Frontend::FConstNodeHandle GetConstNodeHandle() const;
+
+	TArray<UMetasoundEditorGraphNode*> GetNodes() const;
+	void SetDataType(FName InNewType);
+	virtual void OnDataTypeChanged() { }
+};
+
+UCLASS()
+class METASOUNDEDITOR_API UMetasoundEditorGraphInputLiteral : public UObject
+{
+	GENERATED_BODY()
+
+public:
+	virtual void UpdatePreviewInstance(const Metasound::FVertexKey& InParameterName, IAudioInstanceTransmitter& InInstanceTransmitter) const
+	{
+	}
+
+	virtual FMetasoundFrontendLiteral GetDefault() const
+	{
+		return FMetasoundFrontendLiteral();
+	}
+
+	virtual EMetasoundFrontendLiteralType GetLiteralType() const
+	{
+		return EMetasoundFrontendLiteralType::None;
+	}
+};
+
+UCLASS()
+class METASOUNDEDITOR_API UMetasoundEditorGraphInput : public UMetasoundEditorGraphVariable
+{
+	GENERATED_BODY()
+
+protected:
+	virtual Metasound::Frontend::FNodeHandle AddNodeHandle(const FString& InNodeName, const FText& InNodeDisplayName, FName InDataType) override;
+
+
 public:
-	UMetasoundEditorGraphInputNode* CreateInputNode(EMetasoundFrontendLiteralType LiteralType, UClass* InLiteralObjectClass, bool bInSelectNewNode);
+	UPROPERTY(VisibleAnywhere, Category = DefaultValue)
+	UMetasoundEditorGraphInputLiteral* Literal;
+
+	void UpdateDocumentInput() const;
+	void UpdatePreviewInstance(const Metasound::FVertexKey& InParameterName, IAudioInstanceTransmitter& InInstanceTransmitter) const;
+
+	void OnDataTypeChanged() override;
+	void OnLiteralChanged();
+
+#if WITH_EDITORONLY_DATA
+	virtual void PostEditUndo() override
+	{
+		Super::PostEditUndo();
+
+		UpdateDocumentInput();
+	}
+#endif // WITH_EDITORONLY_DATA
+};
+
+UCLASS()
+class METASOUNDEDITOR_API UMetasoundEditorGraphOutput : public UMetasoundEditorGraphVariable
+{
+	GENERATED_BODY()
+
+protected:
+	virtual Metasound::Frontend::FNodeHandle AddNodeHandle(const FString& InNodeName, const FText& InNodeDisplayName, FName InDataType) override;
+};
+
+UCLASS()
+class METASOUNDEDITOR_API UMetasoundEditorGraph : public UMetasoundEditorGraphBase
+{
+	GENERATED_BODY()
+
+public:
+	UMetasoundEditorGraphInputNode* CreateInputNode(Metasound::Frontend::FNodeHandle InNodeHandle, bool bInSelectNewNode);
 
 	UObject* GetMetasound() const;
 	UObject& GetMetasoundChecked() const;
 
+	void IterateInputs(TUniqueFunction<void(UMetasoundEditorGraphInput&)> InFunction) const;
+	void IterateOutputs(TUniqueFunction<void(UMetasoundEditorGraphOutput&)> InFunction) const;
 
 	virtual void Synchronize();
 
@@ -39,7 +137,24 @@ public:
 
 private:
 	UPROPERTY()
-	UObject* ParentMetasound;
+	TObjectPtr<UObject> ParentMetasound;
+
+	UPROPERTY()
+	TArray<TObjectPtr<UMetasoundEditorGraphInput>> Inputs;
+
+	UPROPERTY()
+	TArray<TObjectPtr<UMetasoundEditorGraphOutput>> Outputs;
+
+public:
+	UMetasoundEditorGraphInput* FindInput(FGuid InNodeID) const;
+	UMetasoundEditorGraphInput* FindOrAddInput(Metasound::Frontend::FNodeHandle InNodeHandle);
+
+	UMetasoundEditorGraphOutput* FindOutput(FGuid InNodeID) const;
+	UMetasoundEditorGraphOutput* FindOrAddOutput(Metasound::Frontend::FNodeHandle InNodeHandle);
+
+	UMetasoundEditorGraphVariable* FindVariable(FGuid InNodeID) const;
+
+	bool RemoveVariable(UMetasoundEditorGraphVariable& InVariable);
 
 	TUniquePtr<IAudioInstanceTransmitter> Transmitter;
 
diff --git a/Engine/Plugins/Runtime/Metasound/Source/MetasoundEditor/Public/MetasoundEditorGraphInputNodes.h b/Engine/Plugins/Runtime/Metasound/Source/MetasoundEditor/Public/MetasoundEditorGraphInputNodes.h
--- a/Engine/Plugins/Runtime/Metasound/Source/MetasoundEditor/Public/MetasoundEditorGraphInputNodes.h
+++ b/Engine/Plugins/Runtime/Metasound/Source/MetasoundEditor/Public/MetasoundEditorGraphInputNodes.h
@@ -1,6 +1,7 @@
 // Copyright Epic Games, Inc. All Rights Reserved.
 #pragma once
 
+#include "Algo/Count.h"
 #include "Algo/Transform.h"
 #include "CoreMinimal.h"
 #include "EdGraph/EdGraphNode.h"
@@ -39,46 +40,36 @@ class METASOUNDEDITOR_API UMetasoundEditorGraphInputNode : public UMetasoundEdit
 public:
 	virtual ~UMetasoundEditorGraphInputNode() = default;
 
-	virtual FMetasoundFrontendLiteral GetLiteralDefault() const
-	{
-		return FMetasoundFrontendLiteral();
-	}
-
+	UPROPERTY()
+	TObjectPtr<UMetasoundEditorGraphInput> Input;
 
-	void UpdateDocumentInput() const
+public:
+	virtual FMetasoundFrontendClassName GetClassName() const override
 	{
-		using namespace Metasound::Frontend;
-
-		FConstNodeHandle NodeHandle = GetNodeHandle();
-		const FString& NodeName = NodeHandle->GetNodeName();
+		if (ensure(Input))
+		{
+			return Input->ClassName;
+		}
 
-		FGraphHandle GraphHandle = GetRootGraphHandle();
-		const FGuid VertexID = GraphHandle->GetVertexIDForInputVertex(NodeName);
-		GraphHandle->SetDefaultInput(VertexID, GetLiteralDefault());
+		return Super::GetClassName();
 	}
 
-	virtual void UpdatePreviewInstance(const Metasound::FVertexKey& InParamterName, IAudioInstanceTransmitter& InInstanceTransmitter) const {}
-
-	void OnLiteralChanged() 
+	void UpdatePreviewInstance(const Metasound::FVertexKey& InParameterName, IAudioInstanceTransmitter& InInstanceTransmitter) const
 	{
-		UpdateDocumentInput();
-
-		if (UMetasoundEditorGraph* MetasoundGraph = Cast<UMetasoundEditorGraph>(GetGraph()))
+		if (ensure(Input))
+		{
+			Input->UpdatePreviewInstance(InParameterName, InInstanceTransmitter);
+		}
+	}
+	
+	virtual FGuid GetNodeID() const override
+	{
+		if (ensure(Input))
 		{
-			if (IAudioInstanceTransmitter* Transmitter = MetasoundGraph->GetMetasoundInstanceTransmitter())
-			{
-				// TODO: fix how this parameter name is determined. It should not be done with a "DisplayName" but rather 
-				// the vertex "Name". Currently user entered vertex names only have their "Names" stored as "DisplayNames" 
-				TArray<Metasound::Frontend::FConstInputHandle> Inputs = GetNodeHandle()->GetConstInputs();
-
-				// An input node should only have one input. 
-				if (ensure(Inputs.Num() == 1))
-				{
-					Metasound::FVertexKey VertexKey = Metasound::FVertexKey(Inputs[0]->GetDisplayName().ToString());
-					UpdatePreviewInstance(VertexKey, *Transmitter);
-				}
-			}
+			return Input->NodeID;
 		}
+
+		return Super::GetNodeID();
 	}
 
 #if WITH_EDITORONLY_DATA
@@ -86,39 +77,34 @@ public:
 	{
 		Super::PostEditUndo();
 
-		OnLiteralChanged();
+		if (ensure(Input))
+		{
+			Input->OnLiteralChanged();
+		}
 	}
 #endif // WITH_EDITORONLY_DATA
 
-	FName GetLiteralDefaultPropertyFName() const
-	{
-		return "Default";
-	}
-
-	virtual EMetasoundFrontendLiteralType GetLiteralType() const
+protected:
+	virtual void SetNodeID(FGuid InNodeID) override
 	{
-		return EMetasoundFrontendLiteralType::None;
+		Input->NodeID = InNodeID;
 	}
 
-protected:
-	UPROPERTY()
-	FName InputTypeName;
-
 	friend class Metasound::Editor::FGraphBuilder;
 };
 
 UCLASS(MinimalAPI)
-class UMetasoundEditorGraphInputBoolNode : public UMetasoundEditorGraphInputNode
+class UMetasoundEditorGraphInputBool : public UMetasoundEditorGraphInputLiteral
 {
 	GENERATED_BODY()
 
 public:
-	virtual ~UMetasoundEditorGraphInputBoolNode() = default;
+	virtual ~UMetasoundEditorGraphInputBool() = default;
 
 	UPROPERTY(EditAnywhere, Category = DefaultValue)
 	bool Default = false;
 
-	virtual FMetasoundFrontendLiteral GetLiteralDefault() const override
+	virtual FMetasoundFrontendLiteral GetDefault() const override
 	{
 		FMetasoundFrontendLiteral Literal;
 		Literal.Set(Default);
@@ -137,17 +123,17 @@ public:
 };
 
 UCLASS(MinimalAPI)
-class UMetasoundEditorGraphInputBoolArrayNode : public UMetasoundEditorGraphInputNode
+class UMetasoundEditorGraphInputBoolArray : public UMetasoundEditorGraphInputLiteral
 {
 	GENERATED_BODY()
 
 public:
-	virtual ~UMetasoundEditorGraphInputBoolArrayNode() = default;
+	virtual ~UMetasoundEditorGraphInputBoolArray() = default;
 
 	UPROPERTY(EditAnywhere, Category = DefaultValue)
 	TArray<bool> Default;
 
-	virtual FMetasoundFrontendLiteral GetLiteralDefault() const override
+	virtual FMetasoundFrontendLiteral GetDefault() const override
 	{
 		FMetasoundFrontendLiteral Literal;
 		Literal.Set(Default);
@@ -167,7 +153,7 @@ public:
 
 // Broken out to be able to customize and swap enum behavior for basic integer literal behavior
 USTRUCT()
-struct FMetasoundEditorGraphInputInt
+struct FMetasoundEditorGraphInputIntRef
 {
 	GENERATED_BODY()
 
@@ -176,17 +162,17 @@ struct FMetasoundEditorGraphInputInt
 };
 
 UCLASS(MinimalAPI)
-class UMetasoundEditorGraphInputIntNode : public UMetasoundEditorGraphInputNode
+class UMetasoundEditorGraphInputInt : public UMetasoundEditorGraphInputLiteral
 {
 	GENERATED_BODY()
 
 public:
-	virtual ~UMetasoundEditorGraphInputIntNode() = default;
+	virtual ~UMetasoundEditorGraphInputInt() = default;
 
 	UPROPERTY(EditAnywhere, Category = DefaultValue)
-	FMetasoundEditorGraphInputInt Default;
+	FMetasoundEditorGraphInputIntRef Default;
 
-	virtual FMetasoundFrontendLiteral GetLiteralDefault() const override
+	virtual FMetasoundFrontendLiteral GetDefault() const override
 	{
 		FMetasoundFrontendLiteral Literal;
 		Literal.Set(Default.Value);
@@ -205,26 +191,24 @@ public:
 };
 
 UCLASS(MinimalAPI)
-class UMetasoundEditorGraphInputIntArrayNode : public UMetasoundEditorGraphInputNode
+class UMetasoundEditorGraphInputIntArray : public UMetasoundEditorGraphInputLiteral
 {
 	GENERATED_BODY()
 
 public:
-	virtual ~UMetasoundEditorGraphInputIntArrayNode() = default;
+	virtual ~UMetasoundEditorGraphInputIntArray() = default;
 
 	UPROPERTY(EditAnywhere, Category = DefaultValue)
-	TArray<FMetasoundEditorGraphInputInt> Default;
+	TArray<FMetasoundEditorGraphInputIntRef> Default;
 
-	virtual FMetasoundFrontendLiteral GetLiteralDefault() const override
+	virtual FMetasoundFrontendLiteral GetDefault() const override
 	{
-		TArray<int32> Values;
-		for (const FMetasoundEditorGraphInputInt& Value : Default)
-		{
-			Values.Add(Value.Value);
-		}
+		TArray<int32> IntArray;
+		Algo::Transform(Default, IntArray, [](const FMetasoundEditorGraphInputIntRef& InValue) { return InValue.Value; });
 
 		FMetasoundFrontendLiteral Literal;
-		Literal.Set(Values);
+		Literal.Set(IntArray);
+
 		return Literal;
 	}
 
@@ -236,24 +220,24 @@ public:
 	void UpdatePreviewInstance(const Metasound::FVertexKey& InParameterName, IAudioInstanceTransmitter& InInstanceTransmitter) const override
 	{
 		TArray<int32> IntArray;
-		Algo::Transform(Default, IntArray, [](const FMetasoundEditorGraphInputInt& InValue) { return InValue.Value; });
+		Algo::Transform(Default, IntArray, [](const FMetasoundEditorGraphInputIntRef& InValue) { return InValue.Value; });
 
 		InInstanceTransmitter.SetParameter(*InParameterName, TArray<int32>{ IntArray });
 	}
 };
 
 UCLASS(MinimalAPI)
-class UMetasoundEditorGraphInputFloatNode : public UMetasoundEditorGraphInputNode
+class UMetasoundEditorGraphInputFloat : public UMetasoundEditorGraphInputLiteral
 {
 	GENERATED_BODY()
 
 public:
-	virtual ~UMetasoundEditorGraphInputFloatNode() = default;
+	virtual ~UMetasoundEditorGraphInputFloat() = default;
 
 	UPROPERTY(EditAnywhere, Category = DefaultValue)
 	float Default = 0.f;
 
-	virtual FMetasoundFrontendLiteral GetLiteralDefault() const override
+	virtual FMetasoundFrontendLiteral GetDefault() const override
 	{
 		FMetasoundFrontendLiteral Literal;
 		Literal.Set(Default);
@@ -272,17 +256,17 @@ public:
 };
 
 UCLASS(MinimalAPI)
-class UMetasoundEditorGraphInputFloatArrayNode : public UMetasoundEditorGraphInputNode
+class UMetasoundEditorGraphInputFloatArray : public UMetasoundEditorGraphInputLiteral
 {
 	GENERATED_BODY()
 
 public:
-	virtual ~UMetasoundEditorGraphInputFloatArrayNode() = default;
+	virtual ~UMetasoundEditorGraphInputFloatArray() = default;
 
 	UPROPERTY(EditAnywhere, Category = DefaultValue)
 	TArray<float> Default;
 
-	virtual FMetasoundFrontendLiteral GetLiteralDefault() const override
+	virtual FMetasoundFrontendLiteral GetDefault() const override
 	{
 		FMetasoundFrontendLiteral Literal;
 		Literal.Set(Default);
@@ -301,17 +285,17 @@ public:
 };
 
 UCLASS(MinimalAPI)
-class UMetasoundEditorGraphInputStringNode : public UMetasoundEditorGraphInputNode
+class UMetasoundEditorGraphInputString : public UMetasoundEditorGraphInputLiteral
 {
 	GENERATED_BODY()
 
 public:
-	virtual ~UMetasoundEditorGraphInputStringNode() = default;
+	virtual ~UMetasoundEditorGraphInputString() = default;
 
 	UPROPERTY(EditAnywhere, Category = DefaultValue)
 	FString Default;
 
-	virtual FMetasoundFrontendLiteral GetLiteralDefault() const override
+	virtual FMetasoundFrontendLiteral GetDefault() const override
 	{
 		FMetasoundFrontendLiteral Literal;
 		Literal.Set(Default);
@@ -330,17 +314,17 @@ public:
 };
 
 UCLASS(MinimalAPI)
-class UMetasoundEditorGraphInputStringArrayNode : public UMetasoundEditorGraphInputNode
+class UMetasoundEditorGraphInputStringArray : public UMetasoundEditorGraphInputLiteral
 {
 	GENERATED_BODY()
 
 public:
-	virtual ~UMetasoundEditorGraphInputStringArrayNode() = default;
+	virtual ~UMetasoundEditorGraphInputStringArray() = default;
 
 	UPROPERTY(EditAnywhere, Category = DefaultValue)
 	TArray<FString> Default;
 
-	virtual FMetasoundFrontendLiteral GetLiteralDefault() const override
+	virtual FMetasoundFrontendLiteral GetDefault() const override
 	{
 		FMetasoundFrontendLiteral Literal;
 		Literal.Set(Default);
@@ -360,7 +344,7 @@ public:
 
 // Broken out to be able to customize and swap AllowedClass based on provided object proxy
 USTRUCT()
-struct FMetasoundEditorGraphInputObject
+struct FMetasoundEditorGraphInputObjectRef
 {
 	GENERATED_BODY()
 
@@ -369,17 +353,17 @@ struct FMetasoundEditorGraphInputObject
 };
 
 UCLASS(MinimalAPI)
-class UMetasoundEditorGraphInputObjectNode : public UMetasoundEditorGraphInputNode
+class UMetasoundEditorGraphInputObject : public UMetasoundEditorGraphInputLiteral
 {
 	GENERATED_BODY()
 
 public:
-	virtual ~UMetasoundEditorGraphInputObjectNode() = default;
+	virtual ~UMetasoundEditorGraphInputObject() = default;
 
 	UPROPERTY(EditAnywhere, Category = DefaultValue)
-	FMetasoundEditorGraphInputObject Default;
+	FMetasoundEditorGraphInputObjectRef Default;
 
-	virtual FMetasoundFrontendLiteral GetLiteralDefault() const override
+	virtual FMetasoundFrontendLiteral GetDefault() const override
 	{
 		FMetasoundFrontendLiteral Literal;
 		Literal.Set(Default.Object);
@@ -399,26 +383,23 @@ public:
 };
 
 UCLASS(MinimalAPI)
-class UMetasoundEditorGraphInputObjectArrayNode : public UMetasoundEditorGraphInputNode
+class UMetasoundEditorGraphInputObjectArray : public UMetasoundEditorGraphInputLiteral
 {
 	GENERATED_BODY()
 
 public:
-	virtual ~UMetasoundEditorGraphInputObjectArrayNode() = default;
+	virtual ~UMetasoundEditorGraphInputObjectArray() = default;
 
 	UPROPERTY(EditAnywhere, Category = DefaultValue)
-	TArray<FMetasoundEditorGraphInputObject> Default;
+	TArray<FMetasoundEditorGraphInputObjectRef> Default;
 
-	virtual FMetasoundFrontendLiteral GetLiteralDefault() const override
+	virtual FMetasoundFrontendLiteral GetDefault() const override
 	{
-		TArray<UObject*> Objects;
-		for (const FMetasoundEditorGraphInputObject& InputObject : Default)
-		{
-			Objects.Add(InputObject.Object);
-		}
+		TArray<UObject*> ObjectArray;
+		Algo::Transform(Default, ObjectArray, [](const FMetasoundEditorGraphInputObjectRef& InValue) { return InValue.Object; });
 
 		FMetasoundFrontendLiteral Literal;
-		Literal.Set(Objects);
+		Literal.Set(ObjectArray);
 		return Literal;
 	}
 
@@ -430,9 +411,9 @@ public:
 	void UpdatePreviewInstance(const Metasound::FVertexKey& InParameterName, IAudioInstanceTransmitter& InInstanceTransmitter) const override
 	{
 		TArray<UObject*> ObjectArray;
-		Algo::Transform(Default, ObjectArray, [](const FMetasoundEditorGraphInputObject& InValue) { return InValue.Object; });
+		Algo::Transform(Default, ObjectArray, [](const FMetasoundEditorGraphInputObjectRef& InValue) { return InValue.Object; });
 		
 		// TODO. We need proxy object here safely.
 		//InInstanceTransmitter.SetParameter(*InParameterName, Default.Object);
 	}
-};
+};
\ No newline at end of file
diff --git a/Engine/Plugins/Runtime/Metasound/Source/MetasoundEditor/Public/MetasoundEditorGraphNode.h b/Engine/Plugins/Runtime/Metasound/Source/MetasoundEditor/Public/MetasoundEditorGraphNode.h
--- a/Engine/Plugins/Runtime/Metasound/Source/MetasoundEditor/Public/MetasoundEditorGraphNode.h
+++ b/Engine/Plugins/Runtime/Metasound/Source/MetasoundEditor/Public/MetasoundEditorGraphNode.h
@@ -18,6 +18,7 @@
 // Forward Declarations
 class UEdGraphPin;
 class UMetasound;
+class UMetasoundEditorGraphOutput;
 
 namespace Metasound
 {
@@ -80,18 +81,14 @@ public:
 
 	Metasound::Frontend::FGraphHandle GetRootGraphHandle() const;
 	Metasound::Frontend::FConstGraphHandle GetConstRootGraphHandle() const;
-
 	Metasound::Frontend::FNodeHandle GetNodeHandle() const;
 	Metasound::Frontend::FConstNodeHandle GetConstNodeHandle() const;
 
-	FGuid GetNodeID() const;
+	virtual FMetasoundFrontendClassName GetClassName() const { return FMetasoundFrontendClassName(); }
+	virtual FGuid GetNodeID() const { return FGuid(); }
 
 protected:
-	UPROPERTY()
-	FMetasoundFrontendClassName ClassName;
-
-	UPROPERTY()
-	FGuid NodeID;
+	virtual void SetNodeID(FGuid InNodeID) { }
 
 	friend class Metasound::Editor::FGraphBuilder;
 };
@@ -101,9 +98,15 @@ class UMetasoundEditorGraphOutputNode : public UMetasoundEditorGraphNode
 {
 	GENERATED_BODY()
 
-protected:
+public:
 	UPROPERTY()
-	FName OutputTypeName;
+	UMetasoundEditorGraphOutput* Output;
+
+	virtual FMetasoundFrontendClassName GetClassName() const override;
+	virtual FGuid GetNodeID() const override;
+
+protected:
+	virtual void SetNodeID(FGuid InNodeID) override;
 
 	friend class Metasound::Editor::FGraphBuilder;
 };
@@ -112,4 +115,23 @@ UCLASS(MinimalAPI)
 class UMetasoundEditorGraphExternalNode : public UMetasoundEditorGraphNode
 {
 	GENERATED_BODY()
+
+protected:
+	UPROPERTY()
+	FMetasoundFrontendClassName ClassName;
+
+	UPROPERTY()
+	FGuid NodeID;
+
+public:
+	virtual FMetasoundFrontendClassName GetClassName() const override { return ClassName; }
+	virtual FGuid GetNodeID() const override { return NodeID; }
+
+protected:
+	virtual void SetNodeID(FGuid InNodeID) override
+	{
+		NodeID = InNodeID;
+	}
+
+	friend class Metasound::Editor::FGraphBuilder;
 };
\ No newline at end of file
diff --git a/Engine/Plugins/Runtime/Metasound/Source/MetasoundEditor/Public/MetasoundEditorGraphSchema.h b/Engine/Plugins/Runtime/Metasound/Source/MetasoundEditor/Public/MetasoundEditorGraphSchema.h
--- a/Engine/Plugins/Runtime/Metasound/Source/MetasoundEditor/Public/MetasoundEditorGraphSchema.h
+++ b/Engine/Plugins/Runtime/Metasound/Source/MetasoundEditor/Public/MetasoundEditorGraphSchema.h
@@ -9,6 +9,7 @@
 #include "EdGraphUtilities.h"
 #include "MetasoundEditorModule.h"
 #include "MetasoundFrontend.h"
+#include "MetasoundFrontendController.h"
 #include "Templates/Function.h"
 #include "UObject/ObjectMacros.h"
 
@@ -28,7 +29,7 @@ namespace Metasound
 	{
 		using FInputFilterFunction = TFunction<bool(const FMetasoundFrontendClassInput&)>;
 		using FOutputFilterFunction = TFunction<bool(const FMetasoundFrontendClassOutput&)>;
-		using FDataTypeFilterFunction = TFunction<bool(const FEditorDataType&)>;
+		using FInterfaceNodeFilterFunction = TFunction<bool(Frontend::FConstNodeHandle)>;
 
 		struct FActionClassFilters
 		{
@@ -98,13 +99,13 @@ struct METASOUNDEDITOR_API FMetasoundGraphSchemaAction_NewInput : public FEdGrap
 	GENERATED_USTRUCT_BODY();
 
 	UPROPERTY()
-	FName NodeTypeName;
+	FGuid NodeID;
 
 	FMetasoundGraphSchemaAction_NewInput()
 		: FEdGraphSchemaAction()
 	{}
 
-	FMetasoundGraphSchemaAction_NewInput(FText InNodeCategory, FText InDisplayName, FName InTypeName, FText InToolTip, const int32 InGrouping);
+	FMetasoundGraphSchemaAction_NewInput(FText InNodeCategory, FText InDisplayName, FGuid InInputNodeID, FText InToolTip, const int32 InGrouping);
 
 	//~ Begin FEdGraphSchemaAction Interface
 	virtual UEdGraphNode* PerformAction(UEdGraph* ParentGraph, UEdGraphPin* FromPin, const FVector2D Location, bool bSelectNewNode = true) override;
@@ -118,13 +119,13 @@ struct METASOUNDEDITOR_API FMetasoundGraphSchemaAction_NewOutput : public FEdGra
 	GENERATED_USTRUCT_BODY();
 
 	UPROPERTY()
-	FName NodeTypeName;
+	FGuid NodeID;
 
 	FMetasoundGraphSchemaAction_NewOutput()
 		: FEdGraphSchemaAction()
 	{}
 
-	FMetasoundGraphSchemaAction_NewOutput(FText InNodeCategory, FText InDisplayName, FName InTypeName, FText InToolTip, const int32 InGrouping);
+	FMetasoundGraphSchemaAction_NewOutput(FText InNodeCategory, FText InDisplayName, FGuid InOutputNodeID, FText InToolTip, const int32 InGrouping);
 
 	//~ Begin FEdGraphSchemaAction Interface
 	virtual UEdGraphNode* PerformAction(UEdGraph* ParentGraph, UEdGraphPin* FromPin, const FVector2D Location, bool bSelectNewNode = true) override;
@@ -244,8 +245,8 @@ private:
 	void GetConversionActions(FGraphActionMenuBuilder& ActionMenuBuilder, Metasound::Editor::FActionClassFilters InFilters = Metasound::Editor::FActionClassFilters(), bool bShowSelectedActions = true) const;
 	void GetFunctionActions(FGraphActionMenuBuilder& ActionMenuBuilder, Metasound::Editor::FActionClassFilters InFilters = Metasound::Editor::FActionClassFilters(), bool bShowSelectedActions = true) const;
 
-	void GetDataTypeInputNodeActions(FGraphActionMenuBuilder& ActionMenuBuilder, Metasound::Editor::FDataTypeFilterFunction InFilter = Metasound::Editor::FDataTypeFilterFunction(), bool bShowSelectedActions = true) const;
-	void GetDataTypeOutputNodeActions(FGraphActionMenuBuilder& ActionMenuBuilder, Metasound::Editor::FDataTypeFilterFunction InFilter = Metasound::Editor::FDataTypeFilterFunction(), bool bShowSelectedActions = true) const;
+	void GetDataTypeInputNodeActions(FGraphContextMenuBuilder& InMenuBuilder, Metasound::Frontend::FConstGraphHandle InGraphHandle, Metasound::Editor::FInterfaceNodeFilterFunction InFilter = Metasound::Editor::FInterfaceNodeFilterFunction(), bool bShowSelectedActions = true) const;
+	void GetDataTypeOutputNodeActions(FGraphContextMenuBuilder& InMenuBuilder, Metasound::Frontend::FConstGraphHandle InGraphHandle, Metasound::Editor::FInterfaceNodeFilterFunction InFilter = Metasound::Editor::FInterfaceNodeFilterFunction(), bool bShowSelectedActions = true) const;
 
 	/** Adds action for creating a comment */
 	void GetCommentAction(FGraphActionMenuBuilder& ActionMenuBuilder, const UEdGraph* CurrentGraph = nullptr) const;
diff --git a/Engine/Plugins/Runtime/Metasound/Source/MetasoundFrontend/Private/MetasoundFrontendInvalidController.h b/Engine/Plugins/Runtime/Metasound/Source/MetasoundFrontend/Private/MetasoundFrontendInvalidController.h
--- a/Engine/Plugins/Runtime/Metasound/Source/MetasoundFrontend/Private/MetasoundFrontendInvalidController.h
+++ b/Engine/Plugins/Runtime/Metasound/Source/MetasoundFrontend/Private/MetasoundFrontendInvalidController.h
@@ -89,8 +89,8 @@ namespace Metasound
 			bool IsConnected() const override { return false; }
 			const FName& GetDataType() const override { return MetasoundFrontendInvalidControllerPrivate::GetInvalid<FName>(); }
 			const FString& GetName() const override { return MetasoundFrontendInvalidControllerPrivate::GetInvalid<FString>(); }
-			const FText& GetDisplayName() const override { return FText::GetEmpty(); }
-			const FText& GetTooltip() const override { return FText::GetEmpty(); }
+			const FText& GetDisplayName() const override { return MetasoundFrontendInvalidControllerPrivate::GetInvalid<FText>(); }
+			const FText& GetTooltip() const override { return MetasoundFrontendInvalidControllerPrivate::GetInvalid<FText>(); }
 			const FMetasoundFrontendLiteral* GetDefaultLiteral() const override { return nullptr; }
 			FGuid GetOwningNodeID() const override { return Metasound::FrontendInvalidID; }
 			TSharedRef<INodeController> GetOwningNode() override;
@@ -156,12 +156,14 @@ namespace Metasound
 			EMetasoundFrontendClassType GetClassType() const override { return EMetasoundFrontendClassType::Invalid; }
 			const FMetasoundFrontendClassName& GetClassName() const override { return MetasoundFrontendInvalidControllerPrivate::GetInvalid<FMetasoundFrontendClassName>(); }
 			FMetasoundFrontendVersionNumber GetClassVersionNumber() const override { return MetasoundFrontendInvalidControllerPrivate::GetInvalid<FMetasoundFrontendVersionNumber>(); }
-			const FText& GetClassDisplayName() const override { return MetasoundFrontendInvalidControllerPrivate::GetInvalid<FText>(); }
-			const FText& GetClassDescription() const override { return MetasoundFrontendInvalidControllerPrivate::GetInvalid<FText>(); }
 			const FMetasoundFrontendInterfaceStyle& GetInputStyle() const override { static const FMetasoundFrontendInterfaceStyle Invalid; return Invalid; }
 			const FMetasoundFrontendInterfaceStyle& GetOutputStyle() const override { static const FMetasoundFrontendInterfaceStyle Invalid; return Invalid; }
 			const FMetasoundFrontendClassStyle& GetClassStyle() const override { static const FMetasoundFrontendClassStyle Invalid; return Invalid; }
 
+			const FText& GetDescription() const override { return MetasoundFrontendInvalidControllerPrivate::GetInvalid<FText>(); }
+
+			bool IsRequired() const override { return false; }
+
 			TSharedRef<IGraphController> AsGraph() override;
 			TSharedRef<const IGraphController> AsGraph() const override;
 
@@ -172,7 +174,21 @@ namespace Metasound
 			TSharedRef<IGraphController> GetOwningGraph() override;
 			TSharedRef<const IGraphController> GetOwningGraph() const override;
 
+			void IterateInputs(TUniqueFunction<void(FInputHandle)> InFunction) override { }
+			void IterateConstInputs(TUniqueFunction<void(FConstInputHandle)> InFunction) const override { }
+
+			void IterateOutputs(TUniqueFunction<void(FOutputHandle)> InFunction) override { }
+			void IterateConstOutputs(TUniqueFunction<void(FConstOutputHandle)> InFunction) const override { }
+
+			int32 GetNumInputs() const override { return 0; }
+			int32 GetNumOutputs() const override { return 0; }
+
 			const FString& GetNodeName() const override { return MetasoundFrontendInvalidControllerPrivate::GetInvalid<FString>(); }
+			const FText&  GetDisplayName() const override { return MetasoundFrontendInvalidControllerPrivate::GetInvalid<FText>(); }
+			const FText& GetDisplayTitle() const override { return MetasoundFrontendInvalidControllerPrivate::GetInvalid<FText>(); }
+			void SetDescription(const FText& InDescription) override { }
+			void SetDisplayName(const FText& InText) override { }
+
 		protected:
 			FDocumentAccess ShareAccess() override { return FDocumentAccess(); }
 			FConstDocumentAccess ShareAccess() const override { return FConstDocumentAccess(); }
@@ -212,6 +228,9 @@ namespace Metasound
 			TArray<TSharedRef<const INodeController>> GetConstOutputNodes() const override { return TArray<TSharedRef<const INodeController>>(); }
 			TArray<TSharedRef<const INodeController>> GetConstInputNodes() const override { return TArray<TSharedRef<const INodeController>>(); }
 
+			void IterateConstNodes(TUniqueFunction<void(FConstNodeHandle)> InFunction, EMetasoundFrontendClassType InClassType /* = EMetasoundFrontendClassType::Invalid */) const override { }
+			void IterateNodes(TUniqueFunction<void(FNodeHandle)> InFunction, EMetasoundFrontendClassType InClassType /* = EMetasoundFrontendClassType::Invalid */) override { }
+
 			bool ContainsOutputVertexWithName(const FString& InName) const override { return false; }
 			bool ContainsInputVertexWithName(const FString& InName) const override { return false; }
 
@@ -307,8 +326,6 @@ namespace Metasound
 				}
 
 				bool IsValid() const override { return false; }
-				bool IsRequiredInput(const FString& InInputName) const override { return false; }
-				bool IsRequiredOutput(const FString& InOutputName) const override { return false; }
 
 				const TArray<FMetasoundFrontendClassVertex>& GetRequiredInputs() const override { const static TArray<FMetasoundFrontendClassVertex> Empty; return Empty; }
 				const TArray<FMetasoundFrontendClassVertex>& GetRequiredOutputs() const override { const static TArray<FMetasoundFrontendClassVertex> Empty; return Empty; }
diff --git a/Engine/Plugins/Runtime/Metasound/Source/MetasoundFrontend/Private/MetasoundFrontendStandardController.cpp b/Engine/Plugins/Runtime/Metasound/Source/MetasoundFrontend/Private/MetasoundFrontendStandardController.cpp
--- a/Engine/Plugins/Runtime/Metasound/Source/MetasoundFrontend/Private/MetasoundFrontendStandardController.cpp
+++ b/Engine/Plugins/Runtime/Metasound/Source/MetasoundFrontend/Private/MetasoundFrontendStandardController.cpp
@@ -1138,15 +1138,6 @@ namespace Metasound
 			return Metasound::FrontendInvalidID;
 		}
 
-		const FText& FBaseNodeController::GetClassDisplayName() const
-		{
-			if (const FMetasoundFrontendClass* NodeClass = ClassPtr.Get())
-			{
-				return NodeClass->Metadata.DisplayName;
-			}
-			return FText::GetEmpty();
-		}
-
 		FMetasoundFrontendVersionNumber FBaseNodeController::GetClassVersionNumber() const
 		{
 			if (ClassPtr.IsValid())
@@ -1202,13 +1193,13 @@ namespace Metasound
 
 		void FBaseNodeController::SetNodeStyle(const FMetasoundFrontendNodeStyle& InStyle)
 		{
-			if (ClassPtr.IsValid())
+			if (NodePtr.IsValid())
 			{
 				NodePtr->Style = InStyle;
 			}
 		}
 
-		const FText& FBaseNodeController::GetClassDescription() const
+		const FText& FBaseNodeController::GetDescription() const
 		{
 			if (ClassPtr.IsValid())
 			{
@@ -1303,6 +1294,28 @@ namespace Metasound
 			return Inputs;
 		}
 
+		int32 FBaseNodeController::GetNumInputs() const
+		{
+			if (NodePtr.IsValid())
+			{
+				return NodePtr->Interface.Inputs.Num();
+			}
+
+			return 0;
+		}
+
+		void FBaseNodeController::IterateInputs(TUniqueFunction<void(FInputHandle)> InFunction)
+		{
+			for (const FInputControllerParams& Params : GetInputControllerParams())
+			{
+				FInputHandle InputHandle = CreateInputController(Params.VertexID, Params.NodeVertexPtr, Params.ClassInputPtr, AsShared());
+				if (InputHandle->IsValid())
+				{
+					InFunction(InputHandle);
+				}
+			}
+		}
+
 		TArray<FOutputHandle> FBaseNodeController::GetOutputs()
 		{
 			TArray<FOutputHandle> Outputs;
@@ -1321,6 +1334,16 @@ namespace Metasound
 			return Outputs;
 		}
 
+		int32 FBaseNodeController::GetNumOutputs() const
+		{
+			if (NodePtr.IsValid())
+			{
+				return NodePtr->Interface.Outputs.Num();
+			}
+
+			return 0;
+		}
+
 		TArray<FConstInputHandle> FBaseNodeController::GetConstInputs() const
 		{
 			TArray<FConstInputHandle> Inputs;
@@ -1347,6 +1370,52 @@ namespace Metasound
 			return Inputs;
 		}
 
+		void FBaseNodeController::IterateOutputs(TUniqueFunction<void(FOutputHandle)> InFunction)
+		{
+			for (const FOutputControllerParams& Params : GetOutputControllerParams())
+			{
+				FOutputHandle OutputHandle = CreateOutputController(Params.VertexID, Params.NodeVertexPtr, Params.ClassOutputPtr, AsShared());
+				if (OutputHandle->IsValid())
+				{
+					InFunction(OutputHandle);
+				}
+			}
+		}
+
+		const FText& FBaseNodeController::GetDisplayTitle() const
+		{
+			if (ClassPtr.IsValid())
+			{
+				return ClassPtr->Metadata.DisplayName;
+			}
+
+			return FrontendControllerIntrinsics::GetInvalidValueConstRef<FText>();
+		}
+
+		const FText& FBaseNodeController::GetDisplayName() const
+		{
+			if (ClassPtr.IsValid())
+			{
+				return ClassPtr->Metadata.DisplayName;
+			}
+			return FText::GetEmpty();
+		}
+
+		void FBaseNodeController::IterateConstInputs(TUniqueFunction<void(FConstInputHandle)> InFunction) const
+		{
+			// See early use of ConstCastSharedRef in this class for discussion.
+			FNodeHandle ThisNode = ConstCastSharedRef<INodeController>(this->AsShared());
+
+			for (const FInputControllerParams& Params : GetInputControllerParams())
+			{
+				FConstInputHandle InputHandle = CreateInputController(Params.VertexID, Params.NodeVertexPtr, Params.ClassInputPtr, ThisNode);
+				if (InputHandle->IsValid())
+				{
+					InFunction(InputHandle);
+				}
+			}
+		}
+
 		TArray<FConstOutputHandle> FBaseNodeController::GetConstOutputs() const
 		{
 			TArray<FConstOutputHandle> Outputs;
@@ -1366,6 +1435,21 @@ namespace Metasound
 			return Outputs;
 		}
 
+		void FBaseNodeController::IterateConstOutputs(TUniqueFunction<void(FConstOutputHandle)> InFunction) const
+		{
+			// See early use of ConstCastSharedRef in this class for discussion.
+			FNodeHandle ThisNode = ConstCastSharedRef<INodeController>(this->AsShared());
+
+			for (const FOutputControllerParams& Params : GetOutputControllerParams())
+			{
+				FConstOutputHandle OutputHandle = CreateOutputController(Params.VertexID, Params.NodeVertexPtr, Params.ClassOutputPtr, ThisNode);
+				if (OutputHandle->IsValid())
+				{
+					InFunction(OutputHandle);
+				}
+			}
+		}
+
 		TArray<FInputHandle> FBaseNodeController::GetInputsWithVertexName(const FString& InName) 
 		{
 			TArray<FInputHandle> Inputs;
@@ -1453,6 +1537,11 @@ namespace Metasound
 			return FInvalidInputController::GetInvalid();
 		}
 
+		bool FBaseNodeController::IsRequired() const
+		{
+			return false;
+		}
+
 		FConstInputHandle FBaseNodeController::GetInputWithID(FGuid InVertexID) const
 		{
 			FInputControllerParams Params;
@@ -1809,6 +1898,43 @@ namespace Metasound
 			return FInvalidNodeController::GetInvalid();
 		}
 
+		const FText& FOutputNodeController::GetDescription() const
+		{
+			if (OwningGraphClassOutputPtr.IsValid())
+			{
+				return OwningGraphClassOutputPtr->Metadata.Description;
+			}
+
+			return FrontendControllerIntrinsics::GetInvalidValueConstRef<FText>();
+		}
+
+		const FText& FOutputNodeController::GetDisplayName() const
+		{
+			if (OwningGraphClassOutputPtr.IsValid())
+			{
+				return OwningGraphClassOutputPtr->Metadata.DisplayName;
+			}
+
+			return FrontendControllerIntrinsics::GetInvalidValueConstRef<FText>();
+		}
+
+		void FOutputNodeController::SetDescription(const FText& InDescription)
+		{
+			if (OwningGraphClassOutputPtr.IsValid())
+			{
+				FClassOutputAccessPtr ClassOutput = ConstCastAccessPtr<FMetasoundFrontendClassOutput>(OwningGraphClassOutputPtr);
+				ClassOutput->Metadata.Description = InDescription;
+			}
+		}
+
+		void FOutputNodeController::SetDisplayName(const FText& InDisplayName)
+		{
+			if (OwningGraphClassOutputPtr.IsValid())
+			{
+				ConstCastAccessPtr<FMetasoundFrontendClassOutput>(OwningGraphClassOutputPtr)->Metadata.DisplayName = InDisplayName;
+			}
+		}
+
 		FConstNodeHandle FOutputNodeController::CreateConstOutputNodeHandle(const FOutputNodeController::FInitParams& InParams)
 		{
 			if (InParams.NodePtr.IsValid() && InParams.ClassPtr.IsValid())
@@ -1833,6 +1959,29 @@ namespace Metasound
 			return FInvalidNodeController::GetInvalid();
 		}
 
+		const FText& FOutputNodeController::GetDisplayTitle() const
+		{
+			static FText OutputDisplayTitle = LOCTEXT("OutputNode_Title", "Output");
+			return OutputDisplayTitle;
+		}
+
+		bool FOutputNodeController::IsRequired() const
+		{
+			if (NodePtr.IsValid() && OwningGraph->IsValid())
+			{
+				const FString& Name = NodePtr->Name;
+				const TArray<FMetasoundFrontendClassVertex>& RequiredOutputs = OwningGraph->GetOwningDocument()->GetRequiredOutputs();
+				for (const FMetasoundFrontendClassVertex& OutputVertex : RequiredOutputs)
+				{
+					if (OutputVertex.Name == Name)
+					{
+						return true;
+					}
+				}
+			}
+
+			return false;
+		}
 
 		bool FOutputNodeController::IsValid() const
 		{
@@ -1849,7 +1998,7 @@ namespace Metasound
 			return MakeShared<FOutputNodeOutputController>(FOutputNodeOutputController::FInitParams{InVertexID, InNodeVertexPtr, InClassOutputPtr, OwningGraphClassOutputPtr, GraphPtr, InOwningNode});
 		}
 
-		FDocumentAccess FOutputNodeController::ShareAccess() 
+		FDocumentAccess FOutputNodeController::ShareAccess()
 		{
 			FDocumentAccess Access = FBaseNodeController::ShareAccess();
 
@@ -1944,7 +2093,68 @@ namespace Metasound
 			return MakeShared<FInputNodeOutputController>(FInputNodeOutputController::FInitParams{InVertexID, InNodeVertexPtr, InClassOutputPtr, OwningGraphClassInputPtr, GraphPtr, InOwningNode});
 		}
 
-		FDocumentAccess FInputNodeController::ShareAccess() 
+		const FText& FInputNodeController::GetDescription() const
+		{
+			if (OwningGraphClassInputPtr.IsValid())
+			{
+				return OwningGraphClassInputPtr->Metadata.Description;
+			}
+
+			return FrontendControllerIntrinsics::GetInvalidValueConstRef<FText>();
+		}
+
+		const FText& FInputNodeController::GetDisplayName() const
+		{
+			if (OwningGraphClassInputPtr.IsValid())
+			{
+				return OwningGraphClassInputPtr->Metadata.DisplayName;
+			}
+
+			return FrontendControllerIntrinsics::GetInvalidValueConstRef<FText>();
+		}
+
+		const FText& FInputNodeController::GetDisplayTitle() const
+		{
+			static FText InputDisplayTitle = LOCTEXT("InputNode_Title", "Input");
+			return InputDisplayTitle;
+		}
+
+		bool FInputNodeController::IsRequired() const
+		{
+			if (NodePtr.IsValid() && OwningGraph->IsValid())
+			{
+				const FString& Name = NodePtr->Name;
+				const TArray<FMetasoundFrontendClassVertex>& RequiredInputs = OwningGraph->GetOwningDocument()->GetRequiredInputs();
+				for (const FMetasoundFrontendClassVertex& InputVertex : RequiredInputs)
+				{
+					if (InputVertex.Name == Name)
+					{
+						return true;
+					}
+				}
+			}
+
+			return false;
+		}
+
+		void FInputNodeController::SetDescription(const FText& InDescription)
+		{
+			if (OwningGraphClassInputPtr.IsValid())
+			{
+				FClassInputAccessPtr ClassInput = ConstCastAccessPtr<FMetasoundFrontendClassInput>(OwningGraphClassInputPtr);
+				ClassInput->Metadata.Description = InDescription;
+			}
+		}
+
+		void FInputNodeController::SetDisplayName(const FText& InDisplayName)
+		{
+			if (OwningGraphClassInputPtr.IsValid())
+			{
+				ConstCastAccessPtr<FMetasoundFrontendClassInput>(OwningGraphClassInputPtr)->Metadata.DisplayName = InDisplayName;
+			}
+		}
+
+		FDocumentAccess FInputNodeController::ShareAccess()
 		{
 			FDocumentAccess Access = FBaseNodeController::ShareAccess();
 
@@ -2133,6 +2343,58 @@ namespace Metasound
 			return GetNodesByPredicate(IsInputNode);
 		}
 
+		void FGraphController::IterateNodes(TUniqueFunction<void(FNodeHandle)> InFunction, EMetasoundFrontendClassType InClassType)
+		{
+			using namespace FrontendControllerIntrinsics;
+
+			if (GraphClassPtr.IsValid())
+			{
+				for (const FMetasoundFrontendNode& Node : GraphClassPtr->Graph.Nodes)
+				{
+					FConstClassAccessPtr NodeClassPtr = OwningDocument->FindClassWithID(Node.ClassID);
+					if (const FMetasoundFrontendClass* NodeClass = NodeClassPtr.Get())
+					{
+						if (InClassType == EMetasoundFrontendClassType::Invalid || NodeClass->Metadata.Type == InClassType)
+						{
+							FNodeAccessPtr NodePtr = GraphClassPtr.GetMemberAccessPtr<FMetasoundFrontendNode>(FGetNodeFromGraphClassByNodeID(Node.ID));
+							FNodeHandle NodeHandle = GetNodeHandle(FGraphController::FNodeAndClass{ NodePtr, NodeClassPtr });
+							InFunction(NodeHandle);
+						}
+					}
+					else
+					{
+						UE_LOG(LogMetasound, Warning, TEXT("Failed to find class for node [NodeID:%s, ClassID:%s]"), *Node.ID.ToString(), *Node.ClassID.ToString());
+					}
+				}
+			}
+		}
+
+		void FGraphController::IterateConstNodes(TUniqueFunction<void(FConstNodeHandle)> InFunction, EMetasoundFrontendClassType InClassType) const
+		{
+			using namespace FrontendControllerIntrinsics;
+
+			if (GraphClassPtr.IsValid())
+			{
+				for (const FMetasoundFrontendNode& Node : GraphClassPtr->Graph.Nodes)
+				{
+					FConstClassAccessPtr NodeClassPtr = OwningDocument->FindClassWithID(Node.ClassID);
+					if (const FMetasoundFrontendClass* NodeClass = NodeClassPtr.Get())
+					{
+						if (InClassType == EMetasoundFrontendClassType::Invalid || NodeClass->Metadata.Type == InClassType)
+						{
+							FConstNodeAccessPtr NodePtr = GraphClassPtr.GetMemberAccessPtr<const FMetasoundFrontendNode>(FGetConstNodeFromGraphClassByNodeID(Node.ID));
+							FConstNodeHandle NodeHandle = GetNodeHandle(FGraphController::FConstNodeAndClass{ NodePtr, NodeClassPtr });
+							InFunction(NodeHandle);
+						}
+					}
+					else
+					{
+						UE_LOG(LogMetasound, Warning, TEXT("Failed to find class for node [NodeID:%s, ClassID:%s]"), *Node.ID.ToString(), *Node.ClassID.ToString());
+					}
+				}
+			}
+		}
+
 		TArray<FConstNodeHandle> FGraphController::GetConstOutputNodes() const
 		{
 			auto IsOutputNode = [](const FMetasoundFrontendClass& NodeClass, const FMetasoundFrontendNode& Node) 
@@ -3196,28 +3458,6 @@ namespace Metasound
 			return DocumentPtr.IsValid();
 		}
 
-		bool FDocumentController::IsRequiredInput(const FString& InInputName) const
-		{
-			if (DocumentPtr.IsValid())
-			{ 
-				auto IsInputWithSameName = [&](const FMetasoundFrontendClassVertex& Desc) { return Desc.Name == InInputName; };
-
-				return DocumentPtr->Archetype.Interface.Inputs.ContainsByPredicate(IsInputWithSameName);
-			}
-			return false;
-		}
-
-		bool FDocumentController::IsRequiredOutput(const FString& InOutputName) const
-		{
-			if (DocumentPtr.IsValid())
-			{ 
-				auto IsOutputWithSameName = [&](const FMetasoundFrontendClassVertex& Desc) { return Desc.Name == InOutputName; };
-
-				return DocumentPtr->Archetype.Interface.Outputs.ContainsByPredicate(IsOutputWithSameName);
-			}
-			return false;
-		}
-
 		const TArray<FMetasoundFrontendClassVertex>& FDocumentController::GetRequiredInputs() const
 		{
 			if (DocumentPtr.IsValid())
@@ -3556,6 +3796,4 @@ namespace Metasound
 		}
 	}
 }
-
 #undef LOCTEXT_NAMESPACE 
-
diff --git a/Engine/Plugins/Runtime/Metasound/Source/MetasoundFrontend/Private/MetasoundFrontendStandardController.h b/Engine/Plugins/Runtime/Metasound/Source/MetasoundFrontend/Private/MetasoundFrontendStandardController.h
--- a/Engine/Plugins/Runtime/Metasound/Source/MetasoundFrontend/Private/MetasoundFrontendStandardController.h
+++ b/Engine/Plugins/Runtime/Metasound/Source/MetasoundFrontend/Private/MetasoundFrontendStandardController.h
@@ -345,14 +345,15 @@ namespace Metasound
 			// Info about this node.
 			FGuid GetID() const override;
 			FGuid GetClassID() const override;
-			const FText& GetClassDisplayName() const override;
 
 			FMetasoundFrontendVersionNumber GetClassVersionNumber() const override;
-			const FText& GetClassDescription() const override;
 			const FMetasoundFrontendInterfaceStyle& GetInputStyle() const override;
 			const FMetasoundFrontendInterfaceStyle& GetOutputStyle() const override;
 			const FMetasoundFrontendClassStyle& GetClassStyle() const override;
 
+			/** Description of the given node. */
+			const FText& GetDescription() const override;
+
 			const FMetasoundFrontendNodeStyle& GetNodeStyle() const override;
 			void SetNodeStyle(const FMetasoundFrontendNodeStyle& InStyle) override;
 
@@ -372,9 +373,30 @@ namespace Metasound
 			/** Returns all node inputs. */
 			TArray<FInputHandle> GetInputs() override;
 
+			/** Returns the human-readable display name of the given node. */
+			const FText& GetDisplayName() const override;
+
+			/** Sets the description of the node. */
+			void SetDescription(const FText& InDescription) override { }
+
+			/** Sets the display name of the node. */
+			void SetDisplayName(const FText& InDisplayName) override { };
+
+			/** Returns the title of the given node (what to label in visual node). */
+			const FText& GetDisplayTitle() const override;
+
 			/** Returns all node inputs. */
 			TArray<FConstInputHandle> GetConstInputs() const override;
 
+			void IterateConstInputs(TUniqueFunction<void(FConstInputHandle)> InFunction) const override;
+			void IterateConstOutputs(TUniqueFunction<void(FConstOutputHandle)> InFunction) const override;
+
+			void IterateInputs(TUniqueFunction<void(FInputHandle)> InFunction) override;
+			void IterateOutputs(TUniqueFunction<void(FOutputHandle)> InFunction) override;
+
+			int32 GetNumInputs() const override;
+			int32 GetNumOutputs() const override;
+
 			TArray<FInputHandle> GetInputsWithVertexName(const FString& InName) override;
 			TArray<FConstInputHandle> GetInputsWithVertexName(const FString& InName) const override;
 
@@ -387,6 +409,8 @@ namespace Metasound
 			TArray<FOutputHandle> GetOutputsWithVertexName(const FString& InName) override;
 			TArray<FConstOutputHandle> GetOutputsWithVertexName(const FString& InName) const override;
 
+			bool IsRequired() const override;
+
 			/** Returns an input with the given id. 
 			 *
 			 * If the input does not exist, an invalid handle is returned. 
@@ -550,9 +574,16 @@ namespace Metasound
 			 */
 			static FConstNodeHandle CreateConstOutputNodeHandle(const FInitParams& InParams);
 
+
 			virtual ~FOutputNodeController() = default;
 
 			bool IsValid() const override;
+			const FText& GetDescription() const override;
+			const FText& GetDisplayName() const override;
+			void SetDescription(const FText& InDescription) override;
+			void SetDisplayName(const FText& InText) override;
+			const FText& GetDisplayTitle() const override;
+			bool IsRequired() const override;
 
 		protected:
 
@@ -606,7 +637,13 @@ namespace Metasound
 
 			virtual ~FInputNodeController() = default;
 
+			const FText& GetDescription() const override;
+			const FText& GetDisplayName() const override;
+			const FText& GetDisplayTitle() const override;
+			bool IsRequired() const override;
 			bool IsValid() const override;
+			void SetDescription(const FText& InDescription) override;
+			void SetDisplayName(const FText& InText) override;
 
 		protected:
 
@@ -682,6 +719,9 @@ namespace Metasound
 			TArray<FNodeHandle> GetInputNodes() override;
 			TArray<FConstNodeHandle> GetConstInputNodes() const override;
 
+			void IterateNodes(TUniqueFunction<void(FNodeHandle)> InFunction, EMetasoundFrontendClassType InClassType) override;
+			void IterateConstNodes(TUniqueFunction<void(FConstNodeHandle)> InFunction, EMetasoundFrontendClassType InClassType) const override;
+
 			TArray<FNodeHandle> GetOutputNodes() override;
 			TArray<FConstNodeHandle> GetConstOutputNodes() const override;
 
@@ -862,9 +902,6 @@ namespace Metasound
 
 			bool IsValid() const override;
 
-			bool IsRequiredInput(const FString& InInputName) const override;
-			bool IsRequiredOutput(const FString& InOutputName) const override;
-
 			const TArray<FMetasoundFrontendClassVertex>& GetRequiredInputs() const override;
 			const TArray<FMetasoundFrontendClassVertex>& GetRequiredOutputs() const override;
 
diff --git a/Engine/Plugins/Runtime/Metasound/Source/MetasoundFrontend/Public/MetasoundFrontendController.h b/Engine/Plugins/Runtime/Metasound/Source/MetasoundFrontend/Public/MetasoundFrontendController.h
--- a/Engine/Plugins/Runtime/Metasound/Source/MetasoundFrontend/Public/MetasoundFrontendController.h
+++ b/Engine/Plugins/Runtime/Metasound/Source/MetasoundFrontend/Public/MetasoundFrontendController.h
@@ -303,6 +303,32 @@ namespace Metasound
 			/** Returns all node inputs. */
 			virtual TArray<FConstInputHandle> GetConstInputs() const = 0;
 
+			/** Returns the display name of the given node (what to distinguish and label in visual arrays, such as context menus). */
+			virtual const FText& GetDisplayName() const = 0;
+
+			/** Sets the description of the node. */
+			virtual void SetDescription(const FText& InDescription) = 0;
+
+			/** Sets the display name of the node. */
+			virtual void SetDisplayName(const FText& InDisplayName) = 0;
+
+			/** Returns the title of the given node (what to label when displayed as visual node). */
+			virtual const FText& GetDisplayTitle() const = 0;
+
+			/** Iterate over inputs */
+			virtual void IterateInputs(TUniqueFunction<void(FInputHandle)> InFunction) = 0;
+			virtual void IterateConstInputs(TUniqueFunction<void(FConstInputHandle)> InFunction) const = 0;
+
+			/** Returns number of node inputs. */
+			virtual int32 GetNumInputs() const = 0;
+
+			/** Iterate over outputs */
+			virtual void IterateOutputs(TUniqueFunction<void(FOutputHandle)> InFunction) = 0;
+			virtual void IterateConstOutputs(TUniqueFunction<void(FConstOutputHandle)> InFunction) const = 0;
+
+			/** Returns number of node outputs. */
+			virtual int32 GetNumOutputs() const = 0;
+
 			virtual TArray<FInputHandle> GetInputsWithVertexName(const FString& InName) = 0;
 			virtual TArray<FConstInputHandle> GetInputsWithVertexName(const FString& InName) const = 0;
 
@@ -315,15 +341,18 @@ namespace Metasound
 			virtual TArray<FOutputHandle> GetOutputsWithVertexName(const FString& InName) = 0;
 			virtual TArray<FConstOutputHandle> GetOutputsWithVertexName(const FString& InName) const = 0;
 
-			/** Returns an input with the given id. 
+			/** Returns true if node is required to satisfy the document archetype. */
+			virtual bool IsRequired() const = 0;
+
+			/** Returns an input with the given id.
 			 *
-			 * If the input does not exist, an invalid handle is returned. 
+			 * If the input does not exist, an invalid handle is returned.
 			 */
 			virtual FInputHandle GetInputWithID(FGuid InVertexID) = 0;
 
 			/** Returns an input with the given name. 
 			 *
-			 * If the input does not exist, an invalid handle is returned. 
+			 * If the input does not exist, an invalid handle is returned.
 			 */
 			virtual FConstInputHandle GetInputWithID(FGuid InVertexID) const = 0;
 
@@ -354,12 +383,13 @@ namespace Metasound
 			/** Returns the name of the node class. */
 			virtual const FMetasoundFrontendClassName& GetClassName() const = 0;
 			virtual FMetasoundFrontendVersionNumber GetClassVersionNumber() const = 0;
-			virtual const FText& GetClassDisplayName() const = 0;
-			virtual const FText& GetClassDescription() const = 0;
 			virtual const FMetasoundFrontendInterfaceStyle& GetOutputStyle() const = 0;
 			virtual const FMetasoundFrontendInterfaceStyle& GetInputStyle() const = 0;
 			virtual const FMetasoundFrontendClassStyle& GetClassStyle() const = 0;
 
+			/** Description of the given node. */
+			virtual const FText& GetDescription() const = 0;
+
 			/** If the node is also a graph, this returns a graph handle.
 			 * If the node is not also a graph, it will return an invalid handle.
 			 */
@@ -433,6 +463,12 @@ namespace Metasound
 			/** Returns all input nodes in the graph. */
 			virtual TArray<FConstNodeHandle> GetConstInputNodes() const = 0;
 
+			/** Iterate over all input nodes with the given function. If ClassType is specified, only iterate over given type. */
+			virtual void IterateNodes(TUniqueFunction<void(FNodeHandle)> InFunction, EMetasoundFrontendClassType InClassType = EMetasoundFrontendClassType::Invalid) = 0;
+
+			/** Iterate over all nodes with the given function. If ClassType is specified, only iterate over given type. */
+			virtual void IterateConstNodes(TUniqueFunction<void(FConstNodeHandle)> InFunction, EMetasoundFrontendClassType InClassType = EMetasoundFrontendClassType::Invalid) const = 0;
+
 			/** Returns true if an output vertex with the given Name exists.
 			 *
 			 * @param InName - Name of vertex.
@@ -654,12 +690,6 @@ namespace Metasound
 			/** Returns true if the controller is in a valid state. */
 			virtual bool IsValid() const = 0;
 
-			/** Returns true if an input with the provided name is required to satisfy the document archetype. */
-			virtual bool IsRequiredInput(const FString& InInputName) const = 0;
-
-			/** Returns true if an input with the provided name is required to satisfy the document archetype. */
-			virtual bool IsRequiredOutput(const FString& InOutputName) const = 0;
-
 			/** Returns an array of inputs which describe the required inputs needed to satisfy the document archetype. */
 			virtual const TArray<FMetasoundFrontendClassVertex>& GetRequiredInputs() const = 0;
 
diff --git a/Engine/Plugins/Runtime/Metasound/Source/MetasoundFrontend/Public/MetasoundFrontendDocument.h b/Engine/Plugins/Runtime/Metasound/Source/MetasoundFrontend/Public/MetasoundFrontendDocument.h
--- a/Engine/Plugins/Runtime/Metasound/Source/MetasoundFrontend/Public/MetasoundFrontendDocument.h
+++ b/Engine/Plugins/Runtime/Metasound/Source/MetasoundFrontend/Public/MetasoundFrontendDocument.h
@@ -151,9 +151,10 @@ struct METASOUNDFRONTEND_API FMetasoundFrontendNodeStyleDisplay
 	UPROPERTY()
 	EMetasoundFrontendNodeStyleDisplayVisibility Visibility = EMetasoundFrontendNodeStyleDisplayVisibility::Visible;
 
-	// Position of node
+	// Map of visual node guid to 2D location. May have more than one if the node allows displaying in
+	// more than one place on the graph (Only functionally relevant for nodes that cannot contain inputs.)
 	UPROPERTY()
-	FVector2D Location = FVector2D::ZeroVector;
+	TMap<FGuid, FVector2D> Locations;
 };
 
 USTRUCT()
