Fix for crash in random node get

diff --git a/Engine/Plugins/Runtime/Metasound/Source/MetasoundFrontend/Private/MetasoundArrayRandomNode.cpp b/Engine/Plugins/Runtime/Metasound/Source/MetasoundFrontend/Private/MetasoundArrayRandomNode.cpp
--- a/Engine/Plugins/Runtime/Metasound/Source/MetasoundFrontend/Private/MetasoundArrayRandomNode.cpp
+++ b/Engine/Plugins/Runtime/Metasound/Source/MetasoundFrontend/Private/MetasoundArrayRandomNode.cpp
@@ -79,9 +79,7 @@ namespace Metasound
 		SetSeed(InSeed);
 		MaxIndex = InMaxIndex;
 		SetNoRepeatOrder(InNoRepeatOrder);
-		PreviousIndicesQueue = MakeUnique<TCircularQueue<int32>>(InNoRepeatOrder);
-		PreviousIndices.Reset();
-
+		check(!InNoRepeatOrder || !PreviousIndicesQueue->IsFull());
 	}
 
 	void FArrayRandomGet::SetSeed(int32 InSeed)
@@ -105,7 +103,7 @@ namespace Metasound
 
 		if (InNoRepeatOrder != NoRepeatOrder)
 		{
-			PreviousIndicesQueue = MakeUnique<TCircularQueue<int32>>(InNoRepeatOrder);
+			PreviousIndicesQueue = MakeUnique<TCircularQueue<int32>>(InNoRepeatOrder + 1);
 			PreviousIndices.Reset();
 			NoRepeatOrder = InNoRepeatOrder;
 		}
@@ -175,19 +173,26 @@ namespace Metasound
 		check(ChosenIndex != INDEX_NONE);
 
 		// Dequeue and remove the oldest previous index
-		if (PreviousIndices.Num() == NoRepeatOrder)
+		if (NoRepeatOrder > 0)
 		{
-			int32 OldPrevIndex;
-			PreviousIndicesQueue->Dequeue(OldPrevIndex);
-			PreviousIndices.Remove(OldPrevIndex);
-			check(PreviousIndices.Num() == NoRepeatOrder - 1);
-		}
+			if (PreviousIndices.Num() == NoRepeatOrder)
+			{
+				check(PreviousIndicesQueue->Count() == PreviousIndices.Num());
+				int32 OldPrevIndex;
+				PreviousIndicesQueue->Dequeue(OldPrevIndex);
+				PreviousIndices.Remove(OldPrevIndex);
+				check(PreviousIndices.Num() == NoRepeatOrder - 1);
+			}
+
+			check(PreviousIndices.Num() < NoRepeatOrder);
+			check(!PreviousIndicesQueue->IsFull());
 
-		// We should only have less than *or* equal to the NoRepeatOrder
-		check(PreviousIndices.Num() < NoRepeatOrder);
-		PreviousIndicesQueue->Enqueue(ChosenIndex);
-		check(!PreviousIndices.Contains(ChosenIndex));
-		PreviousIndices.Add(ChosenIndex);
+			bool bSuccess = PreviousIndicesQueue->Enqueue(ChosenIndex);
+			check(bSuccess);
+			check(!PreviousIndices.Contains(ChosenIndex));
+			PreviousIndices.Add(ChosenIndex);
+			check(PreviousIndicesQueue->Count() == PreviousIndices.Num());
+		}
 
 		return ChosenIndex;
 	}
diff --git a/Engine/Plugins/Runtime/Metasound/Source/MetasoundFrontend/Public/MetasoundArrayRandomNode.h b/Engine/Plugins/Runtime/Metasound/Source/MetasoundFrontend/Public/MetasoundArrayRandomNode.h
--- a/Engine/Plugins/Runtime/Metasound/Source/MetasoundFrontend/Public/MetasoundArrayRandomNode.h
+++ b/Engine/Plugins/Runtime/Metasound/Source/MetasoundFrontend/Public/MetasoundArrayRandomNode.h
@@ -55,7 +55,7 @@ namespace Metasound
 		// The current index into the array of indicies (wraps between 0 and ShuffleIndices.Num())
 		TArray<int32> PreviousIndices;
 		TUniquePtr<TCircularQueue<int32>> PreviousIndicesQueue;
-		int32 NoRepeatOrder = 0;
+		int32 NoRepeatOrder = INDEX_NONE;
 
 		// Array of indices (in order 0 to Num)
 		int32 MaxIndex = 0;
