Changing name of "wave player node" to "wave player". Node is redundant.

diff --git a/Engine/Plugins/Runtime/Metasound/Source/MetasoundEngine/Private/MetasoundWavePlayerNode.cpp b/Engine/Plugins/Runtime/Metasound/Source/MetasoundEngine/Private/MetasoundWavePlayerNode.cpp
--- a/Engine/Plugins/Runtime/Metasound/Source/MetasoundEngine/Private/MetasoundWavePlayerNode.cpp
+++ b/Engine/Plugins/Runtime/Metasound/Source/MetasoundEngine/Private/MetasoundWavePlayerNode.cpp
@@ -428,7 +428,7 @@ namespace Metasound
 			Info.ClassName = { Metasound::EngineNodes::Namespace, TEXT("Wave Player"), Metasound::EngineNodes::StereoVariant };
 			Info.MajorVersion = 1;
 			Info.MinorVersion = 0;
-			Info.DisplayName = LOCTEXT("Metasound_WavePlayerNodeDisplayName", "Wave Player Node");
+			Info.DisplayName = LOCTEXT("Metasound_WavePlayerNodeDisplayName", "Wave Player");
 			Info.Description = LOCTEXT("Metasound_WavePlayerNodeDescription", "Plays a supplied Wave");
 			Info.Author = PluginAuthor;
 			Info.PromptIfMissing = PluginNodeMissingPrompt;
