Metasound LFOs

diff --git a/Engine/Plugins/Runtime/Metasound/Source/MetasoundStandardNodes/Private/MetasoundOscillatorNodes.cpp b/Engine/Plugins/Runtime/Metasound/Source/MetasoundStandardNodes/Private/MetasoundOscillatorNodes.cpp
--- a/Engine/Plugins/Runtime/Metasound/Source/MetasoundStandardNodes/Private/MetasoundOscillatorNodes.cpp
+++ b/Engine/Plugins/Runtime/Metasound/Source/MetasoundStandardNodes/Private/MetasoundOscillatorNodes.cpp
@@ -1050,5 +1050,216 @@ namespace Metasound
 
 	METASOUND_REGISTER_NODE(FTriangleOscilatorNode);
 #pragma endregion Triangle
+
+#pragma region LFO
+
+	enum class ELfoWaveshapeType
+	{
+		Sine,
+		Saw,
+		Triangle,
+		Square,
+	};
+
+	DECLARE_METASOUND_ENUM(ELfoWaveshapeType, ELfoWaveshapeType::Sine, METASOUNDSTANDARDNODES_API,
+		FEnumLfoWaveshapeType, FEnumLfoWaveshapeTypeInfo, FEnumLfoWaveshapeTypeReadRef, FEnumLfoWaveshapeTypeWriteRef);
+
+	DEFINE_METASOUND_ENUM_BEGIN(ELfoWaveshapeType, FEnumLfoWaveshapeType, "LfoWaveshapeType")
+		DEFINE_METASOUND_ENUM_ENTRY(ELfoWaveshapeType::Sine, LOCTEXT("LfoWaveShapeSineDescription", "Sine"), LOCTEXT("LfoWaveShapeSineDescriptionTT", "Sinewave Low Frequency Oscillator")),
+		DEFINE_METASOUND_ENUM_ENTRY(ELfoWaveshapeType::Saw, LOCTEXT("LfoWaveShapeSawDescription", "Saw"), LOCTEXT("LfoWaveShapeSawDescriptionTT", "Sawtooth Low Frequency Oscillator")),
+		DEFINE_METASOUND_ENUM_ENTRY(ELfoWaveshapeType::Triangle, LOCTEXT("LfoWaveShapeTriangleDescription", "Triangle"), LOCTEXT("LfoWaveShapeTriangleDescriptionTT", "Triangle shape Frequency Oscillator")),
+		DEFINE_METASOUND_ENUM_ENTRY(ELfoWaveshapeType::Square, LOCTEXT("LfoWaveShapeSquareDescription", "Square"), LOCTEXT("LfoWaveShapeSquareDescriptionTT", "Square shape Low Frequency Oscillator"))
+	DEFINE_METASOUND_ENUM_END()
+
+	// Block rate Oscillator Swiss Army Knife
+	class FLfoOperator : public TExecutableOperator<FLfoOperator>
+	{
+	public:
+		// Common pins
+		static constexpr const TCHAR* WaveshapePinName = TEXT("Shape");
+		static constexpr const TCHAR* LfoOutPinName = TEXT("Out");
+	
+		static const FVertexInterface& GetVertexInterface()
+		{
+			static const FVertexInterface Interface
+			{
+				FInputVertexInterface{
+					TInputDataVertexModel<bool>(FOscilatorFactoryBase::BiPolarPinName, LOCTEXT("LfoBiPolarDescription", "If the output is bipolar (-1..1) or unipolar (0..1)"), true),
+					TInputDataVertexModel<float>(FOscilatorFactoryBase::BaseFrequencyPinName, LOCTEXT("LfoFrequencyDescription", "Frequency of LFO (Hz), clamped at blockrate"), 5.f),
+					TInputDataVertexModel<FTrigger>(FOscilatorFactoryBase::PhaseResetPinName, LOCTEXT("LfoPhaseResetDescription", "Phase Reset (block rate only)")),
+					TInputDataVertexModel<float>(FOscilatorFactoryBase::PhaseOffsetPinName, LOCTEXT("LfoPhaseOffsetDescription", "Phase Offset In Degrees (0..360)"), 0.f),
+					TInputDataVertexModel<float>(FSquareOscilatorNode::FFactory::PulseWidthPinName, LOCTEXT("LfoPulseWidthDescription", "Pulse Width (0..1)"), 0.25f),
+					TInputDataVertexModel<FEnumLfoWaveshapeType>(WaveshapePinName, LOCTEXT("LfoShapeDescription", "Waveshape of the LFO")),
+				},
+				FOutputVertexInterface{
+					TOutputDataVertexModel<float>(LfoOutPinName, LOCTEXT("LfoOutputDescription", "Output of the LFO (blockrate)"))
+				}
+			};
+			return Interface;
+		}
+
+		static const FNodeClassMetadata& GetNodeInfo()
+		{
+			auto InitNodeInfo = []() -> FNodeClassMetadata
+			{
+				FNodeClassMetadata Info;
+				Info.ClassName = { Metasound::StandardNodes::Namespace, TEXT("LFO"), Metasound::StandardNodes::AudioVariant };
+				Info.MajorVersion = 1;
+				Info.MinorVersion = 0;
+				Info.DisplayName = LOCTEXT("Metasound_LfoNodeDisplayName", "LFO");
+				Info.Description = LOCTEXT("Metasound_LfoNodeDescription", "Low frequency oscillator < blockrate");
+				Info.Author = PluginAuthor;
+				Info.PromptIfMissing = PluginNodeMissingPrompt;
+				Info.DefaultInterface = GetVertexInterface();
+				Info.CategoryHierarchy.Emplace(StandardNodes::Generators);
+				return Info;
+			};
+			static const FNodeClassMetadata Info = InitNodeInfo();
+			return Info;
+		}
+		
+		static TUniquePtr<IOperator> CreateOperator(const FCreateOperatorParams& InParams, FBuildErrorArray& OutErrors)
+		{
+			const FLfoNode& Node = static_cast<const FLfoNode&>(InParams.Node);
+			const FDataReferenceCollection& InputCol = InParams.InputDataReferences;
+			const FOperatorSettings& Settings = InParams.OperatorSettings;
+			const FInputVertexInterface& InputInterface = GetVertexInterface().GetInputInterface();
+
+			return MakeUnique<FLfoOperator>(
+				  Settings
+				, InputCol.GetDataReadReferenceOrConstructWithVertexDefault<float>(InputInterface, FOscilatorFactoryBase::BaseFrequencyPinName)
+				, InputCol.GetDataReadReferenceOrConstruct<FEnumLfoWaveshapeType>(WaveshapePinName)
+				, InputCol.GetDataReadReferenceOrConstruct<FTrigger>(FOscilatorFactoryBase::PhaseResetPinName, Settings)
+				, InputCol.GetDataReadReferenceOrConstructWithVertexDefault<float>(InputInterface,FOscilatorFactoryBase::PhaseOffsetPinName)
+				, InputCol.GetDataReadReferenceOrConstructWithVertexDefault<bool>(InputInterface,FOscilatorFactoryBase::BiPolarPinName)
+				, InputCol.GetDataReadReferenceOrConstructWithVertexDefault<float>(InputInterface,FSquareOscilatorNode::FFactory::PulseWidthPinName)
+			);
+		}
+
+		FLfoOperator(const FOperatorSettings& InSettings, FFloatReadRef&& InFrequency, FEnumLfoWaveshapeTypeReadRef&& InType, 
+			FTriggerReadRef&& InPhaseReset, FFloatReadRef&& InPhaseOffset, FBoolReadRef&& InBiPolar, FFloatReadRef&& InPulseWidth)
+			: BlockRate{ InSettings.GetActualBlockRate() }
+			, Phase{ 0.0f }
+			, Frequency{ MoveTemp(InFrequency) }
+			, Waveshape{ MoveTemp(InType) }
+			, PhaseReset{ MoveTemp(InPhaseReset) }
+			, PhaseOffset{ MoveTemp(InPhaseOffset) }
+			, BiPolar{ MoveTemp(InBiPolar) }
+			, PulseWidth{ MoveTemp(InPulseWidth) }
+			, Output{ FFloatWriteRef::CreateNew(0.f) }
+		{
+			ResetPhase();
+		}
+
+		FDataReferenceCollection GetInputs() const override
+		{
+			FDataReferenceCollection InputDataReferences;
+			InputDataReferences.AddDataReadReference(FOscilatorFactoryBase::BaseFrequencyPinName, Frequency);
+			InputDataReferences.AddDataReadReference(WaveshapePinName, Waveshape);
+			InputDataReferences.AddDataReadReference(FOscilatorFactoryBase::PhaseOffsetPinName, PhaseOffset);
+			InputDataReferences.AddDataReadReference(FOscilatorFactoryBase::PhaseResetPinName, PhaseReset);
+			InputDataReferences.AddDataReadReference(FOscilatorFactoryBase::BiPolarPinName, BiPolar);
+			InputDataReferences.AddDataReadReference(FSquareOscilatorNode::FFactory::PulseWidthPinName, PulseWidth);
+			return InputDataReferences;
+		}
+
+		FDataReferenceCollection GetOutputs() const override
+		{
+			FDataReferenceCollection OutputDataReferences;
+			OutputDataReferences.AddDataReadReference(LfoOutPinName, Output);
+			return OutputDataReferences;
+		}
+
+		void ResetPhase()
+		{
+			float ClampedDegrees = FMath::Clamp(*PhaseOffset, 0.f, 360.f);
+			Phase = ClampedDegrees / 360.f;
+		}
+
+		void Execute()
+		{
+			using namespace Generators;
+		
+			// Prevent LFO going faster than the block rate Nyquist
+			const float Nyquist = BlockRate / 2.f;
+			const float ClampedFreq = FMath::Clamp(*Frequency, 0.f, Nyquist);
+			const float DeltaPhase = ClampedFreq * (1.f / BlockRate);
+			const float ClampPulseWidth = FMath::Clamp(*PulseWidth, 0.f, 1.f);
+
+			FGeneratorArgs Args{ BlockRate, ClampedFreq, ClampPulseWidth };
+
+			// We are not sample accurate.
+			if (PhaseReset->IsTriggeredInBlock()) 
+			{
+				ResetPhase();
+			}
+			
+			// Wrap phase. (0..1)
+			Wrap(Phase);
+
+			float Value = 0.f;
+			switch (*Waveshape)
+			{
+				case ELfoWaveshapeType::Sine:
+				{
+					FBhaskaraGenerator Generator;
+					Value = Generator(Phase, DeltaPhase, Args);
+					break;
+				}
+				case ELfoWaveshapeType::Saw:
+				{
+					FSawPolySmoothGenerator Generator;
+					Value = Generator(Phase, DeltaPhase, Args);
+					break;
+				}
+				case ELfoWaveshapeType::Triangle:
+				{
+					FTriangleGenerator Generator;
+					Value = Generator(Phase, DeltaPhase, Args);
+					break;
+				}
+				case ELfoWaveshapeType::Square:
+				{
+					FSquarePolysmoothGenerator Generator;
+					Value = Generator(Phase, DeltaPhase, Args);
+					break;
+				}
+				default:
+				{
+					checkNoEntry();
+					break;
+				}
+			}
+
+			if (!*BiPolar)
+			{
+				Value = Audio::GetUnipolar(Value);
+			}
+
+			*Output = Value;
+
+			Phase += DeltaPhase;
+		}
+
+	private:
+		float BlockRate = 0.f;
+		float Phase = 0.f;
+		Generators::FWrapPhase Wrap;
+		FFloatReadRef Frequency;
+		FEnumLfoWaveshapeTypeReadRef Waveshape;
+		FTriggerReadRef PhaseReset;
+		FFloatReadRef PhaseOffset;
+		FBoolReadRef BiPolar;
+		FFloatReadRef PulseWidth;
+		FFloatWriteRef Output;
+	};
+
+	FLfoNode::FLfoNode(const FNodeInitData& InInitData)
+		: FNodeFacade(InInitData.InstanceName, InInitData.InstanceID, TFacadeOperatorClass<FLfoOperator>())
+	{}
+
+	METASOUND_REGISTER_NODE(FLfoNode);
+
+#pragma endregion LFO
 }
-#undef LOCTEXT_NAMESPACE //MetasoundStandardNodes
\ No newline at end of file
+#undef LOCTEXT_NAMESPACE //MetasoundStandardNodes
diff --git a/Engine/Plugins/Runtime/Metasound/Source/MetasoundStandardNodes/Public/MetasoundOscillatorNodes.h b/Engine/Plugins/Runtime/Metasound/Source/MetasoundStandardNodes/Public/MetasoundOscillatorNodes.h
--- a/Engine/Plugins/Runtime/Metasound/Source/MetasoundStandardNodes/Public/MetasoundOscillatorNodes.h
+++ b/Engine/Plugins/Runtime/Metasound/Source/MetasoundStandardNodes/Public/MetasoundOscillatorNodes.h
@@ -75,7 +75,7 @@ namespace Metasound
 		FSineOscilatorNode(const FString& InInstanceName, const FGuid& InInstanceID, float InDefaultFrequency, bool bInDefaultEnablement);
 		FSineOscilatorNode(const FNodeInitData& InInitData);
 	};
-
+		
 	class METASOUNDSTANDARDNODES_API FSawOscilatorNode : public FOscilatorNodeBase
 	{
 	public:
@@ -99,4 +99,11 @@ namespace Metasound
 		FSquareOscilatorNode(const FString& InInstanceName, const FGuid& InInstanceID, float InDefaultFrequency, bool bInDefaultEnablement);
 		FSquareOscilatorNode(const FNodeInitData& InInitData);
 	};
+
+	class METASOUNDSTANDARDNODES_API FLfoNode : public FNodeFacade
+	{
+	public:		
+		FLfoNode(const FNodeInitData& InInitData);
+		virtual ~FLfoNode() = default;
+	};
 }
