Fix for random bool crash

Adding Random Time type

diff --git a/Engine/Plugins/Runtime/Metasound/Source/MetasoundStandardNodes/Private/MetasoundRandomNode.cpp b/Engine/Plugins/Runtime/Metasound/Source/MetasoundStandardNodes/Private/MetasoundRandomNode.cpp
--- a/Engine/Plugins/Runtime/Metasound/Source/MetasoundStandardNodes/Private/MetasoundRandomNode.cpp
+++ b/Engine/Plugins/Runtime/Metasound/Source/MetasoundStandardNodes/Private/MetasoundRandomNode.cpp
@@ -112,6 +112,9 @@ namespace Metasound
  
 	using FRandomNodeBool = TRandomNode<bool>;
 	METASOUND_REGISTER_NODE(FRandomNodeBool)
+
+	using FRandomNodeTime = TRandomNode<FTime>;
+	METASOUND_REGISTER_NODE(FRandomNodeTime)
 }
 
 #undef LOCTEXT_NAMESPACE
diff --git a/Engine/Plugins/Runtime/Metasound/Source/MetasoundStandardNodes/Public/MetasoundRandomNode.h b/Engine/Plugins/Runtime/Metasound/Source/MetasoundStandardNodes/Public/MetasoundRandomNode.h
--- a/Engine/Plugins/Runtime/Metasound/Source/MetasoundStandardNodes/Public/MetasoundRandomNode.h
+++ b/Engine/Plugins/Runtime/Metasound/Source/MetasoundStandardNodes/Public/MetasoundRandomNode.h
@@ -10,11 +10,37 @@
 #include "MetasoundTrigger.h"
 #include "Internationalization/Text.h"
 #include "MetasoundStandardNodesCategories.h"
+#include "MetasoundTime.h"
 
 #define LOCTEXT_NAMESPACE "MetasoundStandardNodes"
 
 namespace Metasound
 {
+	namespace RandomNodeNames
+	{
+		/** Input vertex names. */
+		const FString& GetInputNextTriggerName();
+		const FString& GetInputResetTriggerName();
+		const FString& GetInputSeedName();
+		const FString& GetInputMinName();
+		const FString& GetInputMaxName();
+
+		/** Output vertex names. */
+		const FString& GetOutputOnNextTriggerName();
+		const FString& GetOutputOnResetTriggerName();
+		const FString& GetOutputValueName();
+
+		/** Descriptions. */
+		static FText GetNextTriggerDescription();
+		static FText GetResetDescription();
+		static FText GetSeedDescription();
+		static FText GetMinDescription();
+		static FText GetMaxDescription();
+		static FText GetOutputDescription();
+		static FText GetOutputOnNextDescription();
+		static FText GetOutputOnResetDescription();
+	}
+
 	template<typename ValueType>
 	struct TRandomNodeSpecialization
 	{
@@ -31,7 +57,7 @@
 
 		static FText GetDisplayName()
 		{
-			return LOCTEXT("Int32RandomValueDisplayName", "Random Int");
+			return LOCTEXT("Int32RandomValueDisplayName", "Random (Int)");
 		}
 
 		static FText GetDescription()
@@ -44,16 +70,30 @@
 			return true;
 		}
 
-		static void GetDefaultRange(int32& OutMin, int32& OutMax)
+		static int32 GetDefaultMin()
 		{
-			OutMin = 0;
-			OutMax = 100;
+			return 0;
 		}
 
+		static int32 GetDefaultMax()
+		{
+			return 100;
+		}
+
 		static int32 GetNextValue(FRandomStream& InStream, int32 InMin, int32 InMax)
 		{
 			return InStream.RandRange(InMin, InMax);
 		}
+
+		static TDataReadReference<int32> CreateMinValueRef(const FDataReferenceCollection& InputCollection, const FInputVertexInterface& InputInterface)
+		{
+			return InputCollection.GetDataReadReferenceOrConstructWithVertexDefault<int32>(InputInterface, RandomNodeNames::GetInputMinName());
+		}
+
+		static TDataReadReference<int32> CreateMaxValueRef(const FDataReferenceCollection& InputCollection, const FInputVertexInterface& InputInterface)
+		{
+			return InputCollection.GetDataReadReferenceOrConstructWithVertexDefault<int32>(InputInterface, RandomNodeNames::GetInputMaxName());
+		}
 	};
 
 	template<>
@@ -66,7 +106,7 @@
 
 		static FText GetDisplayName()
 		{
-			return LOCTEXT("FloatRandomValueDisplayName", "Random Float");
+			return LOCTEXT("FloatRandomValueDisplayName", "Random (Float)");
 		}
 
 		static FText GetDescription()
@@ -79,16 +119,30 @@
 			return true;
 		}
 
-		static void GetDefaultRange(float& OutMin, float& OutMax)
+		static float GetDefaultMin()
 		{
-			OutMin = 0.0f;
-			OutMax = 1.0f;
+			return 0.0f;
 		}
 
+		static float GetDefaultMax()
+		{
+			return 1.0f;
+		}
+
 		static float GetNextValue(FRandomStream& InStream, float InMin, float InMax)
 		{
 			return InStream.FRandRange(InMin, InMax);
 		}
+
+		static TDataReadReference<float> CreateMinValueRef(const FDataReferenceCollection& InputCollection, const FInputVertexInterface& InputInterface)
+		{
+			return InputCollection.GetDataReadReferenceOrConstructWithVertexDefault<float>(InputInterface, RandomNodeNames::GetInputMinName());
+		}
+
+		static TDataReadReference<float> CreateMaxValueRef(const FDataReferenceCollection& InputCollection, const FInputVertexInterface& InputInterface)
+		{
+			return InputCollection.GetDataReadReferenceOrConstructWithVertexDefault<float>(InputInterface, RandomNodeNames::GetInputMaxName());
+		}
 	};
 
 	template<>
@@ -101,7 +155,7 @@
 
 		static FText GetDisplayName()
 		{
-			return LOCTEXT("BoolRandomValueDisplayName", "Random Bool");
+			return LOCTEXT("BoolRandomValueDisplayName", "Random (Bool)");
 		}
 
 		static FText GetDescription()
@@ -114,40 +168,84 @@
 			return false;
 		}
 
-		static void GetDefaultRange(bool& OutMin, bool& OutMax)
+		static bool GetDefaultMin()
 		{
+			return false;
 		}
 
+		static bool GetDefaultMax()
+		{
+			return true;
+		}
+
 		static bool GetNextValue(FRandomStream& InStream, bool, bool)
 		{
 			return (bool)InStream.RandRange(0, 1);
 		}
+
+		static TDataReadReference<bool> CreateMinValueRef(const FDataReferenceCollection&, const FInputVertexInterface&)
+		{
+			return TDataReadReference<bool>::CreateNew();
+		}
+
+		static TDataReadReference<bool> CreateMaxValueRef(const FDataReferenceCollection&, const FInputVertexInterface&)
+		{
+			return TDataReadReference<bool>::CreateNew();
+		}
 	};
 
-	namespace RandomNodeNames
+	template<>
+	struct TRandomNodeSpecialization<FTime>
 	{
-		/** Input vertex names. */
-		const FString& GetInputNextTriggerName();
-		const FString& GetInputResetTriggerName();
-		const FString& GetInputSeedName();
-		const FString& GetInputMinName();
-		const FString& GetInputMaxName();
+		static FName GetClassName()
+		{
+			return FName("RandomTime");
+		}
 
-		/** Output vertex names. */
-		const FString& GetOutputOnNextTriggerName();
-		const FString& GetOutputOnResetTriggerName();
-		const FString& GetOutputValueName();
+		static FText GetDisplayName()
+		{
+			return LOCTEXT("TimeRandomValueDisplayName", "Random (Time)");
+		}
 
-		/** Descriptions. */
-		static FText GetNextTriggerDescription();
-		static FText GetResetDescription();
-		static FText GetSeedDescription();
-		static FText GetMinDescription();
-		static FText GetMaxDescription();
-		static FText GetOutputDescription();
-		static FText GetOutputOnNextDescription();
-		static FText GetOutputOnResetDescription();
-	}
+		static FText GetDescription()
+		{
+			return LOCTEXT("TimeRandomDiscription", "Generates a random time value.");
+		}
+
+		static bool HasRange()
+		{
+			return true;
+		}
+
+		static float GetDefaultMin()
+		{
+			return 0.0f;
+		}
+
+		static float GetDefaultMax()
+		{
+			return 1.0f;
+		}
+
+
+		static FTime GetNextValue(FRandomStream& InStream, FTime InMin, FTime InMax)
+		{
+			float Seconds = InStream.FRandRange(InMin.GetSeconds(), InMax.GetSeconds());
+			return FTime(Seconds);
+		}
+
+		static TDataReadReference<FTime> CreateMinValueRef(const FDataReferenceCollection& InputCollection, const FInputVertexInterface& InputInterface)
+		{
+			return InputCollection.GetDataReadReferenceOrConstructWithVertexDefault<FTime, float>(InputInterface, RandomNodeNames::GetInputMinName());
+		}
+
+		static TDataReadReference<FTime> CreateMaxValueRef(const FDataReferenceCollection& InputCollection, const FInputVertexInterface& InputInterface)
+		{
+			return InputCollection.GetDataReadReferenceOrConstructWithVertexDefault<FTime, float>(InputInterface, RandomNodeNames::GetInputMaxName());
+		}
+	};
+
+
 	
 	template<typename ValueType>
 	class TRandomNodeOperator : public TExecutableOperator<TRandomNodeOperator<ValueType>>
@@ -161,17 +259,13 @@
 
 			if (TRandomNodeSpecialization<ValueType>::HasRange())
 			{
-				ValueType MinValue;
-				ValueType MaxValue;
-				TRandomNodeSpecialization<ValueType>::GetDefaultRange(MinValue, MaxValue);
-
 				static const FVertexInterface DefaultInterface(
 					FInputVertexInterface(
 						TInputDataVertexModel<FTrigger>(GetInputNextTriggerName(), GetNextTriggerDescription()),
 						TInputDataVertexModel<FTrigger>(GetInputResetTriggerName(), GetResetDescription()),
 						TInputDataVertexModel<int32>(GetInputSeedName(), GetSeedDescription(), DefaultSeed),
-						TInputDataVertexModel<ValueType>(GetInputMinName(), GetMinDescription(), MinValue),
-						TInputDataVertexModel<ValueType>(GetInputMaxName(), GetMaxDescription(), MaxValue)
+						TInputDataVertexModel<ValueType>(GetInputMinName(), GetMinDescription(), TRandomNodeSpecialization<ValueType>::GetDefaultMin()),
+						TInputDataVertexModel<ValueType>(GetInputMaxName(), GetMaxDescription(), TRandomNodeSpecialization<ValueType>::GetDefaultMax())
 					),
 					FOutputVertexInterface(
 						TOutputDataVertexModel<FTrigger>(GetOutputOnNextTriggerName(), GetOutputOnNextDescription()),
@@ -233,11 +327,18 @@
 			FTriggerReadRef InResetTrigger = InputCollection.GetDataReadReferenceOrConstruct<FTrigger>(GetInputResetTriggerName(), InParams.OperatorSettings);
 			FInt32ReadRef InSeedValue = InputCollection.GetDataReadReferenceOrConstructWithVertexDefault<int32>(InputInterface, GetInputSeedName());
 
-			// note: this is ignored by random bool
-			TDataReadReference<ValueType> InMinValue = InputCollection.GetDataReadReferenceOrConstructWithVertexDefault<ValueType>(InputInterface, GetInputMinName());
-			TDataReadReference<ValueType> InMaxValue = InputCollection.GetDataReadReferenceOrConstructWithVertexDefault<ValueType>(InputInterface, GetInputMaxName());
+			// note: random bool does not have a range
+			if (TRandomNodeSpecialization<ValueType>::HasRange())
+			{
+				TDataReadReference<ValueType> InMinValue = TRandomNodeSpecialization<ValueType>::CreateMinValueRef(InputCollection, InputInterface);
+				TDataReadReference<ValueType> InMaxValue = TRandomNodeSpecialization<ValueType>::CreateMaxValueRef(InputCollection, InputInterface);
 
-			return MakeUnique<TRandomNodeOperator<ValueType>>(InParams.OperatorSettings, InNextTrigger, InResetTrigger, InSeedValue, InMinValue, InMaxValue);
+				return MakeUnique<TRandomNodeOperator<ValueType>>(InParams.OperatorSettings, InNextTrigger, InResetTrigger, InSeedValue, InMinValue, InMaxValue);
+			}
+			else
+			{
+				return MakeUnique<TRandomNodeOperator<ValueType>>(InParams.OperatorSettings, InNextTrigger, InResetTrigger, InSeedValue);
+			}
 		}
 
 
@@ -265,6 +366,28 @@
 			RandomStream.Reset();
 		}
 
+		TRandomNodeOperator(const FOperatorSettings& InSettings,
+			TDataReadReference<FTrigger> InNextTrigger,
+			TDataReadReference<FTrigger> InResetTrigger,
+			TDataReadReference<int32> InSeedValue)
+			: NextTrigger(InNextTrigger)
+			, ResetTrigger(InResetTrigger)
+			, SeedValue(InSeedValue)
+			, MinValue(TDataReadReference<ValueType>::CreateNew()) // Create stub default for no range case
+			, MaxValue(TDataReadReference<ValueType>::CreateNew())
+			, TriggerOutOnNext(FTriggerWriteRef::CreateNew(InSettings))
+			, TriggerOutOnReset(FTriggerWriteRef::CreateNew(InSettings))
+			, OutputValue(TDataWriteReferenceFactory<ValueType>::CreateAny(InSettings))
+			, bIsDefaultSeeded(*SeedValue == DefaultSeed)
+			, bIsRandomStreamInitialized(false)
+		{
+			// We need to initialize the output value to *something*
+			*OutputValue = *MinValue;
+
+			EvaluateSeedChanges();
+			RandomStream.Reset();
+		}
+
 		virtual ~TRandomNodeOperator() = default;
 
 
