Adding metasound map range node

diff --git a/Engine/Plugins/Runtime/Metasound/Source/MetasoundStandardNodes/Private/MetasoundMapRangeNode.cpp b/Engine/Plugins/Runtime/Metasound/Source/MetasoundStandardNodes/Private/MetasoundMapRangeNode.cpp
new file mode 100644
--- /dev/null
+++ b/Engine/Plugins/Runtime/Metasound/Source/MetasoundStandardNodes/Private/MetasoundMapRangeNode.cpp
@@ -0,0 +1,88 @@
+// Copyright Epic Games, Inc. All Rights Reserved.
+
+#include "MetasoundMapRangeNode.h"
+
+#include "MetasoundNodeRegistrationMacro.h"
+
+#define LOCTEXT_NAMESPACE "MetasoundStandardNodes"
+
+namespace Metasound
+{
+	namespace MetasoundMapRangeNodePrivate
+	{
+		FNodeClassMetadata CreateNodeClassMetadata(const FName& InDataTypeName, const FName& InOperatorName, const FText& InDisplayName, const FText& InDescription, const FVertexInterface& InDefaultInterface)
+		{
+			FNodeClassMetadata Metadata
+			{
+				FNodeClassName{FName("MapRange"), InOperatorName, InDataTypeName},
+				1, // Major Version
+				0, // Minor Version
+				InDisplayName,
+				InDescription,
+				PluginAuthor,
+				PluginNodeMissingPrompt,
+				InDefaultInterface,
+				{LOCTEXT("MapRangeCategory", "Utils")},
+				{TEXT("MapRange")},
+				FNodeDisplayStyle{}
+			};
+
+			return Metadata;
+		}
+	}
+
+
+	namespace MapRangeVertexNames
+	{
+		METASOUNDSTANDARDNODES_API const FString& GetInputValueName()
+		{
+			static const FString Name = TEXT("In");
+			return Name;
+		}
+
+		METASOUNDSTANDARDNODES_API const FString& GetInputInRangeAName()
+		{
+			static const FString Name = TEXT("In Range A");
+			return Name;
+		}
+
+		METASOUNDSTANDARDNODES_API const FString& GetInputInRangeBName()
+		{
+			static const FString Name = TEXT("In Range B");
+			return Name;
+		}
+
+		METASOUNDSTANDARDNODES_API const FString& GetInputOutRangeAName()
+		{
+			static const FString Name = TEXT("Out Range A");
+			return Name;
+		}
+
+		METASOUNDSTANDARDNODES_API const FString& GetInputOutRangeBName()
+		{
+			static const FString Name = TEXT("Out Range B");
+			return Name;
+		}
+
+		METASOUNDSTANDARDNODES_API const FString& GetInputClampedName()
+		{
+			static const FString Name = TEXT("Clamped");
+			return Name;
+		}
+
+		METASOUNDSTANDARDNODES_API const FString& GetOutputValueName()
+		{
+			static const FString Name = TEXT("Value");
+			return Name;
+		}
+	}
+
+
+	using FMapRangeNodeInt32 = TMapRangeNode<int32>;
+ 	METASOUND_REGISTER_NODE(FMapRangeNodeInt32)
+
+	using FMapRangeNodeFloat = TMapRangeNode<float>;
+	METASOUND_REGISTER_NODE(FMapRangeNodeFloat)
+}
+
+#undef LOCTEXT_NAMESPACE
diff --git a/Engine/Plugins/Runtime/Metasound/Source/MetasoundStandardNodes/Public/MetasoundMapRangeNode.h b/Engine/Plugins/Runtime/Metasound/Source/MetasoundStandardNodes/Public/MetasoundMapRangeNode.h
new file mode 100644
--- /dev/null
+++ b/Engine/Plugins/Runtime/Metasound/Source/MetasoundStandardNodes/Public/MetasoundMapRangeNode.h
@@ -0,0 +1,232 @@
+// Copyright Epic Games, Inc. All Rights Reserved.
+
+#pragma once
+
+#include "CoreMinimal.h"
+#include "MetasoundFacade.h"
+#include "MetasoundExecutableOperator.h"
+#include "MetasoundPrimitives.h"
+#include "MetasoundStandardNodesNames.h"
+#include "MetasoundTrigger.h"
+#include "Internationalization/Text.h"
+
+#define LOCTEXT_NAMESPACE "MetasoundStandardNodes"
+
+namespace Metasound
+{
+	namespace MetasoundMapRangeNodePrivate
+	{
+		METASOUNDSTANDARDNODES_API FNodeClassMetadata CreateNodeClassMetadata(const FName& InDataTypeName, const FName& InOperatorName, const FText& InDisplayName, const FText& InDescription, const FVertexInterface& InDefaultInterface);
+
+		template<typename ValueType>
+		struct TMapRange
+		{
+			bool bIsSupported = false;
+		};
+
+		template<>
+		struct TMapRange<int32>
+		{
+			static int32 GetDefaultInRangeA() { return 0; }
+			static int32 GetDefaultInRangeB() { return 100; }
+			static int32 GetDefaultOutRangeA() { return 0; }
+			static int32 GetDefaultOutRangeB() { return 100; }
+
+			static int32 MapRangeClamped(int32 InRangeA, int32 InRangeB, int32 OutRangeA, int32 OutRangeB, int32 InValue)
+			{
+				return (int32)FMath::GetMappedRangeValueClamped({(float)InRangeA, (float)InRangeB}, {(float)OutRangeA, (float)OutRangeB}, (float)InValue);
+			}
+
+			static int32 MapRangeUnclamped(int32 InRangeA, int32 InRangeB, int32 OutRangeA, int32 OutRangeB, int32 InValue)
+			{
+				return (int32)FMath::GetMappedRangeValueUnclamped({ (float)InRangeA, (float)InRangeB }, { (float)OutRangeA, (float)OutRangeB }, (float)InValue);
+			}
+		};
+
+
+		template<>
+		struct TMapRange<float>
+		{
+			static float GetDefaultInRangeA() { return 0.0f; }
+			static float GetDefaultInRangeB() { return 1.0f; }
+			static float GetDefaultOutRangeA() { return 0.0f; }
+			static float GetDefaultOutRangeB() { return 1.0f; }
+
+			static float MapRangeClamped(float InRangeA, float InRangeB, float OutRangeA, float OutRangeB, float InValue)
+			{
+				return FMath::GetMappedRangeValueClamped({ InRangeA, InRangeB }, { OutRangeA, OutRangeB }, InValue);
+			}
+
+			static float MapRangeUnclamped(float InRangeA, float InRangeB, float OutRangeA, float OutRangeB, float InValue)
+			{
+				return FMath::GetMappedRangeValueUnclamped({ InRangeA, InRangeB }, { OutRangeA, OutRangeB }, InValue);
+			}
+		};
+	}
+
+	namespace MapRangeVertexNames
+	{
+		METASOUNDSTANDARDNODES_API const FString& GetInputValueName();
+		METASOUNDSTANDARDNODES_API const FString& GetInputInRangeAName();
+		METASOUNDSTANDARDNODES_API const FString& GetInputInRangeBName();
+		METASOUNDSTANDARDNODES_API const FString& GetInputOutRangeAName();
+		METASOUNDSTANDARDNODES_API const FString& GetInputOutRangeBName();
+		METASOUNDSTANDARDNODES_API const FString& GetInputClampedName();
+		METASOUNDSTANDARDNODES_API const FString& GetOutputValueName();
+	}
+
+	template<typename ValueType>
+	class TMapRangeOperator : public TExecutableOperator<TMapRangeOperator<ValueType>>
+	{
+	public:
+		static const FVertexInterface& GetDefaultInterface()
+		{
+			using namespace MapRangeVertexNames;
+			using namespace MetasoundMapRangeNodePrivate;
+
+			static const FVertexInterface DefaultInterface(
+				FInputVertexInterface(
+					TInputDataVertexModel<ValueType>(GetInputValueName(), LOCTEXT("InputValueDesc", "Input value to map.")),
+					TInputDataVertexModel<ValueType>(GetInputInRangeAName(), LOCTEXT("InRangeADesc", "The min input value range."), TMapRange<ValueType>::GetDefaultInRangeA()),
+					TInputDataVertexModel<ValueType>(GetInputInRangeBName(), LOCTEXT("InRangeBDesc", "The max input value range."), TMapRange<ValueType>::GetDefaultInRangeB()),
+					TInputDataVertexModel<ValueType>(GetInputOutRangeAName(), LOCTEXT("OutRangeADesc", "The min output value range."), TMapRange<ValueType>::GetDefaultOutRangeA()),
+					TInputDataVertexModel<ValueType>(GetInputOutRangeBName(), LOCTEXT("OutRangeBDesc", "The max output value range."), TMapRange<ValueType>::GetDefaultOutRangeB()),
+					TInputDataVertexModel<bool>(GetInputClampedName(), LOCTEXT("ClampedDesc", "Whether or not to clamp the input to the specified input range."), true)
+				),
+				FOutputVertexInterface(
+					TOutputDataVertexModel<ValueType>(GetOutputValueName(), LOCTEXT("ValueOutput", "The output value cached in the node."))
+				)
+			);
+
+			return DefaultInterface;
+		}
+
+		static const FNodeClassMetadata& GetNodeInfo()
+		{
+			auto CreateNodeClassMetadata = []() -> FNodeClassMetadata
+			{
+				FName DataTypeName = GetMetasoundDataTypeName<ValueType>();
+				FName OperatorName = TEXT("MapRange");
+				FText NodeDisplayName = FText::Format(LOCTEXT("MapRangeDisplayNamePattern", "Map Range ({0})"), FText::FromString(GetMetasoundDataTypeString<ValueType>()));
+				FText NodeDescription = LOCTEXT("MapRangeDescription", "Maps an input value in the given input range to the given output range.");
+				FVertexInterface NodeInterface = GetDefaultInterface();
+
+				return MetasoundMapRangeNodePrivate::CreateNodeClassMetadata(DataTypeName, OperatorName, NodeDisplayName, NodeDescription, NodeInterface);
+			};
+
+			static const FNodeClassMetadata Metadata = CreateNodeClassMetadata();
+			return Metadata;
+		}
+
+		static TUniquePtr<IOperator> CreateOperator(const FCreateOperatorParams& InParams, TArray<TUniquePtr<IOperatorBuildError>>& OutErrors)
+		{
+			using namespace MapRangeVertexNames;
+
+			const FInputVertexInterface& InputInterface = InParams.Node.GetVertexInterface().GetInputInterface();
+			const FDataReferenceCollection& InputCollection = InParams.InputDataReferences;
+
+			TDataReadReference<ValueType> InputValue = InputCollection.GetDataReadReferenceOrConstructWithVertexDefault<ValueType>(InputInterface, GetInputValueName());
+			TDataReadReference<ValueType> InputInRangeA = InputCollection.GetDataReadReferenceOrConstructWithVertexDefault<ValueType>(InputInterface, GetInputInRangeAName());
+			TDataReadReference<ValueType> InputInRangeB = InputCollection.GetDataReadReferenceOrConstructWithVertexDefault<ValueType>(InputInterface, GetInputInRangeBName());
+			TDataReadReference<ValueType> InputOutRangeA = InputCollection.GetDataReadReferenceOrConstructWithVertexDefault<ValueType>(InputInterface, GetInputOutRangeAName());
+			TDataReadReference<ValueType> InputOutRangeB = InputCollection.GetDataReadReferenceOrConstructWithVertexDefault<ValueType>(InputInterface, GetInputOutRangeBName());
+			FBoolReadRef bInputClamped = InputCollection.GetDataReadReferenceOrConstructWithVertexDefault<bool>(InputInterface, GetInputClampedName());
+
+			return MakeUnique<TMapRangeOperator<ValueType>>(InParams.OperatorSettings, InputValue, InputInRangeA, InputInRangeB, InputOutRangeA, InputOutRangeB, bInputClamped);
+		}
+
+
+		TMapRangeOperator(const FOperatorSettings& InSettings, 
+			const TDataReadReference<ValueType>& InInputValue,
+			const TDataReadReference<ValueType>& InInputInRangeA,
+			const TDataReadReference<ValueType>& InInputInRangeB,
+			const TDataReadReference<ValueType>& InInputOutRangeA,
+			const TDataReadReference<ValueType>& InInputOutRangeB,
+			const FBoolReadRef& bInClamped)
+			: InputValue(InInputValue)
+			, InputInRangeA(InInputInRangeA)
+			, InputInRangeB(InInputInRangeB)
+			, InputOutRangeA(InInputOutRangeA)
+			, InputOutRangeB(InInputOutRangeB)
+			, bInputClamped(bInClamped)
+			, OutputValue(TDataWriteReferenceFactory<ValueType>::CreateAny(InSettings))
+		{
+			DoMapping();
+		}
+
+		virtual ~TMapRangeOperator() = default;
+
+		virtual FDataReferenceCollection GetInputs() const override
+		{
+			using namespace MapRangeVertexNames;
+
+			FDataReferenceCollection Inputs;
+			Inputs.AddDataReadReference(GetInputValueName(), InputValue);
+			Inputs.AddDataReadReference(GetInputInRangeAName(), InputInRangeA);
+			Inputs.AddDataReadReference(GetInputInRangeBName(), InputInRangeB);
+			Inputs.AddDataReadReference(GetInputOutRangeAName(), InputOutRangeA);
+			Inputs.AddDataReadReference(GetInputOutRangeBName(), InputOutRangeB);
+			Inputs.AddDataReadReference(GetInputClampedName(), bInputClamped);
+
+			return Inputs;
+		}
+
+		virtual FDataReferenceCollection GetOutputs() const override
+		{
+			using namespace MapRangeVertexNames;
+
+			FDataReferenceCollection Outputs;
+			Outputs.AddDataReadReference(GetOutputValueName(), OutputValue);
+
+			return Outputs;
+		}
+
+		void DoMapping()
+		{
+			using namespace MetasoundMapRangeNodePrivate;
+
+			if (*bInputClamped)
+			{
+				*OutputValue = TMapRange<ValueType>::MapRangeClamped(*InputInRangeA, *InputInRangeB, *InputOutRangeA, *InputOutRangeB, *InputValue);
+			}
+			else
+			{
+				*OutputValue = TMapRange<ValueType>::MapRangeUnclamped(*InputInRangeA, *InputInRangeB, *InputOutRangeA, *InputOutRangeB, *InputValue);
+			}
+		}
+
+		void Execute()
+		{
+			DoMapping();
+		}
+
+	private:
+
+		TDataReadReference<ValueType> InputValue;
+		TDataReadReference<ValueType> InputInRangeA;
+		TDataReadReference<ValueType> InputInRangeB;
+		TDataReadReference<ValueType> InputOutRangeA;
+		TDataReadReference<ValueType> InputOutRangeB;
+		FBoolReadRef bInputClamped;
+		TDataWriteReference<ValueType> OutputValue;
+	};
+
+	/** TMapRangeNode
+	 *
+	 *  Performs a mapping of input to output in the given input and output ranges. Offers a clamping or unclamping option.
+	 */
+	template<typename ValueType>
+	class METASOUNDSTANDARDNODES_API TMapRangeNode : public FNodeFacade
+	{
+	public:
+		/**
+		 * Constructor used by the Metasound Frontend.
+		 */
+		TMapRangeNode(const FNodeInitData& InInitData)
+			: FNodeFacade(InInitData.InstanceName, InInitData.InstanceID, TFacadeOperatorClass<TMapRangeOperator<ValueType>>())
+		{}
+
+		virtual ~TMapRangeNode() = default;
+	};
+} // namespace Metasound
+
