Fixes from CL 15745119

-- Stereo/mono output nodes disappear when changing metasound graph format. This recovers the old behavior.
-- Fix for crashing when deleting comments in metasound graph

diff --git a/Engine/Plugins/Runtime/Metasound/Source/MetasoundEditor/Private/MetasoundEditorGraphBuilder.cpp b/Engine/Plugins/Runtime/Metasound/Source/MetasoundEditor/Private/MetasoundEditorGraphBuilder.cpp
--- a/Engine/Plugins/Runtime/Metasound/Source/MetasoundEditor/Private/MetasoundEditorGraphBuilder.cpp
+++ b/Engine/Plugins/Runtime/Metasound/Source/MetasoundEditor/Private/MetasoundEditorGraphBuilder.cpp
@@ -880,7 +880,7 @@ namespace Metasound
 			UMetasoundEditorGraph* Graph = CastChecked<UMetasoundEditorGraph>(InNode.GetGraph());
 			if (!Node)
 			{
-				Graph->RemoveNode(Node);
+				Graph->RemoveNode(&InNode);
 				return true;
 			}
 
@@ -1306,10 +1306,32 @@ namespace Metasound
 					bIsEditorGraphDirty = true;
 
 					TArray<UMetasoundEditorGraphNode*> AddedNodes;
-					for (const TPair<FGuid, FVector2D>& Location : Style.Display.Locations)
+					if (Style.Display.Locations.Num() > 0)
 					{
-						FVector2D NewLocation = FVector2D::ZeroVector;
-						UMetasoundEditorGraphNode* NewNode = AddNode(InMetasound, Node, Location.Value, false /* bInSelectNewNode */);
+						for (const TPair<FGuid, FVector2D>& Location : Style.Display.Locations)
+						{
+							UMetasoundEditorGraphNode* NewNode = AddNode(InMetasound, Node, Location.Value, false /* bInSelectNewNode */);
+							if (ensure(NewNode))
+							{
+								FAssociatedNodes& AssociatedNodeData = AssociatedNodes.FindOrAdd(Node->GetID());
+								if (AssociatedNodeData.Node->IsValid())
+								{
+									ensure(AssociatedNodeData.Node == Node);
+								}
+								else
+								{
+									AssociatedNodeData.Node = Node;
+								}
+
+								AddedNodes.Add(NewNode);
+								AssociatedNodeData.EditorNodes.Add(NewNode);
+							}
+						}
+					}
+					else
+					{
+						// TODO: pass in a non-origin position, ideally the position of the output node that this is replacing (e.g. the old mono node, that may have been moved by a user)
+						UMetasoundEditorGraphNode* NewNode = AddNode(InMetasound, Node, { 0.0f, 0.0f }, false /* bInSelectNewNode */);
 						if (ensure(NewNode))
 						{
 							FAssociatedNodes& AssociatedNodeData = AssociatedNodes.FindOrAdd(Node->GetID());
