Add Trigger Simulation
- Add ExecutableDataType statics required for updating and ownership of triggers on Receive and Input nodes to avoid triggers spamming when set from boolean literals (Could potentially be useful for AudioBuffer inputs transmission).
- Misc clean-up

diff --git a/Engine/Plugins/Runtime/Metasound/Source/MetasoundEditor/Private/MetasoundEditor.cpp b/Engine/Plugins/Runtime/Metasound/Source/MetasoundEditor/Private/MetasoundEditor.cpp
--- a/Engine/Plugins/Runtime/Metasound/Source/MetasoundEditor/Private/MetasoundEditor.cpp
+++ b/Engine/Plugins/Runtime/Metasound/Source/MetasoundEditor/Private/MetasoundEditor.cpp
@@ -16,6 +16,7 @@
 #include "GraphEditor.h"
 #include "GraphEditorActions.h"
 #include "HAL/PlatformApplicationMisc.h"
+#include "IAudioExtensionPlugin.h"
 #include "IDetailsView.h"
 #include "Kismet2/BlueprintEditorUtils.h"
 #include "Metasound.h"
@@ -717,7 +718,6 @@ namespace Metasound
 					}
 				}
 
-
 				MetasoundGraphEditor->RegisterActiveTimer(0.0f,
 					FWidgetActiveTimerDelegate::CreateLambda([](double InCurrentTime, float InDeltaTime)
 						{
@@ -735,19 +735,18 @@ namespace Metasound
 			}
 		}
 
-		void FEditor::PlayNode()
+		void FEditor::ExecuteNode()
 		{
 			const FGraphPanelSelectionSet SelectedNodes = MetasoundGraphEditor->GetSelectedNodes();
 			for (FGraphPanelSelectionSet::TConstIterator NodeIt(SelectedNodes); NodeIt; ++NodeIt)
 			{
-				PlaySingleNode(CastChecked<UEdGraphNode>(*NodeIt));
+				ExecuteNode(CastChecked<UEdGraphNode>(*NodeIt));
 			}
 		}
 
-		bool FEditor::CanPlayNode() const
+		bool FEditor::CanExecuteNode() const
 		{
-			// TODO: Implement node playback
-			return false;
+			return true;
 		}
 
 		void FEditor::Stop()
@@ -771,9 +770,22 @@ namespace Metasound
 			}
 		}
 
-		void FEditor::PlaySingleNode(UEdGraphNode* Node)
+		void FEditor::ExecuteNode(UEdGraphNode* Node)
 		{
-			// TODO: Implement? Will we support single node playback?
+			if (UMetasoundEditorGraphInputNode* InputNode = Cast<UMetasoundEditorGraphInputNode>(Node))
+			{
+				// TODO: fix how identifying the parameter to update is determined. It should not be done
+				// with a "DisplayName" but rather the vertex Guid.
+				if (UMetasoundEditorGraph* MetasoundGraph = Cast<UMetasoundEditorGraph>(InputNode->GetOuter()))
+				{
+					if (IAudioInstanceTransmitter* Transmitter = MetasoundGraph->GetMetasoundInstanceTransmitter())
+					{
+						Metasound::Frontend::FConstNodeHandle NodeHandle = InputNode->GetConstNodeHandle();
+						Metasound::FVertexKey VertexKey = Metasound::FVertexKey(NodeHandle->GetDisplayName().ToString());
+						Transmitter->SetParameter(*VertexKey, true);
+					}
+				}
+			}
 		}
 
 		void FEditor::EditMetasoundSettings()
@@ -874,6 +886,10 @@ namespace Metasound
 					FExecuteAction::CreateLambda([this]() { PasteNodes(); }),
 					FCanExecuteAction::CreateSP(this, &FEditor::CanPasteNodes));
 
+				GraphEditorCommands->MapAction(FGenericCommands::Get().Delete,
+					FExecuteAction::CreateSP(this, &FEditor::DeleteSelectedNodes),
+					FCanExecuteAction::CreateLambda([this]() { return CanDeleteNodes(); }));
+
 				GraphEditorCommands->MapAction(FGenericCommands::Get().Duplicate,
 					FExecuteAction::CreateLambda([this] { CopySelectedNodes(); PasteNodes(); }),
 					FCanExecuteAction::CreateLambda([this]() { return CanCopyNodes(); }));
@@ -914,7 +930,7 @@ namespace Metasound
 			SGraphEditor::FGraphEditorEvents InEvents;
 			InEvents.OnSelectionChanged = SGraphEditor::FOnSelectionChanged::CreateSP(this, &FEditor::OnSelectedNodesChanged);
 			InEvents.OnTextCommitted = FOnNodeTextCommitted::CreateSP(this, &FEditor::OnNodeTitleCommitted);
-			InEvents.OnNodeDoubleClicked = FSingleNodeEvent::CreateSP(this, &FEditor::PlaySingleNode);
+			InEvents.OnNodeDoubleClicked = FSingleNodeEvent::CreateSP(this, &FEditor::ExecuteNode);
 
 			FMetasoundAssetBase* MetasoundAsset = IMetasoundUObjectRegistry::Get().GetObjectAsAssetBase(Metasound);
 			check(MetasoundAsset);
@@ -1189,7 +1205,7 @@ namespace Metasound
 
 		bool FEditor::CanDeleteNodes() const
 		{
-			if (MetasoundGraphEditor->GetSelectedNodes().Num() == 0)
+			if (MetasoundGraphEditor->GetSelectedNodes().IsEmpty())
 			{
 				return false;
 			}
diff --git a/Engine/Plugins/Runtime/Metasound/Source/MetasoundEditor/Private/MetasoundEditor.h b/Engine/Plugins/Runtime/Metasound/Source/MetasoundEditor/Private/MetasoundEditor.h
--- a/Engine/Plugins/Runtime/Metasound/Source/MetasoundEditor/Private/MetasoundEditor.h
+++ b/Engine/Plugins/Runtime/Metasound/Source/MetasoundEditor/Private/MetasoundEditor.h
@@ -188,19 +188,19 @@ namespace Metasound
 
 			/** Toolbar command methods */
 			void Import();
+			void ExecuteNode();
 			void Export();
 			void Play();
-			void PlayNode();
 			void Stop();
 
 			/** Whether we can play the current selection of nodes */
-			bool CanPlayNode() const;
+			bool CanExecuteNode() const;
 
 			/** Either play the Metasound or stop currently playing sound */
 			void TogglePlayback();
 
-			/** Plays a single specified node */
-			void PlaySingleNode(UEdGraphNode* Node);
+			/** Executes specified node (If supported) */
+			void ExecuteNode(UEdGraphNode* Node);
 
 			/** Sync the content browser to the current selection of nodes */
 			void SyncInBrowser();
diff --git a/Engine/Plugins/Runtime/Metasound/Source/MetasoundEditor/Private/MetasoundEditorGraphSchema.cpp b/Engine/Plugins/Runtime/Metasound/Source/MetasoundEditor/Private/MetasoundEditorGraphSchema.cpp
--- a/Engine/Plugins/Runtime/Metasound/Source/MetasoundEditor/Private/MetasoundEditorGraphSchema.cpp
+++ b/Engine/Plugins/Runtime/Metasound/Source/MetasoundEditor/Private/MetasoundEditorGraphSchema.cpp
@@ -489,13 +489,10 @@ void UMetasoundEditorGraphSchema::GetContextMenuActions(class UToolMenu* Menu, c
 			Section.AddMenuEntry(FGraphEditorCommands::Get().BreakPinLinks);
 		}
 	}
-	else if (Context->Node)
+	else if (Context->Node && Context->Node->IsA<UMetasoundEditorGraphNode>())
 	{
-		const UMetasoundEditorGraphNode* SoundGraphNode = Cast<const UMetasoundEditorGraphNode>(Context->Node);
-		{
-			FToolMenuSection& Section = Menu->AddSection("MetasoundGraphSchemaNodeActions", LOCTEXT("NodeActionsMenuHeader", "Node Actions"));
-			Section.AddMenuEntry(FGraphEditorCommands::Get().BreakNodeLinks);
-		}
+		FToolMenuSection& Section = Menu->AddSection("MetasoundGraphSchemaNodeActions", LOCTEXT("NodeActionsMenuHeader", "Node Actions"));
+		Section.AddMenuEntry(FGraphEditorCommands::Get().BreakNodeLinks);
 	}
 
 	Super::GetContextMenuActions(Menu, Context);
diff --git a/Engine/Plugins/Runtime/Metasound/Source/MetasoundEditor/Private/MetasoundEditorModule.cpp b/Engine/Plugins/Runtime/Metasound/Source/MetasoundEditor/Private/MetasoundEditorModule.cpp
--- a/Engine/Plugins/Runtime/Metasound/Source/MetasoundEditor/Private/MetasoundEditorModule.cpp
+++ b/Engine/Plugins/Runtime/Metasound/Source/MetasoundEditor/Private/MetasoundEditorModule.cpp
@@ -285,6 +285,10 @@ namespace Metasound
 					UMetasoundEditorGraphOutput::StaticClass()->GetFName(),
 					FOnGetDetailCustomizationInstance::CreateLambda([]() { return MakeShared<FMetasoundOutputDetailCustomization>(); }));
 
+				PropertyModule.RegisterCustomPropertyTypeLayout(
+					"MetasoundEditorGraphInputBoolRef",
+					FOnGetPropertyTypeCustomizationInstance::CreateLambda([]() { return MakeShared<FMetasoundInputBoolDetailCustomization>(); }));
+
 				PropertyModule.RegisterCustomPropertyTypeLayout(
 					"MetasoundEditorGraphInputIntRef",
 					FOnGetPropertyTypeCustomizationInstance::CreateLambda([]() { return MakeShared<FMetasoundInputIntDetailCustomization>(); }));
diff --git a/Engine/Plugins/Runtime/Metasound/Source/MetasoundEditor/Private/MetasoundNodeDetailCustomization.cpp b/Engine/Plugins/Runtime/Metasound/Source/MetasoundEditor/Private/MetasoundNodeDetailCustomization.cpp
--- a/Engine/Plugins/Runtime/Metasound/Source/MetasoundEditor/Private/MetasoundNodeDetailCustomization.cpp
+++ b/Engine/Plugins/Runtime/Metasound/Source/MetasoundEditor/Private/MetasoundNodeDetailCustomization.cpp
@@ -62,8 +62,108 @@ namespace Metasound
 			static const FText InputNameText = LOCTEXT("Input_Name", "Input Name");
 			static const FText OutputNameText = LOCTEXT("Output_Name", "Output Name");
 
+			static const FName TriggerTypeName = "Trigger";
+
 			static const FName DataTypeNameIdentifier = "DataTypeName";
 			static const FName ProxyGeneratorClassNameIdentifier = "GeneratorClass";
+
+			void ExecuteInputTrigger(UMetasoundEditorGraphInputLiteral* Literal)
+			{
+				if (!Literal)
+				{
+					return;
+				}
+
+				UMetasoundEditorGraphInput* Input = Cast<UMetasoundEditorGraphInput>(Literal->GetOuter());
+				if (!ensure(Input))
+				{
+					return;
+				}
+
+				// TODO: fix how identifying the parameter to update is determined. It should not be done
+				// with a "DisplayName" but rather the vertex Guid.
+				if (UMetasoundEditorGraph* MetasoundGraph = Cast<UMetasoundEditorGraph>(Input->GetOuter()))
+				{
+					if (IAudioInstanceTransmitter* Transmitter = MetasoundGraph->GetMetasoundInstanceTransmitter())
+					{
+						Metasound::Frontend::FConstNodeHandle NodeHandle = Input->GetConstNodeHandle();
+						Metasound::FVertexKey VertexKey = Metasound::FVertexKey(NodeHandle->GetDisplayName().ToString());
+						Transmitter->SetParameter(*VertexKey, true);
+					}
+				}
+			}
+		}
+
+		void FMetasoundInputBoolDetailCustomization::CacheProxyData(TSharedPtr<IPropertyHandle> ProxyHandle)
+		{
+			DataTypeName = FName();
+
+			const FString* MetadataDataTypeName = ProxyHandle->GetInstanceMetaData(VariableCustomizationPrivate::DataTypeNameIdentifier);
+			if (ensure(MetadataDataTypeName))
+			{
+				DataTypeName = **MetadataDataTypeName;
+			}
+		}
+
+		FText FMetasoundInputBoolDetailCustomization::GetPropertyNameOverride() const
+		{
+			if (DataTypeName == VariableCustomizationPrivate::TriggerTypeName)
+			{
+				return LOCTEXT("TriggerInput_SimulateTitle", "Simulate");
+			}
+
+			return FText::GetEmpty();
+		}
+
+		TSharedRef<SWidget> FMetasoundInputBoolDetailCustomization::CreateStructureWidget(TSharedPtr<IPropertyHandle>& StructPropertyHandle) const
+		{
+			using namespace Frontend;
+
+			if (FMetasoundFrontendRegistryContainer* Registry = FMetasoundFrontendRegistryContainer::Get())
+			{
+				TSharedPtr<IPropertyHandle> ValueProperty = StructPropertyHandle->GetChildHandle(GET_MEMBER_NAME_CHECKED(FMetasoundEditorGraphInputBoolRef, Value));
+				if (ValueProperty.IsValid())
+				{
+					// Not a trigger, so just display as underlying literal type (bool)
+					if (DataTypeName != VariableCustomizationPrivate::TriggerTypeName)
+					{
+						return ValueProperty->CreatePropertyValueWidget();
+					}
+
+					return
+						SNew(SHorizontalBox)
+						+ SHorizontalBox::Slot()
+						.AutoWidth()
+						.Padding(2.0f, 0.0f, 0.0f, 0.0f)
+						.VAlign(VAlign_Center)
+						[
+							SNew(SButton)
+							.ButtonStyle(FAppStyle::Get(), "SimpleButton")
+							.OnClicked_Lambda([ValueProperty]()
+							{
+								TArray<UObject*> OuterObjects;
+								ValueProperty->GetOuterObjects(OuterObjects);
+								for (UObject* Object : OuterObjects)
+								{
+									UMetasoundEditorGraphInputLiteral* Literal = Cast<UMetasoundEditorGraphInputLiteral>(Object);
+									VariableCustomizationPrivate::ExecuteInputTrigger(Literal);
+								}
+								return FReply::Handled();
+							})
+							.ToolTipText(LOCTEXT("TriggerTestToolTip", "Executes trigger if currently previewing Metasound."))
+							.ForegroundColor(FSlateColor::UseForeground())
+							.ContentPadding(0)
+							.IsFocusable(false)
+							[
+								SNew(SImage)
+								.Image(FAppStyle::Get().GetBrush("Icons.CircleArrowDown"))
+								.ColorAndOpacity(FSlateColor::UseForeground())
+							]
+						];
+				}
+			}
+
+			return SNullWidget::NullWidget;
 		}
 
 		void FMetasoundInputIntDetailCustomization::CacheProxyData(TSharedPtr<IPropertyHandle> ProxyHandle)
@@ -253,12 +353,27 @@ namespace Metasound
 				.NewAssetFactories(PropertyCustomizationHelpers::GetNewAssetFactoriesForClasses(AllowedClasses));
 		}
 
+		TSharedRef<SWidget> FMetasoundInputArrayDetailCustomizationBase::CreateNameWidget(TSharedPtr<IPropertyHandle> StructPropertyHandle) const
+		{
+			const FText PropertyName = GetPropertyNameOverride();
+			if (!PropertyName.IsEmpty())
+			{
+				return SNew(STextBlock)
+					.Font(IDetailLayoutBuilder::GetDetailFont())
+					.Text(PropertyName);
+			}
+
+			return SNew(STextBlock)
+				.Text(VariableCustomizationPrivate::DefaultPropertyText)
+				.Font(IDetailLayoutBuilder::GetDetailFont());
+		}
+
 		TSharedRef<SWidget> FMetasoundInputArrayDetailCustomizationBase::CreateValueWidget(TSharedPtr<IPropertyHandleArray> ParentArrayProperty, TSharedPtr<IPropertyHandle> StructPropertyHandle, bool bIsInArray) const
 		{
-			TSharedRef<SWidget> ObjectPickerWidget = CreateStructureWidget(StructPropertyHandle);
+			TSharedRef<SWidget> ValueWidget = CreateStructureWidget(StructPropertyHandle);
 			if (!bIsInArray)
 			{
-				return ObjectPickerWidget;
+				return ValueWidget;
 			}
 
 			TSharedPtr<IPropertyHandle> StructPropertyPtr = StructPropertyHandle;
@@ -295,7 +410,7 @@ namespace Metasound
 				.Padding(1.0f, 0.0f, 0.0f, 0.0f)
 				.VAlign(VAlign_Center)
 				[
-					ObjectPickerWidget
+					ValueWidget
 				]
 				+ SHorizontalBox::Slot()
 				.FillWidth(0.05f)
@@ -327,7 +442,6 @@ namespace Metasound
 			CacheProxyData(ProxyProperty);
 
 			TSharedRef<SWidget> ValueWidget = CreateValueWidget(ParentArrayProperty, StructPropertyHandle, bIsInArray);
-
 			FDetailWidgetRow& ValueRow = ChildBuilder.AddCustomRow(VariableCustomizationPrivate::DefaultPropertyText);
 			if (bIsInArray)
 			{
@@ -340,9 +454,7 @@ namespace Metasound
 			{
 				ValueRow.NameContent()
 				[
-					SNew(STextBlock)
-					.Text(VariableCustomizationPrivate::DefaultPropertyText)
-					.Font(IDetailLayoutBuilder::GetDetailFont())
+					CreateNameWidget(StructPropertyHandle)
 				];
 			}
 
@@ -454,7 +566,7 @@ namespace Metasound
 				]
 				+ SHorizontalBox::Slot()
 				.FillWidth(0.40f)
-				.Padding(1.0f, 0.0f, 0.0f, 0.0f)
+				.Padding(2.0f, 0.0f, 0.0f, 0.0f)
 				.VAlign(VAlign_Center)
 				[
 					SAssignNew(DataTypeArrayCheckbox, SCheckBox)
@@ -720,15 +832,16 @@ namespace Metasound
 		{
 			using namespace Frontend;
 
-			FNodeHandle NodeHandle = GraphVariable->GetNodeHandle();
-			const TArray<FOutputHandle> Outputs = NodeHandle->GetOutputs();
+			FName TypeName;
 
-			if (!ensure(Outputs.Num() == 1))
+			// Just take last type.  If more than one, all types are the same.
+			FConstNodeHandle NodeHandle = GraphVariable->GetConstNodeHandle();
+			NodeHandle->IterateConstOutputs([InTypeName = &TypeName](FConstOutputHandle OutputHandle)
 			{
-				return FName();
-			}
+				*InTypeName = OutputHandle->GetDataType();
+			});
 
-			return Outputs[0]->GetDataType();
+			return TypeName;
 		}
 
 		void FMetasoundInputDetailCustomization::OnDisplayNameChanged(const FText& InNewName)
@@ -736,6 +849,7 @@ namespace Metasound
 			using namespace Frontend;
 
 			bIsNameInvalid = false;
+			DisplayNameEditableTextBox->SetError(FText::GetEmpty());
 
 			if (!ensure(GraphVariable.IsValid()))
 			{
@@ -755,7 +869,9 @@ namespace Metasound
 
 			GraphHandle->IterateConstNodes([this, NodeID, InNewName](FConstNodeHandle NodeToCompare)
 			{
-				if (NodeID != NodeToCompare->GetID())
+				// Disregard display name collisions with hidden nodes
+				const bool bIsVisible = NodeToCompare->GetNodeStyle().Display.Visibility == EMetasoundFrontendNodeStyleDisplayVisibility::Visible;
+				if (NodeID != NodeToCompare->GetID() && bIsVisible)
 				{
 					if (InNewName.CompareToCaseIgnored(NodeToCompare->GetDisplayName()) == 0)
 					{
@@ -869,11 +985,10 @@ namespace Metasound
 				return;
 			}
 
-			static const FString ArrayIdentifier = TEXT(":Array");
 			FString TypeNameString = TypeName.ToString();
-			if (TypeNameString.EndsWith(ArrayIdentifier))
+			if (TypeNameString.EndsWith(VariableCustomizationPrivate::ArrayIdentifier))
 			{
-				TypeNameString = TypeNameString.LeftChop(ArrayIdentifier.Len());
+				TypeNameString = TypeNameString.LeftChop(VariableCustomizationPrivate::ArrayIdentifier.Len());
 			}
 			InDefaultPropertyHandle->SetInstanceMetaData(VariableCustomizationPrivate::DataTypeNameIdentifier, TypeNameString);
 
@@ -901,15 +1016,16 @@ namespace Metasound
 		{
 			using namespace Frontend;
 
-			FNodeHandle NodeHandle = GraphVariable->GetNodeHandle();
-			const TArray<FOutputHandle> Outputs = NodeHandle->GetOutputs();
+			FName TypeName;
 
-			if (!ensure(Outputs.Num() == 1))
+			// Just take last type.  If more than one, all types are the same.
+			FConstNodeHandle NodeHandle = GraphVariable->GetConstNodeHandle();
+			NodeHandle->IterateConstInputs([InTypeName = &TypeName](FConstInputHandle InputHandle)
 			{
-				return FName();
-			}
+				*InTypeName = InputHandle->GetDataType();
+			});
 
-			return Outputs[0]->GetDataType();
+			return TypeName;
 		}
 
 		void FMetasoundOutputDetailCustomization::OnDisplayNameChanged(const FText& InNewName)
@@ -917,6 +1033,7 @@ namespace Metasound
 			using namespace Frontend;
 
 			bIsNameInvalid = false;
+			DisplayNameEditableTextBox->SetError(FText::GetEmpty());
 
 			if (!ensure(GraphVariable.IsValid()))
 			{
@@ -936,7 +1053,9 @@ namespace Metasound
 
 			GraphHandle->IterateConstNodes([this, NodeID, InNewName](FConstNodeHandle NodeToCompare)
 			{
-				if (NodeID != NodeToCompare->GetID())
+				// Disregard display name collisions with hidden nodes
+				const bool bIsVisible = NodeToCompare->GetNodeStyle().Display.Visibility == EMetasoundFrontendNodeStyleDisplayVisibility::Visible;
+				if (NodeID != NodeToCompare->GetID() && bIsVisible)
 				{
 					if (InNewName.CompareToCaseIgnored(NodeToCompare->GetDisplayName()) == 0)
 					{
diff --git a/Engine/Plugins/Runtime/Metasound/Source/MetasoundEditor/Private/MetasoundNodeDetailCustomization.h b/Engine/Plugins/Runtime/Metasound/Source/MetasoundEditor/Private/MetasoundNodeDetailCustomization.h
--- a/Engine/Plugins/Runtime/Metasound/Source/MetasoundEditor/Private/MetasoundNodeDetailCustomization.h
+++ b/Engine/Plugins/Runtime/Metasound/Source/MetasoundEditor/Private/MetasoundNodeDetailCustomization.h
@@ -40,13 +40,26 @@ namespace Metasound
 			//~ End IPropertyTypeCustomization
 
 		protected:
+			virtual FText GetPropertyNameOverride() const { return FText::GetEmpty(); }
 			virtual TSharedRef<SWidget> CreateStructureWidget(TSharedPtr<IPropertyHandle>& PropertyHandle) const = 0;
 			virtual void CacheProxyData(TSharedPtr<IPropertyHandle> ProxyHandle) { }
 
 		private:
+			TSharedRef<SWidget> CreateNameWidget(TSharedPtr<IPropertyHandle> StructPropertyHandle) const;
 			TSharedRef<SWidget> CreateValueWidget(TSharedPtr<IPropertyHandleArray> ParentArrayProperty, TSharedPtr<IPropertyHandle> StructPropertyHandle, bool bIsInArray) const;
 		};
 
+		class FMetasoundInputBoolDetailCustomization : public FMetasoundInputArrayDetailCustomizationBase
+		{
+		protected:
+			virtual FText GetPropertyNameOverride() const override;
+			virtual TSharedRef<SWidget> CreateStructureWidget(TSharedPtr<IPropertyHandle>& StructPropertyHandle) const override;
+			virtual void CacheProxyData(TSharedPtr<IPropertyHandle> ProxyHandle) override;
+
+		private:
+			FName DataTypeName;
+		};
+
 		class FMetasoundInputIntDetailCustomization : public FMetasoundInputArrayDetailCustomizationBase
 		{
 		protected:
diff --git a/Engine/Plugins/Runtime/Metasound/Source/MetasoundEditor/Public/MetasoundEditorGraphInputNodes.h b/Engine/Plugins/Runtime/Metasound/Source/MetasoundEditor/Public/MetasoundEditorGraphInputNodes.h
--- a/Engine/Plugins/Runtime/Metasound/Source/MetasoundEditor/Public/MetasoundEditorGraphInputNodes.h
+++ b/Engine/Plugins/Runtime/Metasound/Source/MetasoundEditor/Public/MetasoundEditorGraphInputNodes.h
@@ -93,6 +93,16 @@ protected:
 	friend class Metasound::Editor::FGraphBuilder;
 };
 
+// Broken out to be able to customize and swap enum behavior for basic integer literal behavior
+USTRUCT()
+struct FMetasoundEditorGraphInputBoolRef
+{
+	GENERATED_BODY()
+
+	UPROPERTY(EditAnywhere, Category = DefaultValue)
+	bool Value = false;
+};
+
 UCLASS(MinimalAPI)
 class UMetasoundEditorGraphInputBool : public UMetasoundEditorGraphInputLiteral
 {
@@ -102,12 +112,12 @@ public:
 	virtual ~UMetasoundEditorGraphInputBool() = default;
 
 	UPROPERTY(EditAnywhere, Category = DefaultValue)
-	bool Default = false;
+	FMetasoundEditorGraphInputBoolRef Default;
 
 	virtual FMetasoundFrontendLiteral GetDefault() const override
 	{
 		FMetasoundFrontendLiteral Literal;
-		Literal.Set(Default);
+		Literal.Set(Default.Value);
 		return Literal;
 	}
 
@@ -118,7 +128,7 @@ public:
 
 	void UpdatePreviewInstance(const Metasound::FVertexKey& InParameterName, IAudioInstanceTransmitter& InInstanceTransmitter) const override
 	{
-		InInstanceTransmitter.SetParameter(*InParameterName, Default);
+		InInstanceTransmitter.SetParameter(*InParameterName, Default.Value);
 	}
 };
 
@@ -131,12 +141,15 @@ public:
 	virtual ~UMetasoundEditorGraphInputBoolArray() = default;
 
 	UPROPERTY(EditAnywhere, Category = DefaultValue)
-	TArray<bool> Default;
+	TArray<FMetasoundEditorGraphInputBoolRef> Default;
 
 	virtual FMetasoundFrontendLiteral GetDefault() const override
 	{
+		TArray<bool> BoolArray;
+		Algo::Transform(Default, BoolArray, [](const FMetasoundEditorGraphInputBoolRef& InValue) { return InValue.Value; });
+
 		FMetasoundFrontendLiteral Literal;
-		Literal.Set(Default);
+		Literal.Set(BoolArray);
 		return Literal;
 	}
 
@@ -147,7 +160,10 @@ public:
 
 	void UpdatePreviewInstance(const Metasound::FVertexKey& InParameterName, IAudioInstanceTransmitter& InInstanceTransmitter) const override
 	{
-		InInstanceTransmitter.SetParameter(*InParameterName, TArray<bool>{Default});
+		TArray<bool> BoolArray;
+		Algo::Transform(Default, BoolArray, [](const FMetasoundEditorGraphInputBoolRef& InValue) { return InValue.Value; });
+
+		InInstanceTransmitter.SetParameter(*InParameterName, MoveTemp(BoolArray));
 	}
 };
 
@@ -222,7 +238,7 @@ public:
 		TArray<int32> IntArray;
 		Algo::Transform(Default, IntArray, [](const FMetasoundEditorGraphInputIntRef& InValue) { return InValue.Value; });
 
-		InInstanceTransmitter.SetParameter(*InParameterName, TArray<int32>{ IntArray });
+		InInstanceTransmitter.SetParameter(*InParameterName, MoveTemp(IntArray));
 	}
 };
 
@@ -280,7 +296,7 @@ public:
 
 	void UpdatePreviewInstance(const Metasound::FVertexKey& InParameterName, IAudioInstanceTransmitter& InInstanceTransmitter) const override
 	{
-		InInstanceTransmitter.SetParameter(*InParameterName, TArray<float>{Default});
+		InInstanceTransmitter.SetParameter(*InParameterName, TArray<float> { Default });
 	}
 };
 
@@ -309,7 +325,7 @@ public:
 
 	void UpdatePreviewInstance(const Metasound::FVertexKey& InParameterName, IAudioInstanceTransmitter& InInstanceTransmitter) const override
 	{
-		InInstanceTransmitter.SetParameter(*InParameterName, FString{ Default });
+		InInstanceTransmitter.SetParameter(*InParameterName, FString { Default });
 	}
 };
 
@@ -338,7 +354,7 @@ public:
 
 	void UpdatePreviewInstance(const Metasound::FVertexKey& InParameterName, IAudioInstanceTransmitter& InInstanceTransmitter) const override
 	{
-		InInstanceTransmitter.SetParameter(*InParameterName, TArray<FString>{Default});
+		InInstanceTransmitter.SetParameter(*InParameterName, TArray<FString> { Default });
 	}
 };
 
diff --git a/Engine/Plugins/Runtime/Metasound/Source/MetasoundFrontend/Private/MetasoundInstanceTransmitter.cpp b/Engine/Plugins/Runtime/Metasound/Source/MetasoundFrontend/Private/MetasoundInstanceTransmitter.cpp
--- a/Engine/Plugins/Runtime/Metasound/Source/MetasoundFrontend/Private/MetasoundInstanceTransmitter.cpp
+++ b/Engine/Plugins/Runtime/Metasound/Source/MetasoundFrontend/Private/MetasoundInstanceTransmitter.cpp
@@ -66,23 +66,20 @@ namespace Metasound
 		return SetParameterWithLiteral(InParameterName, FLiteral(MoveTemp(InValue)));
 	}
 
-	bool FMetasoundInstanceTransmitter::SetParameterWithLiteral(const FName& InParameterName, const FLiteral& InLiteral) 
+	bool FMetasoundInstanceTransmitter::SetParameterWithLiteral(const FName& InParameterName, const FLiteral& InLiteral)
 	{
-		ISender* Sender = FindSender(InParameterName);
-
-		if (nullptr == Sender)
+		if (ISender* Sender = FindSender(InParameterName))
 		{
-			// If not sender exists for parameter name, attempt to add one.
-			if (const FSendInfo* SendInfo = FindSendInfo(InParameterName))
-			{
-				Sender = AddSender(*SendInfo);
-			}
+			return Sender->PushLiteral(InLiteral);
 		}
 
-		if (nullptr != Sender)
+		// If no sender exists for parameter name, attempt to add one.
+		if (const FSendInfo* SendInfo = FindSendInfo(InParameterName))
 		{
-			// Push data using found sender.
-			return Sender->PushLiteral(InLiteral);
+			if (ISender* Sender = AddSender(*SendInfo))
+			{
+				return Sender->PushLiteral(InLiteral);
+			}
 		}
 
 		return false;
@@ -100,8 +97,7 @@ namespace Metasound
 
 	ISender* FMetasoundInstanceTransmitter::FindSender(const FName& InParameterName)
 	{
-		TUniquePtr<ISender>* SenderPtrPtr = InputSends.Find(InParameterName);
-		if (SenderPtrPtr)
+		if (TUniquePtr<ISender>* SenderPtrPtr = InputSends.Find(InParameterName))
 		{
 			return SenderPtrPtr->Get();
 		}
diff --git a/Engine/Plugins/Runtime/Metasound/Source/MetasoundFrontend/Private/MetasoundTrigger.cpp b/Engine/Plugins/Runtime/Metasound/Source/MetasoundFrontend/Private/MetasoundTrigger.cpp
--- a/Engine/Plugins/Runtime/Metasound/Source/MetasoundFrontend/Private/MetasoundTrigger.cpp
+++ b/Engine/Plugins/Runtime/Metasound/Source/MetasoundFrontend/Private/MetasoundTrigger.cpp
@@ -36,7 +36,6 @@ namespace Metasound
 
 		if (InFrameToTrigger < NumFramesPerBlock)
 		{
-			// Update last trigger pos. 
 			UpdateLastTriggerIndexInBlock();
 		}
 
diff --git a/Engine/Plugins/Runtime/Metasound/Source/MetasoundFrontend/Public/MetasoundInputNode.h b/Engine/Plugins/Runtime/Metasound/Source/MetasoundFrontend/Public/MetasoundInputNode.h
--- a/Engine/Plugins/Runtime/Metasound/Source/MetasoundFrontend/Public/MetasoundInputNode.h
+++ b/Engine/Plugins/Runtime/Metasound/Source/MetasoundFrontend/Public/MetasoundInputNode.h
@@ -5,6 +5,7 @@
 #include "Internationalization/Text.h"
 #include "MetasoundOperatorInterface.h"
 #include "MetasoundNodeInterface.h"
+#include "MetasoundTrigger.h"
 
 #define LOCTEXT_NAMESPACE "MetasoundGraphCore"
 
@@ -19,9 +20,13 @@ namespace Metasound
 			using FDataWriteReference = TDataWriteReference<DataType>;
 
 			TInputOperator(const FVertexKey& InDataReferenceName, FDataWriteReference InDataReference)
+				: DataReferenceName(InDataReferenceName)
+				// Executable DataTypes require a copy of the output to operate on whereas non-executable
+				// types do not. Avoid copy by assigning to reference for non-executable types.
+				, Output(TExecutableDataType<DataType>::bIsExecutable ? FDataWriteReference::CreateNew(*InDataReference) : InDataReference)
 			{
 				Inputs.AddDataWriteReference<DataType>(InDataReferenceName, InDataReference);
-				Outputs.AddDataReadReference<DataType>(InDataReferenceName, InDataReference);
+				Outputs.AddDataReadReference<DataType>(InDataReferenceName, Output);
 			}
 
 			virtual ~TInputOperator() {}
@@ -36,12 +41,30 @@ namespace Metasound
 				return Outputs;
 			}
 
+			void Execute()
+			{
+				FDataWriteReference Input = Inputs.GetDataWriteReference<DataType>(DataReferenceName);
+				TExecutableDataType<DataType>::Execute(*Input, *Output);
+			}
+
+			static void ExecuteFunction(IOperator* InOperator)
+			{
+				static_cast<TInputOperator<DataType>*>(InOperator)->Execute();
+			}
+
 			virtual FExecuteFunction GetExecuteFunction() override
 			{
+				if (TExecutableDataType<DataType>::bIsExecutable)
+				{
+					return &TInputOperator<FTrigger>::ExecuteFunction;
+				}
 				return nullptr;
 			}
 
 		private:
+			FVertexKey DataReferenceName;
+			FDataWriteReference Output;
+
 			FDataReferenceCollection Outputs;
 			FDataReferenceCollection Inputs;
 	};
diff --git a/Engine/Plugins/Runtime/Metasound/Source/MetasoundFrontend/Public/MetasoundReceiveNode.h b/Engine/Plugins/Runtime/Metasound/Source/MetasoundFrontend/Public/MetasoundReceiveNode.h
--- a/Engine/Plugins/Runtime/Metasound/Source/MetasoundFrontend/Public/MetasoundReceiveNode.h
+++ b/Engine/Plugins/Runtime/Metasound/Source/MetasoundFrontend/Public/MetasoundReceiveNode.h
@@ -48,7 +48,6 @@ namespace Metasound
 	class TReceiveNode : public FNode
 	{
 	public:
-
 		static FVertexInterface DeclareVertexInterface()
 		{
 			return FVertexInterface(
@@ -89,9 +88,9 @@ namespace Metasound
 	private:
 		class TReceiverOperator : public TExecutableOperator<TReceiverOperator>
 		{
-				TReceiverOperator() = delete;
-			public:
+			TReceiverOperator() = delete;
 
+			public:
 				TReceiverOperator(TDataReadReference<TDataType> InInitDataRef, TDataWriteReference<TDataType> InOutDataRef, TDataReadReference<FSendAddress> InSendAddress, const FOperatorSettings& InOperatorSettings)
 					: bHasNotReceivedData(true)
 					, DefaultData(InInitDataRef)
@@ -135,14 +134,21 @@ namespace Metasound
 						*OutputData = *DefaultData;
 					}
 
+					bool bHasNewData = false;
 					if (ensure(Receiver.IsValid()))
 					{
-						if (Receiver->CanPop())
+						bHasNewData = Receiver->CanPop();
+						if (bHasNewData)
 						{
 							Receiver->Pop(*OutputData);
 							bHasNotReceivedData = false;
 						}
 					}
+
+					if (TExecutableDataType<TDataType>::bIsExecutable)
+					{
+						TExecutableDataType<TDataType>::ExecuteInline(*OutputData, bHasNewData);
+					}
 				}
 
 			private:
diff --git a/Engine/Plugins/Runtime/Metasound/Source/MetasoundFrontend/Public/MetasoundRouter.h b/Engine/Plugins/Runtime/Metasound/Source/MetasoundFrontend/Public/MetasoundRouter.h
--- a/Engine/Plugins/Runtime/Metasound/Source/MetasoundFrontend/Public/MetasoundRouter.h
+++ b/Engine/Plugins/Runtime/Metasound/Source/MetasoundFrontend/Public/MetasoundRouter.h
@@ -188,6 +188,8 @@ namespace Metasound
 
 		virtual ~ISender();
 
+		virtual int64 GetPayloadID() const = 0;
+
 		// Constructs the value with a literal and pushes to data channel. 
 		virtual bool PushLiteral(const FLiteral& InLiteral) = 0;
 
@@ -239,6 +241,7 @@ namespace Metasound
 			return NumAliveSenders.GetValue();
 		}
 
+		virtual int64 GetPayloadID() const { return 0; }
 		virtual bool PushOpaque(ISender& InSender) { return true; };
 		virtual void PopOpaque(IReceiver& InReceiver) {};
 		virtual bool IsEmpty() const { return true; };
@@ -404,6 +407,11 @@ namespace Metasound
 			SetDelay(InitParams.DelayTimeInSeconds);
 		}
 
+		int64 GetPayloadID() const override
+		{
+			return PayloadID;
+		}
+
 		TDataType RetrievePayload()
 		{
 			check(Payload.IsValid());
@@ -413,9 +421,10 @@ namespace Metasound
 		// Push from the SenderBuffer to the data channel.
 		bool PushToDataChannel()
 		{
+			++PayloadID;
 			if (!Payload.IsValid())
 			{
-				Payload.Reset(new TDataType(SenderBuffer.Pop()));
+				Payload = MakeUnique<TDataType>(SenderBuffer.Pop());
 			}
 			else
 			{
@@ -429,6 +438,8 @@ namespace Metasound
 		// This buffer acts as a delay line for this sender specifically.
 		Audio::TCircularAudioBuffer<TDataType> SenderBuffer;
 
+		int64 PayloadID = 0;
+
 		// lazily initialized holder for payload, allocated on first call to Push.
 		// Placed on the heap so that we don't have to assume we know how to construct TDataType.
 		TUniquePtr<TDataType> Payload;
@@ -447,7 +458,7 @@ namespace Metasound
 
 		bool CanPop() const
 		{
-			return !DataChannel->IsEmpty();
+			return DataChannel->GetPayloadID() != LastPayloadID && !DataChannel->IsEmpty();
 		}
 
 		// Pop the latest value from the data channel.
@@ -471,11 +482,18 @@ namespace Metasound
 		{
 		}
 
-		void PushPayload(const TDataType& InDataPayload)
+		void PushPayload(int64 InPayloadID, const TDataType& InDataPayload)
 		{
+			// Attempting to send the same payload again, so ignore.
+			if (InPayloadID == LastPayloadID)
+			{
+				return;
+			}
+
+			LastPayloadID = InPayloadID;
 			if (!Payload.IsValid())
 			{
-				Payload.Reset(new TDataType(InDataPayload));
+				Payload = MakeUnique<TDataType>(InDataPayload);
 			}
 			else
 			{
@@ -484,6 +502,7 @@ namespace Metasound
 		}
 
 	private:
+		int64 LastPayloadID = INDEX_NONE;
 		TUniquePtr<TDataType> Payload;
 		FOperatorSettings OperatorSettings;
 	};
@@ -501,18 +520,26 @@ namespace Metasound
 			, OperatorSettings(InOperatorSettings)
 		{}
 
+		virtual int64 GetPayloadID() const override
+		{
+			return PayloadID;
+		}
+
 		virtual bool PushOpaque(ISender& InSender) override
 		{
 			TSender<TDataType>& CastSender = InSender.GetAs<TSender<TDataType>>();
 
-			FScopeLock ScopeLock(&AtomicDataLock);
-			if (!AtomicData.IsValid())
 			{
-				AtomicData.Reset(new TDataType(CastSender.RetrievePayload()));
-			}
-			else
-			{
-				*AtomicData = CastSender.RetrievePayload();
+				FScopeLock ScopeLock(&AtomicDataLock);
+				PayloadID = InSender.GetPayloadID();
+				if (!AtomicData.IsValid())
+				{
+					AtomicData = MakeUnique<TDataType>(CastSender.RetrievePayload());
+				}
+				else
+				{
+					*AtomicData = CastSender.RetrievePayload();
+				}
 			}
 
 			return true;
@@ -525,7 +552,7 @@ namespace Metasound
 			FScopeLock ScopeLock(&AtomicDataLock);
 			if (AtomicData.IsValid())
 			{
-				CastReceiver.PushPayload(*AtomicData);
+				CastReceiver.PushPayload(PayloadID, *AtomicData);
 			}
 		}
 
@@ -553,7 +580,7 @@ namespace Metasound
 				FScopeLock ScopeLock(&AtomicDataLock);
 				if (!AtomicData.IsValid())
 				{
-					AtomicData.Reset(new TDataType(DataToPush));
+					AtomicData = MakeUnique<TDataType>(DataToPush);
 				}
 				else
 				{
@@ -580,14 +607,14 @@ namespace Metasound
 		}
 
 	private:
-		FThreadSafeCounter ReceiverCounter;
-
 		// NOTE- In the future, this could be changed to use a TCircualrAudioBuffer or a TQueue for lockless operation.
 		// The primary challenge with this is handling multiple senders and receivers.
 		// In order to support multiple senders and receivers at scale, we'd need a good implementation of a bounded MPMC queue for arbitrary datatypes.
 		TUniquePtr<TDataType> AtomicData;
 		FCriticalSection AtomicDataLock;
 
+		int64 PayloadID = 0;
+
 		FOperatorSettings OperatorSettings;
 	};
 
@@ -622,6 +649,11 @@ namespace Metasound
 			return true;
 		}
 
+		int64 GetPayloadID() const override
+		{
+			return 0;
+		}
+
 		bool PushLiteral(const FLiteral& InLiteral) override
 		{
 			ensureMsgf(false, TEXT("Cannot push literal to audio format"));
diff --git a/Engine/Plugins/Runtime/Metasound/Source/MetasoundFrontend/Public/MetasoundTrigger.h b/Engine/Plugins/Runtime/Metasound/Source/MetasoundFrontend/Public/MetasoundTrigger.h
--- a/Engine/Plugins/Runtime/Metasound/Source/MetasoundFrontend/Public/MetasoundTrigger.h
+++ b/Engine/Plugins/Runtime/Metasound/Source/MetasoundFrontend/Public/MetasoundTrigger.h
@@ -3,6 +3,7 @@
 #pragma once
 
 #include "MetasoundDataReference.h"
+#include "MetasoundExecutableOperator.h"
 #include "MetasoundOperatorSettings.h"
 #include "MetasoundTime.h"
 
@@ -153,11 +154,35 @@ namespace Metasound
 
 			TArray<int32> TriggeredFrames;
 
+			bool bHasAdvanced = false;
+			bool bTriggeredFromInit = false;
 			bool bHasTrigger = false;
 			int32 NumFramesPerBlock = 0;
 			FSampleRate SampleRate = 0;
 			int32 LastTriggerIndexInBlock = 0;
 	};
 
+	template<>
+	struct TExecutableDataType<FTrigger>
+	{
+		static constexpr bool bIsExecutable = true;
+
+		static void Execute(FTrigger& InData, FTrigger& OutData)
+		{
+			OutData = InData;
+			InData.AdvanceBlock();
+		}
+
+		static void ExecuteInline(FTrigger& OutData, bool bInUpdated)
+		{
+			// Only advance (i.e. flush old triggers) if not updated.  If
+			// updated, leave trigger alone as it has new information.
+			if (!bInUpdated)
+			{
+				OutData.AdvanceBlock();
+			}
+		}
+	};
+
 	DECLARE_METASOUND_DATA_REFERENCE_TYPES(FTrigger, METASOUNDFRONTEND_API, FTriggerTypeInfo, FTriggerReadRef, FTriggerWriteRef);
 }
diff --git a/Engine/Plugins/Runtime/Metasound/Source/MetasoundGenerator/Private/MetasoundGenerator.cpp b/Engine/Plugins/Runtime/Metasound/Source/MetasoundGenerator/Private/MetasoundGenerator.cpp
--- a/Engine/Plugins/Runtime/Metasound/Source/MetasoundGenerator/Private/MetasoundGenerator.cpp
+++ b/Engine/Plugins/Runtime/Metasound/Source/MetasoundGenerator/Private/MetasoundGenerator.cpp
@@ -155,8 +155,6 @@ namespace Metasound
 			bIsFinished = true;
 		}
 
-		OnPlayTriggerRef->AdvanceBlock();
-
 		return NumSamplesWritten;
 	}
 
diff --git a/Engine/Plugins/Runtime/Metasound/Source/MetasoundGraphCore/Public/MetasoundExecutableOperator.h b/Engine/Plugins/Runtime/Metasound/Source/MetasoundGraphCore/Public/MetasoundExecutableOperator.h
--- a/Engine/Plugins/Runtime/Metasound/Source/MetasoundGraphCore/Public/MetasoundExecutableOperator.h
+++ b/Engine/Plugins/Runtime/Metasound/Source/MetasoundGraphCore/Public/MetasoundExecutableOperator.h
@@ -6,12 +6,33 @@
 
 #include "MetasoundOperatorInterface.h"
 
+
 namespace Metasound
 {
+	// As a general rule, ExecutableDataTypes should be avoided whenever possible
+	// as they incur an undesired cost and are generally not typically necessary.
+	// This is primarily for the special case of trigger types, where state management
+	// cannot be avoided (or rather an avoidable design has yet to be formulated).
+	template<class DataType>
+	struct TExecutableDataType
+	{
+		static constexpr bool bIsExecutable = false;
+
+		static void Execute(const DataType& InData, const DataType& OutData)
+		{
+			// No-Op for base case as most DataTypes (ex POD) are not executable.
+		}
+
+		static void ExecuteInline(DataType& InData, bool bInUpdated)
+		{
+			// No-Op for base case as most DataTypes (ex POD) are not executable.
+		}
+	};
+
 	/** Convenience class for supporting the IOperator interface's GetExecuteFunction virtual member function.
 	 *
-	 * Dervied classes should inherit from this template class as well as implement a void Execute() member
-	 * function. 
+	 * Derived classes should inherit from this template class as well as implement a void Execute() member
+	 * function.
 	 *
 	 * 	class MyOperator : public TExecutableOperator<MyOperator>
 	 * 	{
@@ -19,7 +40,7 @@ namespace Metasound
 	 * 	  ...
 	 * 	  void Execute()
 	 * 	  {
-	 * 	     ...
+	 *		  ...
 	 * 	  }
 	 * 	};
 	 */
