Fix for Base Frequency Zippering in Oscillators, plus fixes for DC offset smoothing for FM Oscillators.

diff --git a/Engine/Plugins/Runtime/Metasound/Source/MetasoundStandardNodes/Private/MetasoundOscillatorNodes.cpp b/Engine/Plugins/Runtime/Metasound/Source/MetasoundStandardNodes/Private/MetasoundOscillatorNodes.cpp
--- a/Engine/Plugins/Runtime/Metasound/Source/MetasoundStandardNodes/Private/MetasoundOscillatorNodes.cpp
+++ b/Engine/Plugins/Runtime/Metasound/Source/MetasoundStandardNodes/Private/MetasoundOscillatorNodes.cpp
@@ -90,27 +90,62 @@ namespace Metasound
 			PhaseWrap Wrap;
 			float CurrentFade = 0.0f;
 			float FadeSmooth = 0.01f;
-
+			float PreviousFreq = -1.f;
+		
 			void operator()(const FGeneratorArgs& InArgs)
 			{
 				int32 RemainingSamplesInBlock = InArgs.AlignedBuffer.Num();
 				float* Out = InArgs.AlignedBuffer.GetData();
-				float OneOverSampleRate = 1.f / InArgs.SampleRate;
-				float PulseWidth = InArgs.PulseWidth;
-				float DeltaPhase = InArgs.FrequencyHz * OneOverSampleRate;
+				const float OneOverSampleRate = 1.f / InArgs.SampleRate;
+		
+				// Constant Freq between this an last update? 
+				if (FMath::IsNearlyEqual(InArgs.FrequencyHz, PreviousFreq) || PreviousFreq < 0.f )
+				{
+					const float DeltaPhase = InArgs.FrequencyHz * OneOverSampleRate;
+					while (RemainingSamplesInBlock > 0)
+					{
+						Wrap(Phase);
+
+						float Output = Osc(Phase, DeltaPhase, InArgs);
+						Output *= CurrentFade;
+						*Out++ = Output;
+
+						CurrentFade = CurrentFade + FadeSmooth * (1.0f - CurrentFade);
+
+						Phase += DeltaPhase;
 
-				while (RemainingSamplesInBlock > 0)
+						RemainingSamplesInBlock--;
+					}
+				}
+				else
 				{
-					Wrap(Phase);
+					// Lerp freq over this block 
+					const float OneOverBlockCount = 1.f / (float)RemainingSamplesInBlock;
+					const float BlockDeltaFreq = InArgs.FrequencyHz - PreviousFreq;
+					const float DeltaFreq = BlockDeltaFreq * OneOverBlockCount;
+
+					float Freq = PreviousFreq;
+					
+					while (RemainingSamplesInBlock > 0)
+					{
+						Wrap(Phase);
 
-					*Out++ = CurrentFade * Osc(Phase, DeltaPhase, InArgs);
+						const float DeltaPhase = Freq * OneOverSampleRate;
+						
+						float Output = Osc(Phase, DeltaPhase, InArgs);
+						Output *= CurrentFade;
+						*Out++ = Output;
 
-					CurrentFade = CurrentFade + FadeSmooth * (1.0f - CurrentFade);
+						CurrentFade = CurrentFade + FadeSmooth * (1.0f - CurrentFade);
 
-					Phase += DeltaPhase;
+						Phase += DeltaPhase;
+						Freq += DeltaFreq;
 
-					RemainingSamplesInBlock--;
+						RemainingSamplesInBlock--;
+					}
 				}
+
+				PreviousFreq = InArgs.FrequencyHz;
 			}
 		};
 
@@ -126,27 +161,67 @@ namespace Metasound
 			Oscillator Osc;
 			PhaseWrap Wrap;
 
+			float CurrentFade = 0.0f;
+			float FadeSmooth = 0.01f;
+			float PreviousFreq = -1.f;
+
 			void operator()(const FGeneratorArgs& InArgs)
 			{
 				int32 RemainingSamplesInBlock = InArgs.AlignedBuffer.Num();
 				float* Out = InArgs.AlignedBuffer.GetData();
 				const float* FM = InArgs.FM.GetData();
-				float OneOverSampleRate = 1.f / InArgs.SampleRate;
-				float PulseWidth = InArgs.PulseWidth;
+				const float OneOverSampleRate = 1.f / InArgs.SampleRate;
+
+				// Constant Base Freq between this and last update? 
+				if (FMath::IsNearlyEqual(InArgs.FrequencyHz, PreviousFreq) || PreviousFreq < 0.f)
+				{
+					while (RemainingSamplesInBlock > 0)
+					{
+						const float PerSampleFreq = InArgs.FrequencyHz + *FM++;
+						const float DeltaPhase = PerSampleFreq * OneOverSampleRate;
+
+						Wrap(Phase);
+
+						float Output = Osc(Phase, DeltaPhase, InArgs);
+						Output *= CurrentFade;
+						*Out++ = Output;
+												
+						CurrentFade = CurrentFade + FadeSmooth * (1.0f - CurrentFade);
 
-				while (RemainingSamplesInBlock > 0)
+						Phase += DeltaPhase;
+
+						RemainingSamplesInBlock--;
+					}
+				}
+				else
 				{
-					float PerSampleFreq = InArgs.FrequencyHz + *FM++;
-					float DeltaPhase = PerSampleFreq * OneOverSampleRate;
+					// Lerp Base Frequency over this block
+					const float OneOverBlockCount = 1.f / (float)RemainingSamplesInBlock;
+					const float BlockDeltaFreq = InArgs.FrequencyHz - PreviousFreq;
+					const float DeltaFreq = BlockDeltaFreq * OneOverBlockCount;
+					
+					float Freq = PreviousFreq;
+
+					while (RemainingSamplesInBlock > 0)
+					{
+						const float ModulatedFreqSum = Freq + *FM++;
+						const float DeltaPhase = ModulatedFreqSum * OneOverSampleRate;
 
-					*Out++ = Osc(Phase, DeltaPhase, InArgs);
+						Wrap(Phase);
 
-					Phase += DeltaPhase;
+						float Output = Osc(Phase, DeltaPhase, InArgs);
+						Output *= CurrentFade;
+						*Out++ = Output;
+						
+						CurrentFade = CurrentFade + FadeSmooth * (1.0f - CurrentFade);
 
-					Wrap(Phase);
+						Phase += DeltaPhase;
+						Freq += DeltaFreq;
 
-					RemainingSamplesInBlock--;
+						RemainingSamplesInBlock--;
+					}
 				}
+				PreviousFreq = InArgs.FrequencyHz;
 			}
 		};	
 		
@@ -1263,3 +1338,4 @@ namespace Metasound
 #pragma endregion LFO
 }
 #undef LOCTEXT_NAMESPACE //MetasoundStandardNodes
+
