IWYU Fixes

diff --git a/Engine/Plugins/Runtime/Metasound/Source/MetasoundEditor/Private/MetasoundEditor.cpp b/Engine/Plugins/Runtime/Metasound/Source/MetasoundEditor/Private/MetasoundEditor.cpp
--- a/Engine/Plugins/Runtime/Metasound/Source/MetasoundEditor/Private/MetasoundEditor.cpp
+++ b/Engine/Plugins/Runtime/Metasound/Source/MetasoundEditor/Private/MetasoundEditor.cpp
@@ -34,6 +34,8 @@
 #include "SMetasoundPalette.h"
 #include "SNodePanel.h"
 #include "Widgets/Docking/SDockTab.h"
+#include "Widgets/Images/SImage.h"
+#include "Widgets/Input/SButton.h"
 #include "Widgets/Notifications/SNotificationList.h"
 #include "Widgets/SBoxPanel.h"
 #include "Templates/SharedPointer.h"
diff --git a/Engine/Plugins/Runtime/Metasound/Source/MetasoundEditor/Private/MetasoundEditorGraph.cpp b/Engine/Plugins/Runtime/Metasound/Source/MetasoundEditor/Private/MetasoundEditorGraph.cpp
--- a/Engine/Plugins/Runtime/Metasound/Source/MetasoundEditor/Private/MetasoundEditorGraph.cpp
+++ b/Engine/Plugins/Runtime/Metasound/Source/MetasoundEditor/Private/MetasoundEditorGraph.cpp
@@ -7,6 +7,7 @@
 #include "MetasoundEditorGraphInputNodes.h"
 #include "MetasoundEditorModule.h"
 #include "MetasoundUObjectRegistry.h"
+#include "ScopedTransaction.h"
 
 #define LOCTEXT_NAMESPACE "MetasoundEditor"
 
