Fix for typo in Noise Generator Tooltip

diff --git a/Engine/Plugins/Runtime/Metasound/Source/MetasoundStandardNodes/Private/MetasoundNoiseGenerator.cpp b/Engine/Plugins/Runtime/Metasound/Source/MetasoundStandardNodes/Private/MetasoundNoiseGenerator.cpp
--- a/Engine/Plugins/Runtime/Metasound/Source/MetasoundStandardNodes/Private/MetasoundNoiseGenerator.cpp
+++ b/Engine/Plugins/Runtime/Metasound/Source/MetasoundStandardNodes/Private/MetasoundNoiseGenerator.cpp
@@ -180,7 +180,7 @@ namespace Metasound
 			Info.MajorVersion = 1;
 			Info.MinorVersion = 0;
 			Info.DisplayName = LOCTEXT("Metasound_NoiseNodeDisplayNameX", "Noise");
-			Info.Description = LOCTEXT("Metasound_NoiseNodeDescription", "Noise Generator that produces differnt types of noise");
+			Info.Description = LOCTEXT("Metasound_NoiseNodeDescription", "Noise Generator that produces different types of noise");
 			Info.Author = PluginAuthor;
 			Info.PromptIfMissing = PluginNodeMissingPrompt;
 			Info.DefaultInterface = DeclareVertexInterface();
