Making metasound factory to create a new metasound instance asynchronous

diff --git a/Engine/Plugins/Runtime/Metasound/Source/MetasoundEngine/Private/MetasoundSource.cpp b/Engine/Plugins/Runtime/Metasound/Source/MetasoundEngine/Private/MetasoundSource.cpp
--- a/Engine/Plugins/Runtime/Metasound/Source/MetasoundEngine/Private/MetasoundSource.cpp
+++ b/Engine/Plugins/Runtime/Metasound/Source/MetasoundEngine/Private/MetasoundSource.cpp
@@ -153,67 +153,20 @@ ISoundGeneratorPtr UMetasoundSource::CreateSoundGenerator(const FSoundGeneratorI
 
 	// Create handles for new root graph
 	Frontend::FConstDocumentHandle NewDocumentHandle = Frontend::IDocumentController::CreateDocumentHandle(Frontend::MakeAccessPtr<const FMetasoundFrontendDocument>(DocumentWithInjectedReceives.AccessPoint, DocumentWithInjectedReceives));
-	Frontend::FConstGraphHandle RootGraph = NewDocumentHandle->GetRootGraph();
 
-	ensureAlways(RootGraph->IsValid());
-
-	// Create the operator from the graph handle.
-	TArray<IOperatorBuilder::FBuildErrorPtr> BuildErrors;
-	TUniquePtr<IOperator> Operator = RootGraph->BuildOperator(InSettings, Environment, BuildErrors);
-
-	// Log build errors
-	for (const IOperatorBuilder::FBuildErrorPtr& Error : BuildErrors)
-	{
-		if (Error.IsValid())
-		{
-			UE_LOG(LogMetasound, Warning, TEXT("MetasoundSource [%s] build error [%s] \"%s\""), *GetName(), *(Error->GetErrorType().ToString()), *(Error->GetErrorDescription().ToString()));
-		}
-	}
-	if (!Operator.IsValid())
+	FMetasoundGeneratorInitParams InitParams = 
 	{
-		UE_LOG(LogMetasound, Error, TEXT("Failed to build Metasound operator from graph in MetasoundSource [%s]"), *GetName());
-	}
-	else
-	{
-		FDataReferenceCollection Outputs = Operator->GetOutputs();
-		TArrayView<FAudioBufferReadRef> OutputBuffers;
-
-		// Get output audio buffers.
-		if (EMetasoundSourceAudioFormat::Stereo == OutputFormat)
-		{
-			if (!Outputs.ContainsDataReadReference<FStereoAudioFormat>(GetAudioOutputName()))
-			{
-				UE_LOG(LogMetasound, Warning, TEXT("MetasoundSource [%s] does not contain stereo output [%s] in output"), *GetName(), *GetAudioOutputName());
-			}
-			OutputBuffers = Outputs.GetDataReadReferenceOrConstruct<FStereoAudioFormat>(GetAudioOutputName(), InSettings)->GetBuffers();
-		}
-		else if (EMetasoundSourceAudioFormat::Mono == OutputFormat)
-		{
-			if (!Outputs.ContainsDataReadReference<FMonoAudioFormat>(GetAudioOutputName()))
-			{
-				UE_LOG(LogMetasound, Warning, TEXT("MetasoundSource [%s] does not contain mono output [%s] in output"), *GetName(), *GetAudioOutputName());
-			}
-			OutputBuffers = Outputs.GetDataReadReferenceOrConstruct<FMonoAudioFormat>(GetAudioOutputName(), InSettings)->GetBuffers();
-		}
-
-		// References must be cached before moving the operator to the InitParams
-		FDataReferenceCollection Inputs = Operator->GetInputs();
-		FTriggerWriteRef PlayTrigger = Inputs.GetDataWriteReferenceOrConstruct<FTrigger>(GetOnPlayInputName(), InSettings, false);
-		FTriggerReadRef FinishTrigger = Outputs.GetDataReadReferenceOrConstruct<FTrigger>(GetIsFinishedOutputName(), InSettings, false);
-
-		// Create the FMetasoundGenerator.
-		FMetasoundGeneratorInitParams InitParams =
-		{
-			MoveTemp(Operator),
-			OutputBuffers,
-			MoveTemp(PlayTrigger),
-			MoveTemp(FinishTrigger)
-		};
-
-		return ISoundGeneratorPtr(new FMetasoundGenerator(MoveTemp(InitParams)));
-	}
+		InSettings,
+		MoveTemp(DocumentWithInjectedReceives),
+		Environment,
+		NumChannels,
+		GetName(),
+		GetAudioOutputName(),
+		GetOnPlayInputName(),
+		GetIsFinishedOutputName()
+	};
 
-	return ISoundGeneratorPtr(nullptr);
+	return ISoundGeneratorPtr(new FMetasoundGenerator(MoveTemp(InitParams)));
 }
 
 TUniquePtr<IAudioInstanceTransmitter> UMetasoundSource::CreateInstanceTransmitter(const FAudioInstanceTransmitterInitParams& InParams) const
diff --git a/Engine/Plugins/Runtime/Metasound/Source/MetasoundGenerator/Private/MetasoundGenerator.cpp b/Engine/Plugins/Runtime/Metasound/Source/MetasoundGenerator/Private/MetasoundGenerator.cpp
--- a/Engine/Plugins/Runtime/Metasound/Source/MetasoundGenerator/Private/MetasoundGenerator.cpp
+++ b/Engine/Plugins/Runtime/Metasound/Source/MetasoundGenerator/Private/MetasoundGenerator.cpp
@@ -8,78 +8,161 @@
 #include "MetasoundOperatorInterface.h"
 #include "MetasoundOutputNode.h"
 #include "MetasoundTrigger.h"
-
+//#include "MetasoundSource.h"
 
 namespace Metasound
 {
+	FAsyncMetaSoundBuilder::FAsyncMetaSoundBuilder(FMetasoundGeneratorInitParams&& InInitParams)
+		: InitParams(MoveTemp(InInitParams))
+		, PlayTrigger(FTriggerWriteRef::CreateNew(InInitParams.OperatorSettings))
+		, FinishTrigger(FTriggerReadRef::CreateNew(InInitParams.OperatorSettings))
+		, bSuccess(false)
+	{
+	}
+
+	void FAsyncMetaSoundBuilder::DoWork()
+	{
+		// Create handles for new root graph
+		Frontend::FConstDocumentHandle NewDocumentHandle = Frontend::IDocumentController::CreateDocumentHandle(Frontend::MakeAccessPtr<const FMetasoundFrontendDocument>(InitParams.DocumentCopy.AccessPoint, InitParams.DocumentCopy));
+		Frontend::FConstGraphHandle RootGraph = NewDocumentHandle->GetRootGraph();
+		ensureAlways(RootGraph->IsValid());
+
+		TArray<IOperatorBuilder::FBuildErrorPtr> BuildErrors;
+		GraphOperator = RootGraph->BuildOperator(InitParams.OperatorSettings, InitParams.Environment, BuildErrors);
+
+		// Log build errors
+		for (const IOperatorBuilder::FBuildErrorPtr& Error : BuildErrors)
+		{
+			if (Error.IsValid())
+			{
+				UE_LOG(LogMetasound, Warning, TEXT("MetasoundSource [%s] build error [%s] \"%s\""), *InitParams.MetaSoundName, *(Error->GetErrorType().ToString()), *(Error->GetErrorDescription().ToString()));
+			}
+		}
+
+		if (!GraphOperator.IsValid())
+		{
+			UE_LOG(LogMetasound, Error, TEXT("Failed to build Metasound operator from graph in MetasoundSource [%s]"), *InitParams.MetaSoundName);
+			bSuccess = false;
+		}
+		else
+		{
+			FDataReferenceCollection Outputs = GraphOperator->GetOutputs();
+
+			// Get output audio buffers.
+			if (InitParams.NumOutputChannels == 2)
+			{
+				if (!Outputs.ContainsDataReadReference<FStereoAudioFormat>(InitParams.OutputName))
+				{
+					UE_LOG(LogMetasound, Warning, TEXT("MetasoundSource [%s] does not contain stereo output [%s] in output"), *InitParams.MetaSoundName, *InitParams.OutputName);
+				}
+				OutputBuffers = Outputs.GetDataReadReferenceOrConstruct<FStereoAudioFormat>(InitParams.OutputName, InitParams.OperatorSettings)->GetBuffers();
+			}
+			else if (InitParams.NumOutputChannels == 1)
+			{
+				if (!Outputs.ContainsDataReadReference<FMonoAudioFormat>(InitParams.OutputName))
+				{
+					UE_LOG(LogMetasound, Warning, TEXT("MetasoundSource [%s] does not contain mono output [%s] in output"), *InitParams.MetaSoundName, *InitParams.OutputName);
+				}
+				OutputBuffers = Outputs.GetDataReadReferenceOrConstruct<FMonoAudioFormat>(InitParams.OutputName, InitParams.OperatorSettings)->GetBuffers();
+			}
+
+			// References must be cached before moving the operator to the InitParams
+			FDataReferenceCollection Inputs = GraphOperator->GetInputs();
+			PlayTrigger = Inputs.GetDataWriteReferenceOrConstruct<FTrigger>(InitParams.InputName, InitParams.OperatorSettings, false);
+			FinishTrigger = Outputs.GetDataReadReferenceOrConstruct<FTrigger>(InitParams.IsFinishedOutputName, InitParams.OperatorSettings, false);
+			bSuccess = true;
+		}
+	}
+
+	bool FAsyncMetaSoundBuilder::SetDataOnGenerator(FMetasoundGenerator& InGenerator)
+	{
+		if (bSuccess)
+		{
+			FMetasoundGeneratorData GeneratorData =
+			{
+				MoveTemp(GraphOperator),
+				OutputBuffers,
+				MoveTemp(PlayTrigger),
+				MoveTemp(FinishTrigger)
+			};
+
+			InGenerator.SetGraph(MoveTemp(GeneratorData));
+		}
+		return bSuccess;
+	}
+
 	FMetasoundGenerator::FMetasoundGenerator(FMetasoundGeneratorInitParams&& InParams)
-		: bIsPlaying(false)
+		: bIsGraphBuilding(true)
+		, bIsPlaying(false)
 		, bIsFinished(false)
 		, NumChannels(0)
 		, NumFramesPerExecute(0)
 		, NumSamplesPerExecute(0)
-		, OnPlayTriggerRef(InParams.TriggerOnPlayRef)
-		, OnFinishedTriggerRef(InParams.TriggerOnFinishRef)
+		, OnPlayTriggerRef(FTriggerWriteRef::CreateNew(InParams.OperatorSettings))
+		, OnFinishedTriggerRef(FTriggerWriteRef::CreateNew(InParams.OperatorSettings))
 	{
-		SetGraph(MoveTemp(InParams));
+		NumChannels = InParams.NumOutputChannels;
+		NumFramesPerExecute = InParams.OperatorSettings.GetNumFramesPerBlock();
+		NumSamplesPerExecute = NumChannels * NumFramesPerExecute;
+		BuilderTask = MakeUnique<FBuilderTask>(MoveTemp(InParams));
+		BuilderTask->StartBackgroundTask(GBackgroundPriorityThreadPool);
 	}
 
 	FMetasoundGenerator::~FMetasoundGenerator()
 	{
+		if (BuilderTask.IsValid())
+		{
+			BuilderTask->EnsureCompletion();
+			BuilderTask = nullptr;
+		}
 	}
 
-	bool FMetasoundGenerator::UpdateGraphOperator(FMetasoundGeneratorInitParams&& InParams)
+	bool FMetasoundGenerator::IsTickable() const
 	{
-		// multichannel version:
-		const int32 GraphNumChannels = InParams.OutputBuffers.Num();
+		return bIsGraphBuilding;
+	}
 
-		if (GraphNumChannels == NumChannels)
-		{
-			SetGraph(MoveTemp(InParams));
+	void FMetasoundGenerator::Tick(float DeltaTime)
+	{
+		if (BuilderTask->IsDone())
+		{	
+			FAsyncMetaSoundBuilder& Builder = BuilderTask->GetTask();
+			if (Builder.SetDataOnGenerator(*this))
+			{
+				// We're done building our graph now
+				bIsGraphBuilding = false;
+			}
+			else
+			{
+				// Failed to load/generate the graph, kill the metasound
+				bIsFinished = true;
+			}
 
-			return true;
+			BuilderTask = nullptr;
 		}
-
-		return false;
 	}
 
-	void FMetasoundGenerator::SetGraph(FMetasoundGeneratorInitParams&& InParams)
+	void FMetasoundGenerator::SetGraph(FMetasoundGeneratorData&& InData)
 	{
 		InterleavedAudioBuffer.Reset();
-		NumFramesPerExecute = 0;
-		NumSamplesPerExecute = 0;
 
 		GraphOutputAudio.Reset();
-		if (InParams.OutputBuffers.Num() > 0)
+		if (InData.OutputBuffers.Num() > 0)
 		{
-			GraphOutputAudio.Append(InParams.OutputBuffers.GetData(), InParams.OutputBuffers.Num());
+			GraphOutputAudio.Append(InData.OutputBuffers.GetData(), InData.OutputBuffers.Num());
 		}
 
-		OnPlayTriggerRef = InParams.TriggerOnPlayRef;
-		OnFinishedTriggerRef = InParams.TriggerOnFinishRef;
+		OnPlayTriggerRef = InData.TriggerOnPlayRef;
+		OnFinishedTriggerRef = InData.TriggerOnFinishRef;
 
 		// The graph operator and graph audio output contain all the values needed
 		// by the sound generator.
-		RootExecuter.SetOperator(MoveTemp(InParams.GraphOperator));
+		RootExecuter.SetOperator(MoveTemp(InData.GraphOperator));
 
 
 		// Query the graph output to get the number of output audio channels.
 		// Multichannel version:
-		NumChannels = GraphOutputAudio.Num();
-
-
-		// All buffers have same number of frames, so only need to query
-		// first buffer to know number of frames.
-		if (NumChannels > 0)
-		{
-			NumFramesPerExecute = GraphOutputAudio[0]->Num();
-		}
-		else
-		{
-			NumFramesPerExecute = 0;
-		}
-
-		NumSamplesPerExecute = NumFramesPerExecute * NumChannels;
+		check(NumChannels == GraphOutputAudio.Num());
 
 		if (NumSamplesPerExecute > 0)
 		{
@@ -100,16 +183,17 @@ namespace Metasound
 			return 0;
 		}
 
-		if (!bIsPlaying)
+		// Output silent audio if we're still building a graph
+		if (bIsGraphBuilding || NumSamplesPerExecute < 1)
 		{
-			OnPlayTriggerRef->TriggerFrame(0);
-			bIsPlaying = true;
+			FMemory::Memset(OutAudio, 0, sizeof(float) * NumSamplesRemaining);
+			return NumSamplesRemaining;
 		}
 
-		if (NumSamplesPerExecute < 1)
+		if (!bIsPlaying)
 		{
-			FMemory::Memset(OutAudio, 0, sizeof(float) * NumSamplesRemaining);
-			return NumSamplesRemaining;
+			OnPlayTriggerRef->TriggerFrame(0);
+			bIsPlaying = true;
 		}
 
 		// If we have any audio left in the internal overflow buffer from 
diff --git a/Engine/Plugins/Runtime/Metasound/Source/MetasoundGenerator/Public/MetasoundGenerator.h b/Engine/Plugins/Runtime/Metasound/Source/MetasoundGenerator/Public/MetasoundGenerator.h
--- a/Engine/Plugins/Runtime/Metasound/Source/MetasoundGenerator/Public/MetasoundGenerator.h
+++ b/Engine/Plugins/Runtime/Metasound/Source/MetasoundGenerator/Public/MetasoundGenerator.h
@@ -9,10 +9,26 @@
 #include "MetasoundPrimitives.h"
 #include "MetasoundTrigger.h"
 #include "Sound/SoundGenerator.h"
+#include "MetasoundFrontendController.h"
+#include "Async/Async.h"
+#include "Tickable.h"
 
 namespace Metasound
 {
+	// Struct needed for building the metasound graph
 	struct FMetasoundGeneratorInitParams
+	{
+		FOperatorSettings OperatorSettings;
+		FMetasoundFrontendDocument DocumentCopy;
+		FMetasoundEnvironment Environment;
+		int32 NumOutputChannels = 0;
+		FString MetaSoundName;
+		FString OutputName;
+		FString InputName;
+		FString IsFinishedOutputName;
+	};
+
+	struct FMetasoundGeneratorData
 	{
 		TUniquePtr<Metasound::IOperator> GraphOperator;
 		TArrayView<TDataReadReference<FAudioBuffer>> OutputBuffers;
@@ -20,10 +36,37 @@ namespace Metasound
 		FTriggerReadRef TriggerOnFinishRef;
 	};
 
+	class FMetasoundGenerator;
+
+	class FAsyncMetaSoundBuilder : public FNonAbandonableTask
+	{
+	public:
+		FAsyncMetaSoundBuilder(FMetasoundGeneratorInitParams&& InInitParams);
+		~FAsyncMetaSoundBuilder() = default;
+
+		void DoWork();
+		bool SetDataOnGenerator(FMetasoundGenerator& InGenerator);
+
+		FORCEINLINE TStatId GetStatId() const
+		{
+			RETURN_QUICK_DECLARE_CYCLE_STAT(FAsyncMetaSoundBuilder, STATGROUP_ThreadPoolAsyncTasks);
+		}
+
+	private:
+		FMetasoundGeneratorInitParams InitParams;
+		TUniquePtr<IOperator> GraphOperator;
+		TArrayView<FAudioBufferReadRef> OutputBuffers;
+		FTriggerWriteRef PlayTrigger;
+		FTriggerReadRef FinishTrigger;
+		bool bSuccess;
+	};
+
+
 	/** FMetasoundGenerator generates audio from a given metasound IOperator
 	 * which produces a multichannel audio output.
 	 */
-	class METASOUNDGENERATOR_API FMetasoundGenerator : public ISoundGenerator
+	class METASOUNDGENERATOR_API FMetasoundGenerator : public ISoundGenerator, 
+														public FTickableGameObject
 	{
 	public:
 		using FOperatorUniquePtr = TUniquePtr<Metasound::IOperator>;
@@ -38,6 +81,15 @@ namespace Metasound
 
 		virtual ~FMetasoundGenerator();
 
+		//~ Begin FTickableGameObject
+		virtual ETickableTickType GetTickableTickType() const override { return ETickableTickType::Conditional; }
+		virtual bool IsTickableWhenPaused() const override { return true; }
+		virtual bool IsTickableInEditor() const override { return true; }
+		virtual bool IsTickable() const override;
+		virtual void Tick(float DeltaTime) override;
+		virtual TStatId GetStatId() const override { RETURN_QUICK_DECLARE_CYCLE_STAT(MetasoundGenerator, STATGROUP_Tickables); }
+		//~ End FTickableGameObject
+
 		/** Set the value of a graph's input data using the assignment operator.
 		 *
 		 * @param InName - The name of the graph's input data reference.
@@ -94,7 +146,7 @@ namespace Metasound
 		 *
 		 * @return Returns true if the graph operator was successfully swapped. False otherwise.
 		 */
-		bool UpdateGraphOperator(FMetasoundGeneratorInitParams&& InParams);
+//		bool UpdateGraphOperator(FMetasoundGeneratorInitParams&& InParams);
 
 		/** Return the number of audio channels. */
 		int32 GetNumChannels() const;
@@ -107,7 +159,7 @@ namespace Metasound
 
 	private:
 		// Internal set graph after checking compatibility.
-		void SetGraph(FMetasoundGeneratorInitParams&& InParams);
+		void SetGraph(FMetasoundGeneratorData&& InData);
 
 		// Fill OutAudio with data in InBuffer, up to maximum number of samples.
 		// Returns the number of samples used.
@@ -118,6 +170,7 @@ namespace Metasound
 		
 		FExecuter RootExecuter;
 
+		bool bIsGraphBuilding;
 		bool bIsPlaying;
 		bool bIsFinished;
 
@@ -136,6 +189,11 @@ namespace Metasound
 		Audio::AlignedFloatBuffer InterleavedAudioBuffer;
 
 		Audio::AlignedFloatBuffer OverflowBuffer;
+
+		typedef FAsyncTask<FAsyncMetaSoundBuilder> FBuilderTask;
+		TUniquePtr<FBuilderTask> BuilderTask;
+
+		friend class FAsyncMetaSoundBuilder;
 	};
 }
 
diff --git a/Engine/Plugins/Runtime/Metasound/Source/MetasoundGraphCore/Public/MetasoundEnvironment.h b/Engine/Plugins/Runtime/Metasound/Source/MetasoundGraphCore/Public/MetasoundEnvironment.h
--- a/Engine/Plugins/Runtime/Metasound/Source/MetasoundGraphCore/Public/MetasoundEnvironment.h
+++ b/Engine/Plugins/Runtime/Metasound/Source/MetasoundGraphCore/Public/MetasoundEnvironment.h
@@ -175,7 +175,7 @@ namespace Metasound
 	/** FMetasoundEnvironment contains a set of TMetasoundEnvironmentVariables requiring
 	 * that each environment variable has a unique name. 
 	 */
-	class FMetasoundEnvironment
+	class METASOUNDGRAPHCORE_API FMetasoundEnvironment
 	{
 	public:
 		FMetasoundEnvironment() = default;
