Updating vertex default retrieval method

diff --git a/Engine/Plugins/Runtime/Metasound/Source/MetasoundFrontend/Private/MetasoundFrontend.cpp b/Engine/Plugins/Runtime/Metasound/Source/MetasoundFrontend/Private/MetasoundFrontend.cpp
--- a/Engine/Plugins/Runtime/Metasound/Source/MetasoundFrontend/Private/MetasoundFrontend.cpp
+++ b/Engine/Plugins/Runtime/Metasound/Source/MetasoundFrontend/Private/MetasoundFrontend.cpp
@@ -84,8 +84,8 @@ namespace Metasound
 				ClassInput.Metadata.DisplayName = FText::FromString(InputTuple.Value.GetVertexName());
 				ClassInput.Metadata.Description = InputTuple.Value.GetDescription();
 
-				const FLiteral& DefaultLiteral = InputTuple.Value.GetDefaultValue();
-				if (DefaultLiteral.GetType() != ELiteralType::None)
+				FLiteral DefaultLiteral = InputTuple.Value.GetDefaultLiteral();
+				if (DefaultLiteral.GetType() != ELiteralType::Invalid)
 				{
 					ClassInput.DefaultLiteral.SetFromLiteral(DefaultLiteral);
 				}
diff --git a/Engine/Plugins/Runtime/Metasound/Source/MetasoundFrontend/Public/MetasoundArrayNodes.h b/Engine/Plugins/Runtime/Metasound/Source/MetasoundFrontend/Public/MetasoundArrayNodes.h
--- a/Engine/Plugins/Runtime/Metasound/Source/MetasoundFrontend/Public/MetasoundArrayNodes.h
+++ b/Engine/Plugins/Runtime/Metasound/Source/MetasoundFrontend/Public/MetasoundArrayNodes.h
@@ -41,32 +41,6 @@ namespace Metasound
 		{
 			using Type = ElementType;
 		};
-		
-		// TODO: Remove me and swap after fixing FDataReferenceCollection::GetDataReadReferenceOrConstructWithVertexDefault
-		//
-		// This existing implementation of `GetDataReadReferenceOrConstructWithVertexDefault` does not automatically
-		// utilize the set type of the FLiteral or the TDataReadReferenceLiteralFactory 
-		// for that type. The implementation in FDataReferenceCollection should be replaced
-		// with this implementation. It will require reworking of all nodes which
-		// call `GetDataReadReferenceOrConstructWithVertexDefault`.
-		template<typename DataType>
-		TDataReadReference<DataType> GetDataReadReferenceOrConstructWithVertexDefault(const FDataReferenceCollection& InCollection, const FInputVertexInterface& InputVertices, const FString& InName, const FOperatorSettings& InSettings)
-		{
-			using FDataFactory = TDataReadReferenceLiteralFactory<DataType>;
-
-			if (InCollection.ContainsDataReadReference<DataType>(InName))
-			{
-				return InCollection.GetDataReadReference<DataType>(InName);
-			}
-			else
-			{
-				// TODO: replace mechanism in data vertex to "CreateDefaultValue"
-				// instead of get in order to avoid the need to have a `Clone` on 
-				// FLiteral. Then move FLiteral when needed.
-				const FLiteral& Literal = InputVertices[InName].GetDefaultValue();
-				return FDataFactory::CreateExplicitArgs(InSettings, Literal);
-			}
-		}
 	}
 
 	namespace ArrayNodeVertexNames
@@ -139,7 +113,7 @@ namespace Metasound
 			const FInputVertexInterface& Inputs = InParams.Node.GetVertexInterface().GetInputInterface();
 
 			// Get the input array or construct an empty one. 
-			FArrayDataReadReference Array = GetDataReadReferenceOrConstructWithVertexDefault<ArrayType>(InParams.InputDataReferences, Inputs, GetInputArrayName(), InParams.OperatorSettings);
+			FArrayDataReadReference Array = InParams.InputDataReferences.GetDataReadReferenceOrConstructWithVertexDefault<ArrayType>(Inputs, GetInputArrayName(), InParams.OperatorSettings);
 
 			return MakeUnique<TArrayNumOperator>(Array);
 		}
@@ -253,13 +227,13 @@ namespace Metasound
 			const FInputVertexInterface& Inputs = InParams.Node.GetVertexInterface().GetInputInterface();
 
 			// Input Trigger
-			TDataReadReference<FTrigger> Trigger = GetDataReadReferenceOrConstructWithVertexDefault<FTrigger>(InParams.InputDataReferences, Inputs, GetInputTriggerName(), InParams.OperatorSettings);
+			TDataReadReference<FTrigger> Trigger = InParams.InputDataReferences.GetDataReadReferenceOrConstructWithVertexDefault<FTrigger>(Inputs, GetInputTriggerName(), InParams.OperatorSettings);
 			
 			// Input Array
-			FArrayDataReadReference Array = GetDataReadReferenceOrConstructWithVertexDefault<ArrayType>(InParams.InputDataReferences, Inputs, GetInputArrayName(), InParams.OperatorSettings);
+			FArrayDataReadReference Array = InParams.InputDataReferences.GetDataReadReferenceOrConstructWithVertexDefault<ArrayType>(Inputs, GetInputArrayName(), InParams.OperatorSettings);
 
 			// Input Index
-			TDataReadReference<int32> Index = GetDataReadReferenceOrConstructWithVertexDefault<int32>(InParams.InputDataReferences, Inputs, GetInputIndexName(), InParams.OperatorSettings);
+			TDataReadReference<int32> Index = InParams.InputDataReferences.GetDataReadReferenceOrConstructWithVertexDefault<int32>(Inputs, GetInputIndexName(), InParams.OperatorSettings);
 
 			return MakeUnique<TArrayGetOperator>(InParams.OperatorSettings, Trigger, Array, Index);
 		}
@@ -390,14 +364,14 @@ namespace Metasound
 
 			const FInputVertexInterface& Inputs = InParams.Node.GetVertexInterface().GetInputInterface();
 			
-			TDataReadReference<FTrigger> Trigger = GetDataReadReferenceOrConstructWithVertexDefault<FTrigger>(InParams.InputDataReferences, Inputs, GetInputTriggerName(), InParams.OperatorSettings);
+			TDataReadReference<FTrigger> Trigger = InParams.InputDataReferences.GetDataReadReferenceOrConstructWithVertexDefault<FTrigger>(Inputs, GetInputTriggerName(), InParams.OperatorSettings);
 
-			FArrayDataReadReference InitArray = GetDataReadReferenceOrConstructWithVertexDefault<ArrayType>(InParams.InputDataReferences, Inputs, GetInputArrayName(), InParams.OperatorSettings);
+			FArrayDataReadReference InitArray = InParams.InputDataReferences.GetDataReadReferenceOrConstructWithVertexDefault<ArrayType>(Inputs, GetInputArrayName(), InParams.OperatorSettings);
 			FArrayDataWriteReference Array = TDataWriteReferenceFactory<ArrayType>::CreateExplicitArgs(InParams.OperatorSettings, *InitArray);
 
-			TDataReadReference<int32> Index = GetDataReadReferenceOrConstructWithVertexDefault<int32>(InParams.InputDataReferences, Inputs, GetInputIndexName(), InParams.OperatorSettings);
+			TDataReadReference<int32> Index = InParams.InputDataReferences.GetDataReadReferenceOrConstructWithVertexDefault<int32>(Inputs, GetInputIndexName(), InParams.OperatorSettings);
 
-			TDataReadReference<ElementType> Value = GetDataReadReferenceOrConstructWithVertexDefault<ElementType>(InParams.InputDataReferences, Inputs, GetInputValueName(), InParams.OperatorSettings);
+			TDataReadReference<ElementType> Value = InParams.InputDataReferences.GetDataReadReferenceOrConstructWithVertexDefault<ElementType>(Inputs, GetInputValueName(), InParams.OperatorSettings);
 
 			return MakeUnique<TArraySetOperator>(InParams.OperatorSettings, Trigger, InitArray, Array, Index, Value);
 		}
@@ -530,10 +504,10 @@ namespace Metasound
 
 			const FInputVertexInterface& Inputs = InParams.Node.GetVertexInterface().GetInputInterface();
 			
-			TDataReadReference<FTrigger> Trigger = GetDataReadReferenceOrConstructWithVertexDefault<FTrigger>(InParams.InputDataReferences, Inputs, GetInputTriggerName(), InParams.OperatorSettings);
+			TDataReadReference<FTrigger> Trigger = InParams.InputDataReferences.GetDataReadReferenceOrConstructWithVertexDefault<FTrigger>(Inputs, GetInputTriggerName(), InParams.OperatorSettings);
 
-			FArrayDataReadReference LeftArray = GetDataReadReferenceOrConstructWithVertexDefault<ArrayType>(InParams.InputDataReferences, Inputs, GetInputLeftArrayName(), InParams.OperatorSettings);
-			FArrayDataReadReference RightArray = GetDataReadReferenceOrConstructWithVertexDefault<ArrayType>(InParams.InputDataReferences, Inputs, GetInputRightArrayName(), InParams.OperatorSettings);
+			FArrayDataReadReference LeftArray = InParams.InputDataReferences.GetDataReadReferenceOrConstructWithVertexDefault<ArrayType>(Inputs, GetInputLeftArrayName(), InParams.OperatorSettings);
+			FArrayDataReadReference RightArray = InParams.InputDataReferences.GetDataReadReferenceOrConstructWithVertexDefault<ArrayType>(Inputs, GetInputRightArrayName(), InParams.OperatorSettings);
 
 			FArrayDataWriteReference OutArray = TDataWriteReferenceFactory<ArrayType>::CreateExplicitArgs(InParams.OperatorSettings);
 
@@ -656,12 +630,12 @@ namespace Metasound
 
 			const FInputVertexInterface& Inputs = InParams.Node.GetVertexInterface().GetInputInterface();
 			
-			TDataReadReference<FTrigger> Trigger = GetDataReadReferenceOrConstructWithVertexDefault<FTrigger>(InParams.InputDataReferences, Inputs, GetInputTriggerName(), InParams.OperatorSettings);
+			TDataReadReference<FTrigger> Trigger = InParams.InputDataReferences.GetDataReadReferenceOrConstructWithVertexDefault<FTrigger>(Inputs, GetInputTriggerName(), InParams.OperatorSettings);
 
-			FArrayDataReadReference InArray = GetDataReadReferenceOrConstructWithVertexDefault<ArrayType>(InParams.InputDataReferences, Inputs, GetInputArrayName(), InParams.OperatorSettings);
+			FArrayDataReadReference InArray = InParams.InputDataReferences.GetDataReadReferenceOrConstructWithVertexDefault<ArrayType>(Inputs, GetInputArrayName(), InParams.OperatorSettings);
 
-			TDataReadReference<int32> StartIndex = GetDataReadReferenceOrConstructWithVertexDefault<int32>(InParams.InputDataReferences, Inputs, GetInputStartIndexName(), InParams.OperatorSettings);
-			TDataReadReference<int32> EndIndex = GetDataReadReferenceOrConstructWithVertexDefault<int32>(InParams.InputDataReferences, Inputs, GetInputEndIndexName(), InParams.OperatorSettings);
+			TDataReadReference<int32> StartIndex = InParams.InputDataReferences.GetDataReadReferenceOrConstructWithVertexDefault<int32>(Inputs, GetInputStartIndexName(), InParams.OperatorSettings);
+			TDataReadReference<int32> EndIndex = InParams.InputDataReferences.GetDataReadReferenceOrConstructWithVertexDefault<int32>(Inputs, GetInputEndIndexName(), InParams.OperatorSettings);
 
 			FArrayDataWriteReference OutArray = TDataWriteReferenceFactory<ArrayType>::CreateExplicitArgs(InParams.OperatorSettings);
 
diff --git a/Engine/Plugins/Runtime/Metasound/Source/MetasoundFrontend/Public/MetasoundArrayRandomNode.h b/Engine/Plugins/Runtime/Metasound/Source/MetasoundFrontend/Public/MetasoundArrayRandomNode.h
--- a/Engine/Plugins/Runtime/Metasound/Source/MetasoundFrontend/Public/MetasoundArrayRandomNode.h
+++ b/Engine/Plugins/Runtime/Metasound/Source/MetasoundFrontend/Public/MetasoundArrayRandomNode.h
@@ -144,13 +144,13 @@ namespace Metasound
 
 			const FInputVertexInterface& Inputs = InParams.Node.GetVertexInterface().GetInputInterface();
 
-			FTriggerReadRef InTriggerNext = GetDataReadReferenceOrConstructWithVertexDefault<FTrigger>(InParams.InputDataReferences, Inputs, GetInputTriggerNextName(), InParams.OperatorSettings);
-			FTriggerReadRef InTriggerReset = GetDataReadReferenceOrConstructWithVertexDefault<FTrigger>(InParams.InputDataReferences, Inputs, GetInputTriggerResetName(), InParams.OperatorSettings);
-			FArrayDataReadReference InInputArray = GetDataReadReferenceOrConstructWithVertexDefault<ArrayType>(InParams.InputDataReferences, Inputs, GetInputRandomArrayName(), InParams.OperatorSettings);
-			TDataReadReference<WeightArrayType> InInputWeightsArray = GetDataReadReferenceOrConstructWithVertexDefault<WeightArrayType>(InParams.InputDataReferences, Inputs, GetInputRandomArrayName(), InParams.OperatorSettings);
-			FInt32ReadRef InSeedValue = GetDataReadReferenceOrConstructWithVertexDefault<int32>(InParams.InputDataReferences, Inputs, GetInputSeedName(), InParams.OperatorSettings);
-			FInt32ReadRef InNoRepeatOrder = GetDataReadReferenceOrConstructWithVertexDefault<int32>(InParams.InputDataReferences, Inputs, GetInputNoRepeatOrderName(), InParams.OperatorSettings);
-			FBoolReadRef bInEnableSharedState = GetDataReadReferenceOrConstructWithVertexDefault<bool>(InParams.InputDataReferences, Inputs, GetInputEnableSharedStateName(), InParams.OperatorSettings);
+			FTriggerReadRef InTriggerNext = InParams.InputDataReferences.GetDataReadReferenceOrConstructWithVertexDefault<FTrigger>(Inputs, GetInputTriggerNextName(), InParams.OperatorSettings);
+			FTriggerReadRef InTriggerReset = InParams.InputDataReferences.GetDataReadReferenceOrConstructWithVertexDefault<FTrigger>(Inputs, GetInputTriggerResetName(), InParams.OperatorSettings);
+			FArrayDataReadReference InInputArray = InParams.InputDataReferences.GetDataReadReferenceOrConstructWithVertexDefault<ArrayType>(Inputs, GetInputRandomArrayName(), InParams.OperatorSettings);
+			TDataReadReference<WeightArrayType> InInputWeightsArray = InParams.InputDataReferences.GetDataReadReferenceOrConstructWithVertexDefault<WeightArrayType>(Inputs, GetInputRandomArrayName(), InParams.OperatorSettings);
+			FInt32ReadRef InSeedValue = InParams.InputDataReferences.GetDataReadReferenceOrConstructWithVertexDefault<int32>(Inputs, GetInputSeedName(), InParams.OperatorSettings);
+			FInt32ReadRef InNoRepeatOrder = InParams.InputDataReferences.GetDataReadReferenceOrConstructWithVertexDefault<int32>(Inputs, GetInputNoRepeatOrderName(), InParams.OperatorSettings);
+			FBoolReadRef bInEnableSharedState = InParams.InputDataReferences.GetDataReadReferenceOrConstructWithVertexDefault<bool>(Inputs, GetInputEnableSharedStateName(), InParams.OperatorSettings);
 
 			return MakeUnique<TArrayRandomGetOperator>(InParams, InTriggerNext, InTriggerReset, InInputArray, InInputWeightsArray, InSeedValue, InNoRepeatOrder, bInEnableSharedState);
 		}
diff --git a/Engine/Plugins/Runtime/Metasound/Source/MetasoundFrontend/Public/MetasoundArrayShuffleNode.h b/Engine/Plugins/Runtime/Metasound/Source/MetasoundFrontend/Public/MetasoundArrayShuffleNode.h
--- a/Engine/Plugins/Runtime/Metasound/Source/MetasoundFrontend/Public/MetasoundArrayShuffleNode.h
+++ b/Engine/Plugins/Runtime/Metasound/Source/MetasoundFrontend/Public/MetasoundArrayShuffleNode.h
@@ -190,13 +190,13 @@ namespace Metasound
 
 			const FInputVertexInterface& Inputs = InParams.Node.GetVertexInterface().GetInputInterface();
 
-			TDataReadReference<FTrigger> InTriggerNext = GetDataReadReferenceOrConstructWithVertexDefault<FTrigger>(InParams.InputDataReferences, Inputs, GetInputTriggerNextName(), InParams.OperatorSettings);
-			TDataReadReference<FTrigger> InTriggerShuffle = GetDataReadReferenceOrConstructWithVertexDefault<FTrigger>(InParams.InputDataReferences, Inputs, GetInputTriggerShuffleName(), InParams.OperatorSettings);
-			TDataReadReference<FTrigger> InTriggerReset = GetDataReadReferenceOrConstructWithVertexDefault<FTrigger>(InParams.InputDataReferences, Inputs, GetInputTriggerResetName(), InParams.OperatorSettings);
-			FArrayDataReadReference InInputArray = GetDataReadReferenceOrConstructWithVertexDefault<ArrayType>(InParams.InputDataReferences, Inputs, GetInputShuffleArrayName(), InParams.OperatorSettings);
-			TDataReadReference<int32> InSeedValue = GetDataReadReferenceOrConstructWithVertexDefault<int32>(InParams.InputDataReferences, Inputs, GetInputSeedName(), InParams.OperatorSettings);
-			TDataReadReference<bool> bInAutoShuffle = GetDataReadReferenceOrConstructWithVertexDefault<bool>(InParams.InputDataReferences, Inputs, GetInputAutoShuffleName(), InParams.OperatorSettings);
-			TDataReadReference<bool> bInEnableSharedState = GetDataReadReferenceOrConstructWithVertexDefault<bool>(InParams.InputDataReferences, Inputs, GetInputEnableSharedStateName(), InParams.OperatorSettings);
+			TDataReadReference<FTrigger> InTriggerNext = InParams.InputDataReferences.GetDataReadReferenceOrConstructWithVertexDefault<FTrigger>(Inputs, GetInputTriggerNextName(), InParams.OperatorSettings);
+			TDataReadReference<FTrigger> InTriggerShuffle = InParams.InputDataReferences.GetDataReadReferenceOrConstructWithVertexDefault<FTrigger>(Inputs, GetInputTriggerShuffleName(), InParams.OperatorSettings);
+			TDataReadReference<FTrigger> InTriggerReset = InParams.InputDataReferences.GetDataReadReferenceOrConstructWithVertexDefault<FTrigger>(Inputs, GetInputTriggerResetName(), InParams.OperatorSettings);
+			FArrayDataReadReference InInputArray = InParams.InputDataReferences.GetDataReadReferenceOrConstructWithVertexDefault<ArrayType>(Inputs, GetInputShuffleArrayName(), InParams.OperatorSettings);
+			TDataReadReference<int32> InSeedValue = InParams.InputDataReferences.GetDataReadReferenceOrConstructWithVertexDefault<int32>(Inputs, GetInputSeedName(), InParams.OperatorSettings);
+			TDataReadReference<bool> bInAutoShuffle = InParams.InputDataReferences.GetDataReadReferenceOrConstructWithVertexDefault<bool>(Inputs, GetInputAutoShuffleName(), InParams.OperatorSettings);
+			TDataReadReference<bool> bInEnableSharedState = InParams.InputDataReferences.GetDataReadReferenceOrConstructWithVertexDefault<bool>(Inputs, GetInputEnableSharedStateName(), InParams.OperatorSettings);
 
 			return MakeUnique<TArrayShuffleOperator>(InParams, InTriggerNext, InTriggerShuffle, InTriggerReset, InInputArray, InSeedValue, bInAutoShuffle, bInEnableSharedState);
 		}
diff --git a/Engine/Plugins/Runtime/Metasound/Source/MetasoundFrontend/Public/MetasoundInputNode.h b/Engine/Plugins/Runtime/Metasound/Source/MetasoundFrontend/Public/MetasoundInputNode.h
--- a/Engine/Plugins/Runtime/Metasound/Source/MetasoundFrontend/Public/MetasoundInputNode.h
+++ b/Engine/Plugins/Runtime/Metasound/Source/MetasoundFrontend/Public/MetasoundInputNode.h
@@ -23,28 +23,30 @@ namespace Metasound
 				: DataReferenceName(InDataReferenceName)
 				// Executable DataTypes require a copy of the output to operate on whereas non-executable
 				// types do not. Avoid copy by assigning to reference for non-executable types.
-				, Output(TExecutableDataType<DataType>::bIsExecutable ? FDataWriteReference::CreateNew(*InDataReference) : InDataReference)
+				, InputValue(InDataReference)
+				, OutputValue(TExecutableDataType<DataType>::bIsExecutable ? FDataWriteReference::CreateNew(*InDataReference) : InDataReference)
 			{
-				Inputs.AddDataWriteReference<DataType>(InDataReferenceName, InDataReference);
-				Outputs.AddDataReadReference<DataType>(InDataReferenceName, Output);
 			}
 
 			virtual ~TInputOperator() {}
 
 			virtual FDataReferenceCollection GetInputs() const override
 			{
+				FDataReferenceCollection Inputs;
+				Inputs.AddDataWriteReference<DataType>(DataReferenceName, InputValue);
 				return Inputs;
 			}
 
 			virtual FDataReferenceCollection GetOutputs() const override
 			{
+				FDataReferenceCollection Outputs;
+				Outputs.AddDataReadReference<DataType>(DataReferenceName, OutputValue);
 				return Outputs;
 			}
 
 			void Execute()
 			{
-				FDataWriteReference Input = Inputs.GetDataWriteReference<DataType>(DataReferenceName);
-				TExecutableDataType<DataType>::Execute(*Input, *Output);
+				TExecutableDataType<DataType>::Execute(*InputValue, *OutputValue);
 			}
 
 			static void ExecuteFunction(IOperator* InOperator)
@@ -63,10 +65,10 @@ namespace Metasound
 
 		private:
 			FVertexKey DataReferenceName;
-			FDataWriteReference Output;
 
-			FDataReferenceCollection Outputs;
-			FDataReferenceCollection Inputs;
+			FDataWriteReference InputValue;
+			FDataWriteReference OutputValue;
+
 	};
 
 	/** TInputOperatorLiteralFactory creates an input by passing it a literal. */
diff --git a/Engine/Plugins/Runtime/Metasound/Source/MetasoundGraphCore/Private/MetasoundVertex.cpp b/Engine/Plugins/Runtime/Metasound/Source/MetasoundGraphCore/Private/MetasoundVertex.cpp
--- a/Engine/Plugins/Runtime/Metasound/Source/MetasoundGraphCore/Private/MetasoundVertex.cpp
+++ b/Engine/Plugins/Runtime/Metasound/Source/MetasoundGraphCore/Private/MetasoundVertex.cpp
@@ -53,9 +53,9 @@ namespace Metasound
 		return VertexModel->DataTypeName;
 	}
 
-	const FLiteral& FInputDataVertex::GetDefaultValue() const
+	FLiteral FInputDataVertex::GetDefaultLiteral() const
 	{
-		return VertexModel->DefaultValue;
+		return VertexModel->CreateDefaultLiteral();
 	}
 
 	const FText& FInputDataVertex::GetDescription() const
diff --git a/Engine/Plugins/Runtime/Metasound/Source/MetasoundGraphCore/Public/MetasoundDataReferenceCollection.h b/Engine/Plugins/Runtime/Metasound/Source/MetasoundGraphCore/Public/MetasoundDataReferenceCollection.h
--- a/Engine/Plugins/Runtime/Metasound/Source/MetasoundGraphCore/Public/MetasoundDataReferenceCollection.h
+++ b/Engine/Plugins/Runtime/Metasound/Source/MetasoundGraphCore/Public/MetasoundDataReferenceCollection.h
@@ -3,6 +3,7 @@
 #pragma once
 
 #include "CoreMinimal.h"
+#include "MetasoundDataFactory.h"
 #include "MetasoundDataReference.h"
 #include "MetasoundOperatorSettings.h"
 #include "MetasoundVertex.h"
@@ -220,46 +221,16 @@ namespace Metasound
 			/** Returns a readable data ref from the collection or construct one
 			 * if one is not there with the default provided from the given input vertex collection.
 			 *
-			 * @param InName - Name of readable data.
 			 * @param InputVertices - Collection of input vertices to retrieve default from if reference is not available
-			 * @param ConstructorArgs - Arguments to pass to constructor of TDataReadReference<DataType>
-			 *
-			 * @return A readable data reference.
-			 */
-			// TODO: FIXME
-			// 1. Remove extra template parameters "DefaultDataType", "ConstructorArgTypes"
-			// 2. Use data factory instead of read reference class to construct
-			// 3. Pipe in FOperatorSettings to the factory call.
-			// 4. Build from literal, not from DefaultDataType.
-			template<typename DataType, typename DefaultDataType = DataType, typename... ConstructorArgTypes>
-			TDataReadReference<DataType> GetDataReadReferenceOrConstructWithVertexDefault(const FInputVertexInterface& InputVertices, const FString& InName) const
-			{
-				using FDataRef = TDataReadReference<DataType>;
-
-				if (ContainsDataReadReference<DataType>(InName))
-				{
-					return GetDataReadReference<DataType>(InName);
-				}
-				else
-				{
-					const DefaultDataType DefaultValue = InputVertices[InName].GetDefaultValue().Value.Get<DefaultDataType>();
-					return FDataRef::CreateNew(DefaultValue);
-				}
-			}
-
-			/** Returns a readable data ref from the collection or construct one
-			 * if one is not there with the default provided from the given input vertex collection.
-			 *
 			 * @param InName - Name of readable data.
-			 * @param InputVertices - Collection of input vertices to retrieve default from if reference is not available
-			 * @param ConstructorArgs - Arguments to pass to constructor of TDataReadReference<DataType>
+			 * @param InSettings - OperatorSettings of Metasound Operator 
 			 *
 			 * @return A readable data reference.
 			 */
-			template<typename DataType, typename DefaultDataType = DataType, typename... ConstructorArgTypes>
-			TDataReadReference<DataType> GetDataReadReferenceOrConstructWithVertexDefault(const FInputVertexInterface& InputVertices, const FString& InName, const FOperatorSettings& InOperatorSettings) const
+			template<typename DataType>
+			TDataReadReference<DataType> GetDataReadReferenceOrConstructWithVertexDefault(const FInputVertexInterface& InputVertices, const FString& InName, const FOperatorSettings& InSettings) const
 			{
-				using FDataRefType = TDataReadReference<DataType>;
+				using FDataFactory = TDataReadReferenceLiteralFactory<DataType>;
 
 				if (ContainsDataReadReference<DataType>(InName))
 				{
@@ -267,8 +238,14 @@ namespace Metasound
 				}
 				else
 				{
-					const DefaultDataType DefaultValue = InputVertices[InName].GetDefaultValue().Value.Get<DefaultDataType>();
-					return FDataRefType::CreateNew(InOperatorSettings, DefaultValue);
+					if (ensure(InputVertices.Contains(InName)))
+					{
+						return FDataFactory::CreateExplicitArgs(InSettings, InputVertices[InName].GetDefaultLiteral());
+					}
+					else
+					{
+						return FDataFactory::CreateExplicitArgs(InSettings, FLiteral::CreateInvalid());
+					}
 				}
 			}
 
diff --git a/Engine/Plugins/Runtime/Metasound/Source/MetasoundGraphCore/Public/MetasoundVertex.h b/Engine/Plugins/Runtime/Metasound/Source/MetasoundGraphCore/Public/MetasoundVertex.h
--- a/Engine/Plugins/Runtime/Metasound/Source/MetasoundGraphCore/Public/MetasoundVertex.h
+++ b/Engine/Plugins/Runtime/Metasound/Source/MetasoundGraphCore/Public/MetasoundVertex.h
@@ -87,21 +87,76 @@ namespace Metasound
 	 */
 	struct FInputDataVertexModel : public FDataVertexModel
 	{
+	private:
+		
+		// Factory for creating a literal. 
+		struct ILiteralFactory
+		{
+			virtual ~ILiteralFactory() = default;
+			virtual FLiteral CreateLiteral() const = 0;
+			virtual TUniquePtr<ILiteralFactory> Clone() const = 0;
+		};
+
+		// Create a literal with a copyable type.
+		template<typename LiteralValueType>
+		struct TLiteralFactory : ILiteralFactory
+		{
+			LiteralValueType LiteralValue;
+
+			TLiteralFactory(const LiteralValueType& InValue)
+			: LiteralValue(InValue)
+			{
+			}
+			
+			FLiteral CreateLiteral() const override
+			{
+				return FLiteral(LiteralValue);
+			}
+
+			TUniquePtr<ILiteralFactory> Clone() const override
+			{
+				return MakeUnique<TLiteralFactory<LiteralValueType>>(*this);
+			}
+		};
+
+	public:
+
 		FInputDataVertexModel(const FString& InVertexName, const FName& InDataTypeName, const FText& InDescription)
 			: FDataVertexModel(InVertexName, InDataTypeName, InDescription)
 		{
 		}
 
-		FInputDataVertexModel(const FString& InVertexName, const FName& InDataTypeName, const FText& InDescription, const FLiteral& InDefaultValue)
+		template<typename LiteralValueType>
+		FInputDataVertexModel(const FString& InVertexName, const FName& InDataTypeName, const FText& InDescription, const LiteralValueType& InLiteralValue)
 			: FDataVertexModel(InVertexName, InDataTypeName, InDescription)
-			, DefaultValue(InDefaultValue.Clone())
+			, LiteralFactory(MakeUnique<TLiteralFactory<LiteralValueType>>(InLiteralValue))
 		{
 		}
 
+		FInputDataVertexModel(const FInputDataVertexModel& InOther)
+		: FDataVertexModel(InOther)
+		{
+			if (InOther.LiteralFactory)
+			{
+				LiteralFactory = InOther.LiteralFactory->Clone();
+			}
+		}
+
 		/** Create a clone of this FInputDataVertexModel */
 		virtual TUniquePtr<FInputDataVertexModel> Clone() const = 0;
 
-		FLiteral DefaultValue;
+		/** Creates the default literal for this vertex. If a default does not exist,
+		 * then an invalid literal is returned. 
+		 */
+		FLiteral CreateDefaultLiteral() const 
+		{
+			if (LiteralFactory.IsValid())
+			{
+				return LiteralFactory->CreateLiteral();
+			}
+
+			return FLiteral::CreateInvalid();
+		}
 
 		friend bool METASOUNDGRAPHCORE_API operator==(const FInputDataVertex& InLHS, const FInputDataVertex& InRHS);
 		friend bool METASOUNDGRAPHCORE_API operator!=(const FInputDataVertex& InLHS, const FInputDataVertex& InRHS);
@@ -113,6 +168,8 @@ namespace Metasound
 		 * @return True if models are equal, false otherwise.
 		 */
 		virtual bool IsEqual(const FInputDataVertexModel& InOther) const = 0;
+
+		TUniquePtr<ILiteralFactory> LiteralFactory;
 	};
 
 	/** FOutputDataVertexModel
@@ -148,19 +205,9 @@ namespace Metasound
 		 * @InVertexName - Name of vertex.
 		 * @InDescription - Human readable vertex description.
 		 */
-		TBaseVertexModel(const FString& InVertexName, const FText& InDescription)
-		:	VertexModelType(InVertexName, GetMetasoundDataTypeName<DataType>(), InDescription)
-		{
-		}
-
-		/** TBaseVertexModel Constructor
-		 *
-		 * @InVertexName - Name of vertex.
-		 * @InDescription - Human readable vertex description.
-		 * @InDefaultValue - Default value of vertex.
-		 */
-		TBaseVertexModel(const FString& InVertexName, const FText& InDescription, const FLiteral& InDefaultValue)
-		:	VertexModelType(InVertexName, GetMetasoundDataTypeName<DataType>(), InDescription, InDefaultValue)
+		template<typename... ModelArgTypes>
+		TBaseVertexModel(const FString& InVertexName, const FText& InDescription, ModelArgTypes&&... ModelArgs)
+		:	VertexModelType(InVertexName, GetMetasoundDataTypeName<DataType>(), InDescription, Forward<ModelArgTypes>(ModelArgs)...)
 		{
 		}
 
@@ -197,6 +244,7 @@ namespace Metasound
 	template<typename DataType>
 	struct TInputDataVertexModel : public TBaseVertexModel<DataType, FInputDataVertexModel>
 	{
+
 		using TBaseVertexModel<DataType, FInputDataVertexModel>::TBaseVertexModel;
 
 		/** TInputDataVertexModel Constructor
@@ -215,11 +263,14 @@ namespace Metasound
 		 * @InDescription - Human readable vertex description.
 		 * @InDefaultValue - Default Value of vertex
 		 */
-		TInputDataVertexModel(const FString& InVertexName, const FText& InDescription, const FLiteral& InDefaultValue)
+		template<typename LiteralValueType>
+		TInputDataVertexModel(const FString& InVertexName, const FText& InDescription, const LiteralValueType& InDefaultValue)
 		:	TBaseVertexModel<DataType, FInputDataVertexModel>(InVertexName, InDescription, InDefaultValue)
 		{
 		}
 
+		TInputDataVertexModel(const TInputDataVertexModel& InOther) = default;
+
 		/** Return the vertex type name (not to be confused with the data type name).
 		 *
 		 * @return Name of this vertex type.
@@ -235,7 +286,7 @@ namespace Metasound
 		/** Create a clone of this VertexModelType */
 		virtual TUniquePtr<FInputDataVertexModel> Clone() const override
 		{
-			return MakeUnique<TInputDataVertexModel<DataType>>(this->VertexName, this->Description, this->DefaultValue);
+			return MakeUnique<TInputDataVertexModel<DataType>>(*this);
 		}
 	};
 
@@ -284,7 +335,7 @@ namespace Metasound
 			const FText& GetDescription() const;
 
 			/** Default value of the vertex. */
-			const FLiteral& GetDefaultValue() const;
+			FLiteral GetDefaultLiteral() const;
 
 			/** Determine if vertex refers to same data type. 
 			 *
diff --git a/Engine/Plugins/Runtime/Metasound/Source/MetasoundStandardNodes/Private/MetasoundADSRNode.cpp b/Engine/Plugins/Runtime/Metasound/Source/MetasoundStandardNodes/Private/MetasoundADSRNode.cpp
--- a/Engine/Plugins/Runtime/Metasound/Source/MetasoundStandardNodes/Private/MetasoundADSRNode.cpp
+++ b/Engine/Plugins/Runtime/Metasound/Source/MetasoundStandardNodes/Private/MetasoundADSRNode.cpp
@@ -286,10 +286,10 @@ namespace Metasound
 		// TODO: If none of these are connected, could pre-generate ADSR envelope and return a different operator.
 		FADSRReferences ADSRReferences =
 		{
-			InputCollection.GetDataReadReferenceOrConstructWithVertexDefault<FTime, float>(InputInterface, TEXT("Attack")),
-			InputCollection.GetDataReadReferenceOrConstructWithVertexDefault<FTime, float>(InputInterface, TEXT("Decay")),
-			InputCollection.GetDataReadReferenceOrConstructWithVertexDefault<float>(InputInterface, TEXT("Sustain Level")),
-			InputCollection.GetDataReadReferenceOrConstructWithVertexDefault<FTime, float>(InputInterface, TEXT("Release")),
+			InputCollection.GetDataReadReferenceOrConstructWithVertexDefault<FTime>(InputInterface, TEXT("Attack"), InParams.OperatorSettings),
+			InputCollection.GetDataReadReferenceOrConstructWithVertexDefault<FTime>(InputInterface, TEXT("Decay"), InParams.OperatorSettings),
+			InputCollection.GetDataReadReferenceOrConstructWithVertexDefault<float>(InputInterface, TEXT("Sustain Level"), InParams.OperatorSettings),
+			InputCollection.GetDataReadReferenceOrConstructWithVertexDefault<FTime>(InputInterface, TEXT("Release"), InParams.OperatorSettings),
 		};
 
 		return MakeUnique<FADSROperator>(InParams.OperatorSettings, TriggerAttack, TriggerRelease, ADSRReferences);
diff --git a/Engine/Plugins/Runtime/Metasound/Source/MetasoundStandardNodes/Private/MetasoundDelayNode.cpp b/Engine/Plugins/Runtime/Metasound/Source/MetasoundStandardNodes/Private/MetasoundDelayNode.cpp
--- a/Engine/Plugins/Runtime/Metasound/Source/MetasoundStandardNodes/Private/MetasoundDelayNode.cpp
+++ b/Engine/Plugins/Runtime/Metasound/Source/MetasoundStandardNodes/Private/MetasoundDelayNode.cpp
@@ -213,10 +213,10 @@ namespace Metasound
 		const FInputVertexInterface& InputInterface = GetVertexInterface().GetInputInterface();
 
 		FAudioBufferReadRef AudioIn = InputCollection.GetDataReadReferenceOrConstruct<FAudioBuffer>(Delay::InParamNameAudioInput, InParams.OperatorSettings);
-		FTimeReadRef DelayTime = InputCollection.GetDataReadReferenceOrConstructWithVertexDefault<FTime, float>(InputInterface, Delay::InParamNameDelayTime);
-		FFloatReadRef DryLevel = InputCollection.GetDataReadReferenceOrConstructWithVertexDefault<float>(InputInterface, Delay::InParamNameDryLevel);
-		FFloatReadRef WetLevel = InputCollection.GetDataReadReferenceOrConstructWithVertexDefault<float>(InputInterface, Delay::InParamNameWetLevel);
-		FFloatReadRef Feedback = InputCollection.GetDataReadReferenceOrConstructWithVertexDefault<float>(InputInterface, Delay::InParamNameFeedbackAmount);
+		FTimeReadRef DelayTime = InputCollection.GetDataReadReferenceOrConstructWithVertexDefault<FTime>(InputInterface, Delay::InParamNameDelayTime, InParams.OperatorSettings);
+		FFloatReadRef DryLevel = InputCollection.GetDataReadReferenceOrConstructWithVertexDefault<float>(InputInterface, Delay::InParamNameDryLevel, InParams.OperatorSettings);
+		FFloatReadRef WetLevel = InputCollection.GetDataReadReferenceOrConstructWithVertexDefault<float>(InputInterface, Delay::InParamNameWetLevel, InParams.OperatorSettings);
+		FFloatReadRef Feedback = InputCollection.GetDataReadReferenceOrConstructWithVertexDefault<float>(InputInterface, Delay::InParamNameFeedbackAmount, InParams.OperatorSettings);
 
 		return MakeUnique<FDelayOperator>(InParams.OperatorSettings, AudioIn, DelayTime, DryLevel, WetLevel, Feedback);
 	}
diff --git a/Engine/Plugins/Runtime/Metasound/Source/MetasoundStandardNodes/Private/MetasoundEnvelopeFollowerNode.cpp b/Engine/Plugins/Runtime/Metasound/Source/MetasoundStandardNodes/Private/MetasoundEnvelopeFollowerNode.cpp
--- a/Engine/Plugins/Runtime/Metasound/Source/MetasoundStandardNodes/Private/MetasoundEnvelopeFollowerNode.cpp
+++ b/Engine/Plugins/Runtime/Metasound/Source/MetasoundStandardNodes/Private/MetasoundEnvelopeFollowerNode.cpp
@@ -167,8 +167,8 @@ namespace Metasound
 		const FInputVertexInterface& InputInterface = GetVertexInterface().GetInputInterface();
 
 		FAudioBufferReadRef AudioIn = InputCollection.GetDataReadReferenceOrConstruct<FAudioBuffer>(EnvelopeFollower::InParamNameAudioInput, InParams.OperatorSettings);
-		FTimeReadRef AttackTime = InputCollection.GetDataReadReferenceOrConstructWithVertexDefault<FTime, float>(InputInterface, EnvelopeFollower::InParamNameAttackTime);
-		FTimeReadRef ReleaseTime = InputCollection.GetDataReadReferenceOrConstructWithVertexDefault<FTime, float>(InputInterface, EnvelopeFollower::InParamNameReleaseTime);
+		FTimeReadRef AttackTime = InputCollection.GetDataReadReferenceOrConstructWithVertexDefault<FTime>(InputInterface, EnvelopeFollower::InParamNameAttackTime, InParams.OperatorSettings);
+		FTimeReadRef ReleaseTime = InputCollection.GetDataReadReferenceOrConstructWithVertexDefault<FTime>(InputInterface, EnvelopeFollower::InParamNameReleaseTime, InParams.OperatorSettings);
 
 		return MakeUnique<FEnvelopeFollowerOperator>(InParams.OperatorSettings, AudioIn, AttackTime, ReleaseTime);
 	}
diff --git a/Engine/Plugins/Runtime/Metasound/Source/MetasoundStandardNodes/Private/MetasoundITDPannerNode.cpp b/Engine/Plugins/Runtime/Metasound/Source/MetasoundStandardNodes/Private/MetasoundITDPannerNode.cpp
--- a/Engine/Plugins/Runtime/Metasound/Source/MetasoundStandardNodes/Private/MetasoundITDPannerNode.cpp
+++ b/Engine/Plugins/Runtime/Metasound/Source/MetasoundStandardNodes/Private/MetasoundITDPannerNode.cpp
@@ -335,9 +335,9 @@ namespace Metasound
 		using namespace ITDPannerVertexNames;
 
 		FAudioBufferReadRef AudioIn = InputCollection.GetDataReadReferenceOrConstruct<FAudioBuffer>(GetInputAudioName(), InParams.OperatorSettings);
-		FFloatReadRef PanningAngle = InputCollection.GetDataReadReferenceOrConstructWithVertexDefault<float>(InputInterface, GetInputPanAngleName());
-		FFloatReadRef DistanceFactor = InputCollection.GetDataReadReferenceOrConstructWithVertexDefault<float>(InputInterface, GetInputDistanceFactorName());
-		FFloatReadRef HeadWidth = InputCollection.GetDataReadReferenceOrConstructWithVertexDefault<float>(InputInterface, GetInputHeadWidthName());
+		FFloatReadRef PanningAngle = InputCollection.GetDataReadReferenceOrConstructWithVertexDefault<float>(InputInterface, GetInputPanAngleName(), InParams.OperatorSettings);
+		FFloatReadRef DistanceFactor = InputCollection.GetDataReadReferenceOrConstructWithVertexDefault<float>(InputInterface, GetInputDistanceFactorName(), InParams.OperatorSettings);
+		FFloatReadRef HeadWidth = InputCollection.GetDataReadReferenceOrConstructWithVertexDefault<float>(InputInterface, GetInputHeadWidthName(), InParams.OperatorSettings);
 
 		return MakeUnique<FITDPannerOperator>(InParams.OperatorSettings, AudioIn, PanningAngle, DistanceFactor, HeadWidth);
 	}
diff --git a/Engine/Plugins/Runtime/Metasound/Source/MetasoundStandardNodes/Private/MetasoundInterpToNode.cpp b/Engine/Plugins/Runtime/Metasound/Source/MetasoundStandardNodes/Private/MetasoundInterpToNode.cpp
--- a/Engine/Plugins/Runtime/Metasound/Source/MetasoundStandardNodes/Private/MetasoundInterpToNode.cpp
+++ b/Engine/Plugins/Runtime/Metasound/Source/MetasoundStandardNodes/Private/MetasoundInterpToNode.cpp
@@ -150,8 +150,8 @@ namespace Metasound
 		const FDataReferenceCollection& InputCollection = InParams.InputDataReferences;
 		const FInputVertexInterface& InputInterface = GetVertexInterface().GetInputInterface();
 
-		FFloatReadRef TargetValue = InputCollection.GetDataReadReferenceOrConstructWithVertexDefault<float>(InputInterface, InterpTo::InParamNameTarget);
-		FTimeReadRef InterpTime = InputCollection.GetDataReadReferenceOrConstructWithVertexDefault<FTime, float>(InputInterface, InterpTo::InParamNameInterpTime);
+		FFloatReadRef TargetValue = InputCollection.GetDataReadReferenceOrConstructWithVertexDefault<float>(InputInterface, InterpTo::InParamNameTarget, InParams.OperatorSettings);
+		FTimeReadRef InterpTime = InputCollection.GetDataReadReferenceOrConstructWithVertexDefault<FTime>(InputInterface, InterpTo::InParamNameInterpTime, InParams.OperatorSettings);
 
 		return MakeUnique<FInterpToOperator>(InParams.OperatorSettings, TargetValue, InterpTime);
 	}
diff --git a/Engine/Plugins/Runtime/Metasound/Source/MetasoundStandardNodes/Private/MetasoundMIDIToFreqNode.cpp b/Engine/Plugins/Runtime/Metasound/Source/MetasoundStandardNodes/Private/MetasoundMIDIToFreqNode.cpp
--- a/Engine/Plugins/Runtime/Metasound/Source/MetasoundStandardNodes/Private/MetasoundMIDIToFreqNode.cpp
+++ b/Engine/Plugins/Runtime/Metasound/Source/MetasoundStandardNodes/Private/MetasoundMIDIToFreqNode.cpp
@@ -148,7 +148,7 @@ namespace Metasound
 		const FDataReferenceCollection& InputCollection = InParams.InputDataReferences;
 		const FInputVertexInterface& InputInterface = GetVertexInterface().GetInputInterface();
 
-		FFloatReadRef InMIDINote = InputCollection.GetDataReadReferenceOrConstructWithVertexDefault<float>(InputInterface, GetInputMIDIName());
+		FFloatReadRef InMIDINote = InputCollection.GetDataReadReferenceOrConstructWithVertexDefault<float>(InputInterface, GetInputMIDIName(), InParams.OperatorSettings);
 
 		return MakeUnique<FMIDIToFreqOperator>(InParams.OperatorSettings, InMIDINote);
 	}
diff --git a/Engine/Plugins/Runtime/Metasound/Source/MetasoundStandardNodes/Private/MetasoundMixerNode.cpp b/Engine/Plugins/Runtime/Metasound/Source/MetasoundStandardNodes/Private/MetasoundMixerNode.cpp
--- a/Engine/Plugins/Runtime/Metasound/Source/MetasoundStandardNodes/Private/MetasoundMixerNode.cpp
+++ b/Engine/Plugins/Runtime/Metasound/Source/MetasoundStandardNodes/Private/MetasoundMixerNode.cpp
@@ -162,7 +162,7 @@ namespace Metasound
 				InputBuffers.Emplace(InputCol.GetDataReadReference<FAudioBuffer>(AudioPinName));
 				
 				// Make sure for every valid/connected pin, we have a corresponding gain, even if its defaulted.
-				InputGains.Emplace(InputCol.GetDataReadReferenceOrConstructWithVertexDefault<float>(InputInterface, GainPinName));
+				InputGains.Emplace(InputCol.GetDataReadReferenceOrConstructWithVertexDefault<float>(InputInterface, GainPinName, InParams.OperatorSettings));
 			}
 		}
 		
diff --git a/Engine/Plugins/Runtime/Metasound/Source/MetasoundStandardNodes/Private/MetasoundOscNode.cpp b/Engine/Plugins/Runtime/Metasound/Source/MetasoundStandardNodes/Private/MetasoundOscNode.cpp
--- a/Engine/Plugins/Runtime/Metasound/Source/MetasoundStandardNodes/Private/MetasoundOscNode.cpp
+++ b/Engine/Plugins/Runtime/Metasound/Source/MetasoundStandardNodes/Private/MetasoundOscNode.cpp
@@ -168,8 +168,8 @@ namespace Metasound
 		const FOperatorSettings& Settings = InParams.OperatorSettings;
 		const FInputVertexInterface& InputInterface = GetVertexInterface().GetInputInterface();
 
-		FFloatReadRef Frequency = InputCol.GetDataReadReferenceOrConstructWithVertexDefault<float>(InputInterface, FrequencyPinName);
-		FBoolReadRef Enabled = InputCol.GetDataReadReferenceOrConstructWithVertexDefault<bool>(InputInterface, EnabledPinName);
+		FFloatReadRef Frequency = InputCol.GetDataReadReferenceOrConstructWithVertexDefault<float>(InputInterface, FrequencyPinName, InParams.OperatorSettings);
+		FBoolReadRef Enabled = InputCol.GetDataReadReferenceOrConstructWithVertexDefault<bool>(InputInterface, EnabledPinName, InParams.OperatorSettings);
 
 		return MakeUnique<FOscOperator>(InParams.OperatorSettings, Frequency, Enabled);
 	}
diff --git a/Engine/Plugins/Runtime/Metasound/Source/MetasoundStandardNodes/Private/MetasoundOscillatorNodes.cpp b/Engine/Plugins/Runtime/Metasound/Source/MetasoundStandardNodes/Private/MetasoundOscillatorNodes.cpp
--- a/Engine/Plugins/Runtime/Metasound/Source/MetasoundStandardNodes/Private/MetasoundOscillatorNodes.cpp
+++ b/Engine/Plugins/Runtime/Metasound/Source/MetasoundStandardNodes/Private/MetasoundOscillatorNodes.cpp
@@ -1202,12 +1202,12 @@ namespace Metasound
 
 			return MakeUnique<FLfoOperator>(
 				  Settings
-				, InputCol.GetDataReadReferenceOrConstructWithVertexDefault<float>(InputInterface, FOscilatorFactoryBase::BaseFrequencyPinName)
+				, InputCol.GetDataReadReferenceOrConstructWithVertexDefault<float>(InputInterface, FOscilatorFactoryBase::BaseFrequencyPinName, Settings)
 				, InputCol.GetDataReadReferenceOrConstruct<FEnumLfoWaveshapeType>(WaveshapePinName)
 				, InputCol.GetDataReadReferenceOrConstruct<FTrigger>(FOscilatorFactoryBase::PhaseResetPinName, Settings)
-				, InputCol.GetDataReadReferenceOrConstructWithVertexDefault<float>(InputInterface,FOscilatorFactoryBase::PhaseOffsetPinName)
-				, InputCol.GetDataReadReferenceOrConstructWithVertexDefault<bool>(InputInterface,FOscilatorFactoryBase::BiPolarPinName)
-				, InputCol.GetDataReadReferenceOrConstructWithVertexDefault<float>(InputInterface,FSquareOscilatorNode::FFactory::PulseWidthPinName)
+				, InputCol.GetDataReadReferenceOrConstructWithVertexDefault<float>(InputInterface,FOscilatorFactoryBase::PhaseOffsetPinName, Settings)
+				, InputCol.GetDataReadReferenceOrConstructWithVertexDefault<bool>(InputInterface,FOscilatorFactoryBase::BiPolarPinName, Settings)
+				, InputCol.GetDataReadReferenceOrConstructWithVertexDefault<float>(InputInterface,FSquareOscilatorNode::FFactory::PulseWidthPinName, Settings)
 			);
 		}
 
diff --git a/Engine/Plugins/Runtime/Metasound/Source/MetasoundStandardNodes/Private/MetasoundStereoDelayNode.cpp b/Engine/Plugins/Runtime/Metasound/Source/MetasoundStandardNodes/Private/MetasoundStereoDelayNode.cpp
--- a/Engine/Plugins/Runtime/Metasound/Source/MetasoundStandardNodes/Private/MetasoundStereoDelayNode.cpp
+++ b/Engine/Plugins/Runtime/Metasound/Source/MetasoundStandardNodes/Private/MetasoundStereoDelayNode.cpp
@@ -332,11 +332,11 @@ namespace Metasound
 		FAudioBufferReadRef LeftAudioIn = InputCollection.GetDataReadReferenceOrConstruct<FAudioBuffer>(StereoDelay::InParamNameAudioInputLeft, InParams.OperatorSettings);
 		FAudioBufferReadRef RightAudioIn = InputCollection.GetDataReadReferenceOrConstruct<FAudioBuffer>(StereoDelay::InParamNameAudioInputRight, InParams.OperatorSettings);
 		FStereoDelayModeReadRef StereoDelayMode = InputCollection.GetDataReadReferenceOrConstruct<FEnumStereoDelayMode>(StereoDelay::InParamNameDelayMode);
-		FTimeReadRef DelayTime = InputCollection.GetDataReadReferenceOrConstructWithVertexDefault<FTime, float>(InputInterface, StereoDelay::InParamNameDelayTime);
-		FFloatReadRef DelayRatio = InputCollection.GetDataReadReferenceOrConstructWithVertexDefault<float>(InputInterface, StereoDelay::InParamNameDelayRatio);
-		FFloatReadRef DryLevel = InputCollection.GetDataReadReferenceOrConstructWithVertexDefault<float>(InputInterface, StereoDelay::InParamNameDryLevel);
-		FFloatReadRef WetLevel = InputCollection.GetDataReadReferenceOrConstructWithVertexDefault<float>(InputInterface, StereoDelay::InParamNameWetLevel);
-		FFloatReadRef Feedback = InputCollection.GetDataReadReferenceOrConstructWithVertexDefault<float>(InputInterface, StereoDelay::InParamNameFeedbackAmount);
+		FTimeReadRef DelayTime = InputCollection.GetDataReadReferenceOrConstructWithVertexDefault<FTime>(InputInterface, StereoDelay::InParamNameDelayTime, InParams.OperatorSettings);
+		FFloatReadRef DelayRatio = InputCollection.GetDataReadReferenceOrConstructWithVertexDefault<float>(InputInterface, StereoDelay::InParamNameDelayRatio, InParams.OperatorSettings);
+		FFloatReadRef DryLevel = InputCollection.GetDataReadReferenceOrConstructWithVertexDefault<float>(InputInterface, StereoDelay::InParamNameDryLevel, InParams.OperatorSettings);
+		FFloatReadRef WetLevel = InputCollection.GetDataReadReferenceOrConstructWithVertexDefault<float>(InputInterface, StereoDelay::InParamNameWetLevel, InParams.OperatorSettings);
+		FFloatReadRef Feedback = InputCollection.GetDataReadReferenceOrConstructWithVertexDefault<float>(InputInterface, StereoDelay::InParamNameFeedbackAmount, InParams.OperatorSettings);
 
 		return MakeUnique<FStereoDelayOperator>(InParams.OperatorSettings, LeftAudioIn, RightAudioIn, StereoDelayMode, DelayTime, DelayRatio, DryLevel, WetLevel, Feedback);
 	}
diff --git a/Engine/Plugins/Runtime/Metasound/Source/MetasoundStandardNodes/Private/MetasoundStereoPannerNode.cpp b/Engine/Plugins/Runtime/Metasound/Source/MetasoundStandardNodes/Private/MetasoundStereoPannerNode.cpp
--- a/Engine/Plugins/Runtime/Metasound/Source/MetasoundStandardNodes/Private/MetasoundStereoPannerNode.cpp
+++ b/Engine/Plugins/Runtime/Metasound/Source/MetasoundStandardNodes/Private/MetasoundStereoPannerNode.cpp
@@ -276,7 +276,7 @@ namespace Metasound
 		using namespace StereoPannerVertexNames;
 
 		FAudioBufferReadRef AudioIn = InputCollection.GetDataReadReferenceOrConstruct<FAudioBuffer>(GetInputAudioName(), InParams.OperatorSettings);
-		FFloatReadRef PanningAmount = InputCollection.GetDataReadReferenceOrConstructWithVertexDefault<float>(InputInterface, GetInputPanAmountName());
+		FFloatReadRef PanningAmount = InputCollection.GetDataReadReferenceOrConstructWithVertexDefault<float>(InputInterface, GetInputPanAmountName(), InParams.OperatorSettings);
 		FPanningLawReadRef PanningLaw = InputCollection.GetDataReadReferenceOrConstruct<FEnumPanningLaw>(GetInputPanningLawName());
 
 		return MakeUnique<FStereoPannerOperator>(InParams.OperatorSettings, AudioIn, PanningAmount, PanningLaw);
diff --git a/Engine/Plugins/Runtime/Metasound/Source/MetasoundStandardNodes/Private/MetasoundTriggerCompareNode.cpp b/Engine/Plugins/Runtime/Metasound/Source/MetasoundStandardNodes/Private/MetasoundTriggerCompareNode.cpp
--- a/Engine/Plugins/Runtime/Metasound/Source/MetasoundStandardNodes/Private/MetasoundTriggerCompareNode.cpp
+++ b/Engine/Plugins/Runtime/Metasound/Source/MetasoundStandardNodes/Private/MetasoundTriggerCompareNode.cpp
@@ -160,8 +160,8 @@ namespace Metasound
 		const FInputVertexInterface& InputInterface = GetVertexInterface().GetInputInterface();
 
 		FTriggerReadRef InOnTriggerCompare = InputCollection.GetDataReadReferenceOrConstruct<FTrigger>(TriggerCompare::InParamNameCompare, InParams.OperatorSettings);
-		FInt32ReadRef InValueA = InputCollection.GetDataReadReferenceOrConstructWithVertexDefault<int32>(InputInterface, TriggerCompare::InParamNameA);
-		FInt32ReadRef InValueB = InputCollection.GetDataReadReferenceOrConstructWithVertexDefault<int32>(InputInterface, TriggerCompare::InParamNameB);
+		FInt32ReadRef InValueA = InputCollection.GetDataReadReferenceOrConstructWithVertexDefault<int32>(InputInterface, TriggerCompare::InParamNameA, InParams.OperatorSettings);
+		FInt32ReadRef InValueB = InputCollection.GetDataReadReferenceOrConstructWithVertexDefault<int32>(InputInterface, TriggerCompare::InParamNameB, InParams.OperatorSettings);
 		FEnumTriggerComparisonTypeReadRef InComparison = InputCollection.GetDataReadReferenceOrConstruct<FEnumTriggerComparisonType>(TriggerCompare::InParamCompareType);
 
 		return MakeUnique<FTriggerCompareOperator>(InParams.OperatorSettings, InOnTriggerCompare, InValueA, InValueB, InComparison);
diff --git a/Engine/Plugins/Runtime/Metasound/Source/MetasoundStandardNodes/Private/MetasoundTriggerDelayNode.cpp b/Engine/Plugins/Runtime/Metasound/Source/MetasoundStandardNodes/Private/MetasoundTriggerDelayNode.cpp
--- a/Engine/Plugins/Runtime/Metasound/Source/MetasoundStandardNodes/Private/MetasoundTriggerDelayNode.cpp
+++ b/Engine/Plugins/Runtime/Metasound/Source/MetasoundStandardNodes/Private/MetasoundTriggerDelayNode.cpp
@@ -132,7 +132,7 @@ namespace Metasound
 		const FTriggerDelayNode& DelayNode = static_cast<const FTriggerDelayNode&>(InParams.Node);
 
 		const FInputVertexInterface& InputInterface = GetVertexInterface().GetInputInterface();
-		FTimeReadRef Delay = InParams.InputDataReferences.GetDataReadReferenceOrConstructWithVertexDefault<FTime, float>(InputInterface, TEXT("Delay"));
+		FTimeReadRef Delay = InParams.InputDataReferences.GetDataReadReferenceOrConstructWithVertexDefault<FTime>(InputInterface, TEXT("Delay"), InParams.OperatorSettings);
 
 		FTriggerReadRef TriggerIn = InParams.InputDataReferences.GetDataReadReferenceOrConstruct<FTrigger>(TEXT("In"), InParams.OperatorSettings);
 		FTriggerReadRef TriggerReset = InParams.InputDataReferences.GetDataReadReferenceOrConstruct<FTrigger>(TEXT("Reset"), InParams.OperatorSettings);
diff --git a/Engine/Plugins/Runtime/Metasound/Source/MetasoundStandardNodes/Private/MetasoundTriggerOnThresholdNode.cpp b/Engine/Plugins/Runtime/Metasound/Source/MetasoundStandardNodes/Private/MetasoundTriggerOnThresholdNode.cpp
--- a/Engine/Plugins/Runtime/Metasound/Source/MetasoundStandardNodes/Private/MetasoundTriggerOnThresholdNode.cpp
+++ b/Engine/Plugins/Runtime/Metasound/Source/MetasoundStandardNodes/Private/MetasoundTriggerOnThresholdNode.cpp
@@ -236,7 +236,7 @@ namespace Metasound
 		FBufferTriggerTypeReadRef Type = InputCol.GetDataReadReferenceOrConstruct<FEnumBufferTriggerType>(TriggerType);
 
 		FAudioBufferReadRef InputBuffer = InputCol.GetDataReadReferenceOrConstruct<FAudioBuffer>(InPinName, Settings);
-		FFloatReadRef Threshold = InputCol.GetDataReadReferenceOrConstructWithVertexDefault<float>(InputInterface, ThresholdPinName);
+		FFloatReadRef Threshold = InputCol.GetDataReadReferenceOrConstructWithVertexDefault<float>(InputInterface, ThresholdPinName, Settings);
 
 		switch (*Type)
 		{
diff --git a/Engine/Plugins/Runtime/Metasound/Source/MetasoundStandardNodes/Private/MetasoundTriggerPipeNode.cpp b/Engine/Plugins/Runtime/Metasound/Source/MetasoundStandardNodes/Private/MetasoundTriggerPipeNode.cpp
--- a/Engine/Plugins/Runtime/Metasound/Source/MetasoundStandardNodes/Private/MetasoundTriggerPipeNode.cpp
+++ b/Engine/Plugins/Runtime/Metasound/Source/MetasoundStandardNodes/Private/MetasoundTriggerPipeNode.cpp
@@ -132,7 +132,7 @@ namespace Metasound
 		const FDataReferenceCollection& InputCollection = InParams.InputDataReferences;
 		const FInputVertexInterface& InputInterface = GetVertexInterface().GetInputInterface();
 
-		FTimeReadRef Delay = InputCollection.GetDataReadReferenceOrConstructWithVertexDefault<FTime, float>(InputInterface, TEXT("Delay"));
+		FTimeReadRef Delay = InputCollection.GetDataReadReferenceOrConstructWithVertexDefault<FTime>(InputInterface, TEXT("Delay"), InParams.OperatorSettings);
 		FTriggerReadRef TriggerIn = InParams.InputDataReferences.GetDataReadReferenceOrConstruct<FTrigger>(TEXT("In"), InParams.OperatorSettings);
 		FTriggerReadRef TriggerReset = InParams.InputDataReferences.GetDataReadReferenceOrConstruct<FTrigger>(TEXT("Reset"), InParams.OperatorSettings);
 
diff --git a/Engine/Plugins/Runtime/Metasound/Source/MetasoundStandardNodes/Private/MetasoundTriggerRepeatNode.cpp b/Engine/Plugins/Runtime/Metasound/Source/MetasoundStandardNodes/Private/MetasoundTriggerRepeatNode.cpp
--- a/Engine/Plugins/Runtime/Metasound/Source/MetasoundStandardNodes/Private/MetasoundTriggerRepeatNode.cpp
+++ b/Engine/Plugins/Runtime/Metasound/Source/MetasoundStandardNodes/Private/MetasoundTriggerRepeatNode.cpp
@@ -183,7 +183,7 @@ namespace Metasound
 
 		FTriggerReadRef TriggerEnable = InParams.InputDataReferences.GetDataReadReferenceOrConstruct<FTrigger>(GetInputStartName(), InParams.OperatorSettings);
 		FTriggerReadRef TriggerDisable = InParams.InputDataReferences.GetDataReadReferenceOrConstruct<FTrigger>(GetInputStopName(), InParams.OperatorSettings);
-		FTimeReadRef Period = InParams.InputDataReferences.GetDataReadReferenceOrConstructWithVertexDefault<FTime, float>(InputInterface, GetInputPeriodName());
+		FTimeReadRef Period = InParams.InputDataReferences.GetDataReadReferenceOrConstructWithVertexDefault<FTime>(InputInterface, GetInputPeriodName(), InParams.OperatorSettings);
 
 		return MakeUnique<FTriggerRepeatOperator>(InParams.OperatorSettings, TriggerEnable, TriggerDisable, Period);
 	}
diff --git a/Engine/Plugins/Runtime/Metasound/Source/MetasoundStandardNodes/Private/MetasoundTriggerToggleNode.cpp b/Engine/Plugins/Runtime/Metasound/Source/MetasoundStandardNodes/Private/MetasoundTriggerToggleNode.cpp
--- a/Engine/Plugins/Runtime/Metasound/Source/MetasoundStandardNodes/Private/MetasoundTriggerToggleNode.cpp
+++ b/Engine/Plugins/Runtime/Metasound/Source/MetasoundStandardNodes/Private/MetasoundTriggerToggleNode.cpp
@@ -116,7 +116,7 @@ namespace Metasound
 		const FInputVertexInterface& InputInterface = GetVertexInterface().GetInputInterface();
 		FTriggerReadRef InTriggerOn = InParams.InputDataReferences.GetDataReadReferenceOrConstruct<FTrigger>(GetInputOnTriggerName(), InParams.OperatorSettings);
 		FTriggerReadRef InTriggerOff = InParams.InputDataReferences.GetDataReadReferenceOrConstruct<FTrigger>(GetInputOffTriggerName(), InParams.OperatorSettings);
-		FBoolReadRef InInitValue = InParams.InputDataReferences.GetDataReadReferenceOrConstructWithVertexDefault<bool>(InputInterface, GetInputInitName());
+		FBoolReadRef InInitValue = InParams.InputDataReferences.GetDataReadReferenceOrConstructWithVertexDefault<bool>(InputInterface, GetInputInitName(), InParams.OperatorSettings);
 
 		return MakeUnique<FTriggerToggleOperator>(InParams.OperatorSettings, InTriggerOn, InTriggerOff, InInitValue);
 	}
diff --git a/Engine/Plugins/Runtime/Metasound/Source/MetasoundStandardNodes/Private/MetasoundWaveShaperNode.cpp b/Engine/Plugins/Runtime/Metasound/Source/MetasoundStandardNodes/Private/MetasoundWaveShaperNode.cpp
--- a/Engine/Plugins/Runtime/Metasound/Source/MetasoundStandardNodes/Private/MetasoundWaveShaperNode.cpp
+++ b/Engine/Plugins/Runtime/Metasound/Source/MetasoundStandardNodes/Private/MetasoundWaveShaperNode.cpp
@@ -176,9 +176,9 @@ namespace Metasound
 		const FInputVertexInterface& InputInterface = GetVertexInterface().GetInputInterface();
 
 		FAudioBufferReadRef AudioIn = InputCollection.GetDataReadReferenceOrConstruct<FAudioBuffer>(InParamNameAudioInput, InParams.OperatorSettings);
-		FFloatReadRef InAmount = InputCollection.GetDataReadReferenceOrConstructWithVertexDefault<float>(InputInterface, InParamNameWaveShapeAmount);
-		FFloatReadRef InBias = InputCollection.GetDataReadReferenceOrConstructWithVertexDefault<float>(InputInterface, InParamNameWaveShapeBias);
-		FFloatReadRef OutGain = InputCollection.GetDataReadReferenceOrConstructWithVertexDefault<float>(InputInterface, InParamNameOutputGain);
+		FFloatReadRef InAmount = InputCollection.GetDataReadReferenceOrConstructWithVertexDefault<float>(InputInterface, InParamNameWaveShapeAmount, InParams.OperatorSettings);
+		FFloatReadRef InBias = InputCollection.GetDataReadReferenceOrConstructWithVertexDefault<float>(InputInterface, InParamNameWaveShapeBias, InParams.OperatorSettings);
+		FFloatReadRef OutGain = InputCollection.GetDataReadReferenceOrConstructWithVertexDefault<float>(InputInterface, InParamNameOutputGain, InParams.OperatorSettings);
 		FEnumWaveShaperReadRef InType = InputCollection.GetDataReadReferenceOrConstruct<FEnumEWaveShaperType>(InParamNameType);
 
 		return MakeUnique<FWaveShaperOperator>(InParams.OperatorSettings, AudioIn, InAmount, InBias, OutGain, InType);
diff --git a/Engine/Plugins/Runtime/Metasound/Source/MetasoundStandardNodes/Public/MetasoundDebugLogNode.h b/Engine/Plugins/Runtime/Metasound/Source/MetasoundStandardNodes/Public/MetasoundDebugLogNode.h
--- a/Engine/Plugins/Runtime/Metasound/Source/MetasoundStandardNodes/Public/MetasoundDebugLogNode.h
+++ b/Engine/Plugins/Runtime/Metasound/Source/MetasoundStandardNodes/Public/MetasoundDebugLogNode.h
@@ -74,8 +74,8 @@ namespace Metasound
 				const FDataReferenceCollection& InputCollection = InParams.InputDataReferences;
 
 				FTriggerReadRef Trigger = InputCollection.GetDataReadReferenceOrConstruct<FTrigger>(DebugLogVertexNames::GetInputTriggerName(), InParams.OperatorSettings);
-				TDataReadReference<FString> Label = InputCollection.GetDataReadReferenceOrConstructWithVertexDefault<FString>(InputInterface, DebugLogVertexNames::GetLabelDebugLogName());
-				TDataReadReference<DebugLogType> ValueToLog = InputCollection.GetDataReadReferenceOrConstructWithVertexDefault<DebugLogType>(InputInterface, DebugLogVertexNames::GetToLogDebugLogName());
+				TDataReadReference<FString> Label = InputCollection.GetDataReadReferenceOrConstructWithVertexDefault<FString>(InputInterface, DebugLogVertexNames::GetLabelDebugLogName(), InParams.OperatorSettings);
+				TDataReadReference<DebugLogType> ValueToLog = InputCollection.GetDataReadReferenceOrConstructWithVertexDefault<DebugLogType>(InputInterface, DebugLogVertexNames::GetToLogDebugLogName(), InParams.OperatorSettings);
 
 				return MakeUnique<TDebugLogOperator<DebugLogType>>(InParams.OperatorSettings, Trigger, Label, ValueToLog);
 			}
@@ -141,4 +141,4 @@ namespace Metasound
 
 			virtual ~TDebugLogNode() = default;
 	};
-} // namespace Metasound
\ No newline at end of file
+} // namespace Metasound
diff --git a/Engine/Plugins/Runtime/Metasound/Source/MetasoundStandardNodes/Public/MetasoundMapRangeNode.h b/Engine/Plugins/Runtime/Metasound/Source/MetasoundStandardNodes/Public/MetasoundMapRangeNode.h
--- a/Engine/Plugins/Runtime/Metasound/Source/MetasoundStandardNodes/Public/MetasoundMapRangeNode.h
+++ b/Engine/Plugins/Runtime/Metasound/Source/MetasoundStandardNodes/Public/MetasoundMapRangeNode.h
@@ -125,12 +125,12 @@ namespace Metasound
 			const FInputVertexInterface& InputInterface = InParams.Node.GetVertexInterface().GetInputInterface();
 			const FDataReferenceCollection& InputCollection = InParams.InputDataReferences;
 
-			TDataReadReference<ValueType> InputValue = InputCollection.GetDataReadReferenceOrConstructWithVertexDefault<ValueType>(InputInterface, GetInputValueName());
-			TDataReadReference<ValueType> InputInRangeA = InputCollection.GetDataReadReferenceOrConstructWithVertexDefault<ValueType>(InputInterface, GetInputInRangeAName());
-			TDataReadReference<ValueType> InputInRangeB = InputCollection.GetDataReadReferenceOrConstructWithVertexDefault<ValueType>(InputInterface, GetInputInRangeBName());
-			TDataReadReference<ValueType> InputOutRangeA = InputCollection.GetDataReadReferenceOrConstructWithVertexDefault<ValueType>(InputInterface, GetInputOutRangeAName());
-			TDataReadReference<ValueType> InputOutRangeB = InputCollection.GetDataReadReferenceOrConstructWithVertexDefault<ValueType>(InputInterface, GetInputOutRangeBName());
-			FBoolReadRef bInputClamped = InputCollection.GetDataReadReferenceOrConstructWithVertexDefault<bool>(InputInterface, GetInputClampedName());
+			TDataReadReference<ValueType> InputValue = InputCollection.GetDataReadReferenceOrConstructWithVertexDefault<ValueType>(InputInterface, GetInputValueName(), InParams.OperatorSettings);
+			TDataReadReference<ValueType> InputInRangeA = InputCollection.GetDataReadReferenceOrConstructWithVertexDefault<ValueType>(InputInterface, GetInputInRangeAName(), InParams.OperatorSettings);
+			TDataReadReference<ValueType> InputInRangeB = InputCollection.GetDataReadReferenceOrConstructWithVertexDefault<ValueType>(InputInterface, GetInputInRangeBName(), InParams.OperatorSettings);
+			TDataReadReference<ValueType> InputOutRangeA = InputCollection.GetDataReadReferenceOrConstructWithVertexDefault<ValueType>(InputInterface, GetInputOutRangeAName(), InParams.OperatorSettings);
+			TDataReadReference<ValueType> InputOutRangeB = InputCollection.GetDataReadReferenceOrConstructWithVertexDefault<ValueType>(InputInterface, GetInputOutRangeBName(), InParams.OperatorSettings);
+			FBoolReadRef bInputClamped = InputCollection.GetDataReadReferenceOrConstructWithVertexDefault<bool>(InputInterface, GetInputClampedName(), InParams.OperatorSettings);
 
 			return MakeUnique<TMapRangeOperator<ValueType>>(InParams.OperatorSettings, InputValue, InputInRangeA, InputInRangeB, InputOutRangeA, InputOutRangeB, bInputClamped);
 		}
diff --git a/Engine/Plugins/Runtime/Metasound/Source/MetasoundStandardNodes/Public/MetasoundMathNodes.h b/Engine/Plugins/Runtime/Metasound/Source/MetasoundStandardNodes/Public/MetasoundMathNodes.h
--- a/Engine/Plugins/Runtime/Metasound/Source/MetasoundStandardNodes/Public/MetasoundMathNodes.h
+++ b/Engine/Plugins/Runtime/Metasound/Source/MetasoundStandardNodes/Public/MetasoundMathNodes.h
@@ -202,7 +202,7 @@ namespace Metasound
 				const FString& PrimaryOperandName = MathOpNames::PrimaryOperandName;
 				if (InputInterface.Contains(PrimaryOperandName))
 				{
-					DefaultValue = InputInterface[PrimaryOperandName].GetDefaultValue().Clone();
+					DefaultValue = InputInterface[PrimaryOperandName].GetDefaultLiteral();
 				}
 				TDataClassReadRef PrimaryOperand = InParams.InputDataReferences.GetDataReadReferenceOrConstruct<TDataClass>(MathOpNames::PrimaryOperandName, TMathOpClass::GetDefault(InParams.OperatorSettings, DefaultValue));
 
@@ -210,7 +210,7 @@ namespace Metasound
 				const FString& OpName = MathOpNames::AdditionalOperandsName;
 				if (InputInterface.Contains(OpName))
 				{
-					DefaultValue = InputInterface[OpName].GetDefaultValue().Clone();
+					DefaultValue = InputInterface[OpName].GetDefaultLiteral();
 				}
 				TOperandDataClassReadRef Op1 = InParams.InputDataReferences.GetDataReadReferenceOrConstruct<TOperandDataClass>(OpName, TMathOpClass::GetDefaultOp(InParams.OperatorSettings, DefaultValue));
 				TArray<TOperandDataClassReadRef> AdditionalOperandRefs = { Op1 };
diff --git a/Engine/Plugins/Runtime/Metasound/Source/MetasoundStandardNodes/Public/MetasoundRandomNode.h b/Engine/Plugins/Runtime/Metasound/Source/MetasoundStandardNodes/Public/MetasoundRandomNode.h
--- a/Engine/Plugins/Runtime/Metasound/Source/MetasoundStandardNodes/Public/MetasoundRandomNode.h
+++ b/Engine/Plugins/Runtime/Metasound/Source/MetasoundStandardNodes/Public/MetasoundRandomNode.h
@@ -85,14 +85,14 @@ namespace Metasound
 			return InStream.RandRange(InMin, InMax);
 		}
 
-		static TDataReadReference<int32> CreateMinValueRef(const FDataReferenceCollection& InputCollection, const FInputVertexInterface& InputInterface)
+		static TDataReadReference<int32> CreateMinValueRef(const FDataReferenceCollection& InputCollection, const FInputVertexInterface& InputInterface, const FOperatorSettings& InOperatorSettings)
 		{
-			return InputCollection.GetDataReadReferenceOrConstructWithVertexDefault<int32>(InputInterface, RandomNodeNames::GetInputMinName());
+			return InputCollection.GetDataReadReferenceOrConstructWithVertexDefault<int32>(InputInterface, RandomNodeNames::GetInputMinName(), InOperatorSettings);
 		}
 
-		static TDataReadReference<int32> CreateMaxValueRef(const FDataReferenceCollection& InputCollection, const FInputVertexInterface& InputInterface)
+		static TDataReadReference<int32> CreateMaxValueRef(const FDataReferenceCollection& InputCollection, const FInputVertexInterface& InputInterface, const FOperatorSettings& InOperatorSettings)
 		{
-			return InputCollection.GetDataReadReferenceOrConstructWithVertexDefault<int32>(InputInterface, RandomNodeNames::GetInputMaxName());
+			return InputCollection.GetDataReadReferenceOrConstructWithVertexDefault<int32>(InputInterface, RandomNodeNames::GetInputMaxName(), InOperatorSettings);
 		}
 	};
 
@@ -134,14 +134,14 @@ namespace Metasound
 			return InStream.FRandRange(InMin, InMax);
 		}
 
-		static TDataReadReference<float> CreateMinValueRef(const FDataReferenceCollection& InputCollection, const FInputVertexInterface& InputInterface)
+		static TDataReadReference<float> CreateMinValueRef(const FDataReferenceCollection& InputCollection, const FInputVertexInterface& InputInterface, const FOperatorSettings& InOperatorSettings)
 		{
-			return InputCollection.GetDataReadReferenceOrConstructWithVertexDefault<float>(InputInterface, RandomNodeNames::GetInputMinName());
+			return InputCollection.GetDataReadReferenceOrConstructWithVertexDefault<float>(InputInterface, RandomNodeNames::GetInputMinName(), InOperatorSettings);
 		}
 
-		static TDataReadReference<float> CreateMaxValueRef(const FDataReferenceCollection& InputCollection, const FInputVertexInterface& InputInterface)
+		static TDataReadReference<float> CreateMaxValueRef(const FDataReferenceCollection& InputCollection, const FInputVertexInterface& InputInterface, const FOperatorSettings& InOperatorSettings)
 		{
-			return InputCollection.GetDataReadReferenceOrConstructWithVertexDefault<float>(InputInterface, RandomNodeNames::GetInputMaxName());
+			return InputCollection.GetDataReadReferenceOrConstructWithVertexDefault<float>(InputInterface, RandomNodeNames::GetInputMaxName(), InOperatorSettings);
 		}
 	};
 
@@ -183,12 +183,12 @@ namespace Metasound
 			return (bool)InStream.RandRange(0, 1);
 		}
 
-		static TDataReadReference<bool> CreateMinValueRef(const FDataReferenceCollection&, const FInputVertexInterface&)
+		static TDataReadReference<bool> CreateMinValueRef(const FDataReferenceCollection&, const FInputVertexInterface&, const FOperatorSettings&)
 		{
 			return TDataReadReference<bool>::CreateNew();
 		}
 
-		static TDataReadReference<bool> CreateMaxValueRef(const FDataReferenceCollection&, const FInputVertexInterface&)
+		static TDataReadReference<bool> CreateMaxValueRef(const FDataReferenceCollection&, const FInputVertexInterface&, const FOperatorSettings&)
 		{
 			return TDataReadReference<bool>::CreateNew();
 		}
@@ -234,14 +234,14 @@ namespace Metasound
 			return FTime(Seconds);
 		}
 
-		static TDataReadReference<FTime> CreateMinValueRef(const FDataReferenceCollection& InputCollection, const FInputVertexInterface& InputInterface)
+		static TDataReadReference<FTime> CreateMinValueRef(const FDataReferenceCollection& InputCollection, const FInputVertexInterface& InputInterface, const FOperatorSettings& InOperatorSettings)
 		{
-			return InputCollection.GetDataReadReferenceOrConstructWithVertexDefault<FTime, float>(InputInterface, RandomNodeNames::GetInputMinName());
+			return InputCollection.GetDataReadReferenceOrConstructWithVertexDefault<FTime>(InputInterface, RandomNodeNames::GetInputMinName(), InOperatorSettings);
 		}
 
-		static TDataReadReference<FTime> CreateMaxValueRef(const FDataReferenceCollection& InputCollection, const FInputVertexInterface& InputInterface)
+		static TDataReadReference<FTime> CreateMaxValueRef(const FDataReferenceCollection& InputCollection, const FInputVertexInterface& InputInterface, const FOperatorSettings& InOperatorSettings)
 		{
-			return InputCollection.GetDataReadReferenceOrConstructWithVertexDefault<FTime, float>(InputInterface, RandomNodeNames::GetInputMaxName());
+			return InputCollection.GetDataReadReferenceOrConstructWithVertexDefault<FTime>(InputInterface, RandomNodeNames::GetInputMaxName(), InOperatorSettings);
 		}
 	};
 
@@ -325,13 +325,13 @@ namespace Metasound
 
  			FTriggerReadRef InNextTrigger = InputCollection.GetDataReadReferenceOrConstruct<FTrigger>(GetInputNextTriggerName(), InParams.OperatorSettings);
 			FTriggerReadRef InResetTrigger = InputCollection.GetDataReadReferenceOrConstruct<FTrigger>(GetInputResetTriggerName(), InParams.OperatorSettings);
-			FInt32ReadRef InSeedValue = InputCollection.GetDataReadReferenceOrConstructWithVertexDefault<int32>(InputInterface, GetInputSeedName());
+			FInt32ReadRef InSeedValue = InputCollection.GetDataReadReferenceOrConstructWithVertexDefault<int32>(InputInterface, GetInputSeedName(), InParams.OperatorSettings);
 
 			// note: random bool does not have a range
 			if (TRandomNodeSpecialization<ValueType>::HasRange())
 			{
-				TDataReadReference<ValueType> InMinValue = TRandomNodeSpecialization<ValueType>::CreateMinValueRef(InputCollection, InputInterface);
-				TDataReadReference<ValueType> InMaxValue = TRandomNodeSpecialization<ValueType>::CreateMaxValueRef(InputCollection, InputInterface);
+				TDataReadReference<ValueType> InMinValue = TRandomNodeSpecialization<ValueType>::CreateMinValueRef(InputCollection, InputInterface, InParams.OperatorSettings);
+				TDataReadReference<ValueType> InMaxValue = TRandomNodeSpecialization<ValueType>::CreateMaxValueRef(InputCollection, InputInterface, InParams.OperatorSettings);
 
 				return MakeUnique<TRandomNodeOperator<ValueType>>(InParams.OperatorSettings, InNextTrigger, InResetTrigger, InSeedValue, InMinValue, InMaxValue);
 			}
diff --git a/Engine/Plugins/Runtime/Metasound/Source/MetasoundStandardNodes/Public/MetasoundTriggerAccumulatorNode.h b/Engine/Plugins/Runtime/Metasound/Source/MetasoundStandardNodes/Public/MetasoundTriggerAccumulatorNode.h
--- a/Engine/Plugins/Runtime/Metasound/Source/MetasoundStandardNodes/Public/MetasoundTriggerAccumulatorNode.h
+++ b/Engine/Plugins/Runtime/Metasound/Source/MetasoundStandardNodes/Public/MetasoundTriggerAccumulatorNode.h
@@ -82,7 +82,7 @@ namespace Metasound
 			const FInputVertexInterface& InputInterface = InParams.Node.GetVertexInterface().GetInputInterface();
 			const FDataReferenceCollection& InputCollection = InParams.InputDataReferences;
 
-			FBoolReadRef bInAutoReset = InputCollection.GetDataReadReferenceOrConstructWithVertexDefault<bool>(InputInterface, GetInputAutoResetName());
+			FBoolReadRef bInAutoReset = InputCollection.GetDataReadReferenceOrConstructWithVertexDefault<bool>(InputInterface, GetInputAutoResetName(), InParams.OperatorSettings);
 			TArray<FTriggerReadRef> InputTriggers;
 
 			for (uint32 i = 0; i < NumInputs; ++i)
diff --git a/Engine/Plugins/Runtime/Metasound/Source/MetasoundStandardNodes/Public/MetasoundTriggerRouteNode.h b/Engine/Plugins/Runtime/Metasound/Source/MetasoundStandardNodes/Public/MetasoundTriggerRouteNode.h
--- a/Engine/Plugins/Runtime/Metasound/Source/MetasoundStandardNodes/Public/MetasoundTriggerRouteNode.h
+++ b/Engine/Plugins/Runtime/Metasound/Source/MetasoundStandardNodes/Public/MetasoundTriggerRouteNode.h
@@ -88,7 +88,7 @@ namespace Metasound
 			const FInputVertexInterface& InputInterface = InParams.Node.GetVertexInterface().GetInputInterface();
 			const FDataReferenceCollection& InputCollection = InParams.InputDataReferences;
 
-			TDataReadReference<ValueType> InputInitValue = InputCollection.GetDataReadReferenceOrConstructWithVertexDefault<ValueType>(InputInterface, GetInputInitValueName());
+			TDataReadReference<ValueType> InputInitValue = InputCollection.GetDataReadReferenceOrConstructWithVertexDefault<ValueType>(InputInterface, GetInputInitValueName(), InParams.OperatorSettings);
 
 			TArray<FTriggerReadRef> InputTriggers;
 			TArray<TDataReadReference<ValueType>> InputValues;
@@ -96,7 +96,7 @@ namespace Metasound
 			for (uint32 i = 0; i < NumInputs; ++i)
 			{
 				InputTriggers.Add(InputCollection.GetDataReadReferenceOrConstruct<FTrigger>(GetInputTriggerName(i), InParams.OperatorSettings));
-				InputValues.Add(InputCollection.GetDataReadReferenceOrConstructWithVertexDefault<ValueType>(InputInterface, GetInputValueName(i)));
+				InputValues.Add(InputCollection.GetDataReadReferenceOrConstructWithVertexDefault<ValueType>(InputInterface, GetInputValueName(i), InParams.OperatorSettings));
 			}
 
 			return MakeUnique<TTriggerRouteOperator<ValueType, NumInputs>>(InParams.OperatorSettings, InputInitValue, MoveTemp(InputTriggers), MoveTemp(InputValues));
diff --git a/Engine/Plugins/Runtime/Metasound/Source/MetasoundStandardNodes/Public/MetasoundValueNode.h b/Engine/Plugins/Runtime/Metasound/Source/MetasoundStandardNodes/Public/MetasoundValueNode.h
--- a/Engine/Plugins/Runtime/Metasound/Source/MetasoundStandardNodes/Public/MetasoundValueNode.h
+++ b/Engine/Plugins/Runtime/Metasound/Source/MetasoundStandardNodes/Public/MetasoundValueNode.h
@@ -74,8 +74,8 @@ namespace Metasound
 			const FDataReferenceCollection& InputCollection = InParams.InputDataReferences;
 
 			FTriggerReadRef Trigger = InputCollection.GetDataReadReferenceOrConstruct<FTrigger>(GetInputTriggerName(), InParams.OperatorSettings);
-			TDataReadReference<ValueType> InitValue = InputCollection.GetDataReadReferenceOrConstructWithVertexDefault<ValueType>(InputInterface, GetInitValueName());
-			TDataReadReference<ValueType> SetValue = InputCollection.GetDataReadReferenceOrConstructWithVertexDefault<ValueType>(InputInterface, GetSetValueName());
+			TDataReadReference<ValueType> InitValue = InputCollection.GetDataReadReferenceOrConstructWithVertexDefault<ValueType>(InputInterface, GetInitValueName(), InParams.OperatorSettings);
+			TDataReadReference<ValueType> SetValue = InputCollection.GetDataReadReferenceOrConstructWithVertexDefault<ValueType>(InputInterface, GetSetValueName(), InParams.OperatorSettings);
 
 			return MakeUnique<TValueOperator<ValueType>>(InParams.OperatorSettings, Trigger, InitValue, SetValue);
 		}
