Set MetaSound literal to default constructor on vertex default constructor

diff --git a/Engine/Plugins/Runtime/Metasound/Source/MetasoundGraphCore/Public/MetasoundVertex.h b/Engine/Plugins/Runtime/Metasound/Source/MetasoundGraphCore/Public/MetasoundVertex.h
--- a/Engine/Plugins/Runtime/Metasound/Source/MetasoundGraphCore/Public/MetasoundVertex.h
+++ b/Engine/Plugins/Runtime/Metasound/Source/MetasoundGraphCore/Public/MetasoundVertex.h
@@ -123,6 +123,7 @@ namespace Metasound
 
 		FInputDataVertexModel(const FString& InVertexName, const FName& InDataTypeName, const FText& InDescription)
 			: FDataVertexModel(InVertexName, InDataTypeName, InDescription)
+			, LiteralFactory(MakeUnique<TLiteralFactory<FLiteral::FNone>>(FLiteral::FNone{}))
 		{
 		}
 
