BPM to seconds node

diff --git a/Engine/Plugins/Runtime/Metasound/Source/MetasoundStandardNodes/Private/MetasoundBPMToSecondsNode.cpp b/Engine/Plugins/Runtime/Metasound/Source/MetasoundStandardNodes/Private/MetasoundBPMToSecondsNode.cpp
new file mode 100644
--- /dev/null
+++ b/Engine/Plugins/Runtime/Metasound/Source/MetasoundStandardNodes/Private/MetasoundBPMToSecondsNode.cpp
@@ -0,0 +1,193 @@
+// Copyright Epic Games, Inc. All Rights Reserved.
+
+#include "MetasoundBPMToSecondsNode.h"
+
+#include "Internationalization/Text.h"
+#include "MetasoundExecutableOperator.h"
+#include "MetasoundNodeRegistrationMacro.h"
+#include "MetasoundPrimitives.h"
+#include "MetasoundStandardNodesNames.h"
+#include "MetasoundTime.h"
+
+#define LOCTEXT_NAMESPACE "MetasoundStandardNodes"
+
+namespace Metasound
+{
+	namespace BPMToSecondsVertexNames
+	{
+		const FString& GetInputBPMTempoName()
+		{
+			static FString Name = TEXT("BPM");
+			return Name;
+		}
+
+		const FString& GetInputBeatMultiplierName()
+		{
+			static FString Name = TEXT("Beat Multiplier");
+			return Name;
+		}
+
+		const FString& GetInputDivisionsOfWholeNoteName()
+		{
+			static FString Name = TEXT("Divisions of Whole Note");
+			return Name;
+		}
+
+		const FString& GetOutputTimeSecondsName()
+		{
+			static FString Name = TEXT("Seconds");
+			return Name;
+		}
+	}
+
+	class FBPMToSecondsOperator : public TExecutableOperator<FBPMToSecondsOperator>
+	{
+	public:
+
+		static const FNodeClassMetadata& GetNodeInfo();
+		static const FVertexInterface& GetVertexInterface();
+		static TUniquePtr<IOperator> CreateOperator(const FCreateOperatorParams& InParams, FBuildErrorArray& OutErrors);
+
+		FBPMToSecondsOperator(const FOperatorSettings& InSettings, const FFloatReadRef& InBPM, const FFloatReadRef& InBeatMultiplier, const FFloatReadRef& InDivOfWholeNote);
+
+		virtual FDataReferenceCollection GetInputs() const override;
+		virtual FDataReferenceCollection GetOutputs() const override;
+		void Execute();
+
+	private:
+		void UpdateTime();
+
+		FFloatReadRef BPM;
+		FFloatReadRef BeatMultiplier;
+		FFloatReadRef DivOfWholeNote;
+
+		// The output seconds
+		FTimeWriteRef TimeSeconds;
+
+		// Cached midi note value. Used to catch if the value changes to recompute freq output.
+		float PrevBPM = -1.0f;
+		float PrevBeatMultiplier = -1.0f;
+		float PrevDivOfWholeNote = -1.0f;
+	};
+
+	FBPMToSecondsOperator::FBPMToSecondsOperator(const FOperatorSettings& InSettings, const FFloatReadRef& InBPM, const FFloatReadRef& InBeatMultiplier, const FFloatReadRef& InDivOfWholeNote)
+		: BPM(InBPM)
+		, BeatMultiplier(InBeatMultiplier)
+		, DivOfWholeNote(InDivOfWholeNote)
+		, TimeSeconds(TDataWriteReferenceFactory<FTime>::CreateAny(InSettings))
+	{
+		UpdateTime();
+	}
+
+	FDataReferenceCollection FBPMToSecondsOperator::GetInputs() const
+	{
+		using namespace BPMToSecondsVertexNames;
+
+		FDataReferenceCollection InputDataReferences;
+		InputDataReferences.AddDataReadReference(GetInputBPMTempoName(), BPM);
+		InputDataReferences.AddDataReadReference(GetInputBeatMultiplierName(), BeatMultiplier);
+		InputDataReferences.AddDataReadReference(GetInputDivisionsOfWholeNoteName(), DivOfWholeNote);
+
+		return InputDataReferences;
+	}
+
+	FDataReferenceCollection FBPMToSecondsOperator::GetOutputs() const
+	{
+		using namespace BPMToSecondsVertexNames;
+
+		FDataReferenceCollection OutputDataReferences;
+		OutputDataReferences.AddDataReadReference(GetOutputTimeSecondsName(), TimeSeconds);
+		return OutputDataReferences;
+	}
+
+	void FBPMToSecondsOperator::UpdateTime()
+	{
+		float CurrBPM = FMath::Max(*BPM, 1.0f);
+		float CurrBeatMultiplier = FMath::Max(*BeatMultiplier, KINDA_SMALL_NUMBER);
+		float CurrDivOfWholeNote = FMath::Max(*DivOfWholeNote, 1.0f);
+
+		if (!FMath::IsNearlyEqual(PrevBPM, CurrBPM) || 
+			!FMath::IsNearlyEqual(PrevBeatMultiplier, CurrBeatMultiplier) || 
+			!FMath::IsNearlyEqual(PrevDivOfWholeNote, CurrDivOfWholeNote))
+		{
+			PrevBPM = CurrBPM;
+			PrevBeatMultiplier = CurrBeatMultiplier;
+			PrevDivOfWholeNote = CurrDivOfWholeNote;
+
+			check(CurrBPM > 0.0f);
+			check(CurrDivOfWholeNote > 0.0f);
+			const float QuarterNoteTime = 60.0f / CurrBPM;
+			float NewTimeSeconds = 4.0f * (float)CurrBeatMultiplier * QuarterNoteTime / CurrDivOfWholeNote;
+			*TimeSeconds = FTime(NewTimeSeconds);
+		}
+
+	}
+
+	void FBPMToSecondsOperator::Execute()
+	{
+		UpdateTime();
+	}
+
+	const FVertexInterface& FBPMToSecondsOperator::GetVertexInterface()
+	{
+		using namespace BPMToSecondsVertexNames;
+
+		static const FVertexInterface Interface(
+			FInputVertexInterface(
+				TInputDataVertexModel<float>(GetInputBPMTempoName(), LOCTEXT("BPMToSecsBPMTT", "Input BPM (Beats Per Minute)."), 90.0f),
+				TInputDataVertexModel<float>(GetInputBeatMultiplierName(), LOCTEXT("BPMToSecsBeatMultiplierTT", "The multipilier of the BPM."), 1.0f),
+				TInputDataVertexModel<float>(GetInputDivisionsOfWholeNoteName(), LOCTEXT("BPMToSecsDivOfWholeNoteTT", "Divisions of a whole note."), 4.0f)
+			),
+			FOutputVertexInterface(
+				TOutputDataVertexModel<FTime>(GetOutputTimeSecondsName(), LOCTEXT("BPMToSecsOutputTimeTT", "The output time in seconds give the BPM inputs."))
+			)
+		);
+
+		return Interface;
+	}
+
+	const FNodeClassMetadata& FBPMToSecondsOperator::GetNodeInfo()
+	{
+		auto InitNodeInfo = []() -> FNodeClassMetadata
+		{
+			FNodeClassMetadata Info;
+			Info.ClassName = { Metasound::StandardNodes::Namespace, TEXT("BPMToSeconds"), TEXT("") };
+			Info.MajorVersion = 1;
+			Info.MinorVersion = 0;
+			Info.DisplayName = LOCTEXT("BPMToSecondsDisplayName", "BPM To Seconds");
+			Info.Description = LOCTEXT("BPMToSecondsDesc", "Calculates a beat time in seconds from the given BPM, beat multiplier and divisions of a whole note.");
+			Info.Author = PluginAuthor;
+			Info.PromptIfMissing = PluginNodeMissingPrompt;
+			Info.DefaultInterface = GetVertexInterface();
+
+			return Info;
+		};
+
+		static const FNodeClassMetadata Info = InitNodeInfo();
+
+		return Info;
+	}
+
+	TUniquePtr<IOperator> FBPMToSecondsOperator::CreateOperator(const FCreateOperatorParams& InParams, FBuildErrorArray& OutErrors)
+	{
+		using namespace BPMToSecondsVertexNames;
+
+		const FDataReferenceCollection& InputCollection = InParams.InputDataReferences;
+		const FInputVertexInterface& InputInterface = GetVertexInterface().GetInputInterface();
+
+		FFloatReadRef InBPM = InputCollection.GetDataReadReferenceOrConstructWithVertexDefault<float>(InputInterface, GetInputBPMTempoName(), InParams.OperatorSettings);
+		FFloatReadRef InBeatMultiplier = InputCollection.GetDataReadReferenceOrConstructWithVertexDefault<float>(InputInterface, GetInputBeatMultiplierName(), InParams.OperatorSettings);
+		FFloatReadRef InDivOfWholeNote = InputCollection.GetDataReadReferenceOrConstructWithVertexDefault<float>(InputInterface, GetInputDivisionsOfWholeNoteName(), InParams.OperatorSettings);
+
+		return MakeUnique<FBPMToSecondsOperator>(InParams.OperatorSettings, InBPM, InBeatMultiplier, InDivOfWholeNote);
+	}
+
+	FBPMToSecondsNode::FBPMToSecondsNode(const FNodeInitData& InitData)
+		: FNodeFacade(InitData.InstanceName, InitData.InstanceID, TFacadeOperatorClass<FBPMToSecondsOperator>())
+	{
+	}
+
+	METASOUND_REGISTER_NODE(FBPMToSecondsNode)
+}
+
+#undef LOCTEXT_NAMESPACE
diff --git a/Engine/Plugins/Runtime/Metasound/Source/MetasoundStandardNodes/Public/MetasoundBPMToSecondsNode.h b/Engine/Plugins/Runtime/Metasound/Source/MetasoundStandardNodes/Public/MetasoundBPMToSecondsNode.h
new file mode 100644
--- /dev/null
+++ b/Engine/Plugins/Runtime/Metasound/Source/MetasoundStandardNodes/Public/MetasoundBPMToSecondsNode.h
@@ -0,0 +1,23 @@
+// Copyright Epic Games, Inc. All Rights Reserved.
+
+#pragma once
+
+#include "CoreMinimal.h"
+#include "MetasoundFacade.h"
+
+namespace Metasound
+{
+	/** FBPMToSecondsNode
+	 *
+	 *  Calculates a beat time in seconds from the given BPM, beat multiplier and divisions of a whole note. 
+	 */
+	class METASOUNDSTANDARDNODES_API FBPMToSecondsNode : public FNodeFacade
+	{
+	public:
+		/**
+		 * Constructor used by the Metasound Frontend.
+		 */
+		FBPMToSecondsNode(const FNodeInitData& InitData);
+	};
+} // namespace Metasound
+
