Cleaning up some nodes.

- Adding some minor features to various trigger nodes.
- Fixing/standarding pin orders and names
- Added a utility to help with param management (until we get a better param manager)
- Fixed a few minor bugs (trigger pipe/delay resets)
- Refactored value node to have a reset trigger and output triggers on set/reset. Renamed params to be more clear.
- Did a pass on tooltip cleanup and fixes.
- Refactored trigger compare node and added float support
- Refactored trigger counter, adding a couple minor features (count-by, auto-reset)

diff --git a/Engine/Plugins/Runtime/Metasound/Source/MetasoundStandardNodes/Private/MetasoundTriggerAccumulatorNode.cpp b/Engine/Plugins/Runtime/Metasound/Source/MetasoundStandardNodes/Private/MetasoundTriggerAccumulatorNode.cpp
--- a/Engine/Plugins/Runtime/Metasound/Source/MetasoundStandardNodes/Private/MetasoundTriggerAccumulatorNode.cpp
+++ b/Engine/Plugins/Runtime/Metasound/Source/MetasoundStandardNodes/Private/MetasoundTriggerAccumulatorNode.cpp
@@ -38,41 +38,6 @@ namespace Metasound
 		}
 	}
 
-	namespace TriggerAccumulatorVertexNames
-	{
-		const FString GetInputAutoResetName()
-		{
-			return TEXT("Auto Reset");
-		}
-
-		const FText GetInputAutoResetDescription()
-		{
-			return LOCTEXT("TriggerAccumulatorAutoResetDesc", "Set to true to automatically reset the trigger accumulator state once the output is triggered.");
-		}
-
-		const FString GetInputTriggerName(uint32 InIndex)
-		{
-			return FText::Format(LOCTEXT("TriggerAccumulatorTriggerInputName", "In {0}"), InIndex).ToString();
-		}
-
-		const FText GetInputTriggerDescription(uint32 InIndex)
-		{
-			return FText::Format(LOCTEXT("TriggerAccumulatorInputTriggerDesc", "Trigger {0} input. All trigger inputs must be triggered before the output trigger is hit."), InIndex);
-		}
-
-		const FString& GetOutputTriggerName()
-		{
-			static const FString Name = TEXT("Out");
-			return Name;
-		}
-
-		const FText& GetOutputTriggerDescription()
-		{
-			static const FText Desc = LOCTEXT("TriggerAccumulateOutputTriggerDesc", "Triggered when all input triggers have been triggered. Call Reset to reset the state or use \"Auto Reset\"");
-			return Desc;
-		}
-	}
-
 	REGISTER_TRIGGER_ACCUMULATOR_NODE(1)
 	REGISTER_TRIGGER_ACCUMULATOR_NODE(2)
 	REGISTER_TRIGGER_ACCUMULATOR_NODE(3)
diff --git a/Engine/Plugins/Runtime/Metasound/Source/MetasoundStandardNodes/Private/MetasoundTriggerCompareNode.cpp b/Engine/Plugins/Runtime/Metasound/Source/MetasoundStandardNodes/Private/MetasoundTriggerCompareNode.cpp
--- a/Engine/Plugins/Runtime/Metasound/Source/MetasoundStandardNodes/Private/MetasoundTriggerCompareNode.cpp
+++ b/Engine/Plugins/Runtime/Metasound/Source/MetasoundStandardNodes/Private/MetasoundTriggerCompareNode.cpp
@@ -1,218 +1,39 @@
 // Copyright Epic Games, Inc. All Rights Reserved.
 #include "MetasoundTriggerCompareNode.h"
-
-#include "CoreMinimal.h"
-
-#include "MetasoundBuilderInterface.h"
-#include "MetasoundDataReferenceCollection.h"
-#include "MetasoundExecutableOperator.h"
-#include "MetasoundFacade.h"
-#include "MetasoundNode.h"
-#include "MetasoundNodeInterface.h"
-#include "MetasoundNodeRegistrationMacro.h"
-#include "MetasoundDataTypeRegistrationMacro.h"
-#include "MetasoundOperatorInterface.h"
-#include "MetasoundPrimitives.h"
-#include "MetasoundStandardNodesNames.h"
 #include "MetasoundStandardNodesCategories.h"
 
 #define LOCTEXT_NAMESPACE "MetasoundStandardNodes"
 
 namespace Metasound
 {
-	namespace TriggerCompare
-	{
-		static const TCHAR* InParamNameCompare = TEXT("Compare");
-		static const TCHAR* InParamNameA = TEXT("A");
-		static const TCHAR* InParamNameB = TEXT("B");
-		static const TCHAR* InParamCompareType = TEXT("Type");
-		static const TCHAR* OutParamOnTrue = TEXT("True");
-		static const TCHAR* OutParamOnFalse = TEXT("False");
-	}
-
-	enum class ETriggerComparisonType
-	{
-		Equals,
-		NotEquals,
-		LessThan,
-		GreaterThan,
-		LessThanOrEquals,
-		GreaterThanOrEquals
-	};
-
-	DECLARE_METASOUND_ENUM(ETriggerComparisonType, ETriggerComparisonType::Equals, METASOUNDSTANDARDNODES_API,
-		FEnumTriggerComparisonType, FEnumTriggerComparisonTypeInfo, FEnumTriggerComparisonTypeReadRef, FEnumTriggerComparisonTypeWriteRef);
-
-	DEFINE_METASOUND_ENUM_BEGIN(ETriggerComparisonType, FEnumTriggerComparisonType, "TriggerComparisonType")
-		DEFINE_METASOUND_ENUM_ENTRY(ETriggerComparisonType::Equals, LOCTEXT("EqualsDescription", "Equals"), LOCTEXT("EqualsDescriptionTT", "True if A and B are equal.")),
-		DEFINE_METASOUND_ENUM_ENTRY(ETriggerComparisonType::NotEquals, LOCTEXT("NotEqualsDescriptioin", "Not Equals"), LOCTEXT("NotEqualsTT", "True if A and B are not equal.")),
-		DEFINE_METASOUND_ENUM_ENTRY(ETriggerComparisonType::LessThan, LOCTEXT("LessThanDescription", "Less Than"), LOCTEXT("LessThanTT", "True if A is less than B.")),
-		DEFINE_METASOUND_ENUM_ENTRY(ETriggerComparisonType::GreaterThan, LOCTEXT("GreaterThanDescription", "Greater Than"), LOCTEXT("GreaterThanTT", "True if A is greater than B.")),
-		DEFINE_METASOUND_ENUM_ENTRY(ETriggerComparisonType::LessThanOrEquals, LOCTEXT("LessThanOrEqualsDescription", "Less Than Or Equals"), LOCTEXT("LessThanOrEqualsTT", "True if A is less than or equal to B.")),
-		DEFINE_METASOUND_ENUM_ENTRY(ETriggerComparisonType::GreaterThanOrEquals, LOCTEXT("GreaterThanOrEqualsDescription", "Greater Than Or Equals"), LOCTEXT("GreaterThanOrEqualsTT", "True if A is greater than or equal to B.")),
-	DEFINE_METASOUND_ENUM_END()
-
-	class FTriggerCompareOperator : public TExecutableOperator<FTriggerCompareOperator>
-	{
-		public:
-			static const FNodeClassMetadata& GetNodeInfo();
-			static const FVertexInterface& GetVertexInterface();
-			static TUniquePtr<IOperator> CreateOperator(const FCreateOperatorParams& InParams, FBuildErrorArray& OutErrors);
-
-			FTriggerCompareOperator(const FOperatorSettings& InSettings, const FTriggerReadRef& InOnCompareTrigger, const FInt32ReadRef& InValueA, const FInt32ReadRef& InValueB, const FEnumTriggerComparisonTypeReadRef& InTriggerComparisonType);
-
-			virtual FDataReferenceCollection GetInputs() const override;
-			virtual FDataReferenceCollection GetOutputs() const override;
-
-			void Execute();
-
-		private:
-			FTriggerReadRef OnCompareTrigger;
-			FInt32ReadRef ValueA;
-			FInt32ReadRef ValueB;
-			FEnumTriggerComparisonTypeReadRef TriggerComparisonType;
-
-			FTriggerWriteRef TriggerOutOnTrue;
-			FTriggerWriteRef TriggerOutOnFalse;
-	};
-
-	FTriggerCompareOperator::FTriggerCompareOperator(const FOperatorSettings& InSettings, const FTriggerReadRef& InOnCompareTrigger, const FInt32ReadRef& InValueA, const FInt32ReadRef& InValueB, const FEnumTriggerComparisonTypeReadRef& InTriggerComparisonType)
-		: OnCompareTrigger(InOnCompareTrigger)
-		, ValueA(InValueA)
-		, ValueB(InValueB)
-		, TriggerComparisonType(InTriggerComparisonType)
-		, TriggerOutOnTrue(FTriggerWriteRef::CreateNew(InSettings))
-		, TriggerOutOnFalse(FTriggerWriteRef::CreateNew(InSettings))
-	{
-	}
-
-	FDataReferenceCollection FTriggerCompareOperator::GetInputs() const
-	{
-		FDataReferenceCollection InputDataReferences;
-		InputDataReferences.AddDataReadReference(TriggerCompare::InParamNameCompare, OnCompareTrigger);
-		InputDataReferences.AddDataReadReference(TriggerCompare::InParamNameA, ValueA);
-		InputDataReferences.AddDataReadReference(TriggerCompare::InParamNameB, ValueB);
-		InputDataReferences.AddDataReadReference(TriggerCompare::InParamCompareType, TriggerComparisonType);
-		return InputDataReferences;
-	}
-
-	FDataReferenceCollection FTriggerCompareOperator::GetOutputs() const
-	{
-		FDataReferenceCollection OutputDataReferences;
-		OutputDataReferences.AddDataReadReference(TriggerCompare::OutParamOnTrue, TriggerOutOnTrue);
-		OutputDataReferences.AddDataReadReference(TriggerCompare::OutParamOnFalse, TriggerOutOnFalse);
-
-		return OutputDataReferences;
-	}
-
-	void FTriggerCompareOperator::Execute()
+	namespace MetasoundTriggerCompareNodePrivate
 	{
-		TriggerOutOnTrue->AdvanceBlock();
-		TriggerOutOnFalse->AdvanceBlock();
-
-		OnCompareTrigger->ExecuteBlock(
-			[&](int32 StartFrame, int32 EndFrame)
-			{
-			},
-			[this](int32 StartFrame, int32 EndFrame)
-			{
-				int32 CurrValA = *ValueA;
-				int32 CurrValB = *ValueB;
-				bool bIsTrue = false;
-
-				switch (*TriggerComparisonType)
-				{
-					case ETriggerComparisonType::Equals:
-						bIsTrue = (CurrValA == CurrValB);
-						break;
-					case ETriggerComparisonType::NotEquals:
-						bIsTrue = (CurrValA != CurrValB);
-						break;
-					case ETriggerComparisonType::LessThan:
-						bIsTrue = (CurrValA < CurrValB);
-						break;
-					case ETriggerComparisonType::GreaterThan:
-						bIsTrue = (CurrValA > CurrValB);
-						break;
-					case ETriggerComparisonType::LessThanOrEquals:
-						bIsTrue = (CurrValA <= CurrValB);
-						break;
-					case ETriggerComparisonType::GreaterThanOrEquals:
-						bIsTrue = (CurrValA >= CurrValB);
-						break;
-				}
-
-				if (bIsTrue)
-				{
-					TriggerOutOnTrue->TriggerFrame(StartFrame);
-				}
-				else
-				{
-					TriggerOutOnFalse->TriggerFrame(StartFrame);
-				}
-			}
-		);
-	}
-
-	TUniquePtr<IOperator> FTriggerCompareOperator::CreateOperator(const FCreateOperatorParams& InParams, FBuildErrorArray& OutErrors)
-	{
-		const FDataReferenceCollection& InputCollection = InParams.InputDataReferences;
-		const FInputVertexInterface& InputInterface = GetVertexInterface().GetInputInterface();
-
-		FTriggerReadRef InOnTriggerCompare = InputCollection.GetDataReadReferenceOrConstruct<FTrigger>(TriggerCompare::InParamNameCompare, InParams.OperatorSettings);
-		FInt32ReadRef InValueA = InputCollection.GetDataReadReferenceOrConstructWithVertexDefault<int32>(InputInterface, TriggerCompare::InParamNameA, InParams.OperatorSettings);
-		FInt32ReadRef InValueB = InputCollection.GetDataReadReferenceOrConstructWithVertexDefault<int32>(InputInterface, TriggerCompare::InParamNameB, InParams.OperatorSettings);
-		FEnumTriggerComparisonTypeReadRef InComparison = InputCollection.GetDataReadReferenceOrConstruct<FEnumTriggerComparisonType>(TriggerCompare::InParamCompareType);
-
-		return MakeUnique<FTriggerCompareOperator>(InParams.OperatorSettings, InOnTriggerCompare, InValueA, InValueB, InComparison);
-	}
-
-	const FVertexInterface& FTriggerCompareOperator::GetVertexInterface()
-	{
-		static const FVertexInterface Interface(
-			FInputVertexInterface(
-				TInputDataVertexModel<FTrigger>(TriggerCompare::InParamNameCompare, LOCTEXT("TriggerCompareTT", "Trigger to compare A and B.")),
-				TInputDataVertexModel<int32>(TriggerCompare::InParamNameA, LOCTEXT("TriggerCompareATT", "The first value, A, to compare against.")),
-				TInputDataVertexModel<int32>(TriggerCompare::InParamNameB, LOCTEXT("TriggerCompareBTT", "The second value, B, to compare against.")),
-				TInputDataVertexModel<FEnumTriggerComparisonType>(TriggerCompare::InParamCompareType, LOCTEXT("TriggerCompareTypeTT", "How to compare A and B."))
-			),
-			FOutputVertexInterface(
-				TOutputDataVertexModel<FTrigger>(TriggerCompare::OutParamOnTrue, LOCTEXT("TriggerCompareOnTrueTT", "Output trigger for when the comparison is true.")),
-				TOutputDataVertexModel<FTrigger>(TriggerCompare::OutParamOnFalse, LOCTEXT("TriggerCompareOnFalseTT", "Output trigger for when the comparison is false."))
-			)
-		);
-
-		return Interface;
-	}
-
-	const FNodeClassMetadata& FTriggerCompareOperator::GetNodeInfo()
-	{
-		auto InitNodeInfo = []() -> FNodeClassMetadata
+		FNodeClassMetadata CreateNodeClassMetadata(const FName& InDataTypeName, const FName& InOperatorName, const FText& InDisplayName, const FText& InDescription, const FVertexInterface& InDefaultInterface)
 		{
-			FNodeClassMetadata Info;
-			Info.ClassName = {Metasound::StandardNodes::Namespace, TEXT("Trigger Compare"), TEXT("")};
-			Info.MajorVersion = 1;
-			Info.MinorVersion = 0;
-			Info.DisplayName = LOCTEXT("Metasound_TriggerSelectNodeDisplayName", "Trigger Compare");
-			Info.Description = LOCTEXT("Metasound_TriggerSelectNodeDescription", "Output triggers based on comparing integer inputs, A and B.");
-			Info.Author = PluginAuthor;
-			Info.PromptIfMissing = PluginNodeMissingPrompt;
-			Info.DefaultInterface = GetVertexInterface();
-			Info.CategoryHierarchy.Emplace(StandardNodes::TriggerUtils);
-
-			return Info;
-		};
-
-		static const FNodeClassMetadata Info = InitNodeInfo();
-		return Info;
-	}
-
-	FTriggerCompareNode::FTriggerCompareNode(const FNodeInitData& InInitData)
-	:	FNodeFacade(InInitData.InstanceName, InInitData.InstanceID, TFacadeOperatorClass<FTriggerCompareOperator>())
-	{
-	}
-
-	METASOUND_REGISTER_NODE(FTriggerCompareNode)
+			FNodeClassMetadata Metadata
+			{
+				FNodeClassName{FName("TriggerCompare"), InOperatorName, InDataTypeName},
+				1, // Major Version
+				0, // Minor Version
+				InDisplayName,
+				InDescription,
+				PluginAuthor,
+				PluginNodeMissingPrompt,
+				InDefaultInterface,
+				{StandardNodes::TriggerUtils},
+				{TEXT("TriggerCompare")},
+				FNodeDisplayStyle{}
+			};
+
+			return Metadata;
+		}
+	}
+
+	using FTriggerCompareNodeInt32 = TTriggerCompareNode<int32>;
+	METASOUND_REGISTER_NODE(FTriggerCompareNodeInt32)
+
+	using FTriggerCompareNodeFloat = TTriggerCompareNode<float>;
+	METASOUND_REGISTER_NODE(FTriggerCompareNodeFloat)
 }
 
 #undef LOCTEXT_NAMESPACE // MetasoundTriggerDelayNode
diff --git a/Engine/Plugins/Runtime/Metasound/Source/MetasoundStandardNodes/Private/MetasoundTriggerCounterNode.cpp b/Engine/Plugins/Runtime/Metasound/Source/MetasoundStandardNodes/Private/MetasoundTriggerCounterNode.cpp
--- a/Engine/Plugins/Runtime/Metasound/Source/MetasoundStandardNodes/Private/MetasoundTriggerCounterNode.cpp
+++ b/Engine/Plugins/Runtime/Metasound/Source/MetasoundStandardNodes/Private/MetasoundTriggerCounterNode.cpp
@@ -14,16 +14,22 @@
 #include "MetasoundPrimitives.h"
 #include "MetasoundStandardNodesNames.h"
 #include "MetasoundStandardNodesCategories.h"
+#include "MetasoundParamHelper.h"
 
 #define LOCTEXT_NAMESPACE "MetasoundStandardNodes"
 
 namespace Metasound
 {
-	namespace TriggerCounter
+	namespace TriggerCounterVertexNames
 	{
-		static const TCHAR* InParamNameInput = TEXT("In");
-		static const TCHAR* InParamNameReset = TEXT("Reset");
-		static const TCHAR* OutParamNameOutput = TEXT("Out");
+		METASOUND_PARAM(InputInTrigger, "In", "Trigger Input to count.");
+		METASOUND_PARAM(InputReset, "Reset", "Resets the counter back to the start count.");
+		METASOUND_PARAM(InputStartCount, "Start Value", "The value to start output counter by.");
+		METASOUND_PARAM(InputStepSize, "Step Size", "The value to add to the current count each input trigger. Can be negative.");
+		METASOUND_PARAM(InputAutoResetCount, "Auto Reset Count", "The number of input triggers to automatically set the count back to Start Count. If 0, won't auto-reset.");
+		METASOUND_PARAM(OutputOnTrigger, "On Trigger", "Triggered when the input is triggered.");
+		METASOUND_PARAM(OutputOnReset, "On Reset", "Triggered when the input reset trigger is triggered or if the counter automatically resets.");
+		METASOUND_PARAM(OutputCount, "Count", "The current count value of the node.");
 	}
 
 	class FTriggerCounterOperator : public TExecutableOperator<FTriggerCounterOperator>
@@ -33,7 +39,12 @@ namespace Metasound
 			static const FVertexInterface& GetVertexInterface();
 			static TUniquePtr<IOperator> CreateOperator(const FCreateOperatorParams& InParams, FBuildErrorArray& OutErrors);
 
-			FTriggerCounterOperator(const FOperatorSettings& InSettings, const FTriggerReadRef& InTriggerIn, const FTriggerReadRef& InTriggerReset);
+			FTriggerCounterOperator(const FOperatorSettings& InSettings,
+				const FTriggerReadRef& InTriggerIn,
+				const FTriggerReadRef& InTriggerReset,
+				const FInt32ReadRef& InStartCount,
+				const FInt32ReadRef& InCountBy,
+				const FInt32ReadRef& InAutoResetCount);
 
 			virtual FDataReferenceCollection GetInputs() const override;
 			virtual FDataReferenceCollection GetOutputs() const override;
@@ -43,43 +54,81 @@ namespace Metasound
 		private:
 			FTriggerReadRef TriggerIn;
 			FTriggerReadRef TriggerReset;
-			FInt32WriteRef TriggerCount;
+			FInt32ReadRef StartCount;
+			FInt32ReadRef CountBy;
+			FInt32ReadRef AutoResetCount;
+
+			FTriggerWriteRef TriggerOut;
+			FTriggerWriteRef TriggerOnReset;
+			FInt32WriteRef OutValue;
+
+			// Internal state to do the counting math
+			int32 CurrentAutoResetCount = 0;
+			int32 CurrentTriggerCount = 0;
 	};
 
-	FTriggerCounterOperator::FTriggerCounterOperator(const FOperatorSettings& InSettings, const FTriggerReadRef& InTriggerIn, const FTriggerReadRef& InTriggerReset)
-	: TriggerIn(InTriggerIn)
-	, TriggerReset(InTriggerReset)
-	, TriggerCount(FInt32WriteRef::CreateNew(0))
+	FTriggerCounterOperator::FTriggerCounterOperator(const FOperatorSettings& InSettings, 
+		const FTriggerReadRef& InTriggerIn, 
+		const FTriggerReadRef& InTriggerReset,
+		const FInt32ReadRef& InStartCount,
+		const FInt32ReadRef& InCountBy,
+		const FInt32ReadRef& InAutoResetCount)
+		: TriggerIn(InTriggerIn)
+		, TriggerReset(InTriggerReset)
+		, StartCount(InStartCount)
+		, CountBy(InCountBy)
+		, AutoResetCount(InAutoResetCount)
+		, TriggerOut(FTriggerWriteRef::CreateNew(InSettings))
+		, TriggerOnReset(FTriggerWriteRef::CreateNew(InSettings))
+		, OutValue(FInt32WriteRef::CreateNew(*StartCount))
 	{
+		CurrentAutoResetCount = FMath::Max(0, *AutoResetCount);
 	}
 
 	FDataReferenceCollection FTriggerCounterOperator::GetInputs() const
 	{
+		using namespace TriggerCounterVertexNames;
+
 		FDataReferenceCollection InputDataReferences;
-		InputDataReferences.AddDataReadReference(TriggerCounter::InParamNameInput, TriggerIn);
-		InputDataReferences.AddDataReadReference(TriggerCounter::InParamNameReset, TriggerReset);
+		InputDataReferences.AddDataReadReference(METASOUND_GET_PARAM_NAME(InputInTrigger), TriggerIn);
+		InputDataReferences.AddDataReadReference(METASOUND_GET_PARAM_NAME(InputReset), TriggerReset);
+		InputDataReferences.AddDataReadReference(METASOUND_GET_PARAM_NAME(InputStartCount), StartCount);
+		InputDataReferences.AddDataReadReference(METASOUND_GET_PARAM_NAME(InputStepSize), CountBy);
+		InputDataReferences.AddDataReadReference(METASOUND_GET_PARAM_NAME(InputAutoResetCount), AutoResetCount);
+
 		return InputDataReferences;
 	}
 
 	FDataReferenceCollection FTriggerCounterOperator::GetOutputs() const
 	{
+		using namespace TriggerCounterVertexNames;
+
 		FDataReferenceCollection OutputDataReferences;
-		OutputDataReferences.AddDataReadReference(TriggerCounter::OutParamNameOutput, FInt32WriteRef(TriggerCount));
+		OutputDataReferences.AddDataReadReference(METASOUND_GET_PARAM_NAME(OutputOnTrigger), TriggerOut);
+		OutputDataReferences.AddDataReadReference(METASOUND_GET_PARAM_NAME(OutputOnReset), TriggerOnReset);
+		OutputDataReferences.AddDataReadReference(METASOUND_GET_PARAM_NAME(OutputCount), OutValue);
 
 		return OutputDataReferences;
 	}
 
 	void FTriggerCounterOperator::Execute()
-	{
-		// TODO: fix this up when we implement merged trigger execution to avoid edge cases of ordering problems
-		// for when these triggers are executed at differnet times within the same block
+	{	
+		TriggerOnReset->AdvanceBlock();
+		TriggerOut->AdvanceBlock();
+
+		// Update the current auto reset count
+		CurrentAutoResetCount = FMath::Max(0, *AutoResetCount);
+
 		TriggerReset->ExecuteBlock(
 			[&](int32 StartFrame, int32 EndFrame)
 			{
 			},
 			[this](int32 StartFrame, int32 EndFrame)
 			{
-				*TriggerCount = 0;
+				CurrentTriggerCount = 0;
+				*OutValue = *StartCount;
+
+				TriggerOnReset->TriggerFrame(StartFrame);
 			}
 		);
 
@@ -89,29 +138,55 @@ namespace Metasound
 			},
 			[this](int32 StartFrame, int32 EndFrame)
 			{
-				++(*TriggerCount);
+				++CurrentTriggerCount;
+
+				if (CurrentAutoResetCount > 0 && CurrentTriggerCount >= CurrentAutoResetCount)
+				{
+					CurrentTriggerCount = 0;
+					*OutValue = *StartCount;
+					TriggerOnReset->TriggerFrame(StartFrame);
+				}
+				else
+				{
+					*OutValue += *CountBy;
+				}
+
+				TriggerOut->TriggerFrame(StartFrame);
 			}
 		);
+
 	}
 
 	TUniquePtr<IOperator> FTriggerCounterOperator::CreateOperator(const FCreateOperatorParams& InParams, FBuildErrorArray& OutErrors)
 	{
+		using namespace TriggerCounterVertexNames;
+
 		const FInputVertexInterface& InputInterface = GetVertexInterface().GetInputInterface();
-		FTriggerReadRef TriggerIn = InParams.InputDataReferences.GetDataReadReferenceOrConstruct<FTrigger>(TriggerCounter::InParamNameInput, InParams.OperatorSettings);
-		FTriggerReadRef TriggerReset = InParams.InputDataReferences.GetDataReadReferenceOrConstruct<FTrigger>(TriggerCounter::InParamNameReset, InParams.OperatorSettings);
+		FTriggerReadRef TriggerIn = InParams.InputDataReferences.GetDataReadReferenceOrConstruct<FTrigger>(METASOUND_GET_PARAM_NAME(InputInTrigger), InParams.OperatorSettings);
+		FTriggerReadRef TriggerReset = InParams.InputDataReferences.GetDataReadReferenceOrConstruct<FTrigger>(METASOUND_GET_PARAM_NAME(InputReset), InParams.OperatorSettings);
+		FInt32ReadRef StartCount = InParams.InputDataReferences.GetDataReadReferenceOrConstructWithVertexDefault<int32>(InputInterface, METASOUND_GET_PARAM_NAME(InputStartCount), InParams.OperatorSettings);
+		FInt32ReadRef StepSize = InParams.InputDataReferences.GetDataReadReferenceOrConstructWithVertexDefault<int32>(InputInterface, METASOUND_GET_PARAM_NAME(InputStepSize), InParams.OperatorSettings);
+		FInt32ReadRef AutoResetCount = InParams.InputDataReferences.GetDataReadReferenceOrConstructWithVertexDefault<int32>(InputInterface, METASOUND_GET_PARAM_NAME(InputAutoResetCount), InParams.OperatorSettings);
 
-		return MakeUnique<FTriggerCounterOperator>(InParams.OperatorSettings, TriggerIn, TriggerReset);
+		return MakeUnique<FTriggerCounterOperator>(InParams.OperatorSettings, TriggerIn, TriggerReset, StartCount, StepSize, AutoResetCount);
 	}
 
 	const FVertexInterface& FTriggerCounterOperator::GetVertexInterface()
 	{
+		using namespace TriggerCounterVertexNames;
+
 		static const FVertexInterface Interface(
 			FInputVertexInterface(
-				TInputDataVertexModel<FTrigger>(TriggerCounter::InParamNameInput, LOCTEXT("TriggerCounterInTooltip", "Trigger input to count.")),
-				TInputDataVertexModel<FTrigger>(TriggerCounter::InParamNameReset, LOCTEXT("TriggerCounterResetInTooltip", "Resets the counter back to 0."))
+				TInputDataVertexModel<FTrigger>(METASOUND_GET_PARAM_NAME_AND_TT(InputInTrigger)),
+				TInputDataVertexModel<FTrigger>(METASOUND_GET_PARAM_NAME_AND_TT(InputReset)),
+				TInputDataVertexModel<int32>(METASOUND_GET_PARAM_NAME_AND_TT(InputStartCount), 0),
+				TInputDataVertexModel<int32>(METASOUND_GET_PARAM_NAME_AND_TT(InputStepSize), 1),
+				TInputDataVertexModel<int32>(METASOUND_GET_PARAM_NAME_AND_TT(InputAutoResetCount), 0)
 			),
 			FOutputVertexInterface(
-				TOutputDataVertexModel<int32>(TriggerCounter::OutParamNameOutput, LOCTEXT("TriggerCounterOutTooltip", "The current count of triggers."))
+				TOutputDataVertexModel<FTrigger>(METASOUND_GET_PARAM_NAME_AND_TT(OutputOnTrigger)),
+				TOutputDataVertexModel<FTrigger>(METASOUND_GET_PARAM_NAME_AND_TT(OutputOnReset)),
+				TOutputDataVertexModel<int32>(METASOUND_GET_PARAM_NAME_AND_TT(OutputCount))
 			)
 		);
 
@@ -127,7 +202,7 @@ namespace Metasound
 			Info.MajorVersion = 1;
 			Info.MinorVersion = 0;
 			Info.DisplayName = LOCTEXT("Metasound_TriggerCounterNodeDisplayName", "Trigger Counter");
-			Info.Description = LOCTEXT("Metasound_TriggerCounterNodeDescription", "Counts the trigger inputs.");
+			Info.Description = LOCTEXT("Metasound_TriggerCounterNodeDescription", "Counts the trigger inputs. Supports a start count value, counting by a step size, and auto resetting back to the start count.");
 			Info.Author = PluginAuthor;
 			Info.PromptIfMissing = PluginNodeMissingPrompt;
 			Info.DefaultInterface = GetVertexInterface();
diff --git a/Engine/Plugins/Runtime/Metasound/Source/MetasoundStandardNodes/Private/MetasoundTriggerDelayNode.cpp b/Engine/Plugins/Runtime/Metasound/Source/MetasoundStandardNodes/Private/MetasoundTriggerDelayNode.cpp
--- a/Engine/Plugins/Runtime/Metasound/Source/MetasoundStandardNodes/Private/MetasoundTriggerDelayNode.cpp
+++ b/Engine/Plugins/Runtime/Metasound/Source/MetasoundStandardNodes/Private/MetasoundTriggerDelayNode.cpp
@@ -15,11 +15,20 @@
 #include "MetasoundSampleCounter.h"
 #include "MetasoundStandardNodesNames.h"
 #include "MetasoundStandardNodesCategories.h"
+#include "MetasoundParamHelper.h"
 
 #define LOCTEXT_NAMESPACE "MetasoundStandardNodes"
 
 namespace Metasound
 {
+	namespace TriggerDelayVertexNames
+	{
+		METASOUND_PARAM(InputInTrigger, "In", "Input trigger which results in a delayed trigger.");
+		METASOUND_PARAM(InputReset, "Reset", "Resets the trigger delay, clearing the execution task if pending.");
+		METASOUND_PARAM(InputDelayTime, "Delay Time", "Time to delay and execute deferred trigger in seconds.");
+		METASOUND_PARAM(OutputOnTrigger, "Out", "The delayed output trigger.");
+	}
+
 	METASOUND_REGISTER_NODE(FTriggerDelayNode)
 
 	class FTriggerDelayOperator : public TExecutableOperator<FTriggerDelayOperator>
@@ -44,35 +53,39 @@ namespace Metasound
 			FTriggerReadRef TriggerReset;
 			FTriggerWriteRef TriggerOut;
 
-			FTimeReadRef Delay;
+			FTimeReadRef DelayTime;
 
 			FSampleCount FramesPerBlock;
 	};
 
-	FTriggerDelayOperator::FTriggerDelayOperator(const FOperatorSettings& InSettings, const FTriggerReadRef& InTriggerReset, const FTriggerReadRef& InTriggerIn, const FTimeReadRef& InDelay)
-	:	NextTriggerCounter(-1, InSettings.GetSampleRate())
-	,	LastFrameReset(-1)
-	,	TriggerIn(InTriggerIn)
-	,	TriggerReset(InTriggerReset)
-	,	TriggerOut(FTriggerWriteRef::CreateNew(InSettings))
-	,	Delay(InDelay)
-	,	FramesPerBlock(InSettings.GetNumFramesPerBlock())
+	FTriggerDelayOperator::FTriggerDelayOperator(const FOperatorSettings& InSettings, const FTriggerReadRef& InTriggerReset, const FTriggerReadRef& InTriggerIn, const FTimeReadRef& InTimeDelay)
+	: NextTriggerCounter(-1, InSettings.GetSampleRate())
+	, LastFrameReset(-1)
+	, TriggerIn(InTriggerIn)
+	, TriggerReset(InTriggerReset)
+	, TriggerOut(FTriggerWriteRef::CreateNew(InSettings))
+	, DelayTime(InTimeDelay)
+	, FramesPerBlock(InSettings.GetNumFramesPerBlock())
 	{
 	}
 
 	FDataReferenceCollection FTriggerDelayOperator::GetInputs() const
 	{
+		using namespace TriggerDelayVertexNames;
+
 		FDataReferenceCollection InputDataReferences;
-		InputDataReferences.AddDataReadReference(TEXT("Delay"), FTimeReadRef(Delay));
-		InputDataReferences.AddDataReadReference(TEXT("In"), FTriggerReadRef(TriggerIn));
-		InputDataReferences.AddDataReadReference(TEXT("Reset"), FTriggerReadRef(TriggerIn));
+		InputDataReferences.AddDataReadReference(METASOUND_GET_PARAM_NAME(InputInTrigger), TriggerIn);
+		InputDataReferences.AddDataReadReference(METASOUND_GET_PARAM_NAME(InputReset), TriggerReset);
+		InputDataReferences.AddDataReadReference(METASOUND_GET_PARAM_NAME(InputDelayTime), DelayTime);
 		return InputDataReferences;
 	}
 
 	FDataReferenceCollection FTriggerDelayOperator::GetOutputs() const
 	{
+		using namespace TriggerDelayVertexNames;
+
 		FDataReferenceCollection OutputDataReferences;
-		OutputDataReferences.AddDataReadReference(TEXT("Out"), FTriggerReadRef(TriggerOut));
+		OutputDataReferences.AddDataReadReference(METASOUND_GET_PARAM_NAME(OutputOnTrigger), TriggerOut);
 
 		return OutputDataReferences;
 	}
@@ -102,7 +115,7 @@ namespace Metasound
 			{
 				if (StartFrame > LastFrameReset)
 				{
-					NextTriggerCounter.SetNumSamples(*Delay);
+					NextTriggerCounter.SetNumSamples(*DelayTime);
 					LastFrameReset = -1;
 				}
 				else
@@ -129,28 +142,32 @@ namespace Metasound
 
 	TUniquePtr<IOperator> FTriggerDelayOperator::CreateOperator(const FCreateOperatorParams& InParams, FBuildErrorArray& OutErrors)
 	{
+		using namespace TriggerDelayVertexNames;
+
 		const FTriggerDelayNode& DelayNode = static_cast<const FTriggerDelayNode&>(InParams.Node);
 
 		const FInputVertexInterface& InputInterface = GetVertexInterface().GetInputInterface();
 		FTimeReadRef Delay = InParams.InputDataReferences.GetDataReadReferenceOrConstructWithVertexDefault<FTime>(InputInterface, TEXT("Delay"), InParams.OperatorSettings);
 
-		FTriggerReadRef TriggerIn = InParams.InputDataReferences.GetDataReadReferenceOrConstruct<FTrigger>(TEXT("In"), InParams.OperatorSettings);
-		FTriggerReadRef TriggerReset = InParams.InputDataReferences.GetDataReadReferenceOrConstruct<FTrigger>(TEXT("Reset"), InParams.OperatorSettings);
+		FTriggerReadRef TriggerIn = InParams.InputDataReferences.GetDataReadReferenceOrConstruct<FTrigger>(METASOUND_GET_PARAM_NAME(InputInTrigger), InParams.OperatorSettings);
+		FTriggerReadRef TriggerReset = InParams.InputDataReferences.GetDataReadReferenceOrConstruct<FTrigger>(METASOUND_GET_PARAM_NAME(InputReset), InParams.OperatorSettings);
 
 		return MakeUnique<FTriggerDelayOperator>(InParams.OperatorSettings, TriggerReset, TriggerIn, Delay);
 	}
 
 	const FVertexInterface& FTriggerDelayOperator::GetVertexInterface()
 	{
+		using namespace TriggerDelayVertexNames;
+
 		static const FVertexInterface Interface(
 			FInputVertexInterface(
-				TInputDataVertexModel<FTime>(TEXT("Delay"), LOCTEXT("TriggerDelayTooltip", "Time to delay and execute deferred trigger in seconds."), 1.0f),
-				TInputDataVertexModel<FTrigger>(TEXT("In"), LOCTEXT("TriggerDelayInTooltip", "Triggers delay.")),
-				TInputDataVertexModel<FTrigger>(TEXT("Reset"), LOCTEXT("TriggerDelayResetInTooltip", "Resets the trigger delay, clearing the execution task if pending."))
+				TInputDataVertexModel<FTrigger>(METASOUND_GET_PARAM_NAME_AND_TT(InputInTrigger)),
+				TInputDataVertexModel<FTrigger>(METASOUND_GET_PARAM_NAME_AND_TT(InputReset)),
+				TInputDataVertexModel<FTime>(METASOUND_GET_PARAM_NAME_AND_TT(InputDelayTime), 1.0f)
 			),
 			FOutputVertexInterface(
-				TOutputDataVertexModel<FTrigger>(TEXT("Out"), LOCTEXT("TriggerOutTooltip", "The deferred output trigger"))
-			)
+				TOutputDataVertexModel<FTrigger>(METASOUND_GET_PARAM_NAME_AND_TT(OutputOnTrigger))
+		)
 		);
 
 		return Interface;
diff --git a/Engine/Plugins/Runtime/Metasound/Source/MetasoundStandardNodes/Private/MetasoundTriggerPipeNode.cpp b/Engine/Plugins/Runtime/Metasound/Source/MetasoundStandardNodes/Private/MetasoundTriggerPipeNode.cpp
--- a/Engine/Plugins/Runtime/Metasound/Source/MetasoundStandardNodes/Private/MetasoundTriggerPipeNode.cpp
+++ b/Engine/Plugins/Runtime/Metasound/Source/MetasoundStandardNodes/Private/MetasoundTriggerPipeNode.cpp
@@ -16,12 +16,20 @@
 #include "MetasoundTrigger.h"
 #include "MetasoundSampleCounter.h"
 #include "MetasoundStandardNodesCategories.h"
+#include "MetasoundParamHelper.h"
 
 #define LOCTEXT_NAMESPACE "MetasoundStandardNodes"
 
-
 namespace Metasound
 {
+	namespace TriggerPipeVertexNames
+	{
+		METASOUND_PARAM(InputInTrigger, "In", "Trigger to execute at a future time by the given delay amount.");
+		METASOUND_PARAM(InputReset, "Reset", "Resets the trigger delay, clearing any pending execution tasks.");
+		METASOUND_PARAM(InputDelayTime, "Delay Time", "Time to delay and execute deferred input trigger execution(s) in seconds.");
+		METASOUND_PARAM(OutputOutTrigger, "Out", "The delayed output trigger(s).");
+	}
+
 	METASOUND_REGISTER_NODE(FTriggerPipeNode)
 
 	class FTriggerPipeOperator : public TExecutableOperator<FTriggerPipeOperator>
@@ -44,35 +52,39 @@ namespace Metasound
 		FTriggerReadRef TriggerReset;
 		FTriggerWriteRef TriggerOut;
 
-		FTimeReadRef Delay;
+		FTimeReadRef DelayTime;
 
 		FSampleCount FramesPerBlock;
 		FSampleRate SampleRate;
 	};
 
-	FTriggerPipeOperator::FTriggerPipeOperator(const FOperatorSettings& InSettings, const FTriggerReadRef& InTriggerReset, const FTriggerReadRef& InTriggerIn, const FTimeReadRef& InDelay)
-	:	TriggerIn(InTriggerIn)
-	,	TriggerReset(InTriggerReset)
-	,	TriggerOut(FTriggerWriteRef::CreateNew(InSettings))
-	,	Delay(InDelay)
-	,	FramesPerBlock(InSettings.GetNumFramesPerBlock())
-	,	SampleRate(InSettings.GetSampleRate())
+	FTriggerPipeOperator::FTriggerPipeOperator(const FOperatorSettings& InSettings, const FTriggerReadRef& InTriggerReset, const FTriggerReadRef& InTriggerIn, const FTimeReadRef& InDelayTime)
+	: TriggerIn(InTriggerIn)
+	, TriggerReset(InTriggerReset)
+	, TriggerOut(FTriggerWriteRef::CreateNew(InSettings))
+	, DelayTime(InDelayTime)
+	, FramesPerBlock(InSettings.GetNumFramesPerBlock())
+	, SampleRate(InSettings.GetSampleRate())
 	{
 	}
 
 	FDataReferenceCollection FTriggerPipeOperator::GetInputs() const
 	{
+		using namespace TriggerPipeVertexNames;
+
 		FDataReferenceCollection InputDataReferences;
-		InputDataReferences.AddDataReadReference(TEXT("Delay"), FTimeReadRef(Delay));
-		InputDataReferences.AddDataReadReference(TEXT("In"), FTriggerReadRef(TriggerIn));
-		InputDataReferences.AddDataReadReference(TEXT("Reset"), FTriggerReadRef(TriggerIn));
+		InputDataReferences.AddDataReadReference(METASOUND_GET_PARAM_NAME(InputInTrigger), TriggerIn);
+		InputDataReferences.AddDataReadReference(METASOUND_GET_PARAM_NAME(InputReset), TriggerReset);
+		InputDataReferences.AddDataReadReference(METASOUND_GET_PARAM_NAME(InputDelayTime), DelayTime);
 		return InputDataReferences;
 	}
 
 	FDataReferenceCollection FTriggerPipeOperator::GetOutputs() const
 	{
+		using namespace TriggerPipeVertexNames;
+
 		FDataReferenceCollection OutputDataReferences;
-		OutputDataReferences.AddDataReadReference(TEXT("Out"), FTriggerReadRef(TriggerOut));
+		OutputDataReferences.AddDataReadReference(METASOUND_GET_PARAM_NAME(OutputOutTrigger), TriggerOut);
 
 		return OutputDataReferences;
 	}
@@ -88,7 +100,7 @@ namespace Metasound
 			},
 			[this](int32 StartFrame, int32 EndFrame)
 			{
-				SamplesUntilTrigger.AddUnique(Delay->GetSeconds() * SampleRate);
+				SamplesUntilTrigger.AddUnique(DelayTime->GetSeconds() * SampleRate);
 			}
 		);
 
@@ -128,27 +140,31 @@ namespace Metasound
 
 	TUniquePtr<IOperator> FTriggerPipeOperator::CreateOperator(const FCreateOperatorParams& InParams, FBuildErrorArray& OutErrors)
 	{
+		using namespace TriggerPipeVertexNames;
+
 		const FTriggerPipeNode& DelayNode = static_cast<const FTriggerPipeNode&>(InParams.Node);
 		const FDataReferenceCollection& InputCollection = InParams.InputDataReferences;
 		const FInputVertexInterface& InputInterface = GetVertexInterface().GetInputInterface();
 
-		FTimeReadRef Delay = InputCollection.GetDataReadReferenceOrConstructWithVertexDefault<FTime>(InputInterface, TEXT("Delay"), InParams.OperatorSettings);
-		FTriggerReadRef TriggerIn = InParams.InputDataReferences.GetDataReadReferenceOrConstruct<FTrigger>(TEXT("In"), InParams.OperatorSettings);
-		FTriggerReadRef TriggerReset = InParams.InputDataReferences.GetDataReadReferenceOrConstruct<FTrigger>(TEXT("Reset"), InParams.OperatorSettings);
+		FTimeReadRef DelayTime = InputCollection.GetDataReadReferenceOrConstructWithVertexDefault<FTime>(InputInterface, METASOUND_GET_PARAM_NAME(InputDelayTime), InParams.OperatorSettings);
+		FTriggerReadRef TriggerIn = InParams.InputDataReferences.GetDataReadReferenceOrConstruct<FTrigger>(METASOUND_GET_PARAM_NAME(InputDelayTime), InParams.OperatorSettings);
+		FTriggerReadRef TriggerReset = InParams.InputDataReferences.GetDataReadReferenceOrConstruct<FTrigger>(METASOUND_GET_PARAM_NAME(InputReset), InParams.OperatorSettings);
 
-		return MakeUnique<FTriggerPipeOperator>(InParams.OperatorSettings, TriggerReset, TriggerIn, Delay);
+		return MakeUnique<FTriggerPipeOperator>(InParams.OperatorSettings, TriggerReset, TriggerIn, DelayTime);
 	}
 
 	const FVertexInterface& FTriggerPipeOperator::GetVertexInterface()
 	{
+		using namespace TriggerPipeVertexNames;
+
 		static const FVertexInterface Interface(
 			FInputVertexInterface(
-				TInputDataVertexModel<FTime>(TEXT("Delay"), LOCTEXT("DelayTooltip", "Time to delay and execute deferred input trigger execution(s) in seconds."), 1.0f),
-				TInputDataVertexModel<FTrigger>(TEXT("In"), LOCTEXT("TriggerInTooltip", "Trigger to execute at a future time by the given delay amount.")),
-				TInputDataVertexModel<FTrigger>(TEXT("Reset"), LOCTEXT("TriggerDelayResetInTooltip", "Resets the trigger delay, clearing any pending execution tasks."))
-			),
+				TInputDataVertexModel<FTrigger>(METASOUND_GET_PARAM_NAME_AND_TT(InputDelayTime)),
+				TInputDataVertexModel<FTrigger>(METASOUND_GET_PARAM_NAME_AND_TT(InputDelayTime)),
+				TInputDataVertexModel<FTime>(METASOUND_GET_PARAM_NAME_AND_TT(InputReset), 1.0f)
+				),
 			FOutputVertexInterface(
-				TOutputDataVertexModel<FTrigger>(TEXT("Out"), LOCTEXT("TriggerOutTooltip", "The deferred output trigger"))
+				TOutputDataVertexModel<FTrigger>(METASOUND_GET_PARAM_NAME_AND_TT(OutputOutTrigger))
 			)
 		);
 
@@ -160,11 +176,11 @@ namespace Metasound
 		auto InitNodeInfo = []() -> FNodeClassMetadata
 		{
 			FNodeClassMetadata Info;
-			Info.ClassName = FNodeClassName(StandardNodes::Namespace, "Pipe", "Trigger");
+			Info.ClassName = FNodeClassName(StandardNodes::Namespace, "Pipe", "");
 			Info.MajorVersion = 1;
 			Info.MinorVersion = 0;
 			Info.DisplayName = LOCTEXT("PipeTriggerNode_NodeDisplayName", "Trigger Pipe");
-			Info.Description = LOCTEXT("Metasound_DelayNodeDescription", "Delays execution of the input trigger by the given delay for all input trigger executions.");
+			Info.Description = LOCTEXT("Metasound_DelayNodeDescription", "Delays execution of the input trigger(s) by the given delay for all input trigger executions.");
 			Info.Author = PluginAuthor;
 			Info.PromptIfMissing = PluginNodeMissingPrompt;
 			Info.DefaultInterface = GetVertexInterface();
diff --git a/Engine/Plugins/Runtime/Metasound/Source/MetasoundStandardNodes/Private/MetasoundTriggerRepeatNode.cpp b/Engine/Plugins/Runtime/Metasound/Source/MetasoundStandardNodes/Private/MetasoundTriggerRepeatNode.cpp
--- a/Engine/Plugins/Runtime/Metasound/Source/MetasoundStandardNodes/Private/MetasoundTriggerRepeatNode.cpp
+++ b/Engine/Plugins/Runtime/Metasound/Source/MetasoundStandardNodes/Private/MetasoundTriggerRepeatNode.cpp
@@ -15,6 +15,8 @@
 #include "MetasoundPrimitives.h"
 #include "MetasoundSampleCounter.h"
 #include "MetasoundStandardNodesNames.h"
+#include "MetasoundStandardNodesCategories.h"
+#include "MetasoundParamHelper.h"
 
 #define LOCTEXT_NAMESPACE "MetasoundStandardNodes"
 
@@ -24,56 +26,12 @@ namespace Metasound
 
 	namespace TriggerRepeatVertexNames
 	{
-		const FString& GetInputStartName()
-		{
-			static FString Name = TEXT("Start");
-			return Name;
-		}
-
-		const FText& GetInputStartDescription()
-		{
-			static FText Desc = LOCTEXT("TriggerEnableTooltip", "Starts executing periodic output triggers.");
-			return Desc;
-		}
-
-		const FString& GetInputStopName()
-		{
-			static FString Name = TEXT("Stop");
-			return Name;
-		}
-
-		const FText& GetInputStopDescription()
-		{
-			static FText Desc = LOCTEXT("TriggerDisableTooltip", "Stops executing periodic output triggers.");
-			return Desc;
-		}
-
-		const FString& GetInputPeriodName()
-		{
-			static FString Name = TEXT("Period");
-			return Name;
-		}
-
-		const FText& GetInputPeriodDescription()
-		{
-			static FText Desc = LOCTEXT("PeriodTooltip", "The period to trigger in seconds.");
-			return Desc;
-		}
-
-		const FString& GetOutputTriggerName()
-		{
-			static FString Name = TEXT("Out");
-			return Name;
-		}
-
-		const FText& GetOutputTriggerDescription()
-		{
-			static FText Desc = LOCTEXT("TriggerOutTooltip", "The periodically generated output trigger");
-			return Desc;
-		}
+		METASOUND_PARAM(InputStart, "Start", "Starts executing periodic output triggers.");
+		METASOUND_PARAM(InputStop, "Stop", "Stops executing periodic output triggers.");
+		METASOUND_PARAM(InputPeriod, "Period", "The period to trigger in seconds.");
+		METASOUND_PARAM(OutputOnTrigger, "Out", "The periodically generated output trigger.");
 	}
 
-
 	class FTriggerRepeatOperator : public TExecutableOperator<FTriggerRepeatOperator>
 	{
 		public:
@@ -104,13 +62,13 @@ namespace Metasound
 	};
 
 	FTriggerRepeatOperator::FTriggerRepeatOperator(const FOperatorSettings& InSettings, const FTriggerReadRef& InTriggerEnable, const FTriggerReadRef& InTriggerDisable, const FTimeReadRef& InPeriod)
-	:	bEnabled(false)
-	,	TriggerOut(FTriggerWriteRef::CreateNew(InSettings))
-	,	TriggerEnable(InTriggerEnable)
-	,	TriggerDisable(InTriggerDisable)
-	,	Period(InPeriod)
-	,	BlockSize(InSettings.GetNumFramesPerBlock())
-	,	SampleCounter(0, InSettings.GetSampleRate())
+		: bEnabled(false)
+		, TriggerOut(FTriggerWriteRef::CreateNew(InSettings))
+		, TriggerEnable(InTriggerEnable)
+		, TriggerDisable(InTriggerDisable)
+		, Period(InPeriod)
+		, BlockSize(InSettings.GetNumFramesPerBlock())
+		, SampleCounter(0, InSettings.GetSampleRate())
 	{
 	}
 
@@ -119,9 +77,9 @@ namespace Metasound
 		using namespace TriggerRepeatVertexNames;
 
 		FDataReferenceCollection InputDataReferences;
-		InputDataReferences.AddDataReadReference(GetInputStartName(), FTriggerReadRef(TriggerEnable));
-		InputDataReferences.AddDataReadReference(GetInputStopName(), FTriggerReadRef(TriggerDisable));
-		InputDataReferences.AddDataReadReference(GetInputPeriodName(), FTimeReadRef(Period));
+		InputDataReferences.AddDataReadReference(METASOUND_GET_PARAM_NAME(InputStart), TriggerEnable);
+		InputDataReferences.AddDataReadReference(METASOUND_GET_PARAM_NAME(InputStop), TriggerDisable);
+		InputDataReferences.AddDataReadReference(METASOUND_GET_PARAM_NAME(InputPeriod), Period);
 		return InputDataReferences;
 	}
 
@@ -130,7 +88,7 @@ namespace Metasound
 		using namespace TriggerRepeatVertexNames;
 
 		FDataReferenceCollection OutputDataReferences;
-		OutputDataReferences.AddDataReadReference(GetOutputTriggerName(), FTriggerReadRef(TriggerOut));
+		OutputDataReferences.AddDataReadReference(METASOUND_GET_PARAM_NAME(OutputOnTrigger), TriggerOut);
 
 		return OutputDataReferences;
 	}
@@ -181,9 +139,9 @@ namespace Metasound
 		const FTriggerRepeatNode& PeriodicTriggerNode = static_cast<const FTriggerRepeatNode&>(InParams.Node);
 		const FInputVertexInterface& InputInterface = GetVertexInterface().GetInputInterface();
 
-		FTriggerReadRef TriggerEnable = InParams.InputDataReferences.GetDataReadReferenceOrConstruct<FTrigger>(GetInputStartName(), InParams.OperatorSettings);
-		FTriggerReadRef TriggerDisable = InParams.InputDataReferences.GetDataReadReferenceOrConstruct<FTrigger>(GetInputStopName(), InParams.OperatorSettings);
-		FTimeReadRef Period = InParams.InputDataReferences.GetDataReadReferenceOrConstructWithVertexDefault<FTime>(InputInterface, GetInputPeriodName(), InParams.OperatorSettings);
+		FTriggerReadRef TriggerEnable = InParams.InputDataReferences.GetDataReadReferenceOrConstruct<FTrigger>(METASOUND_GET_PARAM_NAME(InputStart), InParams.OperatorSettings);
+		FTriggerReadRef TriggerDisable = InParams.InputDataReferences.GetDataReadReferenceOrConstruct<FTrigger>(METASOUND_GET_PARAM_NAME(InputStop), InParams.OperatorSettings);
+		FTimeReadRef Period = InParams.InputDataReferences.GetDataReadReferenceOrConstructWithVertexDefault<FTime>(InputInterface, METASOUND_GET_PARAM_NAME(InputPeriod), InParams.OperatorSettings);
 
 		return MakeUnique<FTriggerRepeatOperator>(InParams.OperatorSettings, TriggerEnable, TriggerDisable, Period);
 	}
@@ -194,12 +152,12 @@ namespace Metasound
 
 		static const FVertexInterface Interface(
 			FInputVertexInterface(
-				TInputDataVertexModel<FTrigger>(GetInputStartName(), GetInputStartDescription()),
-				TInputDataVertexModel<FTrigger>(GetInputStopName(), GetInputStopDescription()),
-				TInputDataVertexModel<FTime>(GetInputPeriodName(), GetInputPeriodDescription(), 0.2f)
+				TInputDataVertexModel<FTrigger>(METASOUND_GET_PARAM_NAME_AND_TT(InputStart)),
+				TInputDataVertexModel<FTrigger>(METASOUND_GET_PARAM_NAME_AND_TT(InputStop)),
+				TInputDataVertexModel<FTime>(METASOUND_GET_PARAM_NAME_AND_TT(InputPeriod), 0.2f)
 			),
 			FOutputVertexInterface(
-				TOutputDataVertexModel<FTrigger>(GetOutputTriggerName(), GetOutputTriggerDescription())
+				TOutputDataVertexModel<FTrigger>(METASOUND_GET_PARAM_NAME_AND_TT(OutputOnTrigger))
 			)
 		);
 
@@ -219,6 +177,7 @@ namespace Metasound
 			Info.Author = PluginAuthor;
 			Info.PromptIfMissing = PluginNodeMissingPrompt;
 			Info.DefaultInterface = GetVertexInterface();
+			Info.CategoryHierarchy.Emplace(StandardNodes::TriggerUtils);
 
 			return Info;
 		};
diff --git a/Engine/Plugins/Runtime/Metasound/Source/MetasoundStandardNodes/Private/MetasoundValueNode.cpp b/Engine/Plugins/Runtime/Metasound/Source/MetasoundStandardNodes/Private/MetasoundValueNode.cpp
--- a/Engine/Plugins/Runtime/Metasound/Source/MetasoundStandardNodes/Private/MetasoundValueNode.cpp
+++ b/Engine/Plugins/Runtime/Metasound/Source/MetasoundStandardNodes/Private/MetasoundValueNode.cpp
@@ -32,34 +32,6 @@ namespace Metasound
 	}
 
 
-	namespace ValueVertexNames
-	{
-		const FString& GetInitValueName()
-		{
-			static const FString Name = TEXT("Init");
-			return Name;
-		}
-
-		const FString& GetSetValueName()
-		{
-			static const FString Name = TEXT("Value");
-			return Name;
-		}
-
-		const FString& GetInputTriggerName()
-		{
-			static const FString Name = TEXT("Set");
-			return Name;
-		}
-
-		const FString& GetOutputValueName()
-		{
-			static const FString Name = TEXT("Value");
-			return Name;
-		}
-	}
-
-
 	using FValueNodeInt32 = TValueNode<int32>;
  	METASOUND_REGISTER_NODE(FValueNodeInt32)
 
diff --git a/Engine/Plugins/Runtime/Metasound/Source/MetasoundStandardNodes/Public/MetasoundParamHelper.h b/Engine/Plugins/Runtime/Metasound/Source/MetasoundStandardNodes/Public/MetasoundParamHelper.h
new file mode 100644
--- /dev/null
+++ b/Engine/Plugins/Runtime/Metasound/Source/MetasoundStandardNodes/Public/MetasoundParamHelper.h
@@ -0,0 +1,55 @@
+
+// Copyright Epic Games, Inc. All Rights Reserved.
+
+
+/*
+
+ Macros to make parameter definitions a little easier
+ Use in a namespace
+
+namespace Metasound
+{
+	namespace MyNodeVertexNames
+	{
+		METASOUND_PARAM(InputParam1, "Param 1", "Tooltip for param1.");
+		METASOUND_PARAM(InputParam2, "Param 2", "Tooltip for param2.");
+		METASOUND_PARAM(OutputParam1, "Out Param1", "Tooltip for output param 2.");
+		METASOUND_PARAM(OutputParam2, "Out Param2", "Tooltip for output param 2.");
+	}
+}
+
+Then to retrieve just do when in a:
+
+using namespace MyNodeVertexNames;
+
+METASOUND_GET_PARAM_NAME(InputParam1)
+METASOUND_GET_PARAM_TT(InputParam1)
+
+For variable parameters (e.g. templated on number), do the following:
+
+	namespace TriggerAccumulatorVertexNames
+	{
+		METASOUND_PARAM(InputParam, "Param {0}", "Do {0} tooltip.");
+	}
+
+Then to retrieve the name/tooltip use:
+
+METASOUND_GET_VARIABLE_PARAM_NAME(InputParam, NUMBER);
+
+Where NUMBER is a number (or whatever) to slot into the format specifier you defined in the param def.
+
+*/
+
+
+#define METASOUND_PARAM(NAME, NAME_TEXT, TOOLTIP_TEXT) \
+	static const TCHAR* NAME##Name = TEXT(NAME_TEXT); \
+	static const FText NAME##Tooltip = LOCTEXT(#NAME "Tooltip", TOOLTIP_TEXT); \
+
+
+#define METASOUND_GET_PARAM_NAME(NAME) NAME##Name
+#define METASOUND_GET_PARAM_TT(NAME) NAME##Tooltip
+#define METASOUND_GET_PARAM_NAME_AND_TT(NAME) NAME##Name, NAME##Tooltip
+
+#define METASOUND_GET_VARIABLE_PARAM_NAME(NAME, INDEX) FString::Format(NAME##Name, {INDEX})
+#define METASOUND_GET_VARIABLE_TT(NAME, INDEX)  FText::Format(NAME##Tooltip, INDEX)
+#define METASOUND_GET_VARIABLE_PARAM_NAME_AND_TT(NAME, INDEX)  FString::Format(NAME##Name, {INDEX}), FText::Format(NAME##Tooltip, INDEX)
\ No newline at end of file
diff --git a/Engine/Plugins/Runtime/Metasound/Source/MetasoundStandardNodes/Public/MetasoundTriggerAccumulatorNode.h b/Engine/Plugins/Runtime/Metasound/Source/MetasoundStandardNodes/Public/MetasoundTriggerAccumulatorNode.h
--- a/Engine/Plugins/Runtime/Metasound/Source/MetasoundStandardNodes/Public/MetasoundTriggerAccumulatorNode.h
+++ b/Engine/Plugins/Runtime/Metasound/Source/MetasoundStandardNodes/Public/MetasoundTriggerAccumulatorNode.h
@@ -8,9 +8,11 @@
 #include "MetasoundStandardNodesNames.h"
 #include "MetasoundTrigger.h"
 #include "Internationalization/Text.h"
+#include "MetasoundParamHelper.h"
 
 #define LOCTEXT_NAMESPACE "MetasoundStandardNodes"
 
+
 namespace Metasound
 {
 	namespace MetasoundTriggerAccumulatorNodePrivate
@@ -20,14 +22,9 @@ namespace Metasound
 
 	namespace TriggerAccumulatorVertexNames
 	{
-		METASOUNDSTANDARDNODES_API const FString GetInputAutoResetName();
-		METASOUNDSTANDARDNODES_API const FText GetInputAutoResetDescription();
-
-		METASOUNDSTANDARDNODES_API const FString GetInputTriggerName(uint32 InIndex);
-		METASOUNDSTANDARDNODES_API const FText GetInputTriggerDescription(uint32 InIndex);
-
-		METASOUNDSTANDARDNODES_API const FString& GetOutputTriggerName();
-		METASOUNDSTANDARDNODES_API const FText& GetOutputTriggerDescription();
+		METASOUND_PARAM(InputAutoReset, "Auto Reset", "Input trigger which results in a delayed trigger.");	
+		METASOUND_PARAM(InputTrigger, "In {0}", "Trigger {0} input. All trigger inputs must be triggered before the output trigger is hit.");
+		METASOUND_PARAM(OutputOnTrigger, "Out", "Triggered when all input triggers have been triggered. Call Reset to reset the state or use \"Auto Reset\".");
 	}
 
 	template<uint32 NumInputs>
@@ -42,15 +39,15 @@ namespace Metasound
 			{
 				FInputVertexInterface InputInterface;
 
-				InputInterface.Add(TInputDataVertexModel<bool>(GetInputAutoResetName(), GetInputAutoResetDescription()));
-
 				for (uint32 i = 0; i < NumInputs; ++i)
 				{
-					InputInterface.Add(TInputDataVertexModel<FTrigger>(GetInputTriggerName(i), GetInputTriggerDescription(i)));
+					InputInterface.Add(TInputDataVertexModel<FTrigger>(METASOUND_GET_VARIABLE_PARAM_NAME_AND_TT(InputTrigger, i)));
 				}
 
+				InputInterface.Add(TInputDataVertexModel<bool>(METASOUND_GET_PARAM_NAME_AND_TT(InputAutoReset)));
+
 				FOutputVertexInterface OutputInterface;
-				OutputInterface.Add(TOutputDataVertexModel<FTrigger>(GetOutputTriggerName(), GetOutputTriggerDescription()));
+				OutputInterface.Add(TOutputDataVertexModel<FTrigger>(METASOUND_GET_PARAM_NAME_AND_TT(OutputOnTrigger)));
 
 				return FVertexInterface(InputInterface, OutputInterface);
 			};
@@ -82,12 +79,12 @@ namespace Metasound
 			const FInputVertexInterface& InputInterface = InParams.Node.GetVertexInterface().GetInputInterface();
 			const FDataReferenceCollection& InputCollection = InParams.InputDataReferences;
 
-			FBoolReadRef bInAutoReset = InputCollection.GetDataReadReferenceOrConstructWithVertexDefault<bool>(InputInterface, GetInputAutoResetName(), InParams.OperatorSettings);
+			FBoolReadRef bInAutoReset = InputCollection.GetDataReadReferenceOrConstructWithVertexDefault<bool>(InputInterface, METASOUND_GET_PARAM_NAME(InputAutoReset), InParams.OperatorSettings);
 			TArray<FTriggerReadRef> InputTriggers;
 
 			for (uint32 i = 0; i < NumInputs; ++i)
 			{
-				InputTriggers.Add(InputCollection.GetDataReadReferenceOrConstruct<FTrigger>(GetInputTriggerName(i), InParams.OperatorSettings));
+				InputTriggers.Add(InputCollection.GetDataReadReferenceOrConstruct<FTrigger>(METASOUND_GET_VARIABLE_PARAM_NAME(InputTrigger, i), InParams.OperatorSettings));
 			}
 
 			return MakeUnique<TTriggerAccumulatorOperator<NumInputs>>(InParams.OperatorSettings, bInAutoReset, MoveTemp(InputTriggers));
@@ -111,11 +108,11 @@ namespace Metasound
 			using namespace TriggerAccumulatorVertexNames;
 
 			FDataReferenceCollection Inputs;
-			Inputs.AddDataReadReference(GetInputAutoResetName(), bAutoReset);
 			for (uint32 i = 0; i < NumInputs; ++i)
 			{
-				Inputs.AddDataReadReference(GetInputTriggerName(i), InputTriggers[i]);
+				Inputs.AddDataReadReference(METASOUND_GET_VARIABLE_PARAM_NAME(InputTrigger, i), InputTriggers[i]);
 			}
+			Inputs.AddDataReadReference(METASOUND_GET_PARAM_NAME(InputAutoReset), bAutoReset);
 
 			return Inputs;
 		}
@@ -125,7 +122,7 @@ namespace Metasound
 			using namespace TriggerAccumulatorVertexNames;
 
 			FDataReferenceCollection Outputs;
-			Outputs.AddDataReadReference(GetOutputTriggerName(), OutputTrigger);
+			Outputs.AddDataReadReference(METASOUND_GET_PARAM_NAME(OutputOnTrigger), OutputTrigger);
 
 			return Outputs;
 		}
diff --git a/Engine/Plugins/Runtime/Metasound/Source/MetasoundStandardNodes/Public/MetasoundTriggerCompareNode.h b/Engine/Plugins/Runtime/Metasound/Source/MetasoundStandardNodes/Public/MetasoundTriggerCompareNode.h
--- a/Engine/Plugins/Runtime/Metasound/Source/MetasoundStandardNodes/Public/MetasoundTriggerCompareNode.h
+++ b/Engine/Plugins/Runtime/Metasound/Source/MetasoundStandardNodes/Public/MetasoundTriggerCompareNode.h
@@ -7,15 +7,235 @@
 #include "MetasoundDataReferenceCollection.h"
 #include "MetasoundExecutableOperator.h"
 #include "MetasoundFacade.h"
+#include "MetasoundNode.h"
 #include "MetasoundNodeInterface.h"
+#include "MetasoundNodeRegistrationMacro.h"
+#include "MetasoundDataTypeRegistrationMacro.h"
+#include "MetasoundOperatorInterface.h"
+#include "MetasoundPrimitives.h"
+#include "MetasoundStandardNodesNames.h"
+#include "MetasoundStandardNodesCategories.h"
+#include "MetasoundParamHelper.h"
+
+#define LOCTEXT_NAMESPACE "MetasoundStandardNodes"
 
 namespace Metasound
 {
-	class METASOUNDSTANDARDNODES_API FTriggerCompareNode : public FNodeFacade
+	namespace TriggerCompareVertexNames
+	{
+		METASOUND_PARAM(InputCompare, "Compare", "Trigger to compare A and B.");
+		METASOUND_PARAM(InputParamA, "A", "The first value, A, to compare against.");
+		METASOUND_PARAM(InputParamB, "B", "The first value, B, to compare against.");
+		METASOUND_PARAM(InputCompareType, "Type", "How to compare A and B.");
+		METASOUND_PARAM(OutputOnTrue, "True", "Output trigger for when the comparison is true.");
+		METASOUND_PARAM(OutputOnFalse, "False", "Output trigger for when the comparison is false.");
+	}
+
+	namespace MetasoundTriggerCompareNodePrivate
 	{
-		public:
-			FTriggerCompareNode(const FNodeInitData& InInitData);
+		METASOUNDSTANDARDNODES_API FNodeClassMetadata CreateNodeClassMetadata(const FName& InDataTypeName, const FName& InOperatorName, const FText& InDisplayName, const FText& InDescription, const FVertexInterface& InDefaultInterface);
+	}
 
-			virtual ~FTriggerCompareNode() = default;
+	enum class ETriggerComparisonType
+	{
+		Equals,
+		NotEquals,
+		LessThan,
+		GreaterThan,
+		LessThanOrEquals,
+		GreaterThanOrEquals
+	};
+
+	DECLARE_METASOUND_ENUM(ETriggerComparisonType, ETriggerComparisonType::Equals, METASOUNDSTANDARDNODES_API,
+	FEnumTriggerComparisonType, FEnumTriggerComparisonTypeInfo, FEnumTriggerComparisonTypeReadRef, FEnumTriggerComparisonTypeWriteRef);
+
+	DEFINE_METASOUND_ENUM_BEGIN(ETriggerComparisonType, FEnumTriggerComparisonType, "TriggerComparisonType")
+		DEFINE_METASOUND_ENUM_ENTRY(ETriggerComparisonType::Equals, LOCTEXT("EqualsDescription", "Equals"), LOCTEXT("EqualsDescriptionTT", "True if A and B are equal.")),
+		DEFINE_METASOUND_ENUM_ENTRY(ETriggerComparisonType::NotEquals, LOCTEXT("NotEqualsDescriptioin", "Not Equals"), LOCTEXT("NotEqualsTT", "True if A and B are not equal.")),
+		DEFINE_METASOUND_ENUM_ENTRY(ETriggerComparisonType::LessThan, LOCTEXT("LessThanDescription", "Less Than"), LOCTEXT("LessThanTT", "True if A is less than B.")),
+		DEFINE_METASOUND_ENUM_ENTRY(ETriggerComparisonType::GreaterThan, LOCTEXT("GreaterThanDescription", "Greater Than"), LOCTEXT("GreaterThanTT", "True if A is greater than B.")),
+		DEFINE_METASOUND_ENUM_ENTRY(ETriggerComparisonType::LessThanOrEquals, LOCTEXT("LessThanOrEqualsDescription", "Less Than Or Equals"), LOCTEXT("LessThanOrEqualsTT", "True if A is less than or equal to B.")),
+		DEFINE_METASOUND_ENUM_ENTRY(ETriggerComparisonType::GreaterThanOrEquals, LOCTEXT("GreaterThanOrEqualsDescription", "Greater Than Or Equals"), LOCTEXT("GreaterThanOrEqualsTT", "True if A is greater than or equal to B.")),
+	DEFINE_METASOUND_ENUM_END()
+
+
+	template<typename ValueType>
+	class TTriggerCompareNodeOperator : public TExecutableOperator<TTriggerCompareNodeOperator<ValueType>>
+	{
+	public:
+		static const FVertexInterface& GetDefaultInterface()
+		{
+			using namespace TriggerCompareVertexNames;
+			using namespace MetasoundTriggerCompareNodePrivate;
+
+			static const FVertexInterface DefaultInterface(
+				FInputVertexInterface(
+					TInputDataVertexModel<FTrigger>(METASOUND_GET_PARAM_NAME_AND_TT(InputCompare)),
+					TInputDataVertexModel<ValueType>(METASOUND_GET_PARAM_NAME_AND_TT(InputParamA)),
+					TInputDataVertexModel<ValueType>(METASOUND_GET_PARAM_NAME_AND_TT(InputParamB)),
+					TInputDataVertexModel<FEnumTriggerComparisonType>(METASOUND_GET_PARAM_NAME_AND_TT(InputCompareType))
+				),
+				FOutputVertexInterface(
+					TOutputDataVertexModel<FTrigger>(METASOUND_GET_PARAM_NAME_AND_TT(OutputOnTrue)),
+					TOutputDataVertexModel<FTrigger>(METASOUND_GET_PARAM_NAME_AND_TT(OutputOnFalse))
+				)
+			);
+
+			return DefaultInterface;
+		}
+
+		static const FNodeClassMetadata& GetNodeInfo()
+		{
+			auto CreateNodeClassMetadata = []() -> FNodeClassMetadata
+			{
+				FName DataTypeName = GetMetasoundDataTypeName<ValueType>();
+				FName OperatorName = TEXT("Clamp");
+				FText NodeDisplayName = FText::Format(LOCTEXT("TriggerCompareDisplayPattern", "Trigger Compare ({0})"), FText::FromString(GetMetasoundDataTypeString<ValueType>()));
+				FText NodeDescription = LOCTEXT("TriggerCompareDisc", "Output triggers (True or False) based on comparing inputs, A and B.");
+				FVertexInterface NodeInterface = GetDefaultInterface();
+
+				return MetasoundTriggerCompareNodePrivate::CreateNodeClassMetadata(DataTypeName, OperatorName, NodeDisplayName, NodeDescription, NodeInterface);
+			};
+
+			static const FNodeClassMetadata Metadata = CreateNodeClassMetadata();
+			return Metadata;
+		}
+
+		static TUniquePtr<IOperator> CreateOperator(const FCreateOperatorParams& InParams, TArray<TUniquePtr<IOperatorBuildError>>& OutErrors)
+		{
+			using namespace TriggerCompareVertexNames;
+
+			const FDataReferenceCollection& InputCollection = InParams.InputDataReferences;
+			const FInputVertexInterface& InputInterface = GetDefaultInterface().GetInputInterface();
+
+			FTriggerReadRef InOnTriggerCompare = InputCollection.GetDataReadReferenceOrConstruct<FTrigger>(METASOUND_GET_PARAM_NAME(InputCompare), InParams.OperatorSettings);
+			TDataReadReference<ValueType> InValueA = InputCollection.GetDataReadReferenceOrConstructWithVertexDefault<ValueType>(InputInterface, METASOUND_GET_PARAM_NAME(InputParamA), InParams.OperatorSettings);
+			TDataReadReference<ValueType> InValueB = InputCollection.GetDataReadReferenceOrConstructWithVertexDefault<ValueType>(InputInterface, METASOUND_GET_PARAM_NAME(InputParamB), InParams.OperatorSettings);
+			FEnumTriggerComparisonTypeReadRef InComparison = InputCollection.GetDataReadReferenceOrConstruct<FEnumTriggerComparisonType>(METASOUND_GET_PARAM_NAME(InputCompareType));
+
+			return MakeUnique<TTriggerCompareNodeOperator<ValueType>>(InParams.OperatorSettings, InOnTriggerCompare, InValueA, InValueB, InComparison);
+
+		}
+
+
+		TTriggerCompareNodeOperator(
+			const FOperatorSettings& InSettings, 
+			const FTriggerReadRef& InOnCompareTrigger, 
+			const TDataReadReference<ValueType>& InValueA,
+			const TDataReadReference<ValueType>& InValueB,
+			const FEnumTriggerComparisonTypeReadRef& InTriggerComparisonType)
+			: OnCompareTrigger(InOnCompareTrigger)
+			, ValueA(InValueA)
+			, ValueB(InValueB)
+			, TriggerComparisonType(InTriggerComparisonType)
+			, TriggerOutOnTrue(FTriggerWriteRef::CreateNew(InSettings))
+			, TriggerOutOnFalse(FTriggerWriteRef::CreateNew(InSettings))
+		{
+		}
+
+		virtual ~TTriggerCompareNodeOperator() = default;
+
+		virtual FDataReferenceCollection GetInputs() const override
+		{
+			using namespace TriggerCompareVertexNames;
+
+			FDataReferenceCollection Inputs;
+			Inputs.AddDataReadReference(METASOUND_GET_PARAM_NAME(InputCompare), OnCompareTrigger);
+			Inputs.AddDataReadReference(METASOUND_GET_PARAM_NAME(InputParamA), ValueA);
+			Inputs.AddDataReadReference(METASOUND_GET_PARAM_NAME(InputParamB), ValueB);
+			Inputs.AddDataReadReference(METASOUND_GET_PARAM_NAME(InputCompareType), TriggerComparisonType);
+
+			return Inputs;
+		}
+
+		virtual FDataReferenceCollection GetOutputs() const override
+		{
+			using namespace TriggerCompareVertexNames;
+
+			FDataReferenceCollection Outputs;
+			Outputs.AddDataReadReference(METASOUND_GET_PARAM_NAME(OutputOnTrue), TriggerOutOnTrue);
+			Outputs.AddDataReadReference(METASOUND_GET_PARAM_NAME(OutputOnFalse), TriggerOutOnFalse);
+
+			return Outputs;
+		}
+
+		void Execute()
+		{
+			TriggerOutOnTrue->AdvanceBlock();
+			TriggerOutOnFalse->AdvanceBlock();
+
+			OnCompareTrigger->ExecuteBlock(
+				[&](int32 StartFrame, int32 EndFrame)
+				{
+				},
+				[this](int32 StartFrame, int32 EndFrame)
+				{
+					ValueType CurrValA = *ValueA;
+					ValueType CurrValB = *ValueB;
+					bool bIsTrue = false;
+
+					switch (*TriggerComparisonType)
+					{
+					case ETriggerComparisonType::Equals:
+						bIsTrue = (CurrValA == CurrValB);
+						break;
+					case ETriggerComparisonType::NotEquals:
+						bIsTrue = (CurrValA != CurrValB);
+						break;
+					case ETriggerComparisonType::LessThan:
+						bIsTrue = (CurrValA < CurrValB);
+						break;
+					case ETriggerComparisonType::GreaterThan:
+						bIsTrue = (CurrValA > CurrValB);
+						break;
+					case ETriggerComparisonType::LessThanOrEquals:
+						bIsTrue = (CurrValA <= CurrValB);
+						break;
+					case ETriggerComparisonType::GreaterThanOrEquals:
+						bIsTrue = (CurrValA >= CurrValB);
+						break;
+					}
+
+					if (bIsTrue)
+					{
+						TriggerOutOnTrue->TriggerFrame(StartFrame);
+					}
+					else
+					{
+						TriggerOutOnFalse->TriggerFrame(StartFrame);
+					}
+				}
+				);
+		}
+
+	private:
+
+		FTriggerReadRef OnCompareTrigger;
+		TDataReadReference<ValueType> ValueA;
+		TDataReadReference<ValueType> ValueB;
+		FEnumTriggerComparisonTypeReadRef TriggerComparisonType;
+
+		FTriggerWriteRef TriggerOutOnTrue;
+		FTriggerWriteRef TriggerOutOnFalse;
+	};
+
+	/** TTriggerCompareNode
+	 *
+	 *  Compares two inputs against enumerated comparison types.
+	 */
+	template<typename ValueType>
+	class METASOUNDSTANDARDNODES_API TTriggerCompareNode : public FNodeFacade
+	{
+	public:
+		/**
+		 * Constructor used by the Metasound Frontend.
+		 */
+		TTriggerCompareNode(const FNodeInitData& InInitData)
+			: FNodeFacade(InInitData.InstanceName, InInitData.InstanceID, TFacadeOperatorClass<TTriggerCompareNodeOperator<ValueType>>())
+		{}
+
+		virtual ~TTriggerCompareNode() = default;
 	};
 }
+
+#undef LOCTEXT_NAMESPACE
\ No newline at end of file
diff --git a/Engine/Plugins/Runtime/Metasound/Source/MetasoundStandardNodes/Public/MetasoundValueNode.h b/Engine/Plugins/Runtime/Metasound/Source/MetasoundStandardNodes/Public/MetasoundValueNode.h
--- a/Engine/Plugins/Runtime/Metasound/Source/MetasoundStandardNodes/Public/MetasoundValueNode.h
+++ b/Engine/Plugins/Runtime/Metasound/Source/MetasoundStandardNodes/Public/MetasoundValueNode.h
@@ -9,6 +9,7 @@
 #include "MetasoundStandardNodesNames.h"
 #include "MetasoundTrigger.h"
 #include "Internationalization/Text.h"
+#include "MetasoundParamHelper.h"
 
 #define LOCTEXT_NAMESPACE "MetasoundStandardNodes"
 
@@ -21,10 +22,13 @@ namespace Metasound
 
 	namespace ValueVertexNames
 	{
-		METASOUNDSTANDARDNODES_API const FString& GetInitValueName();
-		METASOUNDSTANDARDNODES_API const FString& GetSetValueName();
-		METASOUNDSTANDARDNODES_API const FString& GetInputTriggerName();
-		METASOUNDSTANDARDNODES_API const FString& GetOutputValueName();
+		METASOUND_PARAM(InputSetTrigger, "Set", "Trigger to write the set value to the output.");
+		METASOUND_PARAM(InputResetTrigger, "Reset", "Trigger to reset the value to the initial value.");
+		METASOUND_PARAM(InputInitValue, "Init Value", "Value to initialize the output value to.");
+		METASOUND_PARAM(InputTargetValue, "Target Value", "Value to immediately set the output to when triggered.");
+		METASOUND_PARAM(OutputOnSet, "On Set", "Triggered when the set input is triggered.");
+		METASOUND_PARAM(OutputOnReset, "On Reset", "Triggered when the reset input is triggered.");
+		METASOUND_PARAM(OutputValue, "Output Value", "The current output value.");
 	}
 
 	template<typename ValueType>
@@ -37,12 +41,15 @@ namespace Metasound
 
 			static const FVertexInterface DefaultInterface(
 				FInputVertexInterface(
-					TInputDataVertexModel<ValueType>(GetInitValueName(), LOCTEXT("InitValue", "Value to init the output to.")),
-					TInputDataVertexModel<ValueType>(GetSetValueName(), LOCTEXT("ValueInput", "Value to set the output to when triggered.")),
-					TInputDataVertexModel<FTrigger>(GetInputTriggerName(), LOCTEXT("ValueTrigger", "Trigger to write the set value to the output."))
+					TInputDataVertexModel<FTrigger>(METASOUND_GET_PARAM_NAME_AND_TT(InputSetTrigger)),
+					TInputDataVertexModel<FTrigger>(METASOUND_GET_PARAM_NAME_AND_TT(InputResetTrigger)),
+					TInputDataVertexModel<ValueType>(METASOUND_GET_PARAM_NAME_AND_TT(InputInitValue)),
+					TInputDataVertexModel<ValueType>(METASOUND_GET_PARAM_NAME_AND_TT(InputTargetValue))
 				),
 				FOutputVertexInterface(
-					TOutputDataVertexModel<ValueType>(GetOutputValueName(), LOCTEXT("ValueOutput", "The output value cached in the node."))
+					TOutputDataVertexModel<FTrigger>(METASOUND_GET_PARAM_NAME_AND_TT(OutputOnSet)),
+					TOutputDataVertexModel<FTrigger>(METASOUND_GET_PARAM_NAME_AND_TT(OutputOnReset)),
+					TOutputDataVertexModel<ValueType>(METASOUND_GET_PARAM_NAME_AND_TT(OutputValue))
 				)
 			);
 
@@ -73,19 +80,23 @@ namespace Metasound
 			const FInputVertexInterface& InputInterface = InParams.Node.GetVertexInterface().GetInputInterface();
 			const FDataReferenceCollection& InputCollection = InParams.InputDataReferences;
 
-			FTriggerReadRef Trigger = InputCollection.GetDataReadReferenceOrConstruct<FTrigger>(GetInputTriggerName(), InParams.OperatorSettings);
-			TDataReadReference<ValueType> InitValue = InputCollection.GetDataReadReferenceOrConstructWithVertexDefault<ValueType>(InputInterface, GetInitValueName(), InParams.OperatorSettings);
-			TDataReadReference<ValueType> SetValue = InputCollection.GetDataReadReferenceOrConstructWithVertexDefault<ValueType>(InputInterface, GetSetValueName(), InParams.OperatorSettings);
+			FTriggerReadRef SetTrigger = InputCollection.GetDataReadReferenceOrConstruct<FTrigger>(METASOUND_GET_PARAM_NAME(InputSetTrigger), InParams.OperatorSettings);
+			FTriggerReadRef ResetTrigger = InputCollection.GetDataReadReferenceOrConstruct<FTrigger>(METASOUND_GET_PARAM_NAME(InputResetTrigger), InParams.OperatorSettings);
+			TDataReadReference<ValueType> InitValue = InputCollection.GetDataReadReferenceOrConstructWithVertexDefault<ValueType>(InputInterface, METASOUND_GET_PARAM_NAME(InputInitValue), InParams.OperatorSettings);
+			TDataReadReference<ValueType> TargetValue = InputCollection.GetDataReadReferenceOrConstructWithVertexDefault<ValueType>(InputInterface, METASOUND_GET_PARAM_NAME(InputTargetValue), InParams.OperatorSettings);
 
-			return MakeUnique<TValueOperator<ValueType>>(InParams.OperatorSettings, Trigger, InitValue, SetValue);
+			return MakeUnique<TValueOperator<ValueType>>(InParams.OperatorSettings, SetTrigger, ResetTrigger, InitValue, TargetValue);
 		}
 
 
-		TValueOperator(const FOperatorSettings& InSettings, TDataReadReference<FTrigger> InTrigger, TDataReadReference<ValueType> InInitValue, TDataReadReference<ValueType> InSetValue)
-			: Trigger(InTrigger)
+		TValueOperator(const FOperatorSettings& InSettings, const FTriggerReadRef& InSetTrigger, const FTriggerReadRef& InResetTrigger, const TDataReadReference<ValueType>& InInitValue, const TDataReadReference<ValueType>& InTargetValue)
+			: SetTrigger(InSetTrigger)
+			, ResetTrigger(InResetTrigger)
 			, InitValue(InInitValue)
-			, SetValue(InSetValue)
+			, TargetValue(InTargetValue)
 			, OutputValue(TDataWriteReferenceFactory<ValueType>::CreateAny(InSettings))
+			, TriggerOnSet(FTriggerWriteRef::CreateNew(InSettings))
+			, TriggerOnReset(FTriggerWriteRef::CreateNew(InSettings))
 		{
 			*OutputValue = *InitValue;
 		}
@@ -98,9 +109,10 @@ namespace Metasound
 			using namespace ValueVertexNames;
 
 			FDataReferenceCollection Inputs;
-			Inputs.AddDataReadReference(GetInputTriggerName(), Trigger);
-			Inputs.AddDataReadReference(GetInitValueName(), InitValue);
-			Inputs.AddDataReadReference(GetSetValueName(), SetValue);
+			Inputs.AddDataReadReference(METASOUND_GET_PARAM_NAME(InputSetTrigger), SetTrigger);
+			Inputs.AddDataReadReference(METASOUND_GET_PARAM_NAME(InputResetTrigger), ResetTrigger);
+			Inputs.AddDataReadReference(METASOUND_GET_PARAM_NAME(InputInitValue), InitValue);
+			Inputs.AddDataReadReference(METASOUND_GET_PARAM_NAME(InputTargetValue), TargetValue);
 
 			return Inputs;
 		}
@@ -110,25 +122,59 @@ namespace Metasound
 			using namespace ValueVertexNames;
 
 			FDataReferenceCollection Outputs;
-			Outputs.AddDataReadReference(GetOutputValueName(), OutputValue);
+			Outputs.AddDataReadReference(METASOUND_GET_PARAM_NAME(OutputOnSet), TriggerOnSet);
+			Outputs.AddDataReadReference(METASOUND_GET_PARAM_NAME(OutputOnReset), TriggerOnReset);
+			Outputs.AddDataReadReference(METASOUND_GET_PARAM_NAME(OutputValue), OutputValue);
 
 			return Outputs;
 		}
 
 		void Execute()
 		{
-			if (*Trigger)
+			TriggerOnReset->AdvanceBlock();
+			TriggerOnSet->AdvanceBlock();
+
+			if (*ResetTrigger)
 			{
-				*OutputValue = *SetValue;
+				*OutputValue = *InitValue;
 			}
+
+			if (*SetTrigger)
+			{
+				*OutputValue = *TargetValue;
+			}
+
+			ResetTrigger->ExecuteBlock(
+				[&](int32 StartFrame, int32 EndFrame)
+				{
+				},
+				[this](int32 StartFrame, int32 EndFrame)
+				{
+					TriggerOnReset->TriggerFrame(StartFrame);
+				}
+			);
+
+			SetTrigger->ExecuteBlock(
+				[&](int32 StartFrame, int32 EndFrame)
+				{
+				},
+				[this](int32 StartFrame, int32 EndFrame)
+				{
+					TriggerOnSet->TriggerFrame(StartFrame);
+				}
+			);
+
 		}
 
 	private:
 
-		TDataReadReference<FTrigger> Trigger;
+		TDataReadReference<FTrigger> SetTrigger;
+		TDataReadReference<FTrigger> ResetTrigger;
 		TDataReadReference<ValueType> InitValue;
-		TDataReadReference<ValueType> SetValue;
+		TDataReadReference<ValueType> TargetValue;
 		TDataWriteReference<ValueType> OutputValue;
+		TDataWriteReference<FTrigger> TriggerOnSet;
+		TDataWriteReference<FTrigger> TriggerOnReset;
 	};
 
 	/** TValueNode
