Fixes to ITD panner and stereo panner

- fixing some comments and adding clamping
- removing unimplemented funcs

diff --git a/Engine/Plugins/Runtime/Metasound/Source/MetasoundStandardNodes/Private/MetasoundITDPannerNode.cpp b/Engine/Plugins/Runtime/Metasound/Source/MetasoundStandardNodes/Private/MetasoundITDPannerNode.cpp
--- a/Engine/Plugins/Runtime/Metasound/Source/MetasoundStandardNodes/Private/MetasoundITDPannerNode.cpp
+++ b/Engine/Plugins/Runtime/Metasound/Source/MetasoundStandardNodes/Private/MetasoundITDPannerNode.cpp
@@ -53,7 +53,7 @@ namespace Metasound
 
 		const FText& GetInputDistanceFactorDescription()
 		{
-			static FText Desc = LOCTEXT("ITDPannerNodeDistanceFactorDesc", "The normalized distance factor (0.0 to 1.0) to use for ILD (Inter-aural level difference) calculations. 0.0 is near, 1.0f is far. The further away something is the less there is a difference in levels (gain) between the ears.");
+			static FText Desc = LOCTEXT("ITDPannerNodeDistanceFactorDesc", "The normalized distance factor (0.0 to 1.0) to use for ILD (Inter-aural level difference) calculations. 0.0 is near, 1.0 is far. The further away something is the less there is a difference in levels (gain) between the ears.");
 			return Desc;
 		}
 
@@ -113,10 +113,7 @@ namespace Metasound
 		void Execute();
 
 	private:
-		void UpdateParams();
-		void UpdateXY();
-		void UpdateGains();
-		void UpdateDelays();
+		void UpdateParams(bool bIsInit);
 
 		FAudioBufferReadRef AudioInput;
 		FFloatReadRef PanningAngle;
@@ -166,7 +163,7 @@ namespace Metasound
 		CurrDistanceFactor = FMath::Clamp(*DistanceFactor, 0.0f, 1.0f);
 		CurrHeadWidth = FMath::Max(*InHeadWidth, 0.0f);
 
-		UpdateParams();
+		UpdateParams(true);
 
 		PrevLeftGain = CurrLeftGain;
 		PrevRightGain = CurrRightGain;
@@ -196,7 +193,7 @@ namespace Metasound
 		return OutputDataReferences;
 	}
 
-	void FITDPannerOperator::UpdateParams()
+	void FITDPannerOperator::UpdateParams(bool bIsInit)
 	{
 		// ****************
 		// Update the x-y values
@@ -232,28 +229,32 @@ namespace Metasound
 
 		if (DeltaTimeSeconds > 0.0f)
 		{
-			LeftDelay.SetEasedDelayMsec(1000.0f * DeltaTimeSeconds);
-			RightDelay.SetEasedDelayMsec(0.0f);
+			LeftDelay.SetEasedDelayMsec(1000.0f * DeltaTimeSeconds, bIsInit);
+			RightDelay.SetEasedDelayMsec(0.0f, bIsInit);
 		}
 		else
 		{
-			LeftDelay.SetEasedDelayMsec(0.0f);
-			RightDelay.SetEasedDelayMsec(-1000.0f * DeltaTimeSeconds);
+			LeftDelay.SetEasedDelayMsec(0.0f, bIsInit);
+			RightDelay.SetEasedDelayMsec(-1000.0f * DeltaTimeSeconds, bIsInit);
 		}
 	}
 	
 
 	void FITDPannerOperator::Execute()
 	{
-		if (!FMath::IsNearlyEqual(*HeadWidth, CurrHeadWidth) || 
-			!FMath::IsNearlyEqual(*PanningAngle, CurrAngle) ||
-			!FMath::IsNearlyEqual(*DistanceFactor, CurrDistanceFactor))
+		float NewHeadWidth = FMath::Max(*HeadWidth, 0.0f);
+		float NewAngle = FMath::Clamp(*PanningAngle, 0.0f, 360.0f);
+		float NewDistanceFactor = FMath::Clamp(*DistanceFactor, 0.0f, 1.0f);
+
+		if (!FMath::IsNearlyEqual(NewAngle, CurrHeadWidth) ||
+			!FMath::IsNearlyEqual(NewDistanceFactor, CurrAngle) ||
+			!FMath::IsNearlyEqual(NewHeadWidth, CurrDistanceFactor))
 		{
-			CurrHeadWidth = *HeadWidth;
-			CurrAngle = *PanningAngle;
-			CurrDistanceFactor = *DistanceFactor;
+			CurrHeadWidth = NewHeadWidth;
+			CurrAngle = NewAngle;
+			CurrDistanceFactor = NewDistanceFactor;
 
-			UpdateParams();
+			UpdateParams(false);
 		}
 
 		const float* InputBufferPtr = AudioInput->GetData();
@@ -291,7 +292,7 @@ namespace Metasound
 		static const FVertexInterface Interface(
 			FInputVertexInterface(
 				TInputDataVertexModel<FAudioBuffer>(GetInputAudioName(), GetInputAudioDescription()),
-				TInputDataVertexModel<float>(GetInputPanAngleName(), GetInputPanAngleDescription(), 0.0f),
+				TInputDataVertexModel<float>(GetInputPanAngleName(), GetInputPanAngleDescription(), 90.0f),
 				TInputDataVertexModel<float>(GetInputDistanceFactorName(), GetInputDistanceFactorDescription(), 0.0f),
 				TInputDataVertexModel<float>(GetInputHeadWidthName(), GetInputHeadWidthDescription(), 34.0f)
 			),
diff --git a/Engine/Plugins/Runtime/Metasound/Source/MetasoundStandardNodes/Private/MetasoundStereoPannerNode.cpp b/Engine/Plugins/Runtime/Metasound/Source/MetasoundStandardNodes/Private/MetasoundStereoPannerNode.cpp
--- a/Engine/Plugins/Runtime/Metasound/Source/MetasoundStandardNodes/Private/MetasoundStereoPannerNode.cpp
+++ b/Engine/Plugins/Runtime/Metasound/Source/MetasoundStandardNodes/Private/MetasoundStereoPannerNode.cpp
@@ -137,14 +137,13 @@ namespace Metasound
 		const FAudioBufferReadRef& InAudioInput, 
 		const FFloatReadRef& InPanningAmount,
 		const FPanningLawReadRef& InPanningLaw)
-
 		: AudioInput(InAudioInput)
 		, PanningAmount(InPanningAmount)
 		, PanningLaw(InPanningLaw)
 		, AudioLeftOutput(FAudioBufferWriteRef::CreateNew(InSettings))
 		, AudioRightOutput(FAudioBufferWriteRef::CreateNew(InSettings))
 	{
-		PrevPanningAmount = *PanningAmount;
+		PrevPanningAmount = FMath::Clamp(*PanningAmount, -1.0f, 1.0f);
 
 		ComputePanGains(PrevPanningAmount, PrevLeftPan, PrevRightPan);
 	}
