Adding min, max, clamp nodes

diff --git a/Engine/Plugins/Runtime/Metasound/Source/MetasoundStandardNodes/Private/MetasoundClampNode.cpp b/Engine/Plugins/Runtime/Metasound/Source/MetasoundStandardNodes/Private/MetasoundClampNode.cpp
new file mode 100644
--- /dev/null
+++ b/Engine/Plugins/Runtime/Metasound/Source/MetasoundStandardNodes/Private/MetasoundClampNode.cpp
@@ -0,0 +1,287 @@
+// Copyright Epic Games, Inc. All Rights Reserved.
+
+#include "MetasoundFacade.h"
+#include "MetasoundNodeRegistrationMacro.h"
+#include "MetasoundExecutableOperator.h"
+#include "MetasoundPrimitives.h"
+#include "MetasoundStandardNodesNames.h"
+#include "MetasoundAudioBuffer.h"
+#include "Internationalization/Text.h"
+#include "MetasoundParamHelper.h"
+
+#define LOCTEXT_NAMESPACE "MetasoundStandardNodes"
+
+namespace Metasound
+{
+	namespace ClampVertexNames
+	{
+		METASOUND_PARAM(InputValue, "In", "Input value to clamp.");
+		METASOUND_PARAM(InputMinValue, "Min", "The min value to clamp.");
+		METASOUND_PARAM(InputMaxValue, "Max", "The max value to clamp.");
+		METASOUND_PARAM(OutputValue, "Value", "The clamped value.");
+	}
+
+	namespace MetasoundClampNodePrivate
+	{
+		FNodeClassMetadata CreateNodeClassMetadata(const FName& InDataTypeName, const FName& InOperatorName, const FText& InDisplayName, const FText& InDescription, const FVertexInterface& InDefaultInterface)
+		{
+			FNodeClassMetadata Metadata
+			{
+				FNodeClassName{FName("Clamp"), InOperatorName, InDataTypeName},
+				1, // Major Version
+				0, // Minor Version
+				InDisplayName,
+				InDescription,
+				PluginAuthor,
+				PluginNodeMissingPrompt,
+				InDefaultInterface,
+				{LOCTEXT("ClampCategory", "Utils")},
+				{TEXT("Clamp")},
+				FNodeDisplayStyle{}
+			};
+
+			return Metadata;
+		}
+
+		template<typename ValueType>
+		struct TClamp
+		{
+			bool bIsSupported = false;
+		};
+
+		template<>
+		struct TClamp<int32>
+		{
+			static void GetClamped(int32 InValue, int32 InMin, int32 InMax, int32& OutClamped)
+			{
+				OutClamped = FMath::Clamp(InValue, InMin, InMax);
+			}
+
+			static TDataReadReference<int32> CreateInValue(const FCreateOperatorParams& InParams)
+			{
+				using namespace ClampVertexNames;
+				const FInputVertexInterface& InputInterface = InParams.Node.GetVertexInterface().GetInputInterface();
+				return InParams.InputDataReferences.GetDataReadReferenceOrConstructWithVertexDefault<int32>(InputInterface, METASOUND_GET_PARAM_NAME(InputValue), InParams.OperatorSettings);
+			}
+
+			static TDataReadReference<int32> CreateInMin(const FCreateOperatorParams& InParams)
+			{
+				using namespace ClampVertexNames;
+				const FInputVertexInterface& InputInterface = InParams.Node.GetVertexInterface().GetInputInterface();
+				return InParams.InputDataReferences.GetDataReadReferenceOrConstructWithVertexDefault<int32>(InputInterface, METASOUND_GET_PARAM_NAME(InputMinValue), InParams.OperatorSettings);
+			}
+
+			static TDataReadReference<int32> CreateInMax(const FCreateOperatorParams& InParams)
+			{
+				using namespace ClampVertexNames;
+				const FInputVertexInterface& InputInterface = InParams.Node.GetVertexInterface().GetInputInterface();
+				return InParams.InputDataReferences.GetDataReadReferenceOrConstructWithVertexDefault<int32>(InputInterface, METASOUND_GET_PARAM_NAME(InputMaxValue), InParams.OperatorSettings);
+			}
+		};
+
+
+		template<>
+		struct TClamp<float>
+		{
+			static void GetClamped(float InValue, float InMin, float InMax, float& OutClamped)
+			{
+				OutClamped = FMath::Clamp(InValue, InMin, InMax);
+			}
+
+			static TDataReadReference<float> CreateInValue(const FCreateOperatorParams& InParams)
+			{
+				using namespace ClampVertexNames;
+				const FInputVertexInterface& InputInterface = InParams.Node.GetVertexInterface().GetInputInterface();
+				return InParams.InputDataReferences.GetDataReadReferenceOrConstructWithVertexDefault<float>(InputInterface, METASOUND_GET_PARAM_NAME(InputValue), InParams.OperatorSettings);
+			}
+
+			static TDataReadReference<float> CreateInMin(const FCreateOperatorParams& InParams)
+			{
+				using namespace ClampVertexNames;
+				const FInputVertexInterface& InputInterface = InParams.Node.GetVertexInterface().GetInputInterface();
+				return InParams.InputDataReferences.GetDataReadReferenceOrConstructWithVertexDefault<float>(InputInterface, METASOUND_GET_PARAM_NAME(InputMinValue), InParams.OperatorSettings);
+			}
+
+			static TDataReadReference<float> CreateInMax(const FCreateOperatorParams& InParams)
+			{
+				using namespace ClampVertexNames;
+				const FInputVertexInterface& InputInterface = InParams.Node.GetVertexInterface().GetInputInterface();
+				return InParams.InputDataReferences.GetDataReadReferenceOrConstructWithVertexDefault<float>(InputInterface, METASOUND_GET_PARAM_NAME(InputMaxValue), InParams.OperatorSettings);
+			}
+		};
+
+		template<>
+		struct TClamp<FAudioBuffer>
+		{
+			static void GetClamped(const FAudioBuffer& InValue, const FAudioBuffer& InMin, const FAudioBuffer& InMax, FAudioBuffer& OutClamped)
+			{
+				float* OutClampedBufferPtr = OutClamped.GetData();
+				const float* InBufferPtr = InValue.GetData();
+				const float* MinBufferPtr = InMin.GetData();
+				const float* MaxBufferPtr = InMax.GetData();
+				int32 NumSamples = InValue.Num();
+				for (int32 i = 0; i < NumSamples; ++i)
+				{
+					OutClampedBufferPtr[i] = FMath::Clamp(InBufferPtr[i], MinBufferPtr[i], MaxBufferPtr[i]);
+				}
+			}
+
+			static TDataReadReference<FAudioBuffer> CreateInValue(const FCreateOperatorParams& InParams)
+			{
+				using namespace ClampVertexNames;
+				return InParams.InputDataReferences.GetDataReadReferenceOrConstruct<FAudioBuffer>(METASOUND_GET_PARAM_NAME(InputValue), InParams.OperatorSettings);
+			}
+
+			static TDataReadReference<FAudioBuffer> CreateInMin(const FCreateOperatorParams& InParams)
+			{
+				using namespace ClampVertexNames;
+				return InParams.InputDataReferences.GetDataReadReferenceOrConstruct<FAudioBuffer>(METASOUND_GET_PARAM_NAME(InputMinValue), InParams.OperatorSettings);
+			}
+
+			static TDataReadReference<FAudioBuffer> CreateInMax(const FCreateOperatorParams& InParams)
+			{
+				using namespace ClampVertexNames;
+				return InParams.InputDataReferences.GetDataReadReferenceOrConstruct<FAudioBuffer>(METASOUND_GET_PARAM_NAME(InputMaxValue), InParams.OperatorSettings);
+			}
+		};
+	}
+
+	template<typename ValueType>
+	class TClampNodeOperator : public TExecutableOperator<TClampNodeOperator<ValueType>>
+	{
+	public:
+		static const FVertexInterface& GetDefaultInterface()
+		{
+			using namespace ClampVertexNames;
+			using namespace MetasoundClampNodePrivate;
+
+			static const FVertexInterface DefaultInterface(
+				FInputVertexInterface(
+					TInputDataVertexModel<ValueType>(METASOUND_GET_PARAM_NAME_AND_TT(InputValue)),
+					TInputDataVertexModel<ValueType>(METASOUND_GET_PARAM_NAME_AND_TT(InputMinValue)),
+					TInputDataVertexModel<ValueType>(METASOUND_GET_PARAM_NAME_AND_TT(InputMaxValue))
+				),
+				FOutputVertexInterface(
+					TOutputDataVertexModel<ValueType>(METASOUND_GET_PARAM_NAME_AND_TT(OutputValue))
+				)
+			);
+
+			return DefaultInterface;
+		}
+
+		static const FNodeClassMetadata& GetNodeInfo()
+		{
+			auto CreateNodeClassMetadata = []() -> FNodeClassMetadata
+			{
+				FName DataTypeName = GetMetasoundDataTypeName<ValueType>();
+				FName OperatorName = TEXT("Clamp");
+				FText NodeDisplayName = FText::Format(LOCTEXT("ClampDisplayNamePattern", "Clamp ({0})"), FText::FromString(GetMetasoundDataTypeString<ValueType>()));
+				FText NodeDescription = LOCTEXT("MinDesc", "Returns the clamped value of the input within the given value range.");
+				FVertexInterface NodeInterface = GetDefaultInterface();
+
+				return MetasoundClampNodePrivate::CreateNodeClassMetadata(DataTypeName, OperatorName, NodeDisplayName, NodeDescription, NodeInterface);
+			};
+
+			static const FNodeClassMetadata Metadata = CreateNodeClassMetadata();
+			return Metadata;
+		}
+
+		static TUniquePtr<IOperator> CreateOperator(const FCreateOperatorParams& InParams, TArray<TUniquePtr<IOperatorBuildError>>& OutErrors)
+		{
+			using namespace ClampVertexNames;
+			using namespace MetasoundClampNodePrivate;
+
+			TDataReadReference<ValueType> InputValue = TClamp<ValueType>::CreateInValue(InParams);
+			TDataReadReference<ValueType> InputMinValue = TClamp<ValueType>::CreateInMin(InParams);
+			TDataReadReference<ValueType> InputMaxValue = TClamp<ValueType>::CreateInMax(InParams);
+
+			return MakeUnique<TClampNodeOperator<ValueType>>(InParams.OperatorSettings, InputValue, InputMinValue, InputMaxValue);
+		}
+
+
+		TClampNodeOperator(const FOperatorSettings& InSettings,
+			const TDataReadReference<ValueType>& InInputValue,
+			const TDataReadReference<ValueType>& InInputMinValue,
+			const TDataReadReference<ValueType>& InInputMaxValue)
+			: InputValue(InInputValue)
+			, InputMinValue(InInputMinValue)
+			, InputMaxValue(InInputMaxValue)
+			, OutputValue(TDataWriteReferenceFactory<ValueType>::CreateAny(InSettings))
+		{
+			GetClamped();
+		}
+
+		virtual ~TClampNodeOperator() = default;
+
+		virtual FDataReferenceCollection GetInputs() const override
+		{
+			using namespace ClampVertexNames;
+
+			FDataReferenceCollection Inputs;
+			Inputs.AddDataReadReference(METASOUND_GET_PARAM_NAME(InputValue), InputValue);
+			Inputs.AddDataReadReference(METASOUND_GET_PARAM_NAME(InputMinValue), InputMinValue);
+			Inputs.AddDataReadReference(METASOUND_GET_PARAM_NAME(InputMaxValue), InputMaxValue);
+
+			return Inputs;
+		}
+
+		virtual FDataReferenceCollection GetOutputs() const override
+		{
+			using namespace ClampVertexNames;
+
+			FDataReferenceCollection Outputs;
+			Outputs.AddDataReadReference(METASOUND_GET_PARAM_NAME(OutputValue), OutputValue);
+
+			return Outputs;
+		}
+
+		void GetClamped()
+		{
+			using namespace MetasoundClampNodePrivate;
+
+			TClamp<ValueType>::GetClamped(*InputValue, *InputMinValue, *InputMaxValue, *OutputValue);
+		}
+
+		void Execute()
+		{
+			GetClamped();
+		}
+
+	private:
+
+		TDataReadReference<ValueType> InputValue;
+		TDataReadReference<ValueType> InputMinValue;
+		TDataReadReference<ValueType> InputMaxValue;
+		TDataWriteReference<ValueType> OutputValue;
+	};
+
+	/** TClampNode
+ *
+ *  Returns the clamped value of the input in the specified range.
+ */
+	template<typename ValueType>
+	class METASOUNDSTANDARDNODES_API TClampNode : public FNodeFacade
+	{
+	public:
+		/**
+		 * Constructor used by the Metasound Frontend.
+		 */
+		TClampNode(const FNodeInitData& InInitData)
+			: FNodeFacade(InInitData.InstanceName, InInitData.InstanceID, TFacadeOperatorClass<TClampNodeOperator<ValueType>>())
+		{}
+
+		virtual ~TClampNode() = default;
+	};
+
+	using FClampNodeInt32 = TClampNode<int32>;
+ 	METASOUND_REGISTER_NODE(FClampNodeInt32)
+
+	using FClampNodeFloat = TClampNode<float>;
+	METASOUND_REGISTER_NODE(FClampNodeFloat)
+
+	using FClampNodeAudio = TClampNode<FAudioBuffer>;
+	METASOUND_REGISTER_NODE(FClampNodeAudio)
+
+}
+
+#undef LOCTEXT_NAMESPACE
diff --git a/Engine/Plugins/Runtime/Metasound/Source/MetasoundStandardNodes/Private/MetasoundMaxNode.cpp b/Engine/Plugins/Runtime/Metasound/Source/MetasoundStandardNodes/Private/MetasoundMaxNode.cpp
new file mode 100644
--- /dev/null
+++ b/Engine/Plugins/Runtime/Metasound/Source/MetasoundStandardNodes/Private/MetasoundMaxNode.cpp
@@ -0,0 +1,261 @@
+// Copyright Epic Games, Inc. All Rights Reserved.
+
+#include "MetasoundNodeRegistrationMacro.h"
+#include "CoreMinimal.h"
+#include "MetasoundFacade.h"
+#include "MetasoundExecutableOperator.h"
+#include "MetasoundPrimitives.h"
+#include "MetasoundStandardNodesNames.h"
+#include "MetasoundAudioBuffer.h"
+#include "Internationalization/Text.h"
+#include "MetasoundParamHelper.h"
+
+#define LOCTEXT_NAMESPACE "MetasoundStandardNodes"
+
+namespace Metasound
+{
+	namespace MaxVertexNames
+	{
+		METASOUND_PARAM(InputAValue, "A", "Input value A.");
+		METASOUND_PARAM(InputBValue, "B", "Input value B.");
+		METASOUND_PARAM(OutputValue, "Value", "The max of A and B.");
+	}
+
+	namespace MetasoundMaxNodePrivate
+	{
+		FNodeClassMetadata CreateNodeClassMetadata(const FName& InDataTypeName, const FName& InOperatorName, const FText& InDisplayName, const FText& InDescription, const FVertexInterface& InDefaultInterface)
+		{
+			FNodeClassMetadata Metadata
+			{
+				FNodeClassName{FName("Max"), InOperatorName, InDataTypeName},
+				1, // Major Version
+				0, // Minor Version
+				InDisplayName,
+				InDescription,
+				PluginAuthor,
+				PluginNodeMissingPrompt,
+				InDefaultInterface,
+				{LOCTEXT("MaxCategory", "Utils")},
+				{TEXT("Max")},
+				FNodeDisplayStyle{}
+			};
+
+			return Metadata;
+		}
+
+		template<typename ValueType>
+		struct TMax
+		{
+			bool bSupported = false;
+		};
+
+		template<>
+		struct TMax<int32>
+		{
+			static void GetMax(int32 InA, int32 InB, int32& OutMax)
+			{
+				OutMax = FMath::Max(InA, InB);
+			}
+
+			static TDataReadReference<int32> CreateInRefA(const FCreateOperatorParams& InParams)
+			{
+				using namespace MaxVertexNames;
+				const FInputVertexInterface& InputInterface = InParams.Node.GetVertexInterface().GetInputInterface();
+				return InParams.InputDataReferences.GetDataReadReferenceOrConstructWithVertexDefault<int32>(InputInterface, METASOUND_GET_PARAM_NAME(InputAValue), InParams.OperatorSettings);
+			}
+
+			static TDataReadReference<int32> CreateInRefB(const FCreateOperatorParams& InParams)
+			{
+				using namespace MaxVertexNames;
+				const FInputVertexInterface& InputInterface = InParams.Node.GetVertexInterface().GetInputInterface();
+				return InParams.InputDataReferences.GetDataReadReferenceOrConstructWithVertexDefault<int32>(InputInterface, METASOUND_GET_PARAM_NAME(InputBValue), InParams.OperatorSettings);
+			}
+		};
+
+
+		template<>
+		struct TMax<float>
+		{
+			static void GetMax(float InA, float InB, float& OutMax)
+			{
+				OutMax = FMath::Max(InA, InB);
+			}
+
+			static TDataReadReference<float> CreateInRefA(const FCreateOperatorParams& InParams)
+			{
+				using namespace MaxVertexNames;
+				const FInputVertexInterface& InputInterface = InParams.Node.GetVertexInterface().GetInputInterface();
+				return InParams.InputDataReferences.GetDataReadReferenceOrConstructWithVertexDefault<float>(InputInterface, METASOUND_GET_PARAM_NAME(InputAValue), InParams.OperatorSettings);
+			}
+
+			static TDataReadReference<float> CreateInRefB(const FCreateOperatorParams& InParams)
+			{
+				using namespace MaxVertexNames;
+				const FInputVertexInterface& InputInterface = InParams.Node.GetVertexInterface().GetInputInterface();
+				return InParams.InputDataReferences.GetDataReadReferenceOrConstructWithVertexDefault<float>(InputInterface, METASOUND_GET_PARAM_NAME(InputBValue), InParams.OperatorSettings);
+			}
+
+			static bool IsAudioBuffer() { return false; }
+		};
+
+		template<>
+		struct TMax<FAudioBuffer>
+		{
+			static void GetMax(const FAudioBuffer& InA, const FAudioBuffer& InB, FAudioBuffer& OutMax)
+			{
+				float* OutMaxBufferPtr = OutMax.GetData();
+				const float* APtr = InA.GetData();
+				const float* BPtr = InB.GetData();
+				int32 NumSamples = InA.Num();
+				for (int32 i = 0; i < NumSamples; ++i)
+				{
+					OutMaxBufferPtr[i] = FMath::Max(APtr[i], BPtr[i]);
+				}
+			}
+
+			static TDataReadReference<FAudioBuffer> CreateInRefA(const FCreateOperatorParams& InParams)
+			{
+				using namespace MaxVertexNames;
+				return InParams.InputDataReferences.GetDataReadReferenceOrConstruct<FAudioBuffer>(METASOUND_GET_PARAM_NAME(InputAValue), InParams.OperatorSettings);
+			}
+
+			static TDataReadReference<FAudioBuffer> CreateInRefB(const FCreateOperatorParams& InParams)
+			{
+				using namespace MaxVertexNames;
+				return InParams.InputDataReferences.GetDataReadReferenceOrConstruct<FAudioBuffer>(METASOUND_GET_PARAM_NAME(InputBValue), InParams.OperatorSettings);
+			}
+		};
+	}
+
+	template<typename ValueType>
+	class TMaxNodeOperator : public TExecutableOperator<TMaxNodeOperator<ValueType>>
+	{
+	public:
+		static const FVertexInterface& GetDefaultInterface()
+		{
+			using namespace MaxVertexNames;
+			using namespace MetasoundMaxNodePrivate;
+
+			static const FVertexInterface DefaultInterface(
+				FInputVertexInterface(
+					TInputDataVertexModel<ValueType>(METASOUND_GET_PARAM_NAME_AND_TT(InputAValue)),
+					TInputDataVertexModel<ValueType>(METASOUND_GET_PARAM_NAME_AND_TT(InputBValue))
+				),
+				FOutputVertexInterface(
+					TOutputDataVertexModel<ValueType>(METASOUND_GET_PARAM_NAME_AND_TT(OutputValue))
+				)
+			);
+
+			return DefaultInterface;
+		}
+
+		static const FNodeClassMetadata& GetNodeInfo()
+		{
+			auto CreateNodeClassMetadata = []() -> FNodeClassMetadata
+			{
+				FName DataTypeName = GetMetasoundDataTypeName<ValueType>();
+				FName OperatorName = TEXT("Max");
+				FText NodeDisplayName = FText::Format(LOCTEXT("MaxDisplayNamePattern", "Max ({0})"), FText::FromString(GetMetasoundDataTypeString<ValueType>()));
+				FText NodeDescription = LOCTEXT("MaxDesc", "Returns the max of A and B.");
+				FVertexInterface NodeInterface = GetDefaultInterface();
+
+				return MetasoundMaxNodePrivate::CreateNodeClassMetadata(DataTypeName, OperatorName, NodeDisplayName, NodeDescription, NodeInterface);
+			};
+
+			static const FNodeClassMetadata Metadata = CreateNodeClassMetadata();
+			return Metadata;
+		}
+
+		static TUniquePtr<IOperator> CreateOperator(const FCreateOperatorParams& InParams, TArray<TUniquePtr<IOperatorBuildError>>& OutErrors)
+		{
+			using namespace MaxVertexNames;
+			using namespace MetasoundMaxNodePrivate;
+
+			TDataReadReference<ValueType> InputA = TMax<ValueType>::CreateInRefA(InParams);
+			TDataReadReference<ValueType> InputB = TMax<ValueType>::CreateInRefB(InParams);
+
+			return MakeUnique<TMaxNodeOperator<ValueType>>(InParams.OperatorSettings, InputA, InputB);
+		}
+
+
+		TMaxNodeOperator(const FOperatorSettings& InSettings,
+			const TDataReadReference<ValueType>& InInputA,
+			const TDataReadReference<ValueType>& InInputB)
+			: InputA(InInputA)
+			, InputB(InInputB)
+			, OutputValue(TDataWriteReferenceFactory<ValueType>::CreateAny(InSettings))
+		{
+			GetMax();
+		}
+
+		virtual ~TMaxNodeOperator() = default;
+
+		virtual FDataReferenceCollection GetInputs() const override
+		{
+			using namespace MaxVertexNames;
+
+			FDataReferenceCollection Inputs;
+			Inputs.AddDataReadReference(METASOUND_GET_PARAM_NAME(InputAValue), InputA);
+			Inputs.AddDataReadReference(METASOUND_GET_PARAM_NAME(InputBValue), InputB);
+
+			return Inputs;
+		}
+
+		virtual FDataReferenceCollection GetOutputs() const override
+		{
+			using namespace MaxVertexNames;
+
+			FDataReferenceCollection Outputs;
+			Outputs.AddDataReadReference(METASOUND_GET_PARAM_NAME(OutputValue), OutputValue);
+
+			return Outputs;
+		}
+
+		void GetMax()
+		{
+			using namespace MetasoundMaxNodePrivate;
+
+			TMax<ValueType>::GetMax(*InputA, *InputB, *OutputValue);
+		}
+
+		void Execute()
+		{
+			GetMax();
+		}
+
+	private:
+
+		TDataReadReference<ValueType> InputA;
+		TDataReadReference<ValueType> InputB;
+		TDataWriteReference<ValueType> OutputValue;
+	};
+
+	/** TMaxNode
+	 *
+	 *  Returns the min of both inputs
+	 */
+	template<typename ValueType>
+	class METASOUNDSTANDARDNODES_API TMaxNode : public FNodeFacade
+	{
+	public:
+		/**
+		 * Constructor used by the Metasound Frontend.
+		 */
+		TMaxNode(const FNodeInitData& InInitData)
+			: FNodeFacade(InInitData.InstanceName, InInitData.InstanceID, TFacadeOperatorClass<TMaxNodeOperator<ValueType>>())
+		{}
+
+		virtual ~TMaxNode() = default;
+	};
+
+	using FMaxNodeInt32 = TMaxNode<int32>;
+ 	METASOUND_REGISTER_NODE(FMaxNodeInt32)
+
+	using FMaxNodeFloat = TMaxNode<float>;
+	METASOUND_REGISTER_NODE(FMaxNodeFloat)
+
+	using FMaxNodeAudioBuffer = TMaxNode<FAudioBuffer>;
+	METASOUND_REGISTER_NODE(FMaxNodeAudioBuffer)
+}
+
+#undef LOCTEXT_NAMESPACE
diff --git a/Engine/Plugins/Runtime/Metasound/Source/MetasoundStandardNodes/Private/MetasoundMinNode.cpp b/Engine/Plugins/Runtime/Metasound/Source/MetasoundStandardNodes/Private/MetasoundMinNode.cpp
new file mode 100644
--- /dev/null
+++ b/Engine/Plugins/Runtime/Metasound/Source/MetasoundStandardNodes/Private/MetasoundMinNode.cpp
@@ -0,0 +1,257 @@
+// Copyright Epic Games, Inc. All Rights Reserved.
+
+#include "CoreMinimal.h"
+#include "MetasoundFacade.h"
+#include "MetasoundExecutableOperator.h"
+#include "MetasoundPrimitives.h"
+#include "MetasoundStandardNodesNames.h"
+#include "MetasoundAudioBuffer.h"
+#include "Internationalization/Text.h"
+#include "MetasoundNodeRegistrationMacro.h"
+#include "MetasoundParamHelper.h"
+
+#define LOCTEXT_NAMESPACE "MetasoundStandardNodes"
+
+namespace Metasound
+{
+	namespace MinVertexNames
+	{
+		METASOUND_PARAM(InputAValue, "A", "Input value A.");
+		METASOUND_PARAM(InputBValue, "B", "Input value B.");
+		METASOUND_PARAM(OutputValue, "Value", "The min of A and B.");
+	}
+
+	namespace MetasoundMinNodePrivate
+	{
+		FNodeClassMetadata CreateNodeClassMetadata(const FName& InDataTypeName, const FName& InOperatorName, const FText& InDisplayName, const FText& InDescription, const FVertexInterface& InDefaultInterface)
+		{
+			FNodeClassMetadata Metadata
+			{
+				FNodeClassName{FName("Min"), InOperatorName, InDataTypeName},
+				1, // Major Version
+				0, // Minor Version
+				InDisplayName,
+				InDescription,
+				PluginAuthor,
+				PluginNodeMissingPrompt,
+				InDefaultInterface,
+				{LOCTEXT("MaxCategory", "Utils")},
+				{TEXT("Min")},
+				FNodeDisplayStyle{}
+			};
+
+			return Metadata;
+		}
+
+		template<typename ValueType>
+		struct TMin
+		{
+			bool bIsSupported = false;
+		};
+
+		template<>
+		struct TMin<int32>
+		{
+			static void GetMin(int32 InA, int32 InB, int32& OutMin)
+			{
+				OutMin = FMath::Min(InA, InB);
+			}
+
+			static TDataReadReference<int32> CreateInRefA(const FCreateOperatorParams& InParams)
+			{
+				using namespace MinVertexNames;
+				const FInputVertexInterface& InputInterface = InParams.Node.GetVertexInterface().GetInputInterface();
+				return InParams.InputDataReferences.GetDataReadReferenceOrConstructWithVertexDefault<int32>(InputInterface, METASOUND_GET_PARAM_NAME(InputAValue), InParams.OperatorSettings);
+			}
+
+			static TDataReadReference<int32> CreateInRefB(const FCreateOperatorParams& InParams)
+			{
+				using namespace MinVertexNames;
+				const FInputVertexInterface& InputInterface = InParams.Node.GetVertexInterface().GetInputInterface();
+				return InParams.InputDataReferences.GetDataReadReferenceOrConstructWithVertexDefault<int32>(InputInterface, METASOUND_GET_PARAM_NAME(InputBValue), InParams.OperatorSettings);
+			}
+		};
+
+
+		template<>
+		struct TMin<float>
+		{
+			static void GetMin(float InA, float InB, float& OutMin)
+			{
+				OutMin = FMath::Min(InA, InB);
+			}
+
+			static TDataReadReference<float> CreateInRefA(const FCreateOperatorParams& InParams)
+			{
+				using namespace MinVertexNames;
+				const FInputVertexInterface& InputInterface = InParams.Node.GetVertexInterface().GetInputInterface();
+				return InParams.InputDataReferences.GetDataReadReferenceOrConstructWithVertexDefault<float>(InputInterface, METASOUND_GET_PARAM_NAME(InputAValue), InParams.OperatorSettings);
+			}
+
+			static TDataReadReference<float> CreateInRefB(const FCreateOperatorParams& InParams)
+			{
+				using namespace MinVertexNames;
+				const FInputVertexInterface& InputInterface = InParams.Node.GetVertexInterface().GetInputInterface();
+				return InParams.InputDataReferences.GetDataReadReferenceOrConstructWithVertexDefault<float>(InputInterface, METASOUND_GET_PARAM_NAME(InputBValue), InParams.OperatorSettings);
+			}
+		};
+
+
+		template<>
+		struct TMin<FAudioBuffer>
+		{
+			static void GetMin(const FAudioBuffer& InA, const FAudioBuffer& InB, FAudioBuffer& OutMax)
+			{
+				float* OutMaxBufferPtr = OutMax.GetData();
+				const float* APtr = InA.GetData();
+				const float* BPtr = InB.GetData();
+				for (int32 i = 0; i < InA.Num(); ++i)
+				{
+					OutMaxBufferPtr[i] = FMath::Min(APtr[i], BPtr[i]);
+				}
+			}
+
+			static TDataReadReference<FAudioBuffer> CreateInRefA(const FCreateOperatorParams& InParams)
+			{
+				using namespace MinVertexNames;
+				return InParams.InputDataReferences.GetDataReadReferenceOrConstruct<FAudioBuffer>(METASOUND_GET_PARAM_NAME(InputAValue), InParams.OperatorSettings);
+			}
+
+			static TDataReadReference<FAudioBuffer> CreateInRefB(const FCreateOperatorParams& InParams)
+			{
+				using namespace MinVertexNames;
+				return InParams.InputDataReferences.GetDataReadReferenceOrConstruct<FAudioBuffer>(METASOUND_GET_PARAM_NAME(InputBValue), InParams.OperatorSettings);
+			}
+		};
+
+	}
+
+	template<typename ValueType>
+	class TMinNodeOperator : public TExecutableOperator<TMinNodeOperator<ValueType>>
+	{
+	public:
+		static const FVertexInterface& GetDefaultInterface()
+		{
+			using namespace MinVertexNames;
+			using namespace MetasoundMinNodePrivate;
+
+			static const FVertexInterface DefaultInterface(
+				FInputVertexInterface(
+					TInputDataVertexModel<ValueType>(METASOUND_GET_PARAM_NAME_AND_TT(InputAValue)),
+					TInputDataVertexModel<ValueType>(METASOUND_GET_PARAM_NAME_AND_TT(InputBValue))
+				),
+				FOutputVertexInterface(
+					TOutputDataVertexModel<ValueType>(METASOUND_GET_PARAM_NAME_AND_TT(OutputValue))
+				)
+			);
+
+			return DefaultInterface;
+		}
+
+		static const FNodeClassMetadata& GetNodeInfo()
+		{
+			auto CreateNodeClassMetadata = []() -> FNodeClassMetadata
+			{
+				FName DataTypeName = GetMetasoundDataTypeName<ValueType>();
+				FName OperatorName = TEXT("Min");
+				FText NodeDisplayName = FText::Format(LOCTEXT("MinDisplayNamePattern", "Min ({0})"), FText::FromString(GetMetasoundDataTypeString<ValueType>()));
+				FText NodeDescription = LOCTEXT("MinDesc", "Returns the min of A and B.");
+				FVertexInterface NodeInterface = GetDefaultInterface();
+
+				return MetasoundMinNodePrivate::CreateNodeClassMetadata(DataTypeName, OperatorName, NodeDisplayName, NodeDescription, NodeInterface);
+			};
+
+			static const FNodeClassMetadata Metadata = CreateNodeClassMetadata();
+			return Metadata;
+		}
+
+		static TUniquePtr<IOperator> CreateOperator(const FCreateOperatorParams& InParams, TArray<TUniquePtr<IOperatorBuildError>>& OutErrors)
+		{
+			TDataReadReference<ValueType> InputA = MetasoundMinNodePrivate::TMin<ValueType>::CreateInRefA(InParams);
+			TDataReadReference<ValueType> InputB = MetasoundMinNodePrivate::TMin<ValueType>::CreateInRefB(InParams);
+
+			return MakeUnique<TMinNodeOperator<ValueType>>(InParams.OperatorSettings, InputA, InputB);
+		}
+
+
+		TMinNodeOperator(const FOperatorSettings& InSettings,
+			const TDataReadReference<ValueType>& InInputA,
+			const TDataReadReference<ValueType>& InInputB)
+			: InputA(InInputA)
+			, InputB(InInputB)
+			, OutputValue(TDataWriteReferenceFactory<ValueType>::CreateAny(InSettings))
+		{
+			GetMin();
+		}
+
+		virtual ~TMinNodeOperator() = default;
+
+		virtual FDataReferenceCollection GetInputs() const override
+		{
+			using namespace MinVertexNames;
+
+			FDataReferenceCollection Inputs;
+			Inputs.AddDataReadReference(METASOUND_GET_PARAM_NAME(InputAValue), InputA);
+			Inputs.AddDataReadReference(METASOUND_GET_PARAM_NAME(InputBValue), InputB);
+
+			return Inputs;
+		}
+
+		virtual FDataReferenceCollection GetOutputs() const override
+		{
+			using namespace MinVertexNames;
+
+			FDataReferenceCollection Outputs;
+			Outputs.AddDataReadReference(METASOUND_GET_PARAM_NAME(OutputValue), OutputValue);
+
+			return Outputs;
+		}
+
+		void GetMin()
+		{
+			using namespace MetasoundMinNodePrivate;
+
+			TMin<ValueType>::GetMin(*InputA, *InputB, *OutputValue);
+		}
+
+		void Execute()
+		{
+			GetMin();
+		}
+
+	private:
+
+		TDataReadReference<ValueType> InputA;
+		TDataReadReference<ValueType> InputB;
+		TDataWriteReference<ValueType> OutputValue;
+	};
+
+	/** TMinNode
+	 *
+	 *  Returns the min of both inputs
+	 */
+	template<typename ValueType>
+	class METASOUNDSTANDARDNODES_API TMinNode : public FNodeFacade
+	{
+	public:
+		/**
+		 * Constructor used by the Metasound Frontend.
+		 */
+		TMinNode(const FNodeInitData& InInitData)
+			: FNodeFacade(InInitData.InstanceName, InInitData.InstanceID, TFacadeOperatorClass<TMinNodeOperator<ValueType>>())
+		{}
+
+		virtual ~TMinNode() = default;
+	};
+
+	using FMinNodeInt32 = TMinNode<int32>;
+	METASOUND_REGISTER_NODE(FMinNodeInt32)
+
+	using FMinNodeFloat = TMinNode<float>;
+	METASOUND_REGISTER_NODE(FMinNodeFloat)
+
+	using FMinNodeAudio = TMinNode<FAudioBuffer>;
+	METASOUND_REGISTER_NODE(FMinNodeAudio)
+}
+
+#undef LOCTEXT_NAMESPACE
