CIS Fix for uninitialized USTRUCT property.

diff --git a/Engine/Plugins/Runtime/Metasound/Source/MetasoundFrontend/Public/MetasoundFrontendDocument.h b/Engine/Plugins/Runtime/Metasound/Source/MetasoundFrontend/Public/MetasoundFrontendDocument.h
--- a/Engine/Plugins/Runtime/Metasound/Source/MetasoundFrontend/Public/MetasoundFrontendDocument.h
+++ b/Engine/Plugins/Runtime/Metasound/Source/MetasoundFrontend/Public/MetasoundFrontendDocument.h
@@ -246,7 +246,7 @@ struct FMetasoundFrontendStyleEdge
 	GENERATED_BODY()
 
 	UPROPERTY()
-	EMetasoundFrontendStyleEdgeDisplay Display;
+	EMetasoundFrontendStyleEdgeDisplay Display = EMetasoundFrontendStyleEdgeDisplay::Default;
 };
 
 // Styling for a class of edges dependent upon edge data type.
