Fixing pitch zippering in wave player node

diff --git a/Engine/Plugins/Runtime/Metasound/Source/MetasoundEngine/MetasoundEngine.Build.cs b/Engine/Plugins/Runtime/Metasound/Source/MetasoundEngine/MetasoundEngine.Build.cs
--- a/Engine/Plugins/Runtime/Metasound/Source/MetasoundEngine/MetasoundEngine.Build.cs
+++ b/Engine/Plugins/Runtime/Metasound/Source/MetasoundEngine/MetasoundEngine.Build.cs
@@ -21,6 +21,7 @@ namespace UnrealBuildTool.Rules
 					"MetasoundGenerator",
 					"AudioExtensions",
 					"AudioCodecEngine",
+					"SignalProcessing"
 				}
 			);
 
@@ -34,7 +35,7 @@ namespace UnrealBuildTool.Rules
 					"MetasoundStandardNodes",
 					"AudioCodecEngine",
 					"Serialization",
-					"AudioPlatformConfiguration"
+					"AudioPlatformConfiguration",
 				}
 			);
 		}
diff --git a/Engine/Plugins/Runtime/Metasound/Source/MetasoundEngine/Private/MetasoundWave.cpp b/Engine/Plugins/Runtime/Metasound/Source/MetasoundEngine/Private/MetasoundWave.cpp
--- a/Engine/Plugins/Runtime/Metasound/Source/MetasoundEngine/Private/MetasoundWave.cpp
+++ b/Engine/Plugins/Runtime/Metasound/Source/MetasoundEngine/Private/MetasoundWave.cpp
@@ -8,6 +8,7 @@
 #include "DecoderInputFactory.h"
 #include "AudioDeviceManager.h"
 #include "AudioDevice.h"
+#include "DSP/ParamInterpolator.h"
 
 namespace Audio
 {
@@ -22,6 +23,7 @@ namespace Audio
 			Output.Reset();
 			Decoder.Reset();
 			bDecoderIsDone = false;
+			bIsSeekable = InWave->IsSeekableStreaming();
 
 			return false;
 		}
@@ -30,17 +32,16 @@ namespace Audio
 		InputSampleRate = InWave->GetSampleRate();
 		OutputSampleRate = InInitParams.OutputSampleRate;
 		FsOutToInRatio = (OutputSampleRate / InputSampleRate);
-		MaxPitchShiftRatio = FMath::Pow(2.f, InInitParams.MaxPitchShiftMagnitudeAllowedInOctaves);
-		MaxPitchShiftCents = InInitParams.MaxPitchShiftMagnitudeAllowedInOctaves * 1200.f;
+		MaxPitchShiftRatio = FMath::Pow(2.0f, InInitParams.MaxPitchShiftMagnitudeAllowedInOctaves);
+		MaxPitchShiftCents = InInitParams.MaxPitchShiftMagnitudeAllowedInOctaves * 1200.0f;
 		StartTimeSeconds = InInitParams.StartTimeSeconds;
 
-
 		NumChannels = InWave->GetNumChannels();
 		DecodeBlockSizeInFrames = 64;
 		DecodeBlockSizeInSamples = DecodeBlockSizeInFrames * NumChannels;
 
 		// set Circular Buffer capacity
-		OutputCircularBuffer.SetCapacity(InInitParams.OutputBlockSizeInFrames * NumChannels * (1.f + FsOutToInRatio * MaxPitchShiftRatio) * 2);
+		OutputCircularBuffer.SetCapacity(InInitParams.OutputBlockSizeInFrames * NumChannels * (1.0f + FsOutToInRatio * MaxPitchShiftRatio) * 2);
 
 		TotalNumFramesOutput = 0;
 		TotalNumFramesDecoded = 0;
@@ -50,21 +51,22 @@ namespace Audio
 		bDecoderIsDone = !bSuccessful;
 
 		// initialize SRC object (will be re-initialized for pitch shifting)
-		Resampler.Init(Audio::EResamplingMethod::Linear, 1.f / FsOutToInRatio, NumChannels);
+		Resampler.Init(Audio::EResamplingMethod::Linear, FsOutToInRatio, NumChannels);
+		PitchShifter.Reset(NumChannels, InInitParams.InitialPitchShiftSemitones);
 
 		return bSuccessful;
 	}
 
+
 	uint32 FSimpleDecoderWrapper::GenerateAudio(float* OutputDest, int32 NumOutputFrames, float PitchShiftInCents, bool bIsLooping)
 	{
 		const uint32 NumOutputSamples = NumOutputFrames * NumChannels;
 
 		if (OutputCircularBuffer.Num() < NumOutputSamples)
 		{
-			const float SampleRateRatio = GetSampleRateRatio(PitchShiftInCents);
 
-			// (multiply by two just to be safe we can handle SRC output)
-			const int32 MaxNumResamplerOutputFramesPerBlock = FMath::CeilToInt(SampleRateRatio * DecodeBlockSizeInFrames) * 2;
+			// (multiply by two just to be sure we can handle SRC output size)
+			const int32 MaxNumResamplerOutputFramesPerBlock = FMath::CeilToInt(FsOutToInRatio * DecodeBlockSizeInFrames) * 2;
 			const int32 MaxNumResamplerOutputSamplesPerBlock = MaxNumResamplerOutputFramesPerBlock * NumChannels;
 
 			PreSrcBuffer.Reset(DecodeBlockSizeInSamples);
@@ -73,8 +75,9 @@ namespace Audio
 			PostSrcBuffer.Reset(MaxNumResamplerOutputSamplesPerBlock);
 			PostSrcBuffer.AddUninitialized(MaxNumResamplerOutputSamplesPerBlock);
 
-			// combine pitch shift and SRC ratios
-			Resampler.SetSampleRateRatio(SampleRateRatio);
+			Resampler.SetSampleRateRatio(FsOutToInRatio);
+			PitchShifter.UpdatePitchShift(FMath::Clamp(PitchShiftInCents, -MaxPitchShiftCents, MaxPitchShiftCents) / 100.0f);
+
 
 			// perform SRC and push to circular buffer until we have enough frames for the output
 			while (Decoder && !(bDecoderIsDone || bDecoderHasLooped) && (OutputCircularBuffer.Num() < NumOutputSamples))
@@ -86,10 +89,13 @@ namespace Audio
 				int32 NumResamplerOutputFrames = 0;
 				int32 Error = Resampler.ProcessAudio(PreSrcBuffer.GetData(), NumFramesDecoded, bDecoderIsDone, PostSrcBuffer.GetData(), MaxNumResamplerOutputFramesPerBlock, NumResamplerOutputFrames);
 				ensure(Error == 0);
-				OutputCircularBuffer.Push(PostSrcBuffer.GetData(), NumResamplerOutputFrames * NumChannels);
 
 				bDecoderIsDone = DecodeResult == Audio::IDecoder::EDecodeResult::Finished;
 				bDecoderHasLooped = DecodeResult == Audio::IDecoder::EDecodeResult::Looped;
+
+				// perform linear pitch shift into OutputCircularBuffer
+				const TArrayView<float> BufferToPitchShift(PostSrcBuffer.GetData(), NumResamplerOutputFrames * NumChannels);
+				PitchShifter.ProcessAudio(BufferToPitchShift, OutputCircularBuffer);
 			}
 		}
 
@@ -119,10 +125,14 @@ namespace Audio
 
 	void FSimpleDecoderWrapper::SeekToTime(const float InSeconds)
 	{
-		if (Input.IsValid())
+		if (Input.IsValid() && bIsSeekable)
 		{
 			Input->SeekToTime(InSeconds);
 		}
+		if (!bIsSeekable)
+		{
+			ensureMsgf(false, TEXT("Attempting to seek on a sound wave that is not set to Seekable"));
+		}
 	}
 
 	bool FSimpleDecoderWrapper::InitializeDecodersInternal(const FSoundWaveProxyPtr& Wave)
@@ -163,19 +173,7 @@ namespace Audio
 		// Decoder:
 		Decoder = Codec->CreateDecoder(Input.Get(), Output.Get());
 
-		// TODO: Delete this
-		// END: Delete this
-
 		// return true if all the components were successfully create
 		return Input.IsValid() && Output.IsValid() && Decoder.IsValid();
 	}
-
-	float FSimpleDecoderWrapper::GetSampleRateRatio(float PitchShiftInCents) const
-	{
-		PitchShiftInCents = FMath::Clamp(PitchShiftInCents, -MaxPitchShiftCents, MaxPitchShiftCents);
-
-		// combine pitch shift due to actual pitch shift and SRC
-		return FsOutToInRatio / FMath::Pow(2.f, PitchShiftInCents / 1200.f);
-	}
-
 } // namespace Audio
diff --git a/Engine/Plugins/Runtime/Metasound/Source/MetasoundEngine/Private/MetasoundWavePlayerNode.cpp b/Engine/Plugins/Runtime/Metasound/Source/MetasoundEngine/Private/MetasoundWavePlayerNode.cpp
--- a/Engine/Plugins/Runtime/Metasound/Source/MetasoundEngine/Private/MetasoundWavePlayerNode.cpp
+++ b/Engine/Plugins/Runtime/Metasound/Source/MetasoundEngine/Private/MetasoundWavePlayerNode.cpp
@@ -221,6 +221,7 @@ namespace Metasound
 				{
 					bIsPlaying = false;
 					ResetDecoder();
+					TrigggerOnDone->TriggerFrame(CurrAudioFrame);
 					++StopTrigIndex;
 				}
 			}
@@ -253,7 +254,8 @@ namespace Metasound
 			Audio::FSimpleDecoderWrapper::InitParams Params;
 			Params.OutputBlockSizeInFrames = OutputBlockSizeInFrames;
 			Params.OutputSampleRate = OutputSampleRate;
-			Params.MaxPitchShiftMagnitudeAllowedInOctaves = 4.f;
+			Params.MaxPitchShiftMagnitudeAllowedInOctaves = 6.f;
+			Params.InitialPitchShiftSemitones = *PitchShiftSemiTones;
 			Params.StartTimeSeconds = FMath::Max(0.f, *SeekTimeSeconds);
 
 			if (false == Wave->IsSoundWaveValid())
diff --git a/Engine/Plugins/Runtime/Metasound/Source/MetasoundEngine/Public/MetasoundWave.h b/Engine/Plugins/Runtime/Metasound/Source/MetasoundEngine/Public/MetasoundWave.h
--- a/Engine/Plugins/Runtime/Metasound/Source/MetasoundEngine/Public/MetasoundWave.h
+++ b/Engine/Plugins/Runtime/Metasound/Source/MetasoundEngine/Public/MetasoundWave.h
@@ -8,6 +8,7 @@
 #include "Sound/SoundWave.h"
 #include "IAudioCodecRegistry.h"
 #include "IAudioCodec.h"
+#include "DSP/InterpolatedLinearPitchShifter.h"
 
 
 namespace Metasound
@@ -78,10 +79,11 @@ namespace Audio
 	public:
 		struct InitParams
 		{
-			float OutputSampleRate;
-			uint32 OutputBlockSizeInFrames;
-			float MaxPitchShiftMagnitudeAllowedInOctaves = 4.f;
-			float StartTimeSeconds = 0.f;
+			uint32 OutputBlockSizeInFrames{ 512 };
+			float OutputSampleRate { 44100.0f };
+			float MaxPitchShiftMagnitudeAllowedInOctaves { 4.0f };
+			float InitialPitchShiftSemitones { 0.0f };
+			float StartTimeSeconds { 0.0f };
 		};
 
 		bool CanGenerateAudio() const
@@ -92,13 +94,11 @@ namespace Audio
 		bool Initialize(const InitParams& InInitParams, const FSoundWaveProxyPtr& InWave);
 
 		// returns number of samples written.   
-		uint32 GenerateAudio(float* OutputDest, int32 NumOutputFrames, float PitchShiftInCents = 0.f, bool bIsLooping = false);
+		uint32 GenerateAudio(float* OutputDest, int32 NumOutputFrames, float PitchShiftInCents = 0.0f, bool bIsLooping = false);
 
 		void SeekToTime(const float InSeconds);
 
 	private:
-		float GetSampleRateRatio(float PitchShiftInCents) const;
-
 		// actual decoder objects
 		TUniquePtr<Audio::IDecoder> Decoder;
 		TUniquePtr<Audio::IDecoderOutput> Output;
@@ -107,8 +107,9 @@ namespace Audio
 		// init helper for decoders
 		bool InitializeDecodersInternal(const FSoundWaveProxyPtr& Wave);
 
-		// SRC object
+		// SRC objects
 		Audio::FResampler Resampler;
+		FLinearPitchShifter PitchShifter;
 
 		// buffers
 		TArray<float> PreSrcBuffer;
@@ -116,26 +117,24 @@ namespace Audio
 		Audio::TCircularAudioBuffer<float> OutputCircularBuffer;
 
 		// meta data:
-		float InputSampleRate;
-		float OutputSampleRate;
-		float FsOutToInRatio;
-		float MaxPitchShiftCents;
-		float MaxPitchShiftRatio;
-
-		uint32 NumChannels;
-		uint32 DecodeBlockSizeInFrames;
-		uint32 DecodeBlockSizeInSamples;
+		float InputSampleRate{ 44100.f };
+		float OutputSampleRate{ 44100.f };
+		float FsOutToInRatio{ 1.f };
+		float MaxPitchShiftCents{ 12.0f };
+		float MaxPitchShiftRatio{ 2.0f };
+
+		uint32 NumChannels{ 1 };
+		uint32 DecodeBlockSizeInFrames{ 64 };
+		uint32 DecodeBlockSizeInSamples{ 64 };
 		
 		// cached values
-		float StartTimeSeconds{ 0.f };
-		float LastPitchShiftCents{ 0.f };
+		float StartTimeSeconds{ 0.0f };
 		int32 TotalNumFramesOutput{ 0 };
 		int32 TotalNumFramesDecoded{ 0 };
 
 		bool bDecoderIsDone{ true };
+		bool bIsSeekable{ false };
 		bool bDecoderHasLooped{ false };
 
 	}; // class FSimpleDecoderWrapper
-
-
 } // namespace Audio
diff --git a/Engine/Source/Runtime/SignalProcessing/Private/InterpolatedLinearPitchShifter.cpp b/Engine/Source/Runtime/SignalProcessing/Private/InterpolatedLinearPitchShifter.cpp
new file mode 100644
--- /dev/null
+++ b/Engine/Source/Runtime/SignalProcessing/Private/InterpolatedLinearPitchShifter.cpp
@@ -0,0 +1,128 @@
+// Copyright Epic Games, Inc. All Rights Reserved.
+
+#include "DSP/InterpolatedLinearPitchShifter.h"
+
+namespace Audio
+{
+	void FLinearPitchShifter::Reset(int32 InNumChannels, float InInitialPitchShiftSemitones, int32 InInterpLengthFrames)
+	{
+		InterpLengthFrames = InInterpLengthFrames;
+		InterpFramesRemaining = 0;
+		PitchShiftRatio.SetValue(FMath::Pow(2.0f, InInitialPitchShiftSemitones / 12.0f), 0);
+		NumChannels = InNumChannels;
+
+		PreviousFrame.Reset();
+		PreviousFrame.AddZeroed(NumChannels);
+
+		bInterpolateBetweenBuffers = false;
+	}
+
+	int32 FLinearPitchShifter::ProcessAudio(const TArrayView<float> InputBuffer, Audio::TCircularAudioBuffer<float>& OutputBuffer)
+	{
+		const int32 NumInputFrames = InputBuffer.Num() / NumChannels;
+		int32 OutputFramesRendered = 0; // return value
+
+		// NOTE: CurrentIndex is a float in the range (-1.0f, (float)(NumInputFrames - 1.0f))
+		// the fractional portion of CurrentIndex is used to interpolate between Floor(CurrentIndex) and Ceil(CurrentIndex)
+		// if CurrentIndex < -1.0f then we missed an interpolation that should have occured in the previous buffer
+		check(CurrentIndex > -1.0f);
+
+
+		// Handle interpolations between index -1.0f and 0.0f
+		// i.e. interpolating between our last buffer and the current buffer
+		if (bInterpolateBetweenBuffers)
+		{
+			// Alpha is between (0.0 and 1.0]
+			// 0.0 -> Final frame of the previous Input buffer
+			// 1.0 -> 0th frame of the current Input buffer
+			float Alpha = CurrentIndex + 1.0f;
+
+			// Once Alpha reaches 1.x, we are on frame 0.x of the current buffer
+			while (Alpha < 1.0f)
+			{
+				for (int32 Chan = 0; Chan < NumChannels; ++Chan)
+				{
+					const float A = PreviousFrame[Chan];
+					const float B = InputBuffer[Chan];
+					OutputBuffer.Push(FMath::Lerp(A, B, Alpha));
+				}
+
+				Alpha += GetNextIndexDelta();
+				++OutputFramesRendered;
+			}
+
+			// Alpha == 1.0 means CurrentIndex == 0.0
+			CurrentIndex = Alpha - 1.0f;
+			bInterpolateBetweenBuffers = false;
+		}
+
+		// Early Exit: copy full input buffer if no work needs to be done
+		// (i.e., not interpolating and pitch shift ratio is 1.0f)
+		if (!InterpFramesRemaining && FMath::IsNearlyEqual(1.0f, PitchShiftRatio.GetTarget()))
+		{
+			OutputBuffer.Push(InputBuffer.GetData(), NumInputFrames * NumChannels);
+			return (NumInputFrames + OutputFramesRendered);
+		}
+
+		// Normal case: linear interpolation across the input buffer
+		while (FMath::CeilToInt(CurrentIndex) < NumInputFrames)
+		{
+			const int32 i = NumChannels * (int32)FMath::Floor(CurrentIndex);
+			float const* A = InputBuffer.GetData() + i;
+			float const* B = A + NumChannels;
+			const float Alpha = FMath::Frac(CurrentIndex);
+
+			for (int32 Chan = 0; Chan < NumChannels; ++Chan)
+			{
+				OutputBuffer.Push(FMath::Lerp(*A++, *B++, Alpha));
+			}
+
+			++OutputFramesRendered;
+			CurrentIndex += GetNextIndexDelta();
+		}
+
+		// wrap our fractional index by the buffer size
+		CurrentIndex -= (float)NumInputFrames;
+
+		// If -1 < CurrentIndex < 0.f, we are going to need to interpolate between
+		// the final frame of this buffer and the 0th frame of the next buffer.
+		// So we cache the final frame and raise a flag to use on the next call
+		if (CurrentIndex < 0.0f)
+		{
+			const int32 FinalFrameIndex = (NumInputFrames - 1) * NumChannels;
+
+			PreviousFrame.Reset(NumChannels);
+			for (int32 Chan = 0; Chan < NumChannels; ++Chan)
+			{
+				PreviousFrame.Add(InputBuffer[FinalFrameIndex + Chan]);
+			}
+
+			bInterpolateBetweenBuffers = true;
+		}
+
+		return OutputFramesRendered;
+	}
+
+	void FLinearPitchShifter::UpdatePitchShift(float InNewPitchSemitones)
+	{
+		InterpFramesRemaining = InterpLengthFrames;
+		PitchShiftRatio.SetValue(FMath::Pow(2.0f, InNewPitchSemitones / 12.0f), InterpLengthFrames);
+
+		if (FMath::IsNearlyEqual(PitchShiftRatio.GetValue(), PitchShiftRatio.GetTarget()))
+		{
+			InterpFramesRemaining = 0; // already at target
+		}
+	}
+
+	float FLinearPitchShifter::GetNextIndexDelta()
+	{
+		if (InterpFramesRemaining)
+		{
+			--InterpFramesRemaining;
+			return PitchShiftRatio.Update();
+		}
+
+		return PitchShiftRatio.GetTarget();
+	}
+
+} // namespace Audio
\ No newline at end of file
diff --git a/Engine/Source/Runtime/SignalProcessing/Public/DSP/InterpolatedLinearPitchShifter.h b/Engine/Source/Runtime/SignalProcessing/Public/DSP/InterpolatedLinearPitchShifter.h
new file mode 100644
--- /dev/null
+++ b/Engine/Source/Runtime/SignalProcessing/Public/DSP/InterpolatedLinearPitchShifter.h
@@ -0,0 +1,38 @@
+// Copyright Epic Games, Inc. All Rights Reserved.
+
+#pragma once
+
+#include "CoreMinimal.h"
+#include "DSP/Dsp.h"
+#include "DSP/ParamInterpolator.h"
+
+namespace Audio
+{
+	class SIGNALPROCESSING_API FLinearPitchShifter
+	{
+	public:
+		// ctor
+		FLinearPitchShifter() = default;
+
+		void Reset(int32 InNumChannels, float InInitialPitchShiftSemitones = 0.0f, int32 InInterpLengthFrames = 100);
+
+		// Sample rate converts the input audio buffer and pushes converted audio to the output circular buffer.
+		// (Returns the number of output frames generated.)
+		int32 ProcessAudio(const TArrayView<float> InputBuffer, Audio::TCircularAudioBuffer<float>& OutputBuffer);
+
+		void UpdatePitchShift(float InNewPitchSemitones);
+
+	private:
+		// returns next fractional index increment
+		float GetNextIndexDelta();
+
+		TArray<float> PreviousFrame;
+		FParam PitchShiftRatio;
+		float CurrentIndex{ 0.0f };
+		int32 NumChannels{ 1 };
+		int32 InterpLengthFrames{ 100 };
+		int32 InterpFramesRemaining{ 0 };
+		bool bInterpolateBetweenBuffers{ false };
+
+	}; // class FLinearPitchShifter
+} // namespace Audio
\ No newline at end of file
