- Display General Sound properties in new inspector view in Metasound Editor
- Disallow deleting outputs from Metasound editor graph.
- Disallow adding outputs to Metasound editor graph (re-enable once reading outputs is required by composition)

diff --git a/Engine/Plugins/Runtime/Metasound/Source/MetasoundEditor/Private/MetasoundDetailCustomization.cpp b/Engine/Plugins/Runtime/Metasound/Source/MetasoundEditor/Private/MetasoundDetailCustomization.cpp
--- a/Engine/Plugins/Runtime/Metasound/Source/MetasoundEditor/Private/MetasoundDetailCustomization.cpp
+++ b/Engine/Plugins/Runtime/Metasound/Source/MetasoundEditor/Private/MetasoundDetailCustomization.cpp
@@ -11,6 +11,7 @@
 #include "IDetailGroup.h"
 #include "Input/Events.h"
 #include "MetasoundAssetBase.h"
+#include "MetasoundEditorSettings.h"
 #include "MetasoundFrontend.h"
 #include "MetasoundFrontendController.h"
 #include "MetasoundUObjectRegistry.h"
@@ -29,13 +30,6 @@
 
 #define LOCTEXT_NAMESPACE "MetasoundEditor"
 
-static int32 ShowLiteralMetasoundInputsInEditorCVar = 0;
-FAutoConsoleVariableRef CVarShowLiteralMetasoundInputsInEditor(
-	TEXT("au.Debug.Editor.Metasounds.ShowLiteralInputs"),
-	ShowLiteralMetasoundInputsInEditorCVar,
-	TEXT("Show literal inputs in the Metasound Editor.\n")
-	TEXT("0: Disabled (default), !0: Enabled"),
-	ECVF_Default);
 
 namespace Metasound
 {
@@ -51,102 +45,6 @@ namespace Metasound
 			return FName(InBasePath.ToString() + TEXT(".") + InPropertyName.ToString());
 		}
 
-		TSet<FString> GetLiteralInputs(IDetailLayoutBuilder& InDetailLayout)
-		{
-			TSet<FString> LiteralInputs;
-			TArray<TWeakObjectPtr<UObject>> Objects;
-			InDetailLayout.GetObjectsBeingCustomized(Objects);
-
-			if (Objects.IsEmpty() || !Objects[0].IsValid())
-			{
-				return LiteralInputs;
-			}
-
-			UObject* Metasound = Objects[0].Get();
-			if (FMetasoundAssetBase* MetasoundAsset = IMetasoundUObjectRegistry::Get().GetObjectAsAssetBase(Metasound))
-			{
-				Frontend::FGraphHandle GraphHandle = MetasoundAsset->GetRootGraphHandle();
-				TArray<Frontend::FNodeHandle> InputNodes = GraphHandle->GetInputNodes();
-				for (Frontend::FNodeHandle& NodeHandle : InputNodes)
-				{
-					if (NodeHandle->GetNodeStyle().Display.Visibility == EMetasoundFrontendNodeStyleDisplayVisibility::Hidden)
-					{
-						LiteralInputs.Add(NodeHandle->GetNodeName());
-					}
-				}
-			}
-
-			return LiteralInputs;
-		}
-
-		template <typename T>
-		void BuildIOFixedArray(IDetailLayoutBuilder& InDetailLayout, FName InCategoryName, FName InPropertyName)
-		{
-			const bool bIsInput = InCategoryName == "Inputs";
-
-			IDetailCategoryBuilder& CategoryBuilder = InDetailLayout.EditCategory(InCategoryName);
-			TSharedPtr<IPropertyHandle> ParentProperty = InDetailLayout.GetProperty(InPropertyName);
-			TSharedPtr<IPropertyHandleArray> ArrayHandle = ParentProperty->AsArray();
-
-			TSet<FString> LiteralInputs;
-			if (bIsInput && !ShowLiteralMetasoundInputsInEditorCVar)
-			{
-				LiteralInputs = GetLiteralInputs(InDetailLayout);
-			}
-
-			uint32 NumElements = 0;
-			ArrayHandle->GetNumElements(NumElements);
-			for (int32 i = 0; i < static_cast<int32>(NumElements); ++i)
-			{
-				TSharedRef<IPropertyHandle> ArrayItemHandle = ArrayHandle->GetElement(i);
-
-				const FName TypeNamePropertyName = GET_MEMBER_NAME_CHECKED(T, TypeName);
-				const FName NamePropertyName = GET_MEMBER_NAME_CHECKED(T, Name);
-				const FName ToolTipPropertyName = GET_MEMBER_NAME_CHECKED(FMetasoundFrontendVertexMetadata, Description);
-				const FName DisplayNamePropertyName = GET_MEMBER_NAME_CHECKED(FMetasoundFrontendVertexMetadata, DisplayName);
-
-				TSharedPtr<IPropertyHandle> TypeProperty = ArrayItemHandle->GetChildHandle(TypeNamePropertyName);
-				TSharedPtr<IPropertyHandle> NameProperty = ArrayItemHandle->GetChildHandle(NamePropertyName);
-
-				TSharedPtr<IPropertyHandle> ToolTipProperty = ArrayItemHandle->GetChildHandle(ToolTipPropertyName, true /* bRecurse */);
-				TSharedPtr<IPropertyHandle> DisplayNameProperty = ArrayItemHandle->GetChildHandle(DisplayNamePropertyName, true /* bRecurse */);
-
-				FString Name;
-				const bool bNameFound = NameProperty->GetValue(Name) == FPropertyAccess::Success;
-
-				// Hide literal members
-				if (LiteralInputs.Contains(Name))
-				{
-					continue;
-				}
-
-				CategoryBuilder.AddCustomRow(ParentProperty->GetPropertyDisplayName())
-				.NameContent()
-				[
-					SNew(STextBlock)
-					.Font(IDetailLayoutBuilder::GetDetailFontBold())
-					.Text(TAttribute<FText>::Create([i, DisplayNameProperty]()
-					{
-						FText DisplayName;
-						DisplayNameProperty->GetValue(DisplayName);
-						return DisplayName;
-					}))
-					.ToolTipText(TAttribute<FText>::Create([ToolTipProperty]()
-					{
-						FText ToolTip;
-						ToolTipProperty->GetValue(ToolTip);
-						return ToolTip;
-					}))
-				];
-			}
-
-			FSimpleDelegate RefreshDelegate = FSimpleDelegate::CreateLambda([DetailLayout = &InDetailLayout]()
-			{
-				DetailLayout->ForceRefreshDetails();
-			});
-			ArrayHandle->SetOnNumElementsChanged(RefreshDelegate);
-		}
-
 		FMetasoundDetailCustomization::FMetasoundDetailCustomization(FName InDocumentPropertyName)
 			: IDetailCustomization()
 			, DocumentPropertyName(InDocumentPropertyName)
@@ -168,53 +66,84 @@ namespace Metasound
 		{
 			using namespace Metasound::Editor;
 
-			// General Category
-			IDetailCategoryBuilder& GeneralCategoryBuilder = DetailLayout.EditCategory("General");
-
-			const FName AuthorPropertyPath = BuildChildPath(GetMetadataPropertyPath(), GET_MEMBER_NAME_CHECKED(FMetasoundFrontendClassMetadata, Author));
-			const FName DescPropertyPath = BuildChildPath(GetMetadataPropertyPath(), GET_MEMBER_NAME_CHECKED(FMetasoundFrontendClassMetadata, Description));
-			const FName VersionPropertyPath = BuildChildPath(GetMetadataPropertyPath(), GET_MEMBER_NAME_CHECKED(FMetasoundFrontendClassMetadata, Version));
-			const FName MajorVersionPropertyPath = BuildChildPath(VersionPropertyPath, GET_MEMBER_NAME_CHECKED(FMetasoundFrontendVersionNumber, Major));
-			const FName MinorVersionPropertyPath = BuildChildPath(VersionPropertyPath, GET_MEMBER_NAME_CHECKED(FMetasoundFrontendVersionNumber, Minor));
-
-			TSharedPtr<IPropertyHandle> AuthorHandle = DetailLayout.GetProperty(AuthorPropertyPath);
-			TSharedPtr<IPropertyHandle> DescHandle = DetailLayout.GetProperty(DescPropertyPath);
-			TSharedPtr<IPropertyHandle> MajorVersionHandle = DetailLayout.GetProperty(MajorVersionPropertyPath);
-			TSharedPtr<IPropertyHandle> MinorVersionHandle = DetailLayout.GetProperty(MinorVersionPropertyPath);
-
-			GeneralCategoryBuilder.AddProperty(AuthorHandle);
-			GeneralCategoryBuilder.AddProperty(DescHandle);
-			GeneralCategoryBuilder.AddProperty(MajorVersionHandle);
-			GeneralCategoryBuilder.AddProperty(MinorVersionHandle);
-
-			// Input/Output Categories
-
-			// If editing multiple metasound objects, all should be the same type, so safe to just check first in array for
-			// required inputs/outputs
-			TArray<TWeakObjectPtr<UObject>> Objects;
-			DetailLayout.GetObjectsBeingCustomized(Objects);
+			EMetasoundActiveDetailView DetailsView = EMetasoundActiveDetailView::Metasound;
+			if (const UMetasoundEditorSettings* EditorSettings = GetDefault<UMetasoundEditorSettings>())
+			{
+				DetailsView = EditorSettings->DetailView;
+			}
 
+			switch (DetailsView)
+			{
+				case EMetasoundActiveDetailView::Metasound:
+				{
+					IDetailCategoryBuilder& GeneralCategoryBuilder = DetailLayout.EditCategory("Metasound");
+					const FName AuthorPropertyPath = BuildChildPath(GetMetadataPropertyPath(), GET_MEMBER_NAME_CHECKED(FMetasoundFrontendClassMetadata, Author));
+					const FName DescPropertyPath = BuildChildPath(GetMetadataPropertyPath(), GET_MEMBER_NAME_CHECKED(FMetasoundFrontendClassMetadata, Description));
+					const FName VersionPropertyPath = BuildChildPath(GetMetadataPropertyPath(), GET_MEMBER_NAME_CHECKED(FMetasoundFrontendClassMetadata, Version));
+					const FName MajorVersionPropertyPath = BuildChildPath(VersionPropertyPath, GET_MEMBER_NAME_CHECKED(FMetasoundFrontendVersionNumber, Major));
+					const FName MinorVersionPropertyPath = BuildChildPath(VersionPropertyPath, GET_MEMBER_NAME_CHECKED(FMetasoundFrontendVersionNumber, Minor));
+
+					TSharedPtr<IPropertyHandle> AuthorHandle = DetailLayout.GetProperty(AuthorPropertyPath);
+					TSharedPtr<IPropertyHandle> DescHandle = DetailLayout.GetProperty(DescPropertyPath);
+					TSharedPtr<IPropertyHandle> MajorVersionHandle = DetailLayout.GetProperty(MajorVersionPropertyPath);
+					TSharedPtr<IPropertyHandle> MinorVersionHandle = DetailLayout.GetProperty(MinorVersionPropertyPath);
+
+					GeneralCategoryBuilder.AddProperty(AuthorHandle);
+					GeneralCategoryBuilder.AddProperty(DescHandle);
+					GeneralCategoryBuilder.AddProperty(MajorVersionHandle);
+					GeneralCategoryBuilder.AddProperty(MinorVersionHandle);
+
+					// Hack to hide categories brought in from UMetasoundSource inherited from USoundBase
+					DetailLayout.HideCategory("Analysis");
+					DetailLayout.HideCategory("Attenuation");
+					DetailLayout.HideCategory("Debug");
+					DetailLayout.HideCategory("Effects");
+					DetailLayout.HideCategory("Loading");
+					DetailLayout.HideCategory("Modulation");
+					DetailLayout.HideCategory("Sound");
+					DetailLayout.HideCategory("SoundWave");
+					DetailLayout.HideCategory("Voice Management");
+				}
+				break;
+
+				case EMetasoundActiveDetailView::General:
+				default:
+					DetailLayout.HideCategory("Metasound");
+
+					const bool bShouldBeInitiallyCollapsed = true;
+					DetailLayout.EditCategory("Analysis").InitiallyCollapsed(bShouldBeInitiallyCollapsed);
+					DetailLayout.EditCategory("Attenuation").InitiallyCollapsed(bShouldBeInitiallyCollapsed);
+					DetailLayout.EditCategory("Debug").InitiallyCollapsed(bShouldBeInitiallyCollapsed);
+					DetailLayout.EditCategory("Effects").InitiallyCollapsed(bShouldBeInitiallyCollapsed);
+					DetailLayout.EditCategory("Modulation").InitiallyCollapsed(bShouldBeInitiallyCollapsed);
+					DetailLayout.EditCategory("Sound").InitiallyCollapsed(bShouldBeInitiallyCollapsed);
+					DetailLayout.EditCategory("SoundWave").InitiallyCollapsed(bShouldBeInitiallyCollapsed);
+					DetailLayout.EditCategory("Voice Management").InitiallyCollapsed(bShouldBeInitiallyCollapsed);
+
+					const bool bRestore = false;
+					DetailLayout.EditCategory("Analysis").RestoreExpansionState(bRestore);
+					DetailLayout.EditCategory("Attenuation").RestoreExpansionState(bRestore);
+					DetailLayout.EditCategory("Debug").RestoreExpansionState(bRestore);
+					DetailLayout.EditCategory("Effects").RestoreExpansionState(bRestore);
+					DetailLayout.EditCategory("Modulation").RestoreExpansionState(bRestore);
+					DetailLayout.EditCategory("Sound").RestoreExpansionState(bRestore);
+					DetailLayout.EditCategory("SoundWave").RestoreExpansionState(bRestore);
+					DetailLayout.EditCategory("Voice Management").RestoreExpansionState(bRestore);
+
+					break;
+			}
 
 			// Hack to hide parent structs for nested metadata properties
 			DetailLayout.HideCategory("CustomView");
 
-			// Hack to hide categories brought in from UMetasoundSource inherited from USoundBase
-			DetailLayout.HideCategory("Analysis");
-			DetailLayout.HideCategory("Attenuation");
 			DetailLayout.HideCategory("Curves");
-			DetailLayout.HideCategory("Debug");
 			DetailLayout.HideCategory("Developer");
-			DetailLayout.HideCategory("Effects");
 			DetailLayout.HideCategory("File Path");
 			DetailLayout.HideCategory("Format");
 			DetailLayout.HideCategory("Info");
 			DetailLayout.HideCategory("Loading");
-			DetailLayout.HideCategory("Modulation");
 			DetailLayout.HideCategory("Playback");
-			DetailLayout.HideCategory("Sound");
-			DetailLayout.HideCategory("SoundWave");
 			DetailLayout.HideCategory("Subtitles");
-			DetailLayout.HideCategory("Voice Management");
 		}
 	} // namespace Editor
 } // namespace Metasound
diff --git a/Engine/Plugins/Runtime/Metasound/Source/MetasoundEditor/Private/MetasoundEditor.cpp b/Engine/Plugins/Runtime/Metasound/Source/MetasoundEditor/Private/MetasoundEditor.cpp
--- a/Engine/Plugins/Runtime/Metasound/Source/MetasoundEditor/Private/MetasoundEditor.cpp
+++ b/Engine/Plugins/Runtime/Metasound/Source/MetasoundEditor/Private/MetasoundEditor.cpp
@@ -15,6 +15,7 @@
 #include "Framework/Notifications/NotificationManager.h"
 #include "GraphEditor.h"
 #include "GraphEditorActions.h"
+#include "HAL/IConsoleManager.h"
 #include "HAL/PlatformApplicationMisc.h"
 #include "IAudioExtensionPlugin.h"
 #include "IDetailsView.h"
@@ -25,6 +26,7 @@
 #include "MetasoundEditorGraphBuilder.h"
 #include "MetasoundEditorGraphNode.h"
 #include "MetasoundEditorGraphSchema.h"
+#include "MetasoundEditorSettings.h"
 #include "MetasoundEditorTabFactory.h"
 #include "MetasoundUObjectRegistry.h"
 #include "Misc/Attribute.h"
@@ -43,8 +45,18 @@
 // This needs to be moved to public directory
 #include "../../GraphEditor/Private/GraphActionNode.h"
 
+
 #define LOCTEXT_NAMESPACE "MetasoundEditor"
 
+static int32 ShowLiteralMetasoundInputsInEditorCVar = 0;
+FAutoConsoleVariableRef CVarShowLiteralMetasoundInputsInEditor(
+	TEXT("au.Debug.Editor.Metasounds.ShowLiteralInputs"),
+	ShowLiteralMetasoundInputsInEditorCVar,
+	TEXT("Show literal inputs in the Metasound Editor.\n")
+	TEXT("0: Disabled (default), !0: Enabled"),
+	ECVF_Default);
+
+
 namespace Metasound
 {
 	namespace Editor
@@ -153,7 +165,7 @@ namespace Metasound
 			.SetGroup(WorkspaceMenuCategoryRef)
 			.SetIcon(FSlateIcon(FEditorStyle::GetStyleSetName(), "LevelEditor.Tabs.Details"));
 
-			InTabManager->RegisterTabSpawner(TabFactory::Names::Metasound, FOnSpawnTab::CreateLambda([InMetasoundMenu = MetasoundGeneralMenu](const FSpawnTabArgs& Args)
+			InTabManager->RegisterTabSpawner(TabFactory::Names::Metasound, FOnSpawnTab::CreateLambda([InMetasoundMenu = MetasoundInterfaceMenu](const FSpawnTabArgs& Args)
 			{
 				return TabFactory::CreateMetasoundTab(InMetasoundMenu, Args);
 			}))
@@ -371,9 +383,9 @@ namespace Metasound
 				MetasoundGraphEditor->NotifyGraphChanged();
 			}
 
-			if (MetasoundGeneralMenu.IsValid())
+			if (MetasoundInterfaceMenu.IsValid())
 			{
-				MetasoundGeneralMenu->RefreshAllActions(true /* bPreserveExpansion */);
+				MetasoundInterfaceMenu->RefreshAllActions(true /* bPreserveExpansion */);
 			}
 
 			FSlateApplication::Get().DismissAllMenus();
@@ -397,9 +409,12 @@ namespace Metasound
 				{
 					if (IMetasoundUObjectRegistry::Get().IsRegisteredClass(MetasoundObj))
 					{
-						FGraphBuilder::SynchronizeGraph(*MetasoundObj);
+						if (FGraphBuilder::SynchronizeGraph(*MetasoundObj))
+						{
+							MetasoundObj->MarkPackageDirty();
+						}
 						MetasoundGraphEditor->NotifyGraphChanged();
-						MetasoundGeneralMenu->RefreshAllActions(true /* bPreserveExpansion */);
+						MetasoundInterfaceMenu->RefreshAllActions(true /* bPreserveExpansion */);
 					}
 				}
 			}
@@ -415,7 +430,7 @@ namespace Metasound
 
 			FPropertyEditorModule& PropertyModule = FModuleManager::LoadModuleChecked<FPropertyEditorModule>("PropertyEditor");
 
-			SAssignNew(MetasoundGeneralMenu, SGraphActionMenu, false)
+			SAssignNew(MetasoundInterfaceMenu, SGraphActionMenu, false)
 				.AlphaSortItems(false)
 // 				.OnActionDoubleClicked(this, &FEditor::OnActionDoubleClicked)
 				.OnActionDragged(this, &FEditor::OnActionDragged)
@@ -544,13 +559,22 @@ namespace Metasound
 
 					ToolbarBuilder.BeginSection("Utilities");
 					{
+						ToolbarBuilder.AddToolBarButton(
+							FEditorCommands::Get().EditGeneralSettings,
+							NAME_None,
+							TAttribute<FText>(),
+							TAttribute<FText>(),
+							TAttribute<FSlateIcon>::Create([this]() { return GetSettingsImage(); }),
+							"EditGeneralSettings"
+						);
+
 						ToolbarBuilder.AddToolBarButton(
 							FEditorCommands::Get().EditMetasoundSettings,
-								NAME_None,
-								TAttribute<FText>(),
-								TAttribute<FText>(),
-								TAttribute<FSlateIcon>::Create([this]() { return GetSettingsImage(); }),
-								"EditMetasoundSettings"
+							NAME_None,
+							TAttribute<FText>(),
+							TAttribute<FText>(),
+							TAttribute<FSlateIcon>::Create([this]() { return GetSettingsImage(); }),
+							"EditMetasoundSettings"
 						);
 					}
 					ToolbarBuilder.EndSection();
@@ -619,6 +643,10 @@ namespace Metasound
 				Commands.EditMetasoundSettings,
 				FExecuteAction::CreateSP(this, &FEditor::EditMetasoundSettings));
 
+			ToolkitCommands->MapAction(
+				Commands.EditGeneralSettings,
+				FExecuteAction::CreateSP(this, &FEditor::EditGeneralSettings));
+
 			ToolkitCommands->MapAction(
 				FGenericCommands::Get().Delete,
 				FExecuteAction::CreateSP(this, &FEditor::DeleteSelected));
@@ -790,11 +818,11 @@ namespace Metasound
 			}
 		}
 
-		void FEditor::EditMetasoundSettings()
+		void FEditor::EditObjectSettings()
 		{
-			if (MetasoundGeneralMenu.IsValid())
+			if (MetasoundInterfaceMenu.IsValid())
 			{
-				MetasoundGeneralMenu->SelectItemByName(FName());
+				MetasoundInterfaceMenu->SelectItemByName(FName());
 			}
 
 			if (MetasoundGraphEditor.IsValid())
@@ -804,9 +832,33 @@ namespace Metasound
 				bManuallyClearingGraphSelection = false;
 			}
 
+			// Clear selection first to force refresh of customization
+			// if swapping from one object-level edit mode to the other
+			// (ex. Metasound Settings to General Settings)
+			SetSelection({ });
 			SetSelection({ Metasound });
 		}
 
+		void FEditor::EditGeneralSettings()
+		{
+			if (UMetasoundEditorSettings* EditorSettings = GetMutableDefault<UMetasoundEditorSettings>())
+			{
+				EditorSettings->DetailView = EMetasoundActiveDetailView::General;
+			}
+
+			EditObjectSettings();
+		}
+
+		void FEditor::EditMetasoundSettings()
+		{
+			if (UMetasoundEditorSettings* EditorSettings = GetMutableDefault<UMetasoundEditorSettings>())
+			{
+				EditorSettings->DetailView = EMetasoundActiveDetailView::Metasound;
+			}
+
+			EditObjectSettings();
+		}
+
 		void FEditor::SyncInBrowser()
 		{
 			TArray<UObject*> ObjectsToSync;
@@ -860,6 +912,9 @@ namespace Metasound
 				GraphEditorCommands->MapAction(FEditorCommands::Get().EditMetasoundSettings,
 					FExecuteAction::CreateSP(this, &FEditor::EditMetasoundSettings));
 
+				GraphEditorCommands->MapAction(FEditorCommands::Get().EditGeneralSettings,
+					FExecuteAction::CreateSP(this, &FEditor::EditGeneralSettings));
+
 				GraphEditorCommands->MapAction(FEditorCommands::Get().AddInput,
 					FExecuteAction::CreateSP(this, &FEditor::AddInput),
 					FCanExecuteAction::CreateSP(this, &FEditor::CanAddInput));
@@ -966,9 +1021,9 @@ namespace Metasound
 				}
 			}
 
-			if (MetasoundGeneralMenu.IsValid() && !bManuallyClearingGraphSelection)
+			if (MetasoundInterfaceMenu.IsValid() && !bManuallyClearingGraphSelection)
 			{
-				MetasoundGeneralMenu->SelectItemByName(FName());
+				MetasoundInterfaceMenu->SelectItemByName(FName());
 			}
 			SetSelection(Selection);
 		}
@@ -987,10 +1042,10 @@ namespace Metasound
 		{
 			using namespace Frontend;
 
-			if (MetasoundGeneralMenu.IsValid())
+			if (MetasoundInterfaceMenu.IsValid())
 			{
 				TArray<TSharedPtr<FEdGraphSchemaAction>> Actions;
-				MetasoundGeneralMenu->GetSelectedActions(Actions);
+				MetasoundInterfaceMenu->GetSelectedActions(Actions);
 
 				if (!Actions.IsEmpty())
 				{
@@ -1083,11 +1138,11 @@ namespace Metasound
 
 						FGraphBuilder::DeleteVariableNodeHandle(*Variable);
 						Graph->RemoveVariable(*Variable);
-						MetasoundGeneralMenu->RefreshAllActions(true /* bPreserveExpansion */);
+						MetasoundInterfaceMenu->RefreshAllActions(true /* bPreserveExpansion */);
 
 						if (!NameToSelect.IsNone())
 						{
-							MetasoundGeneralMenu->SelectItemByName(NameToSelect, ESelectInfo::Direct, SectionId);
+							MetasoundInterfaceMenu->SelectItemByName(NameToSelect, ESelectInfo::Direct, SectionId);
 						}
 					}
 					return;
@@ -1430,14 +1485,14 @@ namespace Metasound
 
 		void FEditor::OnInputNameChanged(FGuid InNodeID)
 		{
-			if (!MetasoundGeneralMenu.IsValid())
+			if (!MetasoundInterfaceMenu.IsValid())
 			{
 				return;
 			}
 
 			TArray<TSharedPtr<FEdGraphSchemaAction>> SelectedActions;
-			MetasoundGeneralMenu->GetSelectedActions(SelectedActions);
-			MetasoundGeneralMenu->RefreshAllActions(/* bPreserveExpansion */ true);
+			MetasoundInterfaceMenu->GetSelectedActions(SelectedActions);
+			MetasoundInterfaceMenu->RefreshAllActions(/* bPreserveExpansion */ true);
 
 			for(const TSharedPtr<FEdGraphSchemaAction>& Action : SelectedActions)
 			{
@@ -1451,7 +1506,7 @@ namespace Metasound
 						if (ensure(OutputHandles.Num() == 1))
 						{
 							const FName ActionName = *OutputHandles[0]->GetDisplayName().ToString();
-							MetasoundGeneralMenu->SelectItemByName(ActionName, ESelectInfo::Direct, Action->GetSectionID());
+							MetasoundInterfaceMenu->SelectItemByName(ActionName, ESelectInfo::Direct, Action->GetSectionID());
 						}
 						break;
 					}
@@ -1461,14 +1516,14 @@ namespace Metasound
 
 		void FEditor::OnOutputNameChanged(FGuid InNodeID)
 		{
-			if (!MetasoundGeneralMenu.IsValid())
+			if (!MetasoundInterfaceMenu.IsValid())
 			{
 				return;
 			}
 
 			TArray<TSharedPtr<FEdGraphSchemaAction>> SelectedActions;
-			MetasoundGeneralMenu->GetSelectedActions(SelectedActions);
-			MetasoundGeneralMenu->RefreshAllActions(/* bPreserveExpansion */ true);
+			MetasoundInterfaceMenu->GetSelectedActions(SelectedActions);
+			MetasoundInterfaceMenu->RefreshAllActions(/* bPreserveExpansion */ true);
 
 			for (const TSharedPtr<FEdGraphSchemaAction>& Action : SelectedActions)
 			{
@@ -1482,7 +1537,7 @@ namespace Metasound
 						if (ensure(InputHandles.Num() == 1))
 						{
 							const FName ActionName = *InputHandles[0]->GetDisplayName().ToString();
-							MetasoundGeneralMenu->SelectItemByName(ActionName, ESelectInfo::Direct, Action->GetSectionID());
+							MetasoundInterfaceMenu->SelectItemByName(ActionName, ESelectInfo::Direct, Action->GetSectionID());
 						}
 						break;
 					}
@@ -1499,7 +1554,8 @@ namespace Metasound
 
 			MetasoundAsset->GetRootGraphHandle()->IterateConstNodes([this, EditorGraph, MetasoundAsset, ActionList = &OutAllActions](const Frontend::FConstNodeHandle& Input)
 			{
-				if (Input->GetNodeStyle().Display.Visibility == EMetasoundFrontendNodeStyleDisplayVisibility::Hidden)
+				const bool bIsLiteralInput = Input->GetNodeStyle().Display.Visibility == EMetasoundFrontendNodeStyleDisplayVisibility::Hidden;
+				if (!ShowLiteralMetasoundInputsInEditorCVar && bIsLiteralInput)
 				{
 					return;
 				}
@@ -1641,6 +1697,10 @@ namespace Metasound
 				{
 					AddNewText = LOCTEXT("AddNewInput", "Input");
 					MetaDataTag = "AddNewInput";
+
+					// TODO: Add back for outputs once reading outputs/Metasound
+					// composition is supported.
+					return CreateAddInputButton(InSectionID, AddNewText, MetaDataTag);
 				}
 				break;
 
@@ -1652,10 +1712,10 @@ namespace Metasound
 				break;
 
 				default:
-					return SNullWidget::NullWidget;
+				break;
 			}
 
-			return CreateAddInputButton(InSectionID, AddNewText, MetaDataTag);
+			return SNullWidget::NullWidget;
 		}
 
 		bool FEditor::CanAddNewElementToSection(int32 InSectionID) const
@@ -1737,12 +1797,12 @@ namespace Metasound
 				return FReply::Unhandled();
 			}
 
-			if (MetasoundGeneralMenu.IsValid())
+			if (MetasoundInterfaceMenu.IsValid())
 			{
-				MetasoundGeneralMenu->RefreshAllActions(/* bPreserveExpansion */ true);
+				MetasoundInterfaceMenu->RefreshAllActions(/* bPreserveExpansion */ true);
 				if (NameToSelect.IsNone())
 				{
-					MetasoundGeneralMenu->SelectItemByName(NameToSelect);
+					MetasoundInterfaceMenu->SelectItemByName(NameToSelect);
 				}
 			}
 			return FReply::Handled();
diff --git a/Engine/Plugins/Runtime/Metasound/Source/MetasoundEditor/Private/MetasoundEditor.h b/Engine/Plugins/Runtime/Metasound/Source/MetasoundEditor/Private/MetasoundEditor.h
--- a/Engine/Plugins/Runtime/Metasound/Source/MetasoundEditor/Private/MetasoundEditor.h
+++ b/Engine/Plugins/Runtime/Metasound/Source/MetasoundEditor/Private/MetasoundEditor.h
@@ -205,7 +205,10 @@ namespace Metasound
 			/** Sync the content browser to the current selection of nodes */
 			void SyncInBrowser();
 
-			/** Show the metasound object's settings in the Inspector */
+			/** Show the Metasound object's General settings in the Inspector */
+			void EditGeneralSettings();
+
+			/** Show the Metasound object's Metasound-specific settings in the Inspector */
 			void EditMetasoundSettings();
 
 			/** Add an input to the currently selected node */
@@ -227,6 +230,8 @@ namespace Metasound
 			TSharedRef<SGraphEditor> CreateGraphEditorWidget();
 
 		private:
+			void EditObjectSettings();
+
 			/** List of open tool panels; used to ensure only one exists at any one time */
 			TMap<FName, TWeakPtr<SDockableTab>> SpawnedToolPanels;
 
@@ -236,8 +241,8 @@ namespace Metasound
 			/** Details tab */
 			TSharedPtr<IDetailsView> MetasoundDetails;
 
-			/** Metasound General info widget */
-			TSharedPtr<SGraphActionMenu> MetasoundGeneralMenu;
+			/** Metasound Interface menu */
+			TSharedPtr<SGraphActionMenu> MetasoundInterfaceMenu;
 
 			/** Metasound Output Meter widget */
 			TSharedPtr<SAudioMeter> MetasoundMeter;
diff --git a/Engine/Plugins/Runtime/Metasound/Source/MetasoundEditor/Private/MetasoundEditorCommands.cpp b/Engine/Plugins/Runtime/Metasound/Source/MetasoundEditor/Private/MetasoundEditorCommands.cpp
--- a/Engine/Plugins/Runtime/Metasound/Source/MetasoundEditor/Private/MetasoundEditorCommands.cpp
+++ b/Engine/Plugins/Runtime/Metasound/Source/MetasoundEditor/Private/MetasoundEditorCommands.cpp
@@ -20,7 +20,8 @@ namespace Metasound
 			UI_COMMAND(AddInput, "Add Input", "Adds an input to the node", EUserInterfaceActionType::Button, FInputChord());
 			UI_COMMAND(DeleteInput, "Delete Input", "Removes an input from the node", EUserInterfaceActionType::Button, FInputChord());
 
-			UI_COMMAND(EditMetasoundSettings, "Settings", "Edit Metasound Settings", EUserInterfaceActionType::Button, FInputChord());
+			UI_COMMAND(EditMetasoundSettings, "Metasound", "Edit Metasound settings", EUserInterfaceActionType::Button, FInputChord());
+			UI_COMMAND(EditGeneralSettings, "General", "Edit general sound settings", EUserInterfaceActionType::Button, FInputChord());
 		}
 	} // namespace Editor
 } // namespace Metasound
diff --git a/Engine/Plugins/Runtime/Metasound/Source/MetasoundEditor/Private/MetasoundEditorGraph.cpp b/Engine/Plugins/Runtime/Metasound/Source/MetasoundEditor/Private/MetasoundEditorGraph.cpp
--- a/Engine/Plugins/Runtime/Metasound/Source/MetasoundEditor/Private/MetasoundEditorGraph.cpp
+++ b/Engine/Plugins/Runtime/Metasound/Source/MetasoundEditor/Private/MetasoundEditorGraph.cpp
@@ -244,9 +244,6 @@ UMetasoundEditorGraphInputNode* UMetasoundEditorGraph::CreateInputNode(Metasound
 			NewInputNode->AllocateDefaultPins();
 		}
 
-		PostEditChange();
-		MarkPackageDirty();
-
 		return NewInputNode;
 	}
 	
diff --git a/Engine/Plugins/Runtime/Metasound/Source/MetasoundEditor/Private/MetasoundEditorGraphBuilder.cpp b/Engine/Plugins/Runtime/Metasound/Source/MetasoundEditor/Private/MetasoundEditorGraphBuilder.cpp
--- a/Engine/Plugins/Runtime/Metasound/Source/MetasoundEditor/Private/MetasoundEditorGraphBuilder.cpp
+++ b/Engine/Plugins/Runtime/Metasound/Source/MetasoundEditor/Private/MetasoundEditorGraphBuilder.cpp
@@ -283,24 +283,6 @@ namespace Metasound
 			return nullptr;
 		}
 
-		UMetasoundEditorGraphInputNode* FGraphBuilder::AddInputNode(UObject& InMetasound, const FString& InName, const FName InTypeName, const FText& InToolTip, FVector2D InLocation, EMetasoundFrontendNodeStyleDisplayVisibility Visibility, bool bInSelectNewNode)
-		{
-			using namespace Frontend;
-
-			FNodeHandle NodeHandle = AddInputNodeHandle(InMetasound, InName, InTypeName, InToolTip, Visibility);
-			UMetasoundEditorGraphInputNode* NewGraphNode = AddInputNode(InMetasound, NodeHandle, InLocation, bInSelectNewNode);
-			if (ensure(NewGraphNode))
-			{
-				return NewGraphNode;
-			}
-			else
-			{
-				GraphBuilderPrivate::DeleteNode(InMetasound, NodeHandle);
-			}
-
-			return nullptr;
-		}
-
 		void FGraphBuilder::AddOrUpdateLiteralInput(UObject& InMetasound, Frontend::FNodeHandle InNodeHandle, UEdGraphPin& InInputPin, bool bForcePinValueAsDefault)
 		{
 			using namespace Frontend;
@@ -346,9 +328,6 @@ namespace Metasound
 				}
 
 				ensure(InputHandle->Connect(*OutputHandle));
-
-				InMetasound.PostEditChange();
-				InMetasound.MarkPackageDirty();
 			}
 			else
 			{
@@ -360,9 +339,6 @@ namespace Metasound
 					const FGuid VertexID = GraphHandle->GetVertexIDForInputVertex(InputNode->GetNodeName());
 
 					GraphHandle->SetDefaultInput(VertexID, PinDataTypeDefaultLiteral);
-
-					InMetasound.PostEditChange();
-					InMetasound.MarkPackageDirty();
 				}
 			}
 		}
@@ -562,12 +538,6 @@ namespace Metasound
 			return NodeHandle;
 		}
 
-		UMetasoundEditorGraphOutputNode* FGraphBuilder::AddOutputNode(UObject& InMetasound, const FString& InName, const FName InTypeName, const FText& InToolTip, FVector2D InLocation, bool bInSelectNewNode)
-		{
-			Frontend::FNodeHandle NodeHandle = AddOutputNodeHandle(InMetasound, InName, InTypeName, InToolTip);
-			return AddOutputNode(InMetasound, NodeHandle, InLocation, bInSelectNewNode);
-		}
-
 		Frontend::FNodeHandle FGraphBuilder::AddOutputNodeHandle(UObject& InMetasound, const FString& InName, const FName InTypeName, const FText& InToolTip)
 		{
 			FMetasoundAssetBase* MetasoundAsset = IMetasoundUObjectRegistry::Get().GetObjectAsAssetBase(&InMetasound);
@@ -812,9 +782,6 @@ namespace Metasound
 					}
 				}
 			}
-
-			InMetasound.PostEditChange();
-			InMetasound.MarkPackageDirty();
 		}
 
 		void FGraphBuilder::DeleteVariableNodeHandle(UMetasoundEditorGraphVariable& InVariable)
@@ -873,8 +840,6 @@ namespace Metasound
 				return false;
 			}
 
-			InNode.MarkPackageDirty();
-
 			// If node isn't a MetasoundEditorGraphNode, just remove and return (ex. comment nodes)
 			UMetasoundEditorGraphNode* Node = Cast<UMetasoundEditorGraphNode>(&InNode);
 			UMetasoundEditorGraph* Graph = CastChecked<UMetasoundEditorGraph>(InNode.GetGraph());
@@ -977,7 +942,6 @@ namespace Metasound
 					AddPinToNode(InGraphNode, OutputHandle);
 				}
 			}
-			InGraphNode.MarkPackageDirty();
 		}
 
 		bool FGraphBuilder::IsMatchingInputHandleAndPin(const Frontend::FInputHandle& InInputHandle, const UEdGraphPin& InEditorPin)
@@ -1199,19 +1163,9 @@ namespace Metasound
 					}
 				}
 
-				if (bIsNodeDirty)
-				{
-					EditorNode->MarkPackageDirty();
-				}
-
 				bIsGraphDirty |= bIsNodeDirty;
 			}
 
-			if (bIsGraphDirty)
-			{
-				EditorGraph->MarkPackageDirty();
-			}
-
 			return bIsGraphDirty;
 		}
 
@@ -1369,13 +1323,6 @@ namespace Metasound
 
 			// Synchronize connections.
 			bIsEditorGraphDirty |= SynchronizeConnections(InMetasound);
-
-			if (bIsEditorGraphDirty)
-			{
-				InMetasound.PostEditChange();
-				InMetasound.MarkPackageDirty();
-			}
-
 			return bIsEditorGraphDirty;
 		}
 
diff --git a/Engine/Plugins/Runtime/Metasound/Source/MetasoundEditor/Private/MetasoundEditorGraphBuilder.h b/Engine/Plugins/Runtime/Metasound/Source/MetasoundEditor/Private/MetasoundEditorGraphBuilder.h
--- a/Engine/Plugins/Runtime/Metasound/Source/MetasoundEditor/Private/MetasoundEditorGraphBuilder.h
+++ b/Engine/Plugins/Runtime/Metasound/Source/MetasoundEditor/Private/MetasoundEditorGraphBuilder.h
@@ -55,17 +55,6 @@ namespace Metasound
 			// nor does it remove existing bound Frontend Node (if set) from associated Frontend Graph.
 			static Frontend::FNodeHandle AddNodeHandle(UObject& InMetasound, UMetasoundEditorGraphNode& InGraphNode);
 
-			// Adds an input node to the editor graph that corresponds to the provided node handle.
-			// Generates analogous FNodeHandle.
-			static UMetasoundEditorGraphInputNode* AddInputNode(
-				UObject& InMetasound,
-				const FString& InName,
-				const FName InTypeName,
-				const FText& InToolTip,
-				FVector2D InLocation,
-				EMetasoundFrontendNodeStyleDisplayVisibility Visibility = EMetasoundFrontendNodeStyleDisplayVisibility::Visible,
-				bool bInSelectNewNode = true);
-
 			// Adds a corresponding UMetasoundEditorGraphInputNode for the provided node handle.
 			static UMetasoundEditorGraphInputNode* AddInputNode(UObject& InMetasound, Frontend::FNodeHandle InNodeHandle, FVector2D InLocation, bool bInSelectNewNode = true);
 
@@ -94,21 +83,17 @@ namespace Metasound
 			// Adds an output node to the editor graph that corresponds to the provided node handle.
 			static UMetasoundEditorGraphOutputNode* AddOutputNode(UObject& InMetasound, Frontend::FNodeHandle& InNodeHandle, FVector2D InLocation, bool bInSelectNewNode = true);
 
-			// Adds an output node to the editor graph that corresponds to the provided node handle.
-			// Generates analogous FNodeHandle.
-			static UMetasoundEditorGraphOutputNode* AddOutputNode(UObject& InMetasound, const FString& InName, const FName InTypeName, const FText& InToolTip, FVector2D InLocation, bool bInSelectNewNode = true);
-
 			// Generates analogous FNodeHandle for the given internal node data. Does not bind nor create EdGraph representation of given node.
 			static Frontend::FNodeHandle AddOutputNodeHandle(UObject& InMetasound, const FString& InName, const FName InTypeName, const FText& InToolTip);
 
-			// Attempts to connect frontend node counterparts together for provided pines.  Returns true if succeeded,
+			// Attempts to connect Frontend node counterparts together for provided pins.  Returns true if succeeded,
 			// and breaks pin link and returns false if failed.  If bConnectEdPins is set, will attempt to connect
 			// the Editor Graph representation of the pins.
 			static bool ConnectNodes(UEdGraphPin& InInputPin, UEdGraphPin& InOutputPin, bool bInConnectEdPins);
 
 			// Disconnects pin from any linked input or output nodes, and reflects change
 			// in the Frontend graph.  (Handles generation of literal inputs where needed).
-			// If bAddLiteralInputs - If true, will attempt to connect literal inputs where
+			// If bAddLiteralInputs true, will attempt to connect literal inputs where
 			// applicable post disconnection.
 			static void DisconnectPin(UEdGraphPin& InPin, bool bAddLiteralInputs = true);
 
diff --git a/Engine/Plugins/Runtime/Metasound/Source/MetasoundEditor/Private/MetasoundEditorGraphNode.cpp b/Engine/Plugins/Runtime/Metasound/Source/MetasoundEditor/Private/MetasoundEditorGraphNode.cpp
--- a/Engine/Plugins/Runtime/Metasound/Source/MetasoundEditor/Private/MetasoundEditorGraphNode.cpp
+++ b/Engine/Plugins/Runtime/Metasound/Source/MetasoundEditor/Private/MetasoundEditorGraphNode.cpp
@@ -369,6 +369,11 @@ FGuid UMetasoundEditorGraphOutputNode::GetNodeID() const
 	return FGuid();
 }
 
+bool UMetasoundEditorGraphOutputNode::CanUserDeleteNode() const
+{
+	return !GetNodeHandle()->IsRequired();
+}
+
 void UMetasoundEditorGraphOutputNode::SetNodeID(FGuid InNodeID)
 {
 	if (ensure(Output))
diff --git a/Engine/Plugins/Runtime/Metasound/Source/MetasoundEditor/Private/MetasoundEditorSettings.h b/Engine/Plugins/Runtime/Metasound/Source/MetasoundEditor/Private/MetasoundEditorSettings.h
--- a/Engine/Plugins/Runtime/Metasound/Source/MetasoundEditor/Private/MetasoundEditorSettings.h
+++ b/Engine/Plugins/Runtime/Metasound/Source/MetasoundEditor/Private/MetasoundEditorSettings.h
@@ -9,6 +9,13 @@
 #include "MetasoundEditorSettings.generated.h"
 
 
+UENUM()
+enum class EMetasoundActiveDetailView : uint8
+{
+	Metasound,
+	General
+};
+
 UCLASS(config=EditorPerProjectUserSettings)
 class METASOUNDEDITOR_API UMetasoundEditorSettings : public UObject
 {
@@ -58,4 +65,8 @@ public:
 	/** Trigger pin type color */
 	UPROPERTY(EditAnywhere, config, Category=PinColors)
 	FLinearColor TriggerPinTypeColor;
+
+	/** Determines which details view to show in Metasounds Editor */
+	UPROPERTY(Transient)
+	EMetasoundActiveDetailView DetailView = EMetasoundActiveDetailView::Metasound;
 };
diff --git a/Engine/Plugins/Runtime/Metasound/Source/MetasoundEditor/Private/MetasoundNodeDetailCustomization.cpp b/Engine/Plugins/Runtime/Metasound/Source/MetasoundEditor/Private/MetasoundNodeDetailCustomization.cpp
--- a/Engine/Plugins/Runtime/Metasound/Source/MetasoundEditor/Private/MetasoundNodeDetailCustomization.cpp
+++ b/Engine/Plugins/Runtime/Metasound/Source/MetasoundEditor/Private/MetasoundNodeDetailCustomization.cpp
@@ -747,6 +747,10 @@ namespace Metasound
 						DataTypeString += VariableCustomizationPrivate::ArrayIdentifier;
 					}
 
+					// Have to stop playback to avoid attempting to change live edit data on invalid input type.
+					check(GEditor);
+					GEditor->ResetPreviewAudioComponent();
+
 					InGraphVariable->SetDataType(FName(DataTypeString));
 
 					// Required to rebuild the literal details customization.
@@ -769,6 +773,10 @@ namespace Metasound
 					DataTypeString += VariableCustomizationPrivate::ArrayIdentifier;
 				}
 
+				// Have to stop playback to avoid attempting to change live edit data on invalid input type.
+				check(GEditor);
+				GEditor->ResetPreviewAudioComponent();
+
 				InGraphVariable->SetDataType(FName(DataTypeString));
 
 				// Required to rebuild the literal details customization.
diff --git a/Engine/Plugins/Runtime/Metasound/Source/MetasoundEditor/Public/MetasoundEditorCommands.h b/Engine/Plugins/Runtime/Metasound/Source/MetasoundEditor/Public/MetasoundEditorCommands.h
--- a/Engine/Plugins/Runtime/Metasound/Source/MetasoundEditor/Public/MetasoundEditorCommands.h
+++ b/Engine/Plugins/Runtime/Metasound/Source/MetasoundEditor/Public/MetasoundEditorCommands.h
@@ -46,7 +46,10 @@ namespace Metasound
 			/** Removes an input from the node */
 			TSharedPtr<FUICommandInfo> DeleteInput;
 
-			/** Shows the Metasound's top-level object settings in the Inspector */
+			/** Shows the Metasound's top-level Metasound-specific settings in the Inspector */
+			TSharedPtr<FUICommandInfo> EditGeneralSettings;
+
+			/** Shows the Metasound's top-level general object settings in the Inspector */
 			TSharedPtr<FUICommandInfo> EditMetasoundSettings;
 
 			/** Initialize commands */
diff --git a/Engine/Plugins/Runtime/Metasound/Source/MetasoundEditor/Public/MetasoundEditorGraphNode.h b/Engine/Plugins/Runtime/Metasound/Source/MetasoundEditor/Public/MetasoundEditorGraphNode.h
--- a/Engine/Plugins/Runtime/Metasound/Source/MetasoundEditor/Public/MetasoundEditorGraphNode.h
+++ b/Engine/Plugins/Runtime/Metasound/Source/MetasoundEditor/Public/MetasoundEditorGraphNode.h
@@ -104,6 +104,12 @@ public:
 
 	virtual FMetasoundFrontendClassName GetClassName() const override;
 	virtual FGuid GetNodeID() const override;
+	
+	// Disallow deleting outputs as they require being connected to some
+	// part of the graph by the Frontend Graph Builder (which is enforced
+	// even when the Editor Graph Node does not have a visible input by
+	// way of a literal input.
+	virtual bool CanUserDeleteNode() const override;
 
 protected:
 	virtual void SetNodeID(FGuid InNodeID) override;
diff --git a/Engine/Plugins/Runtime/Metasound/Source/MetasoundEngine/Public/MetasoundSource.h b/Engine/Plugins/Runtime/Metasound/Source/MetasoundEngine/Public/MetasoundSource.h
--- a/Engine/Plugins/Runtime/Metasound/Source/MetasoundEngine/Public/MetasoundSource.h
+++ b/Engine/Plugins/Runtime/Metasound/Source/MetasoundEngine/Public/MetasoundSource.h
@@ -51,7 +51,7 @@ public:
 	UMetasoundSource(const FObjectInitializer& ObjectInitializer);
 
 	// The output audio format of the metasound source.
-	UPROPERTY(EditAnywhere, BlueprintReadOnly, Category = Audio)
+	UPROPERTY(EditAnywhere, BlueprintReadOnly, Category = Metasound)
 	EMetasoundSourceAudioFormat OutputFormat;
 
 #if WITH_EDITOR
diff --git a/Engine/Plugins/Runtime/Metasound/Source/MetasoundFrontend/Public/MetasoundFrontendDocument.h b/Engine/Plugins/Runtime/Metasound/Source/MetasoundFrontend/Public/MetasoundFrontendDocument.h
--- a/Engine/Plugins/Runtime/Metasound/Source/MetasoundFrontend/Public/MetasoundFrontendDocument.h
+++ b/Engine/Plugins/Runtime/Metasound/Source/MetasoundFrontend/Public/MetasoundFrontendDocument.h
@@ -561,31 +561,31 @@ struct FMetasoundFrontendClassMetadata
 
 	FMetasoundFrontendClassMetadata(const Metasound::FNodeClassMetadata& InNodeClassMetadata);
 
-	UPROPERTY(EditAnywhere, Category = General)
+	UPROPERTY(EditAnywhere, Category = Metasound)
 	FMetasoundFrontendClassName ClassName;
 
-	UPROPERTY(EditAnywhere, Category = General)
+	UPROPERTY(EditAnywhere, Category = Metasound)
 	FMetasoundFrontendVersionNumber Version;
 
-	UPROPERTY(VisibleAnywhere, Category = General)
+	UPROPERTY(VisibleAnywhere, Category = Metasound)
 	EMetasoundFrontendClassType Type = EMetasoundFrontendClassType::Invalid;
 
-	UPROPERTY(EditAnywhere, Category = General)
+	UPROPERTY(EditAnywhere, Category = Metasound)
 	FText DisplayName;
 
-	UPROPERTY(EditAnywhere, Category = General)
+	UPROPERTY(EditAnywhere, Category = Metasound)
 	FText Description;
 
-	UPROPERTY(EditAnywhere, Category = General)
+	UPROPERTY(EditAnywhere, Category = Metasound)
 	FText PromptIfMissing;
 
-	UPROPERTY(EditAnywhere, Category = General)
+	UPROPERTY(EditAnywhere, Category = Metasound)
 	FText Author;
 
-	UPROPERTY(EditAnywhere, Category = General)
+	UPROPERTY(EditAnywhere, Category = Metasound)
 	TArray<FName> Keywords;
 
-	UPROPERTY(EditAnywhere, Category = General)
+	UPROPERTY(EditAnywhere, Category = Metasound)
 	TArray<FText> CategoryHierarchy;
 };
 
