- Fix bug with auto-collapse in Metasounds General settings (recollapses undesirably when adding members to child property arrays)
- Testing out reorganize based on usage (not alphabetical) now that we have a manageable number of categories.
- Hide bLooping/Groups
- Move asset data to advanced category (AdvancedDisplay doesn't work with customization) and move to bottom

diff --git a/Engine/Plugins/Runtime/Metasound/Source/MetasoundEditor/Private/MetasoundDetailCustomization.cpp b/Engine/Plugins/Runtime/Metasound/Source/MetasoundEditor/Private/MetasoundDetailCustomization.cpp
--- a/Engine/Plugins/Runtime/Metasound/Source/MetasoundEditor/Private/MetasoundDetailCustomization.cpp
+++ b/Engine/Plugins/Runtime/Metasound/Source/MetasoundEditor/Private/MetasoundDetailCustomization.cpp
@@ -20,6 +20,7 @@
 #include "PropertyRestriction.h"
 #include "SGraphPalette.h"
 #include "SlateCore/Public/Styling/SlateColor.h"
+#include "Sound/SoundWave.h"
 #include "Templates/Casts.h"
 #include "Templates/SharedPointer.h"
 #include "UObject/WeakObjectPtr.h"
@@ -93,15 +94,12 @@ namespace Metasound
 					GeneralCategoryBuilder.AddProperty(MajorVersionHandle);
 					GeneralCategoryBuilder.AddProperty(MinorVersionHandle);
 
-					// Hack to hide categories brought in from UMetasoundSource inherited from USoundBase
 					DetailLayout.HideCategory("Analysis");
 					DetailLayout.HideCategory("Attenuation");
-					DetailLayout.HideCategory("Debug");
 					DetailLayout.HideCategory("Effects");
 					DetailLayout.HideCategory("Loading");
 					DetailLayout.HideCategory("Modulation");
 					DetailLayout.HideCategory("Sound");
-					DetailLayout.HideCategory("SoundWave");
 					DetailLayout.HideCategory("Voice Management");
 				}
 				break;
@@ -111,24 +109,32 @@ namespace Metasound
 					DetailLayout.HideCategory("Metasound");
 
 					const bool bShouldBeInitiallyCollapsed = true;
-					DetailLayout.EditCategory("Analysis").InitiallyCollapsed(bShouldBeInitiallyCollapsed);
+					IDetailCategoryBuilder& SoundCategory = DetailLayout.EditCategory("Sound");
+					SoundCategory.InitiallyCollapsed(bShouldBeInitiallyCollapsed);
+
+					static const TSet<FName> SoundPropsToHide =
+					{
+						GET_MEMBER_NAME_CHECKED(USoundWave, bLooping),
+						GET_MEMBER_NAME_CHECKED(USoundWave, SoundGroup)
+					};
+
+					TArray<TSharedRef<IPropertyHandle>>SoundProperties;
+					SoundCategory.GetDefaultProperties(SoundProperties);
+					for (TSharedRef<IPropertyHandle> Property : SoundProperties)
+					{
+						if (SoundPropsToHide.Contains(Property->GetProperty()->GetFName()))
+						{
+							Property->MarkHiddenByCustomization();
+						}
+					}
+
 					DetailLayout.EditCategory("Attenuation").InitiallyCollapsed(bShouldBeInitiallyCollapsed);
-					DetailLayout.EditCategory("Debug").InitiallyCollapsed(bShouldBeInitiallyCollapsed);
 					DetailLayout.EditCategory("Effects").InitiallyCollapsed(bShouldBeInitiallyCollapsed);
 					DetailLayout.EditCategory("Modulation").InitiallyCollapsed(bShouldBeInitiallyCollapsed);
-					DetailLayout.EditCategory("Sound").InitiallyCollapsed(bShouldBeInitiallyCollapsed);
-					DetailLayout.EditCategory("SoundWave").InitiallyCollapsed(bShouldBeInitiallyCollapsed);
 					DetailLayout.EditCategory("Voice Management").InitiallyCollapsed(bShouldBeInitiallyCollapsed);
 
-					const bool bRestore = false;
-					DetailLayout.EditCategory("Analysis").RestoreExpansionState(bRestore);
-					DetailLayout.EditCategory("Attenuation").RestoreExpansionState(bRestore);
-					DetailLayout.EditCategory("Debug").RestoreExpansionState(bRestore);
-					DetailLayout.EditCategory("Effects").RestoreExpansionState(bRestore);
-					DetailLayout.EditCategory("Modulation").RestoreExpansionState(bRestore);
-					DetailLayout.EditCategory("Sound").RestoreExpansionState(bRestore);
-					DetailLayout.EditCategory("SoundWave").RestoreExpansionState(bRestore);
-					DetailLayout.EditCategory("Voice Management").RestoreExpansionState(bRestore);
+					DetailLayout.EditCategory("Analysis").InitiallyCollapsed(bShouldBeInitiallyCollapsed);
+					DetailLayout.EditCategory("Advanced").InitiallyCollapsed(bShouldBeInitiallyCollapsed);
 
 					break;
 			}
diff --git a/Engine/Source/Runtime/Engine/Classes/Sound/SoundBase.h b/Engine/Source/Runtime/Engine/Classes/Sound/SoundBase.h
--- a/Engine/Source/Runtime/Engine/Classes/Sound/SoundBase.h
+++ b/Engine/Source/Runtime/Engine/Classes/Sound/SoundBase.h
@@ -185,7 +185,7 @@ public:
 	TArray<FSoundSourceBusSendInfo> PreEffectBusSends;
 
 	/** Array of user data stored with the asset */
-	UPROPERTY(EditAnywhere, AdvancedDisplay, Instanced, Category = Sound)
+	UPROPERTY(EditAnywhere, AdvancedDisplay, Instanced, Category = Advanced)
 	TArray<TObjectPtr<UAssetUserData>> AssetUserData;
 
 public:
