Metasounds (Wave Player Node): Added tooltip reminder to set sound waves to be seekable for the Seek Time and Seek Triggers to have anny affect.

diff --git a/Engine/Plugins/Runtime/Metasound/Source/MetasoundEngine/Private/MetasoundWavePlayerNode.cpp b/Engine/Plugins/Runtime/Metasound/Source/MetasoundEngine/Private/MetasoundWavePlayerNode.cpp
--- a/Engine/Plugins/Runtime/Metasound/Source/MetasoundEngine/Private/MetasoundWavePlayerNode.cpp
+++ b/Engine/Plugins/Runtime/Metasound/Source/MetasoundEngine/Private/MetasoundWavePlayerNode.cpp
@@ -408,8 +408,8 @@ namespace Metasound
 				TInputDataVertexModel<FWaveAsset>(TEXT("Wave"), LOCTEXT("WaveTooltip", "The Wave to be decoded")),
 				TInputDataVertexModel<FTrigger>(TEXT("Play"), LOCTEXT("PlayTooltip", "Play the input wave.")),
 				TInputDataVertexModel<FTrigger>(TEXT("Stop"), LOCTEXT("StopTooltip", "Stop the input wave.")),
-				TInputDataVertexModel<FTrigger>(TEXT("Seek"), LOCTEXT("PlayTooltip", "Trigger the playing of the input wave.")),
-				TInputDataVertexModel<float>(TEXT("SeekTime"), LOCTEXT("SeekTimeTooltip", "Seek time in seconds.")),
+				TInputDataVertexModel<FTrigger>(TEXT("Seek"), LOCTEXT("PlayTooltip", "Seek to the given time. (Wave being played must be set to seekable!)")),
+				TInputDataVertexModel<float>(TEXT("SeekTime"), LOCTEXT("SeekTimeTooltip", "Seek time in seconds. (Wave being played must be set to seekable!)")),
 				TInputDataVertexModel<float>(TEXT("PitchShift"), LOCTEXT("PitchShiftTooltip", "Pitch Shift in semi-tones.")),
 				TInputDataVertexModel<bool>(TEXT("Loop"), LOCTEXT("LoopTooltip", "Wave will loop if true"))
 			),
