fix for missing param assert

diff --git a/Engine/Plugins/Runtime/Metasound/Source/MetasoundStandardNodes/Private/MetasoundTriggerDelayNode.cpp b/Engine/Plugins/Runtime/Metasound/Source/MetasoundStandardNodes/Private/MetasoundTriggerDelayNode.cpp
--- a/Engine/Plugins/Runtime/Metasound/Source/MetasoundStandardNodes/Private/MetasoundTriggerDelayNode.cpp
+++ b/Engine/Plugins/Runtime/Metasound/Source/MetasoundStandardNodes/Private/MetasoundTriggerDelayNode.cpp
@@ -147,7 +147,7 @@ namespace Metasound
 		const FTriggerDelayNode& DelayNode = static_cast<const FTriggerDelayNode&>(InParams.Node);
 
 		const FInputVertexInterface& InputInterface = GetVertexInterface().GetInputInterface();
-		FTimeReadRef Delay = InParams.InputDataReferences.GetDataReadReferenceOrConstructWithVertexDefault<FTime>(InputInterface, TEXT("Delay"), InParams.OperatorSettings);
+		FTimeReadRef Delay = InParams.InputDataReferences.GetDataReadReferenceOrConstructWithVertexDefault<FTime>(InputInterface, METASOUND_GET_PARAM_NAME(InputDelayTime), InParams.OperatorSettings);
 
 		FTriggerReadRef TriggerIn = InParams.InputDataReferences.GetDataReadReferenceOrConstruct<FTrigger>(METASOUND_GET_PARAM_NAME(InputInTrigger), InParams.OperatorSettings);
 		FTriggerReadRef TriggerReset = InParams.InputDataReferences.GetDataReadReferenceOrConstruct<FTrigger>(METASOUND_GET_PARAM_NAME(InputReset), InParams.OperatorSettings);
