Refactoring the frequency lerping.

- Removing it by default (because it resulted in hard-coded glide)
- Adjusted to a "feature" - glide.
- Switched from a linear interp to an exp ease algorithm
- Made 2D rotation as the default sine generation technique (note: ignored on FM synthesis)

diff --git a/Engine/Plugins/Runtime/Metasound/Source/MetasoundStandardNodes/Private/MetasoundOscillatorNodes.cpp b/Engine/Plugins/Runtime/Metasound/Source/MetasoundStandardNodes/Private/MetasoundOscillatorNodes.cpp
--- a/Engine/Plugins/Runtime/Metasound/Source/MetasoundStandardNodes/Private/MetasoundOscillatorNodes.cpp
+++ b/Engine/Plugins/Runtime/Metasound/Source/MetasoundStandardNodes/Private/MetasoundOscillatorNodes.cpp
@@ -11,6 +11,7 @@
 #include "MetasoundStandardNodesNames.h"
 #include "MetasoundTrigger.h"
 #include "MetasoundVertex.h"
+#include "MetasoundTime.h"
 #include "MetasoundStandardNodesCategories.h"
 
 #define LOCTEXT_NAMESPACE "MetasoundStandardNodes"
@@ -31,6 +32,7 @@ namespace Metasound
 		{
 			float SampleRate = 0.f;
 			float FrequencyHz = 0.f;
+			float GlideEaseFactor = 0.0f;
 			float PulseWidth = 0.f;
 			TArrayView<float> AlignedBuffer;
 			TArrayView<const float> FM;
@@ -90,17 +92,18 @@ namespace Metasound
 			PhaseWrap Wrap;
 			float CurrentFade = 0.0f;
 			float FadeSmooth = 0.01f;
-			float PreviousFreq = -1.f;
+			float CurrentFreq = -1.f;
 		
 			void operator()(const FGeneratorArgs& InArgs)
 			{
 				int32 RemainingSamplesInBlock = InArgs.AlignedBuffer.Num();
 				float* Out = InArgs.AlignedBuffer.GetData();
 				const float OneOverSampleRate = 1.f / InArgs.SampleRate;
-		
+				
 				// Constant Freq between this an last update? 
-				if (FMath::IsNearlyEqual(InArgs.FrequencyHz, PreviousFreq) || PreviousFreq < 0.f )
+				if (FMath::IsNearlyEqual(InArgs.FrequencyHz, CurrentFreq) || CurrentFreq < 0.f || FMath::IsNearlyEqual(InArgs.GlideEaseFactor, 1.0f))
 				{
+					CurrentFreq = InArgs.FrequencyHz;
 					const float DeltaPhase = InArgs.FrequencyHz * OneOverSampleRate;
 					while (RemainingSamplesInBlock > 0)
 					{
@@ -117,20 +120,13 @@ namespace Metasound
 						RemainingSamplesInBlock--;
 					}
 				}
-				else
-				{
-					// Lerp freq over this block 
-					const float OneOverBlockCount = 1.f / (float)RemainingSamplesInBlock;
-					const float BlockDeltaFreq = InArgs.FrequencyHz - PreviousFreq;
-					const float DeltaFreq = BlockDeltaFreq * OneOverBlockCount;
-
-					float Freq = PreviousFreq;
-					
+ 				else
+ 				{			
 					while (RemainingSamplesInBlock > 0)
 					{
 						Wrap(Phase);
 
-						const float DeltaPhase = Freq * OneOverSampleRate;
+						const float DeltaPhase = CurrentFreq * OneOverSampleRate;
 						
 						float Output = Osc(Phase, DeltaPhase, InArgs);
 						Output *= CurrentFade;
@@ -139,13 +135,13 @@ namespace Metasound
 						CurrentFade = CurrentFade + FadeSmooth * (1.0f - CurrentFade);
 
 						Phase += DeltaPhase;
-						Freq += DeltaFreq;
+
+						// lerp frequency based on the glide factor
+						CurrentFreq = CurrentFreq + InArgs.GlideEaseFactor * (InArgs.FrequencyHz - CurrentFreq);
 
 						RemainingSamplesInBlock--;
 					}
 				}
-
-				PreviousFreq = InArgs.FrequencyHz;
 			}
 		};
 
@@ -163,7 +159,7 @@ namespace Metasound
 
 			float CurrentFade = 0.0f;
 			float FadeSmooth = 0.01f;
-			float PreviousFreq = -1.f;
+			float CurrentFreq = -1.f;
 
 			void operator()(const FGeneratorArgs& InArgs)
 			{
@@ -173,7 +169,7 @@ namespace Metasound
 				const float OneOverSampleRate = 1.f / InArgs.SampleRate;
 
 				// Constant Base Freq between this and last update? 
-				if (FMath::IsNearlyEqual(InArgs.FrequencyHz, PreviousFreq) || PreviousFreq < 0.f)
+				if (FMath::IsNearlyEqual(InArgs.FrequencyHz, CurrentFreq) || CurrentFreq < 0.f || FMath::IsNearlyEqual(InArgs.GlideEaseFactor, 1.0f))
 				{
 					while (RemainingSamplesInBlock > 0)
 					{
@@ -192,19 +188,14 @@ namespace Metasound
 
 						RemainingSamplesInBlock--;
 					}
+
+					CurrentFreq = InArgs.FrequencyHz;
 				}
 				else
-				{
-					// Lerp Base Frequency over this block
-					const float OneOverBlockCount = 1.f / (float)RemainingSamplesInBlock;
-					const float BlockDeltaFreq = InArgs.FrequencyHz - PreviousFreq;
-					const float DeltaFreq = BlockDeltaFreq * OneOverBlockCount;
-					
-					float Freq = PreviousFreq;
-
+				{	
 					while (RemainingSamplesInBlock > 0)
 					{
-						const float ModulatedFreqSum = Freq + *FM++;
+						const float ModulatedFreqSum = CurrentFreq + *FM++;
 						const float DeltaPhase = ModulatedFreqSum * OneOverSampleRate;
 
 						Wrap(Phase);
@@ -216,12 +207,13 @@ namespace Metasound
 						CurrentFade = CurrentFade + FadeSmooth * (1.0f - CurrentFade);
 
 						Phase += DeltaPhase;
-						Freq += DeltaFreq;
+
+						// lerp frequency based on the glide factor
+						CurrentFreq = CurrentFreq + InArgs.GlideEaseFactor * (InArgs.FrequencyHz - CurrentFreq);
 
 						RemainingSamplesInBlock--;
 					}
 				}
-				PreviousFreq = InArgs.FrequencyHz;
 			}
 		};	
 		
@@ -243,6 +235,40 @@ namespace Metasound
 				return Audio::FastSin3(PhaseRadians - PI); // Expects [-PI, PI] 
 			}
 		};
+		struct FSineWaveTableGenerator
+		{
+			static const TArray<float>& GetWaveTable()
+			{
+				auto MakeSineTable = []() -> const TArray<float>
+				{
+					int32 TableSize = 4096;
+					// Generate the table
+					TArray<float> WaveTable;
+					WaveTable.AddUninitialized(TableSize);
+					float* WaveTableData = WaveTable.GetData();
+					for (int32 i = 0; i < TableSize; ++i)
+					{
+						static const float TwoPi = 2.f * PI;
+						float Phase = (float)i / TableSize;
+						WaveTableData[i] = FMath::Sin(Phase * TwoPi);
+					}
+					return WaveTable;
+				};
+
+				static const TArray<float> SineWaveTable = MakeSineTable();
+				return SineWaveTable;
+			}
+
+			float operator()(float InPhase, float, const FGeneratorArgs&) const
+			{			
+				const TArray<float>& WaveTable = GetWaveTable();
+				int32 LastIndex = WaveTable.Num() - 1;
+				int32 TableIndex = (int32)(InPhase * (float)(LastIndex));
+				TableIndex = FMath::Wrap(TableIndex, 0, LastIndex);
+				return WaveTable[TableIndex];
+			}
+		};
+
 		struct F2DRotatorGenerateBlock
 		{
 			Audio::FSinOsc2DRotation Rotator;
@@ -261,6 +287,9 @@ namespace Metasound
 		using FBhaskara					= TGenerateBlock<FBhaskaraGenerator>;
 		using FBhaskaraWithFm			= TGenerateBlockFM<FBhaskaraGenerator>;
 
+		using FSineWaveTable			= TGenerateBlock<FSineWaveTableGenerator>;
+		using FSineWaveTableWithFm		= TGenerateBlockFM<FSineWaveTableGenerator>;
+
 		// Saws.
 		struct FSawGenerator
 		{
@@ -394,6 +423,7 @@ namespace Metasound
 		static constexpr const TCHAR* BaseFrequencyPinName = TEXT("Frequency");
 		static constexpr const TCHAR* PhaseResetPinName = TEXT("Sync");
 		static constexpr const TCHAR* PhaseOffsetPinName = TEXT("Phase Offset");
+		static constexpr const TCHAR* GlideFactorPinName = TEXT("Glide");
 		static constexpr const TCHAR* AudioOutPinName = TEXT("Audio");
 		static constexpr const TCHAR* BiPolarPinName = TEXT("Bi Polar");
 
@@ -408,7 +438,8 @@ namespace Metasound
 					TInputDataVertexModel<float>(BaseFrequencyPinName, LOCTEXT("OscBaseFrequencyDescription", "Base Frequency of Oscillator in HZ"), 440.f),
 					TInputDataVertexModel<FAudioBuffer>(FrequencyModPinName, LOCTEXT("OscModFrequencyDescription", "Modulation Frequency Input (for doing FM)")),
 					TInputDataVertexModel<FTrigger>(PhaseResetPinName, LOCTEXT("OscPhaseResetDescription", "Phase Reset")),
-					TInputDataVertexModel<float>(PhaseOffsetPinName, LOCTEXT("OscPhaseOffsetDescription", "Phase Offset In Degrees (0..360)"), 0.f)
+					TInputDataVertexModel<float>(PhaseOffsetPinName, LOCTEXT("OscPhaseOffsetDescription", "Phase Offset In Degrees (0..360)"), 0.f),
+					TInputDataVertexModel<float>(GlideFactorPinName, LOCTEXT("OscGlideFactorDescription", "The amount of glide to use when changing frequencies. 0.0 = no glide, 1.0 = lots of glide."), 0.f)
 				},
 				FOutputVertexInterface{
 					TOutputDataVertexModel<FAudioBuffer>(TEXT("Audio"), LOCTEXT("AudioTooltip", "The output audio"))
@@ -426,6 +457,7 @@ namespace Metasound
 		FFloatReadRef BaseFrequency;
 		FFloatReadRef PhaseOffset; 
 		FTriggerReadRef PhaseReset;
+		FFloatReadRef GlideFactor;
 		FBoolReadRef BiPolar;
 	};
 
@@ -443,6 +475,7 @@ namespace Metasound
 			, BaseFrequency(InConstructParams.BaseFrequency)
 			, PhaseReset(InConstructParams.PhaseReset)
 			, PhaseOffset(InConstructParams.PhaseOffset)
+			, GlideFactor(InConstructParams.GlideFactor)
 			, BiPolar(InConstructParams.BiPolar)
 			, AudioBuffer(FAudioBufferWriteRef::CreateNew(InConstructParams.Settings))
 		{
@@ -456,6 +489,7 @@ namespace Metasound
 			InputDataReferences.AddDataReadReference(FOscilatorFactoryBase::BaseFrequencyPinName, BaseFrequency);
 			InputDataReferences.AddDataReadReference(FOscilatorFactoryBase::PhaseOffsetPinName, PhaseOffset);
 			InputDataReferences.AddDataReadReference(FOscilatorFactoryBase::PhaseResetPinName, PhaseReset);
+			InputDataReferences.AddDataReadReference(FOscilatorFactoryBase::GlideFactorPinName, GlideFactor);
 			InputDataReferences.AddDataReadReference(FOscilatorFactoryBase::BiPolarPinName, BiPolar);
 			return InputDataReferences;
 		}
@@ -479,20 +513,21 @@ namespace Metasound
 		{
 			// Clamp frequencies into Nyquist range.
 			const float ClampedFreq = FMath::Clamp(*BaseFrequency, -Nyquist, Nyquist);
-			
+			const float ClampedGlideEase = Audio::GetLogFrequencyClamped(*GlideFactor, { 0.0f, 1.0f }, { 1.0f, 0.0001f });
+
 			AudioBuffer->Zero();
 
 			Derived* Self = static_cast<Derived*>(this);
 			PhaseReset->ExecuteBlock
 			(
-				[Self, ClampedFreq](int32 InFrameStart, int32 InFrameEnd) 
+				[Self, ClampedFreq, ClampedGlideEase](int32 InFrameStart, int32 InFrameEnd)
 				{
-					Self->Generate(InFrameStart, InFrameEnd, ClampedFreq);
+					Self->Generate(InFrameStart, InFrameEnd, ClampedFreq, ClampedGlideEase);
 				},
-				[Self, ClampedFreq](int32 InFrameStart, int32 InFrameEnd) 
+				[Self, ClampedFreq, ClampedGlideEase](int32 InFrameStart, int32 InFrameEnd)
 				{
 					Self->ResetPhase(*Self->PhaseOffset);
-					Self->Generate(InFrameStart, InFrameEnd, ClampedFreq);
+					Self->Generate(InFrameStart, InFrameEnd, ClampedFreq, ClampedGlideEase);
 				}
 			);
 
@@ -512,6 +547,7 @@ namespace Metasound
 		FFloatReadRef BaseFrequency;
 		FTriggerReadRef PhaseReset;
 		FFloatReadRef PhaseOffset;
+		FFloatReadRef GlideFactor;
 		FBoolReadRef BiPolar;
 
 		FAudioBufferWriteRef AudioBuffer;
@@ -527,7 +563,7 @@ namespace Metasound
 			: Super(InConstructParams)
 		{}
 
-		void Generate(int32 InStartFrame, int32 InEndFrame, float InClampedFreq)
+		void Generate(int32 InStartFrame, int32 InEndFrame, float InClampedFreq, float InClampedGlideEase)
 		{		
 			int32 NumFrames = InEndFrame - InStartFrame;
 
@@ -537,6 +573,7 @@ namespace Metasound
 				{
 					this->SampleRate, 
 					InClampedFreq,
+					InClampedGlideEase,
 					0.f,
 					MakeArrayView(this->AudioBuffer->GetData() + InStartFrame, NumFrames), // Not aligned.
 				});
@@ -561,7 +598,7 @@ namespace Metasound
 			return Inputs;
 		}
 
-		void Generate(int32 InStartFrame, int32 InEndFrame, float InClampedFreq)
+		void Generate(int32 InStartFrame, int32 InEndFrame, float InClampedFreq, float InClampedGlideEase)
 		{
 			int32 NumFrames = InEndFrame - InStartFrame;
 			if (*this->Enabled && NumFrames > 0)
@@ -570,6 +607,7 @@ namespace Metasound
 				{
 					this->SampleRate, 
 					InClampedFreq,
+					InClampedGlideEase,
 					0.f,
 					MakeArrayView(this->AudioBuffer->GetData() + InStartFrame, NumFrames), // Not aligned.
 					MakeArrayView(this->Fm->GetData() + InStartFrame, NumFrames) // Not aligned.
@@ -581,11 +619,12 @@ namespace Metasound
 		FAudioBufferReadRef Fm;
 	};
 
-	FOscilatorNodeBase::FOscilatorNodeBase(const FString& InInstanceName, const FGuid& InInstanceID, const FNodeClassMetadata& InInfo, const TSharedRef<IOperatorFactory, ESPMode::ThreadSafe>& InFactory, float InDefaultFrequency, bool bInDefaultEnablement)
+	FOscilatorNodeBase::FOscilatorNodeBase(const FString& InInstanceName, const FGuid& InInstanceID, const FNodeClassMetadata& InInfo, const TSharedRef<IOperatorFactory, ESPMode::ThreadSafe>& InFactory, float InDefaultFrequency, float InDefaultGlideFactor, bool bInDefaultEnablement)
 		: FNode(InInstanceName, InInstanceID, InInfo)
 		, Factory(InFactory)
 		, VertexInterface(GetMetadata().DefaultInterface)
 		, DefaultFrequency(InDefaultFrequency)
+		, DefaultGlideFactor(InDefaultGlideFactor)
 		, bDefaultEnablement(bInDefaultEnablement)
 	{}
 	
@@ -595,21 +634,21 @@ namespace Metasound
 
 	enum class ESineGenerationType
 	{
+		Rotation,
 		Sinf,
 		FilterResonanceFeedback,
-		Rotation,
 		Bhaskara,
 		Wavetable,
 	};
 
-	DECLARE_METASOUND_ENUM(ESineGenerationType, ESineGenerationType::Sinf, METASOUNDSTANDARDNODES_API,
+	DECLARE_METASOUND_ENUM(ESineGenerationType, ESineGenerationType::Wavetable, METASOUNDSTANDARDNODES_API,
 		FEnumSineGenerationType, FEnumSineGenerationTypeInfo, FEnumSineGenerationTypeReadRef, FEnumSineGenerationTypeWriteRef);
 
 	DEFINE_METASOUND_ENUM_BEGIN(ESineGenerationType, FEnumSineGenerationType, "SineGenerationType")
-		DEFINE_METASOUND_ENUM_ENTRY(ESineGenerationType::Sinf, LOCTEXT("SinfDescription", "Pure Math"), LOCTEXT("SinfDescriptionTT", "Uses Pure Math (Sinf) to generate the Sine Wave (most expensive)")),
-		DEFINE_METASOUND_ENUM_ENTRY(ESineGenerationType::Rotation, LOCTEXT("RotationDescription", "2D Rotation"), LOCTEXT("RotationDescriptionTT", "Rotates around the unit circle generate the sinewave")),
+		DEFINE_METASOUND_ENUM_ENTRY(ESineGenerationType::Rotation, LOCTEXT("RotationDescription", "2D Rotation"), LOCTEXT("RotationDescriptionTT", "Rotates around the unit circle generate the sine. Note: Glide and audio rate FM modulation is not supported with the 2D rotator.")),
+		DEFINE_METASOUND_ENUM_ENTRY(ESineGenerationType::Sinf, LOCTEXT("SinfDescription", "Pure Math"), LOCTEXT("SinfDescriptionTT", "Uses the standard math library (Sinf) to generate the sine (most expensive)")),
 		DEFINE_METASOUND_ENUM_ENTRY(ESineGenerationType::Bhaskara, LOCTEXT("BhaskaraDescription", "Bhaskara"), LOCTEXT("BhaskaraDescriptionTT", "Sine approximation using Bhaskara technique discovered in 7th century")),
-		//DEFINE_METASOUND_ENUM_ENTRY(ESineGenerationType::Wavetable, LOCTEXT("WavetableDescription", "Wavetable"), LOCTEXT("WavetableDescriptionTT", "Uses wavetable to generate the sinewave")),
+		DEFINE_METASOUND_ENUM_ENTRY(ESineGenerationType::Wavetable, LOCTEXT("WavetableDescription", "Wavetable"), LOCTEXT("WavetableDescriptionTT", "Uses a wavetable to generate the sine")),
 		//DEFINE_METASOUND_ENUM_ENTRY(ESineGenerationType::FilterResonanceFeedback, LOCTEXT("FRFDescription", "Filter resonance"), LOCTEXT("FRFDescriptionTT", "Filter resonance/feedback")),
 	DEFINE_METASOUND_ENUM_END()
 
@@ -645,7 +684,7 @@ namespace Metasound
 			{
 				FVertexInterface Interface = GetCommmonVertexInterface();
 				Interface.GetInputInterface().Add(
-					TInputDataVertexModel<FEnumSineGenerationType>(SineTypePin, LOCTEXT("SineTypeDescription", "Type of the Sinewave Generator"))
+					TInputDataVertexModel<FEnumSineGenerationType>(SineTypePin, LOCTEXT("SineTypeDescription", "Type of the Sinewave Generator"), static_cast<int32>(ESineGenerationType::Wavetable))
 				);
 				return Interface;
 			};
@@ -659,7 +698,6 @@ namespace Metasound
 			const FDataReferenceCollection& InputCol = InParams.InputDataReferences;
 			const FOperatorSettings& Settings = InParams.OperatorSettings;
 			using namespace Generators;
-			
 			FOscillatorOperatorConstructParams OpParams
 			{
 				Settings,
@@ -667,16 +705,18 @@ namespace Metasound
 				InputCol.GetDataReadReferenceOrConstruct<float>(BaseFrequencyPinName, SineNode.GetDefaultFrequency()),
 				InputCol.GetDataReadReferenceOrConstruct<float>(PhaseOffsetPinName, SineNode.GetDefaultPhaseOffset()),
 				InputCol.GetDataReadReferenceOrConstruct<FTrigger>(PhaseResetPinName, Settings),
+				InputCol.GetDataReadReferenceOrConstruct<float>(GlideFactorPinName, SineNode.GetDefaultGlideFactor()),
 				InputCol.GetDataReadReferenceOrConstruct<bool>(BiPolarPinName, true)
 			};
 
 			// TODO: Make this a static prop. For now its a pin.
-			FEnumSineGenerationTypeReadRef Type = InputCol.GetDataReadReferenceOrConstruct<FEnumSineGenerationType>(SineTypePin);
 			
 			// Check to see if we have an FM input connected.
 			bool bHasFM = InputCol.ContainsDataReadReference<FAudioBuffer>(FrequencyModPinName);			
+			FEnumSineGenerationTypeReadRef Type = InputCol.GetDataReadReferenceOrConstruct<FEnumSineGenerationType>(SineTypePin, ESineGenerationType::Rotation);
 			if (bHasFM)
 			{
+
 				// FM Oscillators.
 				FAudioBufferReadRef FmBuffer = InputCol.GetDataReadReference<FAudioBuffer>(FrequencyModPinName);
 				switch (*Type)
@@ -684,6 +724,7 @@ namespace Metasound
 				default:
 				case ESineGenerationType::Sinf: return MakeUnique<TOscillatorOperatorFM<FSinfWithFm>>(OpParams, FmBuffer);
 				case ESineGenerationType::Bhaskara: return MakeUnique<TOscillatorOperatorFM<FBhaskaraWithFm>>(OpParams, FmBuffer);
+				case ESineGenerationType::Wavetable: return MakeUnique<TOscillatorOperatorFM<FSineWaveTableWithFm>>(OpParams, FmBuffer);
 				}
 			}
 			else //HasFM
@@ -694,6 +735,7 @@ namespace Metasound
 				case ESineGenerationType::Sinf: return MakeUnique<TOscillatorOperator<FSinf>>(OpParams);
 				case ESineGenerationType::Rotation: return MakeUnique<TOscillatorOperator<F2DRotatorGenerateBlock>>(OpParams);
 				case ESineGenerationType::Bhaskara: return MakeUnique<TOscillatorOperator<FBhaskara>>(OpParams);
+				case ESineGenerationType::Wavetable: return MakeUnique<TOscillatorOperator<FSineWaveTable>>(OpParams);
 				}
 			} // HasFM
 			return nullptr;
@@ -701,12 +743,12 @@ namespace Metasound
 	private:
 	};
 
-	FSineOscilatorNode::FSineOscilatorNode(const FString& InInstanceName, const FGuid& InInstanceID, float InDefaultFrequency, bool bInDefaultEnablement)
-		: FOscilatorNodeBase(InInstanceName,InInstanceID, FFactory::GetNodeInfo(), MakeShared<FFactory, ESPMode::ThreadSafe>(), InDefaultFrequency, bInDefaultEnablement )
+	FSineOscilatorNode::FSineOscilatorNode(const FString& InInstanceName, const FGuid& InInstanceID, float InDefaultFrequency, float InDefautlGlideFactor, bool bInDefaultEnablement)
+		: FOscilatorNodeBase(InInstanceName,InInstanceID, FFactory::GetNodeInfo(), MakeShared<FFactory, ESPMode::ThreadSafe>(), InDefaultFrequency, InDefautlGlideFactor, bInDefaultEnablement )
 	{}
 
 	FSineOscilatorNode::FSineOscilatorNode(const FNodeInitData& InInitData)
-		: FSineOscilatorNode(InInitData.InstanceName, InInitData.InstanceID, 440.0f, true)
+		: FSineOscilatorNode(InInitData.InstanceName, InInitData.InstanceID, 440.0f, 0.0f, true)
 	{}
 
 	METASOUND_REGISTER_NODE(FSineOscilatorNode);
@@ -792,6 +834,7 @@ namespace Metasound
 			InputCol.GetDataReadReferenceOrConstruct<float>(BaseFrequencyPinName, Node.GetDefaultFrequency()),
 			InputCol.GetDataReadReferenceOrConstruct<float>(PhaseOffsetPinName, Node.GetDefaultPhaseOffset()),
 			InputCol.GetDataReadReferenceOrConstruct<FTrigger>(PhaseResetPinName, Settings),
+			InputCol.GetDataReadReferenceOrConstruct<float>(GlideFactorPinName, Node.GetDefaultGlideFactor()),
 			InputCol.GetDataReadReferenceOrConstruct<bool>(BiPolarPinName, true)
 		};
 
@@ -819,12 +862,12 @@ namespace Metasound
 		return nullptr;
 	}
 
-	FSawOscilatorNode::FSawOscilatorNode(const FString& InInstanceName, const FGuid& InInstanceID, float InDefaultFrequency, bool bInDefaultEnablement)
-		: FOscilatorNodeBase(InInstanceName, InInstanceID, FFactory::GetNodeInfo(), MakeShared<FFactory, ESPMode::ThreadSafe>(), InDefaultFrequency, bInDefaultEnablement)
+	FSawOscilatorNode::FSawOscilatorNode(const FString& InInstanceName, const FGuid& InInstanceID, float InDefaultFrequency, float InDefaultGlideFactor, bool bInDefaultEnablement)
+		: FOscilatorNodeBase(InInstanceName, InInstanceID, FFactory::GetNodeInfo(), MakeShared<FFactory, ESPMode::ThreadSafe>(), InDefaultFrequency, InDefaultGlideFactor, bInDefaultEnablement)
 	{}
 
 	FSawOscilatorNode::FSawOscilatorNode(const FNodeInitData& InInitData)
-		: FSawOscilatorNode(InInitData.InstanceName, InInitData.InstanceID, 440.0f, true)
+		: FSawOscilatorNode(InInitData.InstanceName, InInitData.InstanceID, 440.0f, 0.0f, true)
 	{}
 
 	METASOUND_REGISTER_NODE(FSawOscilatorNode);
@@ -858,17 +901,18 @@ namespace Metasound
 			: Super(InConstructParams)
 			, PulseWidth(InPulseWidth)
 		{}
-		void Generate(int32 InStartFrame, int32 InEndFrame, float InClampedFreq)
+		void Generate(int32 InStartFrame, int32 InEndFrame, float InClampedFreq, float InClampedGlideEase)
 		{
 			int32 NumFrames = InEndFrame - InStartFrame;
-
+			float ClampedPulseWidth = FMath::Clamp(*PulseWidth, 0.0f, 1.0f);
 			if (*this->Enabled && NumFrames > 0)
 			{
 				this->Generator(
 				{
 					this->SampleRate,
 					InClampedFreq,
-					*PulseWidth,
+					InClampedGlideEase,
+					ClampedPulseWidth,
 					MakeArrayView(this->AudioBuffer->GetData() + InStartFrame, NumFrames), // Not aligned.
 				});
 			}
@@ -891,17 +935,18 @@ namespace Metasound
 			check(InFm->GetData());
 			check(InConstructParams.Settings.GetNumFramesPerBlock() == InFm->Num());
 		}
-		void Generate(int32 InStartFrame, int32 InEndFrame, float InClampedFreq)
+		void Generate(int32 InStartFrame, int32 InEndFrame, float InClampedFreq, float InClampedGlideEase)
 		{
 			int32 NumFrames = InEndFrame - InStartFrame;
-
+			float ClampedPulseWidth = FMath::Clamp(*PulseWidth, 0.0f, 1.0f);
 			if (*this->Enabled && NumFrames > 0)
 			{
 				this->Generator(
 				{
 					this->SampleRate,
 					InClampedFreq,
-					*PulseWidth,
+					InClampedGlideEase,
+					ClampedPulseWidth,
 					MakeArrayView(this->AudioBuffer->GetData() + InStartFrame, NumFrames), // Not aligned.
 					MakeArrayView(this->FM->GetData() + InStartFrame, NumFrames), // Not aligned.
 				});
@@ -935,6 +980,7 @@ namespace Metasound
 				InputCol.GetDataReadReferenceOrConstruct<float>(BaseFrequencyPinName, Node.GetDefaultFrequency()),
 				InputCol.GetDataReadReferenceOrConstruct<float>(PhaseOffsetPinName, Node.GetDefaultPhaseOffset()),
 				InputCol.GetDataReadReferenceOrConstruct<FTrigger>(PhaseResetPinName, Settings),
+				InputCol.GetDataReadReferenceOrConstruct<float>(GlideFactorPinName, Node.GetDefaultGlideFactor()),
 				InputCol.GetDataReadReferenceOrConstruct<bool>(BiPolarPinName, true)
 			};
 
@@ -999,12 +1045,12 @@ namespace Metasound
 	private:
 	};
 	
-	FSquareOscilatorNode::FSquareOscilatorNode(const FString& InInstanceName, const FGuid& InInstanceID, float InDefaultFrequency, bool bInDefaultEnablement)
-		: FOscilatorNodeBase(InInstanceName, InInstanceID, FFactory::GetNodeInfo(), MakeShared<FFactory, ESPMode::ThreadSafe>(), InDefaultFrequency, bInDefaultEnablement)
+	FSquareOscilatorNode::FSquareOscilatorNode(const FString& InInstanceName, const FGuid& InInstanceID, float InDefaultFrequency, float InDefaultGlideFactor, bool bInDefaultEnablement)
+		: FOscilatorNodeBase(InInstanceName, InInstanceID, FFactory::GetNodeInfo(), MakeShared<FFactory, ESPMode::ThreadSafe>(), InDefaultFrequency, InDefaultGlideFactor, bInDefaultEnablement)
 	{}
 
 	FSquareOscilatorNode::FSquareOscilatorNode(const FNodeInitData& InInitData)
-		: FSquareOscilatorNode(InInitData.InstanceName, InInitData.InstanceID, 440.0f, true)
+		: FSquareOscilatorNode(InInitData.InstanceName, InInitData.InstanceID, 440.0f, 0.0f, true)
 	{}
 
 	METASOUND_REGISTER_NODE(FSquareOscilatorNode)
@@ -1051,6 +1097,7 @@ namespace Metasound
 				InputCol.GetDataReadReferenceOrConstruct<float>(BaseFrequencyPinName, Node.GetDefaultFrequency()),
 				InputCol.GetDataReadReferenceOrConstruct<float>(PhaseOffsetPinName, Node.GetDefaultPhaseOffset()),
 				InputCol.GetDataReadReferenceOrConstruct<FTrigger>(PhaseResetPinName, Settings),
+				InputCol.GetDataReadReferenceOrConstruct<float>(GlideFactorPinName, Node.GetDefaultGlideFactor()),
 				InputCol.GetDataReadReferenceOrConstruct<bool>(BiPolarPinName, true)
 			};
 
@@ -1113,14 +1160,14 @@ namespace Metasound
 	private:
 	};
 	
-	FTriangleOscilatorNode::FTriangleOscilatorNode(const FString& InInstanceName, const FGuid& InInstanceID, float InDefaultFrequency, bool bInDefaultEnablement)
-		: FOscilatorNodeBase(InInstanceName, InInstanceID, FFactory::GetNodeInfo(), MakeShared<FFactory, ESPMode::ThreadSafe>(), InDefaultFrequency, bInDefaultEnablement)
+	FTriangleOscilatorNode::FTriangleOscilatorNode(const FString& InInstanceName, const FGuid& InInstanceID, float InDefaultFrequency, float InDefaultGlideFactor, bool bInDefaultEnablement)
+		: FOscilatorNodeBase(InInstanceName, InInstanceID, FFactory::GetNodeInfo(), MakeShared<FFactory, ESPMode::ThreadSafe>(), InDefaultFrequency, InDefaultGlideFactor, bInDefaultEnablement)
 	{
 
 	}
 
 	FTriangleOscilatorNode::FTriangleOscilatorNode(const FNodeInitData& InInitData)
-		: FTriangleOscilatorNode(InInitData.InstanceName, InInitData.InstanceID, 440.0f, true)
+		: FTriangleOscilatorNode(InInitData.InstanceName, InInitData.InstanceID, 440.0f, 0.0f, true)
 	{}
 
 	METASOUND_REGISTER_NODE(FTriangleOscilatorNode);
diff --git a/Engine/Plugins/Runtime/Metasound/Source/MetasoundStandardNodes/Public/MetasoundOscillatorNodes.h b/Engine/Plugins/Runtime/Metasound/Source/MetasoundStandardNodes/Public/MetasoundOscillatorNodes.h
--- a/Engine/Plugins/Runtime/Metasound/Source/MetasoundStandardNodes/Public/MetasoundOscillatorNodes.h
+++ b/Engine/Plugins/Runtime/Metasound/Source/MetasoundStandardNodes/Public/MetasoundOscillatorNodes.h
@@ -17,8 +17,7 @@ namespace Metasound
 	class METASOUNDSTANDARDNODES_API FOscilatorNodeBase : public FNode
 	{
 	protected:
-		FOscilatorNodeBase(const FString& InInstanceName, const FGuid& InInstanceID, const FNodeClassMetadata& InInfo, const TSharedRef<IOperatorFactory, ESPMode::ThreadSafe>& InFactory, float InDefaultFrequency, bool bInDefaultEnablement);
-		FOscilatorNodeBase(const FNodeInitData& InInitData);
+		FOscilatorNodeBase(const FString& InInstanceName, const FGuid& InInstanceID, const FNodeClassMetadata& InInfo, const TSharedRef<IOperatorFactory, ESPMode::ThreadSafe>& InFactory, float InDefaultFrequency, float InDefaultGlideTime, bool bInDefaultEnablement);
 
 		float GetDefaultPhaseOffset() const 
 		{
@@ -30,6 +29,11 @@ namespace Metasound
 			return DefaultFrequency;
 		}
 
+		float GetDefaultGlideFactor() const
+		{
+			return DefaultGlideFactor;
+		}
+
 		bool GetDefaultEnablement() const 
 		{
 			return bDefaultEnablement;
@@ -65,6 +69,7 @@ namespace Metasound
 
 		float DefaultPhaseOffset = 0.f;
 		float DefaultFrequency = 440.f;
+		float DefaultGlideFactor = 0.0f;
 		bool bDefaultEnablement = true;
 	};
 
@@ -72,7 +77,7 @@ namespace Metasound
 	{
 	public:
 		class FFactory;		
-		FSineOscilatorNode(const FString& InInstanceName, const FGuid& InInstanceID, float InDefaultFrequency, bool bInDefaultEnablement);
+		FSineOscilatorNode(const FString& InInstanceName, const FGuid& InInstanceID, float InDefaultFrequency, float InDefautlGlideTime, bool bInDefaultEnablement);
 		FSineOscilatorNode(const FNodeInitData& InInitData);
 	};
 		
@@ -80,7 +85,7 @@ namespace Metasound
 	{
 	public:
 		class FFactory;
-		FSawOscilatorNode(const FString& InInstanceName, const FGuid& InInstanceID, float InDefaultFrequency, bool bInDefaultEnablement);
+		FSawOscilatorNode(const FString& InInstanceName, const FGuid& InInstanceID, float InDefaultFrequency, float InDefautlGlideTime, bool bInDefaultEnablement);
 		FSawOscilatorNode(const FNodeInitData& InInitData);
 	};
 
@@ -88,7 +93,7 @@ namespace Metasound
 	{
 	public:
 		class FFactory;
-		FTriangleOscilatorNode(const FString& InInstanceName, const FGuid& InInstanceID, float InDefaultFrequency, bool bInDefaultEnablement);
+		FTriangleOscilatorNode(const FString& InInstanceName, const FGuid& InInstanceID, float InDefaultFrequency, float InDefautlGlideTime, bool bInDefaultEnablement);
 		FTriangleOscilatorNode(const FNodeInitData& InInitData);
 	};
 
@@ -96,7 +101,7 @@ namespace Metasound
 	{
 	public:
 		class FFactory;
-		FSquareOscilatorNode(const FString& InInstanceName, const FGuid& InInstanceID, float InDefaultFrequency, bool bInDefaultEnablement);
+		FSquareOscilatorNode(const FString& InInstanceName, const FGuid& InInstanceID, float InDefaultFrequency, float InDefautlGlideTime, bool bInDefaultEnablement);
 		FSquareOscilatorNode(const FNodeInitData& InInitData);
 	};
 
