Midi to freq node changes

- Using new param utils
- Fixing comparison on midi to freq
- Adding int input support

diff --git a/Engine/Plugins/Runtime/Metasound/Source/MetasoundStandardNodes/Private/MetasoundMIDIToFreqNode.cpp b/Engine/Plugins/Runtime/Metasound/Source/MetasoundStandardNodes/Private/MetasoundMIDIToFreqNode.cpp
--- a/Engine/Plugins/Runtime/Metasound/Source/MetasoundStandardNodes/Private/MetasoundMIDIToFreqNode.cpp
+++ b/Engine/Plugins/Runtime/Metasound/Source/MetasoundStandardNodes/Private/MetasoundMIDIToFreqNode.cpp
@@ -1,46 +1,63 @@
 // Copyright Epic Games, Inc. All Rights Reserved.
 
-#include "MetasoundMIDIToFreqNode.h"
-
 #include "Internationalization/Text.h"
+#include "MetasoundFacade.h"
 #include "MetasoundExecutableOperator.h"
 #include "MetasoundNodeRegistrationMacro.h"
 #include "MetasoundPrimitives.h"
 #include "MetasoundStandardNodesNames.h"
 #include "DSP/Dsp.h"
+#include "MetasoundParamHelper.h"
 
 #define LOCTEXT_NAMESPACE "MetasoundStandardNodes"
 
 namespace Metasound
 {
-	namespace MIDIToFrequencyVertexNames
+	namespace MidiToFrequencyVertexNames
 	{
-		const FString& GetInputMIDIName()
-		{
-			static FString Name = TEXT("MIDI");
-			return Name;
-		}
+		METASOUND_PARAM(InputMidi, "MIDI In", "A value representing a MIDI note value.");
+		METASOUND_PARAM(OutputFreq, "Out Frequency", "Output frequency value in hertz that corresponds to the input Midi note value.");
+	}
 
-		const FText& GetInputMIDIDescription()
+	namespace MidiToFrequencyPrivate
+	{
+		template<typename ValueType>
+		struct TMidiToFreqNodeSpecialization
 		{
-			static FText Desc = LOCTEXT("MIDIToFrequencyInputMIDIName", "A value representing a MIDI note value.");
-			return Desc;
-		}
+		};
+
 
-		const FString& GetOutputFrequencyName()
+		template<>
+		struct TMidiToFreqNodeSpecialization<int32>
 		{
-			static FString Name = TEXT("Frequency");
-			return Name;
-		}
+			static int32 GetFreqValue(int32 InMidi)
+			{
+				return Audio::GetFrequencyFromMidi(FMath::Clamp(InMidi, 0, 127));
+			}
+
+			static bool IsValueEqual(int32 InValueA, int32 InValueB)
+			{
+				return InValueA == InValueB; 
+			}
+		};
 
-		const FText& GetOutputFrequencyDescription()
+		template<>
+		struct TMidiToFreqNodeSpecialization<float>
 		{
-			static FText Desc = LOCTEXT("MIDITOFrequencyNodeOutputFrequencyName", "Output frequency value in hertz that corresponds to the input MIDI note value.");
-			return Desc;
-		}
+			static float GetFreqValue(float InMidi)
+			{
+				return Audio::GetFrequencyFromMidi(FMath::Clamp(InMidi, 0.0f, 127.0f));
+			}
+
+			static bool IsValueEqual(float InValueA, float InValueB)
+			{
+				return FMath::IsNearlyEqual(InValueA, InValueB);
+			}
+		};
 	}
 
-	class FMIDIToFreqOperator : public TExecutableOperator<FMIDIToFreqOperator>
+	template<typename ValueType>
+	class TMidiToFreqOperator : public TExecutableOperator<TMidiToFreqOperator<ValueType>>
 	{
 	public:
 
@@ -48,86 +65,99 @@ namespace Metasound
 		static const FVertexInterface& GetVertexInterface();
 		static TUniquePtr<IOperator> CreateOperator(const FCreateOperatorParams& InParams, FBuildErrorArray& OutErrors);
 
-		FMIDIToFreqOperator(const FOperatorSettings& InSettings, const FFloatReadRef& InMIDINote);
+		TMidiToFreqOperator(const FOperatorSettings& InSettings, const TDataReadReference<ValueType>& InMidiNote);
 
 		virtual FDataReferenceCollection GetInputs() const override;
 		virtual FDataReferenceCollection GetOutputs() const override;
 		void Execute();
 
 	private:
-		// The input midi value
-		FFloatReadRef MIDINote;
+		// The input Midi value
+		TDataReadReference<ValueType> MidiNote;
 
 		// The output frequency
 		FFloatWriteRef FreqOutput;
 
-		// Cached midi note value. Used to catch if the value changes to recompute freq output.
-		int32 PrevMidiNote;
+		// Cached Midi note value. Used to catch if the value changes to recompute freq output.
+		ValueType PrevMidiNote;
 	};
 
-	FMIDIToFreqOperator::FMIDIToFreqOperator(const FOperatorSettings& InSettings, const FFloatReadRef& InMIDINote)
-		: MIDINote(InMIDINote)
-		, FreqOutput(FFloatWriteRef::CreateNew(Audio::GetFrequencyFromMidi(*InMIDINote)))
-		, PrevMidiNote(*InMIDINote)
+	template<typename ValueType>
+	TMidiToFreqOperator<ValueType>::TMidiToFreqOperator(const FOperatorSettings& InSettings, const TDataReadReference<ValueType>& InMidiNote)
+		: MidiNote(InMidiNote)
+		, FreqOutput(FFloatWriteRef::CreateNew(Audio::GetFrequencyFromMidi(*InMidiNote)))
+		, PrevMidiNote(*InMidiNote)
 	{
 	}
 
-	FDataReferenceCollection FMIDIToFreqOperator::GetInputs() const
+	template<typename ValueType>
+	FDataReferenceCollection TMidiToFreqOperator<ValueType>::GetInputs() const
 	{
-		using namespace MIDIToFrequencyVertexNames;
+		using namespace MidiToFrequencyVertexNames;
 
 		FDataReferenceCollection InputDataReferences;
-		InputDataReferences.AddDataReadReference(GetInputMIDIName(), MIDINote);
+		InputDataReferences.AddDataReadReference(METASOUND_GET_PARAM_NAME(InputMidi), MidiNote);
 
 		return InputDataReferences;
 	}
 
-	FDataReferenceCollection FMIDIToFreqOperator::GetOutputs() const
+	template<typename ValueType>
+	FDataReferenceCollection TMidiToFreqOperator<ValueType>::GetOutputs() const
 	{
-		using namespace MIDIToFrequencyVertexNames;
+		using namespace MidiToFrequencyVertexNames;
 
 		FDataReferenceCollection OutputDataReferences;
-		OutputDataReferences.AddDataReadReference(GetOutputFrequencyName(), FreqOutput);
+		OutputDataReferences.AddDataReadReference(METASOUND_GET_PARAM_NAME(OutputFreq), FreqOutput);
 		return OutputDataReferences;
 	}
 
-	void FMIDIToFreqOperator::Execute()
+	template<typename ValueType>
+	void TMidiToFreqOperator<ValueType>::Execute()
 	{
-		// Only do anything if the midi note changes
-		if (*MIDINote != PrevMidiNote)
+		using namespace MidiToFrequencyPrivate;
+
+		// Only do anything if the Midi note changes
+		if (!TMidiToFreqNodeSpecialization<ValueType>::IsValueEqual(*MidiNote, PrevMidiNote))
 		{
-			PrevMidiNote = *MIDINote;
+			PrevMidiNote = *MidiNote;
 
-			*FreqOutput = Audio::GetFrequencyFromMidi(FMath::Clamp(*MIDINote, 0.0f, 127.0f));
+			*FreqOutput = TMidiToFreqNodeSpecialization<ValueType>::GetFreqValue(PrevMidiNote);
 		}
 	}
 
-	const FVertexInterface& FMIDIToFreqOperator::GetVertexInterface()
+	template<typename ValueType>
+	const FVertexInterface& TMidiToFreqOperator<ValueType>::GetVertexInterface()
 	{
-		using namespace MIDIToFrequencyVertexNames;
+		using namespace MidiToFrequencyVertexNames;
 
 		static const FVertexInterface Interface(
 			FInputVertexInterface(
-				TInputDataVertexModel<float>(GetInputMIDIName(), GetInputMIDIDescription(), 60)
+				TInputDataVertexModel<ValueType>(METASOUND_GET_PARAM_NAME_AND_TT(InputMidi), (ValueType)60.0f)
 			),
 			FOutputVertexInterface(
-				TOutputDataVertexModel<float>(GetOutputFrequencyName(), GetOutputFrequencyDescription())
+				TOutputDataVertexModel<float>(METASOUND_GET_PARAM_NAME_AND_TT(OutputFreq))
 			)
 		);
 
 		return Interface;
 	}
 
-	const FNodeClassMetadata& FMIDIToFreqOperator::GetNodeInfo()
+	template<typename ValueType>
+	const FNodeClassMetadata& TMidiToFreqOperator<ValueType>::GetNodeInfo()
 	{
 		auto InitNodeInfo = []() -> FNodeClassMetadata
 		{
+			FName DataTypeName = GetMetasoundDataTypeName<ValueType>();
+			FName OperatorName = TEXT("MIDI To Frequency");
+			FText NodeDisplayName = FText::Format(LOCTEXT("RandomGetArrayOpDisplayNamePattern", "MIDI To Frequency ({0})"), FText::FromString(GetMetasoundDataTypeString<ValueType>()));
+			FText NodeDescription = LOCTEXT("Metasound_MidiToFreqNodeDescription", "Converts a Midi note value to a frequency (hz) value.");
+
 			FNodeClassMetadata Info;
-			Info.ClassName = { Metasound::StandardNodes::Namespace, TEXT("MIDI To Frequency"), TEXT("") };
+			Info.ClassName = { Metasound::StandardNodes::Namespace, OperatorName, DataTypeName };
 			Info.MajorVersion = 1;
 			Info.MinorVersion = 0;
-			Info.DisplayName = LOCTEXT("Metasound_MidiToFreqDisplayName", "MIDI To Frequency");
-			Info.Description = LOCTEXT("Metasound_MidiToFreqNodeDescription", "Converts a MIDI note value to a frequency (hz) value.");
+			Info.DisplayName = NodeDisplayName;
+			Info.Description = NodeDescription;
 			Info.Author = PluginAuthor;
 			Info.PromptIfMissing = PluginNodeMissingPrompt;
 			Info.DefaultInterface = GetVertexInterface();
@@ -140,25 +170,38 @@ namespace Metasound
 		return Info;
 	}
 
-	TUniquePtr<IOperator> FMIDIToFreqOperator::CreateOperator(const FCreateOperatorParams& InParams, FBuildErrorArray& OutErrors)
+	template<typename ValueType>
+	TUniquePtr<IOperator> TMidiToFreqOperator<ValueType>::CreateOperator(const FCreateOperatorParams& InParams, FBuildErrorArray& OutErrors)
 	{
-		using namespace MIDIToFrequencyVertexNames;
+		using namespace MidiToFrequencyVertexNames;
 
-		const FMIDIToFreqNode& MIDIToFreqNode = static_cast<const FMIDIToFreqNode&>(InParams.Node);
 		const FDataReferenceCollection& InputCollection = InParams.InputDataReferences;
 		const FInputVertexInterface& InputInterface = GetVertexInterface().GetInputInterface();
 
-		FFloatReadRef InMIDINote = InputCollection.GetDataReadReferenceOrConstructWithVertexDefault<float>(InputInterface, GetInputMIDIName(), InParams.OperatorSettings);
+		TDataReadReference<ValueType> InMidiNote = InputCollection.GetDataReadReferenceOrConstructWithVertexDefault<ValueType>(InputInterface, METASOUND_GET_PARAM_NAME(InputMidi), InParams.OperatorSettings);
 
-		return MakeUnique<FMIDIToFreqOperator>(InParams.OperatorSettings, InMIDINote);
+		return MakeUnique<TMidiToFreqOperator>(InParams.OperatorSettings, InMidiNote);
 	}
 
-	FMIDIToFreqNode::FMIDIToFreqNode(const FNodeInitData& InitData)
-		: FNodeFacade(InitData.InstanceName, InitData.InstanceID, TFacadeOperatorClass<FMIDIToFreqOperator>())
+	template<typename ValueType>
+	class METASOUNDSTANDARDNODES_API TMidiToFreqNode : public FNodeFacade
 	{
-	}
+	public:
+		/**
+		 * Constructor used by the Metasound Frontend.
+		 */
+		TMidiToFreqNode(const FNodeInitData& InitData)
+			: FNodeFacade(InitData.InstanceName, InitData.InstanceID, TFacadeOperatorClass <TMidiToFreqOperator<ValueType>>())
+		{
+
+		}
+	};
+
+	using FMidiToFreqNodeInt32 = TMidiToFreqNode<int32>;
+	METASOUND_REGISTER_NODE(FMidiToFreqNodeInt32)
 
-	METASOUND_REGISTER_NODE(FMIDIToFreqNode)
+	using FMidiToFreqNodeFloat = TMidiToFreqNode<float>;
+	METASOUND_REGISTER_NODE(FMidiToFreqNodeFloat)
 }
 
 #undef LOCTEXT_NAMESPACE
diff --git a/Engine/Plugins/Runtime/Metasound/Source/MetasoundStandardNodes/Public/MetasoundMIDIToFreqNode.h b/Engine/Plugins/Runtime/Metasound/Source/MetasoundStandardNodes/Public/MetasoundMIDIToFreqNode.h
deleted file mode 100644
