Fixing trigger pipe

diff --git a/Engine/Plugins/Runtime/Metasound/Source/MetasoundStandardNodes/Private/MetasoundTriggerPipeNode.cpp b/Engine/Plugins/Runtime/Metasound/Source/MetasoundStandardNodes/Private/MetasoundTriggerPipeNode.cpp
--- a/Engine/Plugins/Runtime/Metasound/Source/MetasoundStandardNodes/Private/MetasoundTriggerPipeNode.cpp
+++ b/Engine/Plugins/Runtime/Metasound/Source/MetasoundStandardNodes/Private/MetasoundTriggerPipeNode.cpp
@@ -1,5 +1,4 @@
 // Copyright Epic Games, Inc. All Rights Reserved.
-#include "MetasoundTriggerPipeNode.h"
 
 #include "CoreMinimal.h"
 #include "MetasoundBuilderInterface.h"
@@ -30,8 +29,6 @@ namespace Metasound
 		METASOUND_PARAM(OutputOutTrigger, "Out", "The delayed output trigger(s).");
 	}
 
-	METASOUND_REGISTER_NODE(FTriggerPipeNode)
-
 	class FTriggerPipeOperator : public TExecutableOperator<FTriggerPipeOperator>
 	{
 	public:
@@ -100,7 +97,7 @@ namespace Metasound
 			},
 			[this](int32 StartFrame, int32 EndFrame)
 			{
-				SamplesUntilTrigger.AddUnique(DelayTime->GetSeconds() * SampleRate);
+				SamplesUntilTrigger.AddUnique(StartFrame + FMath::Max(DelayTime->GetSeconds(), 0.0) * SampleRate);
 			}
 		);
 
@@ -142,13 +139,12 @@ namespace Metasound
 	{
 		using namespace TriggerPipeVertexNames;
 
-		const FTriggerPipeNode& DelayNode = static_cast<const FTriggerPipeNode&>(InParams.Node);
 		const FDataReferenceCollection& InputCollection = InParams.InputDataReferences;
 		const FInputVertexInterface& InputInterface = GetVertexInterface().GetInputInterface();
 
-		FTimeReadRef DelayTime = InputCollection.GetDataReadReferenceOrConstructWithVertexDefault<FTime>(InputInterface, METASOUND_GET_PARAM_NAME(InputDelayTime), InParams.OperatorSettings);
-		FTriggerReadRef TriggerIn = InParams.InputDataReferences.GetDataReadReferenceOrConstruct<FTrigger>(METASOUND_GET_PARAM_NAME(InputDelayTime), InParams.OperatorSettings);
+		FTriggerReadRef TriggerIn = InParams.InputDataReferences.GetDataReadReferenceOrConstruct<FTrigger>(METASOUND_GET_PARAM_NAME(InputInTrigger), InParams.OperatorSettings);
 		FTriggerReadRef TriggerReset = InParams.InputDataReferences.GetDataReadReferenceOrConstruct<FTrigger>(METASOUND_GET_PARAM_NAME(InputReset), InParams.OperatorSettings);
+		FTimeReadRef DelayTime = InputCollection.GetDataReadReferenceOrConstructWithVertexDefault<FTime>(InputInterface, METASOUND_GET_PARAM_NAME(InputDelayTime), InParams.OperatorSettings);
 
 		return MakeUnique<FTriggerPipeOperator>(InParams.OperatorSettings, TriggerReset, TriggerIn, DelayTime);
 	}
@@ -159,9 +155,9 @@ namespace Metasound
 
 		static const FVertexInterface Interface(
 			FInputVertexInterface(
-				TInputDataVertexModel<FTrigger>(METASOUND_GET_PARAM_NAME_AND_TT(InputDelayTime)),
-				TInputDataVertexModel<FTrigger>(METASOUND_GET_PARAM_NAME_AND_TT(InputDelayTime)),
-				TInputDataVertexModel<FTime>(METASOUND_GET_PARAM_NAME_AND_TT(InputReset), 1.0f)
+				TInputDataVertexModel<FTrigger>(METASOUND_GET_PARAM_NAME_AND_TT(InputInTrigger)),
+				TInputDataVertexModel<FTrigger>(METASOUND_GET_PARAM_NAME_AND_TT(InputReset)),
+				TInputDataVertexModel<FTime>(METASOUND_GET_PARAM_NAME_AND_TT(InputDelayTime), 1.0f)
 				),
 			FOutputVertexInterface(
 				TOutputDataVertexModel<FTrigger>(METASOUND_GET_PARAM_NAME_AND_TT(OutputOutTrigger))
@@ -193,10 +189,18 @@ namespace Metasound
 		return Info;
 	}
 
-	FTriggerPipeNode::FTriggerPipeNode(const FNodeInitData& InInitData)
-		: FNodeFacade(InInitData.InstanceName, InInitData.InstanceID, TFacadeOperatorClass<FTriggerPipeOperator>())
+	class METASOUNDSTANDARDNODES_API FTriggerPipeNode : public FNodeFacade
 	{
-	}
+	public:
+		FTriggerPipeNode(const FNodeInitData& InInitData)
+			: FNodeFacade(InInitData.InstanceName, InInitData.InstanceID, TFacadeOperatorClass<FTriggerPipeOperator>())
+		{
+		}
+
+		virtual ~FTriggerPipeNode() = default;
+	};
+
+	METASOUND_REGISTER_NODE(FTriggerPipeNode)
 }
 
 #undef LOCTEXT_NAMESPACE // MetasoundTriggerPipeNode
diff --git a/Engine/Plugins/Runtime/Metasound/Source/MetasoundStandardNodes/Public/MetasoundTriggerPipeNode.h b/Engine/Plugins/Runtime/Metasound/Source/MetasoundStandardNodes/Public/MetasoundTriggerPipeNode.h
deleted file mode 100644
