AD Envelope

diff --git a/Engine/Plugins/Runtime/Metasound/Source/MetasoundStandardNodes/Private/MetasoundADEnvelopeNode.cpp b/Engine/Plugins/Runtime/Metasound/Source/MetasoundStandardNodes/Private/MetasoundADEnvelopeNode.cpp
new file mode 100644
--- /dev/null
+++ b/Engine/Plugins/Runtime/Metasound/Source/MetasoundStandardNodes/Private/MetasoundADEnvelopeNode.cpp
@@ -0,0 +1,420 @@
+// Copyright Epic Games, Inc. All Rights Reserved.
+
+#include "CoreMinimal.h"
+#include "Internationalization/Text.h"
+#include "MetasoundExecutableOperator.h"
+#include "MetasoundNodeRegistrationMacro.h"
+#include "MetasoundDataTypeRegistrationMacro.h"
+#include "MetasoundPrimitives.h"
+#include "MetasoundStandardNodesNames.h"
+#include "MetasoundTrigger.h"
+#include "MetasoundTime.h"
+#include "MetasoundAudioBuffer.h"
+#include "Internationalization/Text.h"
+#include "DSP/BufferVectorOperations.h"
+#include "DSP/Dsp.h"
+#include "MetasoundParamHelper.h"
+
+#define LOCTEXT_NAMESPACE "MetasoundStandardNodes"
+
+namespace Metasound
+{
+	namespace ADEnvelopeVertexNames
+	{
+		METASOUND_PARAM(InputTrigger, "Trigger", "Trigger to start the attack phase of the envelope generator.");
+		METASOUND_PARAM(InputAttackTime, "Attack Time", "Attack time of the envelope.");
+		METASOUND_PARAM(InputDecayTime, "Decay Time", "Decay time of the envelope.");
+		METASOUND_PARAM(InputAttackCurve, "Attack Curve", "The exponential curve factor of the attack. 1.0 = linear growth, < 1.0 logorithmic growth, > 1.0 exponential growth.");
+		METASOUND_PARAM(InputDecayCurve, "Decay Curve", "The exponential curve factor of the decay. 1.0 = linear decay, < 1.0 exponential decay, > 1.0 logorithmic decay.");
+		METASOUND_PARAM(InputLooping, "Looping", "Set to true to enable looping of the envelope. This will allow the envelope to be an LFO or wave generator.");
+
+		METASOUND_PARAM(OutputOnTrigger, "On Trigger", "Triggers when the envelope is triggered.");
+		METASOUND_PARAM(OutputOnDone, "On Done", "Triggers when the envelope finishes or loops back if looping is enabled.");
+		METASOUND_PARAM(OutpuTADEnvelopeValue, "Out Envelope", "The output value of the envelope.");
+	}
+
+	namespace ADEnvelopeNodePrivate
+	{
+		struct FEnvState
+		{
+			// Where the envelope is. If INDEX_NONE, then the envelope is not triggered
+			int32 CurrentSampleIndex = INDEX_NONE;
+
+			// Number of samples for attack
+			int32 AttackSampleCount = 0;
+
+			// Number of samples for Decay
+			int32 DecaySampleCount = 0;
+
+			// Curve factors for attack/Decay
+			float AttackCurveFactor = 0.0f;
+			float DecayCurveFactor = 0.0f;
+
+			Audio::FExponentialEase EnvEase;
+
+			// Where the envelope value was when it was triggered
+			float StartingEnvelopeValue = 0.0f;
+			float CurrenTADEnvelopeValue = 0.0f;
+
+			bool bLooping = false;
+		};
+
+		template<typename ValueType>
+		struct TADEnvelope
+		{
+		};
+
+		template<>
+		struct TADEnvelope<float>
+		{
+			static void GetNexTADEnvelopeOutput(FEnvState& InState, int32 StartFrame, int32 EndFrame, TArray<int32>& OuFinishedFrames, float& OutADEnvelopeValue)
+			{
+				// Don't need to do anything if we're not generating the envelope at the top of the block since this is a block-rate envelope
+				if (StartFrame > 0 || InState.CurrentSampleIndex == INDEX_NONE)
+				{
+					OutADEnvelopeValue = 0.0f;
+					return;
+				}
+
+				int32 TotalEnvSampleCount = (InState.AttackSampleCount + InState.DecaySampleCount);
+
+				// We are in attack
+				if (InState.CurrentSampleIndex < InState.AttackSampleCount)
+				{
+					float AttackFraction = (float)InState.CurrentSampleIndex++ / InState.AttackSampleCount;
+					OutADEnvelopeValue = InState.StartingEnvelopeValue + (1.0f - InState.StartingEnvelopeValue) * FMath::Pow(AttackFraction, InState.AttackCurveFactor);
+				}
+				// We are in Decay
+				else if (InState.CurrentSampleIndex < TotalEnvSampleCount)
+				{
+					int32 SampleCountInDecayState = InState.CurrentSampleIndex++ - InState.AttackSampleCount;
+					float DecayFracton = (float)SampleCountInDecayState / InState.DecaySampleCount;
+					OutADEnvelopeValue = 1.0f - FMath::Pow(DecayFracton, InState.DecayCurveFactor);
+				}
+				// We are looping so reset the sample index
+				else if (InState.bLooping)
+				{
+					InState.CurrentSampleIndex = 0;
+					OuFinishedFrames.Add(0);
+				}
+				else
+				{
+					// Envelope is done
+					InState.CurrentSampleIndex = INDEX_NONE;
+					OutADEnvelopeValue = 0.0f;
+					OuFinishedFrames.Add(0);
+				}
+			}
+
+			static bool IsAudio() { return false; }
+		};
+
+		template<>
+		struct TADEnvelope<FAudioBuffer>
+		{
+			static void GetNexTADEnvelopeOutput(FEnvState& InState, int32 StartFrame, int32 EndFrame, TArray<int32>& OuFinishedFrames, FAudioBuffer& OutADEnvelopeValue)
+			{
+				// If we are not active zero the buffer and early exit
+				if (InState.CurrentSampleIndex == INDEX_NONE)
+				{
+					OutADEnvelopeValue.Zero();
+					return;
+				}
+
+				float* OutEnvPtr = OutADEnvelopeValue.GetData();
+				int32 TotalEnvSampleCount = (InState.AttackSampleCount + InState.DecaySampleCount);
+				for (int32 i = StartFrame; i < EndFrame; ++i)
+				{
+					// We are in attack
+					if (InState.CurrentSampleIndex <= InState.AttackSampleCount)
+					{
+						float AttackFraction = (float)InState.CurrentSampleIndex++ / InState.AttackSampleCount;
+						float EnvValue = FMath::Pow(AttackFraction, InState.AttackCurveFactor);
+						float TargeTADEnvelopeValue = InState.StartingEnvelopeValue + (1.0f - InState.StartingEnvelopeValue) * EnvValue;
+						InState.EnvEase.SetValue(TargeTADEnvelopeValue);
+						InState.CurrenTADEnvelopeValue = InState.EnvEase.GetNextValue();
+						OutEnvPtr[i] = InState.CurrenTADEnvelopeValue;
+					}
+					// We are in Decay
+					else if (InState.CurrentSampleIndex < TotalEnvSampleCount)
+					{
+						int32 SampleCountInDecayState = InState.CurrentSampleIndex++ - InState.AttackSampleCount;
+						float DecayFracton = (float)SampleCountInDecayState / InState.DecaySampleCount;
+						float TargeTADEnvelopeValue = 1.0f - FMath::Pow(DecayFracton, InState.DecayCurveFactor);
+						InState.EnvEase.SetValue(TargeTADEnvelopeValue);
+						InState.CurrenTADEnvelopeValue = InState.EnvEase.GetNextValue();
+						OutEnvPtr[i] = InState.CurrenTADEnvelopeValue;
+					}
+					// We are looping so reset the sample index
+					else if (InState.bLooping)
+					{
+						InState.StartingEnvelopeValue = 0.0f;
+						InState.CurrenTADEnvelopeValue = 0.0f;
+						InState.CurrentSampleIndex = 0;
+						OuFinishedFrames.Add(i);
+					}
+					else
+					{
+						InState.CurrenTADEnvelopeValue = InState.EnvEase.GetNextValue();
+						OutEnvPtr[i] = InState.CurrenTADEnvelopeValue;
+
+						// Envelope is done
+						if (InState.EnvEase.IsDone())
+						{
+							// Zero out the rest of the envelope
+							int32 NumSamplesLeft = EndFrame - i - 1;
+							if (NumSamplesLeft > 0)
+							{
+								FMemory::Memzero(&OutEnvPtr[i + 1], sizeof(float) * NumSamplesLeft);
+							}
+							InState.CurrentSampleIndex = INDEX_NONE;
+							OuFinishedFrames.Add(i);
+							break;
+						}
+					}
+				}
+			}
+
+			static bool IsAudio() { return true; }
+		};
+	}
+
+	template<typename ValueType>
+	class TADEnvelopeNodeOperator : public TExecutableOperator<TADEnvelopeNodeOperator<ValueType>>
+	{
+	public:
+		static const FVertexInterface& GetDefaultInterface()
+		{
+			using namespace ADEnvelopeVertexNames;
+
+			static const FVertexInterface DefaultInterface(
+				FInputVertexInterface(
+					TInputDataVertexModel<FTrigger>(METASOUND_GET_PARAM_NAME_AND_TT(InputTrigger)),
+					TInputDataVertexModel<FTime>(METASOUND_GET_PARAM_NAME_AND_TT(InputAttackTime), 0.01f),
+					TInputDataVertexModel<FTime>(METASOUND_GET_PARAM_NAME_AND_TT(InputDecayTime), 1.0f),
+					TInputDataVertexModel<float>(METASOUND_GET_PARAM_NAME_AND_TT(InputAttackCurve), 1.0f),
+					TInputDataVertexModel<float>(METASOUND_GET_PARAM_NAME_AND_TT(InputDecayCurve), 1.0f),
+					TInputDataVertexModel<bool>(METASOUND_GET_PARAM_NAME_AND_TT(InputLooping), false)
+				),
+				FOutputVertexInterface(
+					TOutputDataVertexModel<FTrigger>(METASOUND_GET_PARAM_NAME_AND_TT(OutputOnTrigger)),
+					TOutputDataVertexModel<FTrigger>(METASOUND_GET_PARAM_NAME_AND_TT(OutputOnDone)),
+					TOutputDataVertexModel<ValueType>(METASOUND_GET_PARAM_NAME_AND_TT(OutpuTADEnvelopeValue))
+				)
+			);
+
+			return DefaultInterface;
+		}
+
+		static const FNodeClassMetadata& GetNodeInfo()
+		{
+			auto CreateNodeClassMetadata = []() -> FNodeClassMetadata
+			{
+				FName DataTypeName = GetMetasoundDataTypeName<ValueType>();
+				FName OperatorName = TEXT("AD Envelope");
+				FText NodeDisplayName = FText::Format(LOCTEXT("ADEnvelopeDisplayNamePattern", "AD Envelope ({0})"), FText::FromString(GetMetasoundDataTypeString<ValueType>()));
+				FText NodeDescription = LOCTEXT("ADEnevelopeDesc", "Generates an attack-decay envelope value output when triggered.");
+				FVertexInterface NodeInterface = GetDefaultInterface();
+
+				FNodeClassMetadata Metadata
+				{
+					FNodeClassName{FName("AD Envelope"), OperatorName, DataTypeName},
+					1, // Major Version
+					0, // Minor Version
+					NodeDisplayName,
+					NodeDescription,
+					PluginAuthor,
+					PluginNodeMissingPrompt,
+					NodeInterface,
+					{LOCTEXT("EnvelopeCat", "Envelopes")},
+					{TEXT("Envelope")},
+					FNodeDisplayStyle{}
+				};
+
+				return Metadata;
+			};
+
+			static const FNodeClassMetadata Metadata = CreateNodeClassMetadata();
+			return Metadata;
+		}
+
+		static TUniquePtr<IOperator> CreateOperator(const FCreateOperatorParams& InParams, TArray<TUniquePtr<IOperatorBuildError>>& OutErrors)
+		{
+			using namespace ADEnvelopeVertexNames;
+
+			const FInputVertexInterface& InputInterface = GetDefaultInterface().GetInputInterface();
+
+			FTriggerReadRef TriggerIn = InParams.InputDataReferences.GetDataReadReferenceOrConstruct<FTrigger>(METASOUND_GET_PARAM_NAME(InputTrigger), InParams.OperatorSettings);
+			FTimeReadRef AttackTime = InParams.InputDataReferences.GetDataReadReferenceOrConstructWithVertexDefault<FTime>(InputInterface, METASOUND_GET_PARAM_NAME(InputAttackTime), InParams.OperatorSettings);
+			FTimeReadRef DecayTime = InParams.InputDataReferences.GetDataReadReferenceOrConstructWithVertexDefault<FTime>(InputInterface, METASOUND_GET_PARAM_NAME(InputDecayTime), InParams.OperatorSettings);
+			FFloatReadRef AttackCurveFactor = InParams.InputDataReferences.GetDataReadReferenceOrConstructWithVertexDefault<float>(InputInterface, METASOUND_GET_PARAM_NAME(InputAttackCurve), InParams.OperatorSettings);
+			FFloatReadRef DecayCurveFactor = InParams.InputDataReferences.GetDataReadReferenceOrConstructWithVertexDefault<float>(InputInterface, METASOUND_GET_PARAM_NAME(InputDecayCurve), InParams.OperatorSettings);
+			FBoolReadRef bLooping = InParams.InputDataReferences.GetDataReadReferenceOrConstructWithVertexDefault<bool>(InputInterface, METASOUND_GET_PARAM_NAME(InputLooping), InParams.OperatorSettings);
+
+			return MakeUnique<TADEnvelopeNodeOperator<ValueType>>(InParams.OperatorSettings, TriggerIn, AttackTime, DecayTime, AttackCurveFactor, DecayCurveFactor, bLooping);
+		}
+
+		TADEnvelopeNodeOperator(const FOperatorSettings& InSettings,
+			const FTriggerReadRef& InTriggerIn,
+			const FTimeReadRef& InAttackTime,
+			const FTimeReadRef& InDecayTime,
+			const FFloatReadRef& InAttackCurveFactor,
+			const FFloatReadRef& InDecayeCurveFactor,
+			const FBoolReadRef& bInLooping)
+			: TriggerIn(InTriggerIn)
+			, AttackTime(InAttackTime)
+			, DecayTime(InDecayTime)
+			, AttackCurveFactor(InAttackCurveFactor)
+			, DecayCurveFactor(InDecayeCurveFactor)
+			, bLooping(bInLooping)
+			, OnEnvelopeGen(TDataWriteReferenceFactory<FTrigger>::CreateAny(InSettings))
+			, OnDone(TDataWriteReferenceFactory<FTrigger>::CreateAny(InSettings))
+			, OutpuTADEnvelope(TDataWriteReferenceFactory<ValueType>::CreateAny(InSettings))
+		{
+			NumFramesPerBlock = InSettings.GetNumFramesPerBlock();
+
+			EnvState.EnvEase.SetEaseFactor(0.01f);
+
+			if (ADEnvelopeNodePrivate::TADEnvelope<ValueType>::IsAudio())
+			{
+				SampleRate = InSettings.GetSampleRate();
+			}
+			else
+			{
+				SampleRate = InSettings.GetActualBlockRate();
+			}
+		}
+
+		virtual ~TADEnvelopeNodeOperator() = default;
+
+		virtual FDataReferenceCollection GetInputs() const override
+		{
+			using namespace ADEnvelopeVertexNames;
+
+			FDataReferenceCollection Inputs;
+			Inputs.AddDataReadReference(METASOUND_GET_PARAM_NAME(InputTrigger), TriggerIn);
+			Inputs.AddDataReadReference(METASOUND_GET_PARAM_NAME(InputAttackTime), AttackTime);
+			Inputs.AddDataReadReference(METASOUND_GET_PARAM_NAME(InputDecayTime), DecayTime);
+			Inputs.AddDataReadReference(METASOUND_GET_PARAM_NAME(InputAttackCurve), AttackCurveFactor);
+			Inputs.AddDataReadReference(METASOUND_GET_PARAM_NAME(InputDecayCurve), DecayCurveFactor);
+			Inputs.AddDataReadReference(METASOUND_GET_PARAM_NAME(InputLooping), bLooping);
+
+			return Inputs;
+		}
+
+		virtual FDataReferenceCollection GetOutputs() const override
+		{
+			using namespace ADEnvelopeVertexNames;
+
+			FDataReferenceCollection Outputs;
+			Outputs.AddDataReadReference(METASOUND_GET_PARAM_NAME(OutputOnTrigger), OnEnvelopeGen);
+			Outputs.AddDataReadReference(METASOUND_GET_PARAM_NAME(OutputOnDone), OnDone);
+			Outputs.AddDataReadReference(METASOUND_GET_PARAM_NAME(OutpuTADEnvelopeValue), OutpuTADEnvelope);
+
+			return Outputs;
+		}
+
+		void UpdateParams()
+		{
+			float AttackTimeSeconds = AttackTime->GetSeconds();
+			float DecayTimeSeconds = DecayTime->GetSeconds();
+			EnvState.AttackSampleCount = SampleRate * FMath::Max(0.0f, AttackTimeSeconds);
+			EnvState.DecaySampleCount = SampleRate * FMath::Max(0.0f, DecayTimeSeconds);
+			EnvState.AttackCurveFactor = FMath::Max(KINDA_SMALL_NUMBER, *AttackCurveFactor);
+			EnvState.DecayCurveFactor = FMath::Max(KINDA_SMALL_NUMBER, *DecayCurveFactor);
+			EnvState.bLooping = *bLooping;
+		}
+
+
+		void Execute()
+		{
+			using namespace ADEnvelopeNodePrivate;
+
+			OnEnvelopeGen->AdvanceBlock();
+			OnDone->AdvanceBlock();
+
+			// check for any updates to input params
+			UpdateParams();
+
+			TriggerIn->ExecuteBlock(
+				// OnPreTrigger
+				[&](int32 StartFrame, int32 EndFrame)
+				{
+					TArray<int32> FinishedFrames;
+					TADEnvelope<ValueType>::GetNexTADEnvelopeOutput(EnvState, StartFrame, EndFrame, FinishedFrames, *OutpuTADEnvelope);
+
+					for (int32 FrameFinished : FinishedFrames)
+					{
+						OnDone->TriggerFrame(FrameFinished);
+					}
+				},
+				// OnTrigger
+					[&](int32 StartFrame, int32 EndFrame)
+				{
+					// Get latest params
+					UpdateParams();
+
+					// Set the sample index to the top of the envelope
+					EnvState.CurrentSampleIndex = 0;
+					EnvState.StartingEnvelopeValue = EnvState.CurrenTADEnvelopeValue;
+
+					// Generate the output (this will no-op if we're block rate)
+					TArray<int32> FinishedFrames;
+					TADEnvelope<ValueType>::GetNexTADEnvelopeOutput(EnvState, StartFrame, EndFrame, FinishedFrames, *OutpuTADEnvelope);
+					for (int32 FrameFinished : FinishedFrames)
+					{
+						OnDone->TriggerFrame(FrameFinished);
+					}
+
+					// Forward the trigger
+					OnEnvelopeGen->TriggerFrame(StartFrame);
+				}
+			);
+		}
+
+	private:
+
+		FTriggerReadRef TriggerIn;
+		FTimeReadRef AttackTime;
+		FTimeReadRef DecayTime;
+		FFloatReadRef AttackCurveFactor;
+		FFloatReadRef DecayCurveFactor;
+		FBoolReadRef bLooping;
+
+		FTriggerWriteRef OnEnvelopeGen;
+		FTriggerWriteRef OnDone;
+		TDataWriteReference<ValueType> OutpuTADEnvelope;
+
+		// This will either be the block rate or sample rate depending on if this is block-rate or audio-rate envelope
+		float SampleRate = 0.0f;
+		int32 NumFramesPerBlock = 0;
+
+		ADEnvelopeNodePrivate::FEnvState EnvState;
+	};
+
+	/** TADEnvelopeNode
+	 *
+	 *  Creates an Attack/Decay envelope node.
+	 */
+	template<typename ValueType>
+	class METASOUNDSTANDARDNODES_API TADEnvelopeNode : public FNodeFacade
+	{
+	public:
+		/**
+		 * Constructor used by the Metasound Frontend.
+		 */
+		TADEnvelopeNode(const FNodeInitData& InInitData)
+			: FNodeFacade(InInitData.InstanceName, InInitData.InstanceID, TFacadeOperatorClass<TADEnvelopeNodeOperator<ValueType>>())
+		{}
+
+		virtual ~TADEnvelopeNode() = default;
+	};
+
+	using FADEnvelopeNodeFloat = TADEnvelopeNode<float>;
+	METASOUND_REGISTER_NODE(FADEnvelopeNodeFloat)
+
+	using FEnvelopeAudioBuffer = TADEnvelopeNode<FAudioBuffer>;
+	METASOUND_REGISTER_NODE(FEnvelopeAudioBuffer)
+}
+
+#undef LOCTEXT_NAMESPACE
