Add "Promote To Input"
- Various tidy up micro-tasks

diff --git a/Engine/Plugins/Runtime/Metasound/Source/MetasoundEditor/Private/MetasoundEditorGraph.cpp b/Engine/Plugins/Runtime/Metasound/Source/MetasoundEditor/Private/MetasoundEditorGraph.cpp
--- a/Engine/Plugins/Runtime/Metasound/Source/MetasoundEditor/Private/MetasoundEditorGraph.cpp
+++ b/Engine/Plugins/Runtime/Metasound/Source/MetasoundEditor/Private/MetasoundEditorGraph.cpp
@@ -327,7 +327,9 @@ UMetasoundEditorGraphInput* UMetasoundEditorGraph::FindOrAddInput(Metasound::Fro
 		EMetasoundFrontendLiteralType LiteralType = DefaultLiteral.GetType();
 		IMetasoundEditorModule& EditorModule = FModuleManager::GetModuleChecked<IMetasoundEditorModule>("MetasoundEditor");
 		TSubclassOf<UMetasoundEditorGraphInputLiteral> InputLiteralClass = EditorModule.FindInputLiteralClass(LiteralType);
+
 		NewInput->Literal = NewObject<UMetasoundEditorGraphInputLiteral>(this, InputLiteralClass);
+		NewInput->Literal->SetFromLiteral(DefaultLiteral);
 
 		Inputs.Add(NewInput);
 		return NewInput;
diff --git a/Engine/Plugins/Runtime/Metasound/Source/MetasoundEditor/Private/MetasoundEditorGraphBuilder.cpp b/Engine/Plugins/Runtime/Metasound/Source/MetasoundEditor/Private/MetasoundEditorGraphBuilder.cpp
--- a/Engine/Plugins/Runtime/Metasound/Source/MetasoundEditor/Private/MetasoundEditorGraphBuilder.cpp
+++ b/Engine/Plugins/Runtime/Metasound/Source/MetasoundEditor/Private/MetasoundEditorGraphBuilder.cpp
@@ -156,37 +156,41 @@ namespace Metasound
 			return Categories;
 		}
 
-		FString FGraphBuilder::GenerateUniqueInputName(const UObject& InMetasound)
+		FString FGraphBuilder::GenerateUniqueInputName(const UObject& InMetasound, const FString* InBaseName)
 		{
-			FString NameBase = TEXT("Input");
+			static const FString DefaultBaseName = TEXT("Input");
+			const FString* NameBase = InBaseName ? InBaseName : &DefaultBaseName;
+
 			const FMetasoundAssetBase* MetasoundAsset = IMetasoundUObjectRegistry::Get().GetObjectAsAssetBase(&InMetasound);
 			check(MetasoundAsset);
 
 			const Frontend::FConstGraphHandle GraphHandle = MetasoundAsset->GetRootGraphHandle();
 
 			int32 i = 1;
-			FString NewInputName = NameBase + FString::Printf(TEXT("_%02d"), i);
+			FString NewInputName = *NameBase + FString::Printf(TEXT(" %02d"), i);
 			while (GraphHandle->ContainsInputVertexWithName(NewInputName))
 			{
-				NewInputName = NameBase + FString::Printf(TEXT("_%02d"), ++i);
+				NewInputName = *NameBase + FString::Printf(TEXT(" %02d"), ++i);
 			}
 
 			return NewInputName;
 		}
 
-		FString FGraphBuilder::GenerateUniqueOutputName(const UObject& InMetasound)
+		FString FGraphBuilder::GenerateUniqueOutputName(const UObject& InMetasound, const FString* InBaseName)
 		{
-			FString NameBase = TEXT("Output");
+			static const FString DefaultBaseName = TEXT("Input");
+			const FString* NameBase = InBaseName ? InBaseName : &DefaultBaseName;
+
 			const FMetasoundAssetBase* MetasoundAsset = IMetasoundUObjectRegistry::Get().GetObjectAsAssetBase(&InMetasound);
 			check(MetasoundAsset);
 
 			const Frontend::FConstGraphHandle GraphHandle = MetasoundAsset->GetRootGraphHandle();
 
 			int32 i = 1;
-			FString NewOutputName = NameBase + FString::Printf(TEXT("_%02d"), i);
+			FString NewOutputName = *NameBase + FString::Printf(TEXT(" %02d"), i);
 			while (GraphHandle->ContainsOutputVertexWithName(NewOutputName))
 			{
-				NewOutputName = NameBase + FString::Printf(TEXT("_%02d"), ++i);
+				NewOutputName = *NameBase + FString::Printf(TEXT(" %02d"), ++i);
 			}
 
 			return NewOutputName;
@@ -287,16 +291,14 @@ namespace Metasound
 		{
 			using namespace Frontend;
 
-			const FString InInputName = InInputPin.GetName();
-			TArray<FInputHandle> InputHandles = InNodeHandle->GetInputsWithVertexName(InInputName);
-			if (!ensure(InputHandles.Num() == 1))
+			FMetasoundFrontendLiteral PinDataTypeDefaultLiteral;
+			if (!GetPinDefaultLiteral(InInputPin, PinDataTypeDefaultLiteral))
 			{
 				return;
 			}
 
-			FInputHandle InputHandle = InputHandles[0];
-			FMetasoundFrontendLiteral PinDataTypeDefaultLiteral;
-			if (!GetPinDefaultLiteral(InInputPin, InputHandle, PinDataTypeDefaultLiteral))
+			FInputHandle InputHandle = GetInputHandleFromPin(&InInputPin);
+			if (!ensure(InputHandle->IsValid()))
 			{
 				return;
 			}
@@ -343,12 +345,18 @@ namespace Metasound
 			}
 		}
 
-		bool FGraphBuilder::GetPinDefaultLiteral(UEdGraphPin& InInputPin, Frontend::FInputHandle InputHandle, FMetasoundFrontendLiteral& OutDefaultLiteral)
+		bool FGraphBuilder::GetPinDefaultLiteral(UEdGraphPin& InInputPin, FMetasoundFrontendLiteral& OutDefaultLiteral)
 		{
 			using namespace Frontend;
 
 			IMetasoundEditorModule& EditorModule = FModuleManager::GetModuleChecked<IMetasoundEditorModule>("MetasoundEditor");
 
+			FInputHandle InputHandle = GetInputHandleFromPin(&InInputPin);
+			if (!ensure(InputHandle->IsValid()))
+			{
+				return false;
+			}
+
 			const FString& InStringValue = InInputPin.DefaultValue;
 			const FName TypeName = InputHandle->GetDataType();
 			const FEditorDataType DataType = EditorModule.FindDataType(TypeName);
diff --git a/Engine/Plugins/Runtime/Metasound/Source/MetasoundEditor/Private/MetasoundEditorGraphBuilder.h b/Engine/Plugins/Runtime/Metasound/Source/MetasoundEditor/Private/MetasoundEditorGraphBuilder.h
--- a/Engine/Plugins/Runtime/Metasound/Source/MetasoundEditor/Private/MetasoundEditorGraphBuilder.h
+++ b/Engine/Plugins/Runtime/Metasound/Source/MetasoundEditor/Private/MetasoundEditorGraphBuilder.h
@@ -97,9 +97,9 @@ namespace Metasound
 			// applicable post disconnection.
 			static void DisconnectPin(UEdGraphPin& InPin, bool bAddLiteralInputs = true);
 
-			static FString GenerateUniqueInputName(const UObject& InMetasound);
+			static FString GenerateUniqueInputName(const UObject& InMetasound, const FString* InBaseName = nullptr);
 
-			static FString GenerateUniqueOutputName(const UObject& InMetasound);
+			static FString GenerateUniqueOutputName(const UObject& InMetasound, const FString* InBaseName = nullptr);
 
 			static TArray<FString> GetDataTypeNameCategories(const FName& InDataTypeName);
 
@@ -118,7 +118,8 @@ namespace Metasound
 			// 2. the pin's DataType literal default.
 			static void AddOrUpdateLiteralInput(UObject& InMetasound, Frontend::FNodeHandle InNodeHandle, UEdGraphPin& InInputPin, bool bForcePinValueAsDefault = false);
 
-			static bool GetPinDefaultLiteral(UEdGraphPin& InInputPin, Frontend::FInputHandle InputHandle, FMetasoundFrontendLiteral& OutLiteralDefault);
+			// Returns the default literal stored on the respective Frontend Node's Input.
+			static bool GetPinDefaultLiteral(UEdGraphPin& InInputPin, FMetasoundFrontendLiteral& OutLiteralDefault);
 
 			// Deletes Editor Graph Variable's associated Frontend node, as well as any
 			// Editor Graph nodes referencing the given variable.
diff --git a/Engine/Plugins/Runtime/Metasound/Source/MetasoundEditor/Private/MetasoundEditorGraphSchema.cpp b/Engine/Plugins/Runtime/Metasound/Source/MetasoundEditor/Private/MetasoundEditorGraphSchema.cpp
--- a/Engine/Plugins/Runtime/Metasound/Source/MetasoundEditor/Private/MetasoundEditorGraphSchema.cpp
+++ b/Engine/Plugins/Runtime/Metasound/Source/MetasoundEditor/Private/MetasoundEditorGraphSchema.cpp
@@ -35,144 +35,196 @@ namespace Metasound
 {
 	namespace Editor
 	{
-		bool TryConnectNewNodeToPin(UEdGraphNode& NewGraphNode, UEdGraphPin* FromPin)
+		namespace SchemaPrivate
 		{
-			using namespace Metasound::Frontend;
+			static const FText InputDisplayNameFormat = LOCTEXT("DisplayNameAddInputFormat", "Get {0}");
+			static const FText InputTooltipFormat = LOCTEXT("TooltipAddInputFormat", "Adds a getter for the input '{0}' to the graph.");
 
-			if (!FromPin)
+			static const FText OutputDisplayNameFormat = LOCTEXT("DisplayNameAddOutputFormat", "Set {0}");
+			static const FText OutputTooltipFormat = LOCTEXT("TooltipAddOutputFormat", "Adds a setter for the output '{0}' to the graph.");
+
+			enum class EPrimaryContextGroup
 			{
-				return false;
+				Default = 0,
+				Inputs,
+				Outputs,
+				External
+			};
+
+			const FText& GetContextGroupDisplayName(EPrimaryContextGroup InContextGroup)
+			{
+				switch (InContextGroup)
+				{
+					case EPrimaryContextGroup::Inputs:
+					{
+						static const FText InputGroupName = LOCTEXT("InputActions", "Input Actions");
+						return InputGroupName;
+					}
+					case EPrimaryContextGroup::Outputs:
+					{
+						static const FText OutputGroupName = LOCTEXT("OutputActions", "Output Actions");
+						return OutputGroupName;
+					}
+
+					// External nodes & any other group other than those above should not use a primary context group display name
+					case EPrimaryContextGroup::Default:
+					case EPrimaryContextGroup::External:
+					default:
+					{
+						checkNoEntry();
+						return FText::GetEmpty();
+					}
+				}
 			}
 
-			if (FromPin->Direction == EGPD_Input)
+			bool TryConnectNewNodeToPin(UEdGraphNode& NewGraphNode, UEdGraphPin* FromPin)
 			{
-				FInputHandle InputHandle = FGraphBuilder::GetInputHandleFromPin(FromPin);
-				for (UEdGraphPin* Pin : NewGraphNode.Pins)
+				using namespace Metasound::Frontend;
+
+				if (!FromPin)
 				{
-					if (Pin->Direction == EGPD_Output)
+					return false;
+				}
+
+				if (FromPin->Direction == EGPD_Input)
+				{
+					FInputHandle InputHandle = FGraphBuilder::GetInputHandleFromPin(FromPin);
+					for (UEdGraphPin* Pin : NewGraphNode.Pins)
 					{
-						FOutputHandle OutputHandle = FGraphBuilder::GetOutputHandleFromPin(Pin);
-						if (OutputHandle->IsValid() && OutputHandle->GetDataType() == InputHandle->GetDataType())
+						if (Pin->Direction == EGPD_Output)
 						{
-							if (ensure(FGraphBuilder::ConnectNodes(*FromPin, *Pin, true /* bConnectEdPins */)))
+							FOutputHandle OutputHandle = FGraphBuilder::GetOutputHandleFromPin(Pin);
+							if (OutputHandle->IsValid() && OutputHandle->GetDataType() == InputHandle->GetDataType())
 							{
-								NewGraphNode.MarkPackageDirty();
-								return true;
+								if (ensure(FGraphBuilder::ConnectNodes(*FromPin, *Pin, true /* bConnectEdPins */)))
+								{
+									return true;
+								}
 							}
 						}
 					}
 				}
-			}
 
-			if (FromPin->Direction == EGPD_Output)
-			{
-				FOutputHandle OutputHandle = FGraphBuilder::GetOutputHandleFromPin(FromPin);
-				for (UEdGraphPin* Pin : NewGraphNode.Pins)
+				if (FromPin->Direction == EGPD_Output)
 				{
-					if (Pin->Direction == EGPD_Input)
+					FOutputHandle OutputHandle = FGraphBuilder::GetOutputHandleFromPin(FromPin);
+					for (UEdGraphPin* Pin : NewGraphNode.Pins)
 					{
-						FInputHandle InputHandle = FGraphBuilder::GetInputHandleFromPin(Pin);
-						if (InputHandle->IsValid() && InputHandle->GetDataType() == OutputHandle->GetDataType())
+						if (Pin->Direction == EGPD_Input)
 						{
-							if (ensure(FGraphBuilder::ConnectNodes(*Pin, *FromPin, true /* bConnectEdPins */)))
+							FInputHandle InputHandle = FGraphBuilder::GetInputHandleFromPin(Pin);
+							if (InputHandle->IsValid() && InputHandle->GetDataType() == OutputHandle->GetDataType())
 							{
-								NewGraphNode.MarkPackageDirty();
-								return true;
+								if (ensure(FGraphBuilder::ConnectNodes(*Pin, *FromPin, true /* bConnectEdPins */)))
+								{
+									return true;
+								}
 							}
 						}
 					}
 				}
-			}
-
-			return false;
-		}
-
-		TSharedPtr<FEditor> GetEditorForGraph(const UObject& Metasound)
-		{
-			TSharedPtr<IToolkit> FoundAssetEditor = FToolkitManager::Get().FindEditorForAsset(CastChecked<const UObject>(&Metasound));
-			return StaticCastSharedPtr<FEditor, IToolkit>(FoundAssetEditor);
-		}
-
-		TSharedPtr<FEditor> GetEditorForGraph(const UEdGraph& EdGraph)
-		{
-			const UMetasoundEditorGraph* MetasoundGraph = CastChecked<const UMetasoundEditorGraph>(&EdGraph);
-			return GetEditorForGraph(MetasoundGraph->GetMetasoundChecked());
-		}
-
-		FLinearColor GetPinCategoryColor(const FEdGraphPinType& PinType)
-		{
-			const UMetasoundEditorSettings* Settings = GetDefault<UMetasoundEditorSettings>();
-			check(Settings);
 
-			if (PinType.PinCategory == FGraphBuilder::PinCategoryAudio)
-			{
-				return Settings->AudioPinTypeColor;
+				return false;
 			}
 
-			if (PinType.PinCategory == FGraphBuilder::PinCategoryTrigger)
+			TSharedPtr<FEditor> GetEditorForGraph(const UObject& Metasound)
 			{
-				return Settings->TriggerPinTypeColor;
+				TSharedPtr<IToolkit> FoundAssetEditor = FToolkitManager::Get().FindEditorForAsset(CastChecked<const UObject>(&Metasound));
+				return StaticCastSharedPtr<FEditor, IToolkit>(FoundAssetEditor);
 			}
 
-			if (PinType.PinCategory == FGraphBuilder::PinCategoryBoolean)
+			TSharedPtr<FEditor> GetEditorForGraph(const UEdGraph& EdGraph)
 			{
-				return Settings->BooleanPinTypeColor;
+				const UMetasoundEditorGraph* MetasoundGraph = CastChecked<const UMetasoundEditorGraph>(&EdGraph);
+				return GetEditorForGraph(MetasoundGraph->GetMetasoundChecked());
 			}
 
-			if (PinType.PinCategory == FGraphBuilder::PinCategoryFloat)
+			FLinearColor GetPinCategoryColor(const FEdGraphPinType& PinType)
 			{
-				if (PinType.PinSubCategory == FGraphBuilder::PinSubCategoryTime)
+				const UMetasoundEditorSettings* Settings = GetDefault<UMetasoundEditorSettings>();
+				check(Settings);
+
+				if (PinType.PinCategory == FGraphBuilder::PinCategoryAudio)
 				{
-					return Settings->TimePinTypeColor;
+					return Settings->AudioPinTypeColor;
 				}
-				return Settings->FloatPinTypeColor;
-			}
 
-			if (PinType.PinCategory == FGraphBuilder::PinCategoryInt32)
-			{
-				return Settings->IntPinTypeColor;
-			}
+				if (PinType.PinCategory == FGraphBuilder::PinCategoryTrigger)
+				{
+					return Settings->TriggerPinTypeColor;
+				}
 
-			if (PinType.PinCategory == FGraphBuilder::PinCategoryInt64)
-			{
-				return Settings->Int64PinTypeColor;
-			}
+				if (PinType.PinCategory == FGraphBuilder::PinCategoryBoolean)
+				{
+					return Settings->BooleanPinTypeColor;
+				}
 
-			if (PinType.PinCategory == FGraphBuilder::PinCategoryString)
-			{
-				return Settings->StringPinTypeColor;
-			}
+				if (PinType.PinCategory == FGraphBuilder::PinCategoryFloat)
+				{
+					if (PinType.PinSubCategory == FGraphBuilder::PinSubCategoryTime)
+					{
+						return Settings->TimePinTypeColor;
+					}
+					return Settings->FloatPinTypeColor;
+				}
 
-			if (PinType.PinCategory == FGraphBuilder::PinCategoryDouble)
-			{
-				return Settings->DoublePinTypeColor;
-			}
+				if (PinType.PinCategory == FGraphBuilder::PinCategoryInt32)
+				{
+					return Settings->IntPinTypeColor;
+				}
 
-			if (PinType.PinCategory == FGraphBuilder::PinCategoryObject)
-			{
-				return Settings->ObjectPinTypeColor;
-			}
+				if (PinType.PinCategory == FGraphBuilder::PinCategoryInt64)
+				{
+					return Settings->Int64PinTypeColor;
+				}
 
-			return Settings->DefaultPinTypeColor;
-		}
+				if (PinType.PinCategory == FGraphBuilder::PinCategoryString)
+				{
+					return Settings->StringPinTypeColor;
+				}
 
-		template <typename TAction>
-		void GetDataTypeActions(FGraphActionMenuBuilder& ActionMenuBuilder, const TArray<Frontend::FConstNodeHandle>& InNodeHandles, FInterfaceNodeFilterFunction InFilter, const FText& InCategory, const FText& InTooltipFormat, bool bShowSelectedActions)
-		{
-			using namespace Editor;
-			using namespace Frontend;
+				if (PinType.PinCategory == FGraphBuilder::PinCategoryDouble)
+				{
+					return Settings->DoublePinTypeColor;
+				}
 
-			for (FConstNodeHandle NodeHandle : InNodeHandles)
-			{
-				if (InFilter && !InFilter(NodeHandle))
+				if (PinType.PinCategory == FGraphBuilder::PinCategoryObject)
 				{
-					continue;
+					return Settings->ObjectPinTypeColor;
 				}
 
-				const FText NodeDisplayName = NodeHandle->GetDisplayName();
-				const FText Tooltip = FText::Format(InTooltipFormat, NodeDisplayName);
-				TSharedPtr<TAction> NewNodeAction = MakeShared<TAction>(InCategory, NodeDisplayName, NodeHandle->GetID(), Tooltip, 0);
-				ActionMenuBuilder.AddAction(NewNodeAction);
+				return Settings->DefaultPinTypeColor;
+			}
+
+			struct FDataTypeActionQuery
+			{
+				FGraphActionMenuBuilder& ActionMenuBuilder;
+				const TArray<Frontend::FConstNodeHandle>& NodeHandles;
+				FInterfaceNodeFilterFunction Filter;
+				EPrimaryContextGroup ContextGroup;
+				const FText& DisplayNameFormat;
+				const FText& TooltipFormat;
+				bool bShowSelectedActions = false;
+			};
+
+			template <typename TAction>
+			void GetDataTypeActions(const FDataTypeActionQuery& InQuery)
+			{
+				using namespace Editor;
+				using namespace Frontend;
+
+				for (const FConstNodeHandle& NodeHandle : InQuery.NodeHandles)
+				{
+					if (!InQuery.Filter || InQuery.Filter(NodeHandle))
+					{
+						const FText& GroupName = GetContextGroupDisplayName(InQuery.ContextGroup);
+						const FText NodeDisplayName = NodeHandle->GetDisplayName();
+						const FText Tooltip = FText::Format(InQuery.TooltipFormat, NodeDisplayName);
+						const FText DisplayName = FText::Format(InQuery.DisplayNameFormat, NodeDisplayName);
+						TSharedPtr<TAction> NewNodeAction = MakeShared<TAction>(GroupName, DisplayName, NodeHandle->GetID(), Tooltip, static_cast<int32>(InQuery.ContextGroup));
+						InQuery.ActionMenuBuilder.AddAction(NewNodeAction);
+					}
+				}
 			}
 		}
 
@@ -212,7 +264,7 @@ namespace Metasound
 			OutParams.AssociatedPin1 = InputPin;
 			OutParams.AssociatedPin2 = OutputPin;
 
-			OutParams.WireColor = GetPinCategoryColor(OutputPin->PinType);
+			OutParams.WireColor = SchemaPrivate::GetPinCategoryColor(OutputPin->PinType);
 			bool bExecuted = false;
 
 			// Run through the predecessors, and on
@@ -257,7 +309,7 @@ UEdGraphNode* FMetasoundGraphSchemaAction_NewNode::PerformAction(UEdGraph* Paren
 	UObject& ParentMetasound = CastChecked<UMetasoundEditorGraph>(ParentGraph)->GetMetasoundChecked();
 	if (UMetasoundEditorGraphExternalNode* NewGraphNode = FGraphBuilder::AddExternalNode(ParentMetasound, NodeClassInfo, Location, bSelectNewNode))
 	{
-		TryConnectNewNodeToPin(*NewGraphNode, FromPin);
+		SchemaPrivate::TryConnectNewNodeToPin(*NewGraphNode, FromPin);
 		return NewGraphNode;
 	}
 
@@ -277,8 +329,6 @@ UEdGraphNode* FMetasoundGraphSchemaAction_NewInput::PerformAction(UEdGraph* Pare
 
 	const FScopedTransaction Transaction(LOCTEXT("MetasoundEditorNewInputNode", "Add New Metasound Input Node"));
 
-	ParentGraph->Modify();
-
 	UMetasoundEditorGraph* MetasoundGraph = CastChecked<UMetasoundEditorGraph>(ParentGraph);
 	UObject& ParentMetasound = MetasoundGraph->GetMetasoundChecked();
 	ParentMetasound.Modify();
@@ -298,13 +348,71 @@ UEdGraphNode* FMetasoundGraphSchemaAction_NewInput::PerformAction(UEdGraph* Pare
 	if (UMetasoundEditorGraphInputNode* NewGraphNode = FGraphBuilder::AddInputNode(ParentMetasound, NodeHandle, InLocation))
 	{
 		UEdGraphNode* EdGraphNode = CastChecked<UEdGraphNode>(NewGraphNode);
-		TryConnectNewNodeToPin(*EdGraphNode, FromPin);
+		SchemaPrivate::TryConnectNewNodeToPin(*EdGraphNode, FromPin);
 		return EdGraphNode;
 	}
 
 	return nullptr;
 }
 
+FMetasoundGraphSchemaAction_PromoteToInput::FMetasoundGraphSchemaAction_PromoteToInput(FText InNodeCategory, FText InDisplayName, FText InToolTip, const int32 InGrouping)
+	: FEdGraphSchemaAction(MoveTemp(InNodeCategory), MoveTemp(InDisplayName), MoveTemp(InToolTip), InGrouping)
+{
+}
+
+UEdGraphNode* FMetasoundGraphSchemaAction_PromoteToInput::PerformAction(UEdGraph* ParentGraph, UEdGraphPin* FromPin, const FVector2D InLocation, bool bSelectNewNode /* = true */)
+{
+	using namespace Metasound::Editor;
+	using namespace Metasound::Frontend;
+
+	FInputHandle InputHandle = FGraphBuilder::GetInputHandleFromPin(FromPin);
+	if (!ensure(InputHandle->IsValid()))
+	{
+		return nullptr;
+	}
+
+	const FScopedTransaction Transaction(LOCTEXT("PromoteNodeInputToGraphInput", "Promote Metasound Node Input to Graph Input"));
+
+	ParentGraph->Modify();
+
+	UMetasoundEditorGraph* MetasoundGraph = CastChecked<UMetasoundEditorGraph>(ParentGraph);
+	UObject& ParentMetasound = MetasoundGraph->GetMetasoundChecked();
+
+	const FText& InputName = InputHandle->GetOwningNode()->GetDisplayName();
+	const FText& InputNodeName = InputHandle->GetDisplayName();
+	const FString InputNameBase = FText::Format(LOCTEXT("PromoteInputNameBaseFormat", "{0} {1}"), InputName, InputNodeName).ToString();
+	FString NewNodeName = FGraphBuilder::GenerateUniqueInputName(ParentMetasound, &InputNameBase);
+
+	FMetasoundFrontendLiteral DefaultValue;
+	FGraphBuilder::GetPinDefaultLiteral(*FromPin, DefaultValue);
+
+	FNodeHandle NodeHandle = FGraphBuilder::AddInputNodeHandle(ParentMetasound, NewNodeName, InputHandle->GetDataType(), FText::GetEmpty(), EMetasoundFrontendNodeStyleDisplayVisibility::Visible, &DefaultValue);
+	if (ensure(NodeHandle->IsValid()))
+	{
+		UMetasoundEditorGraphInput* Input = MetasoundGraph->FindOrAddInput(NodeHandle);
+		if (ensure(Input))
+		{
+			if (UMetasoundEditorGraphInputNode* NewGraphNode = FGraphBuilder::AddInputNode(ParentMetasound, NodeHandle, InLocation))
+			{
+				UEdGraphNode* EdGraphNode = CastChecked<UEdGraphNode>(NewGraphNode);
+
+				if (ensure(SchemaPrivate::TryConnectNewNodeToPin(*EdGraphNode, FromPin)))
+				{
+					TSharedPtr<FEditor> MetasoundEditor = SchemaPrivate::GetEditorForGraph(*ParentGraph);
+					if (MetasoundEditor.IsValid())
+					{
+						MetasoundEditor->OnInputNameChanged(NewGraphNode->GetNodeID());
+					}
+
+					return EdGraphNode;
+				}
+			}
+		}
+	}
+
+	return nullptr;
+}
+
 FMetasoundGraphSchemaAction_NewOutput::FMetasoundGraphSchemaAction_NewOutput(FText InNodeCategory, FText InDisplayName, FGuid InOutputNodeID, FText InToolTip, const int32 InGrouping)
 	: FEdGraphSchemaAction(MoveTemp(InNodeCategory), MoveTemp(InDisplayName), MoveTemp(InToolTip), InGrouping)
 	, NodeID(InOutputNodeID)
@@ -338,7 +446,7 @@ UEdGraphNode* FMetasoundGraphSchemaAction_NewOutput::PerformAction(UEdGraph* Par
 
 	if (UMetasoundEditorGraphOutputNode* NewGraphNode = FGraphBuilder::AddOutputNode(ParentMetasound, NodeHandle, Location, bSelectNewNode))
 	{
-		TryConnectNewNodeToPin(*NewGraphNode, FromPin);
+		SchemaPrivate::TryConnectNewNodeToPin(*NewGraphNode, FromPin);
 		return NewGraphNode;
 	}
 
@@ -359,7 +467,7 @@ UEdGraphNode* FMetasoundGraphSchemaAction_NewComment::PerformAction(UEdGraph* Pa
 
 	FSlateRect Bounds;
 	FVector2D SpawnLocation = Location;
-	TSharedPtr<FEditor> MetasoundEditor = GetEditorForGraph(*ParentGraph);
+	TSharedPtr<FEditor> MetasoundEditor = SchemaPrivate::GetEditorForGraph(*ParentGraph);
 
 	if (MetasoundEditor.IsValid() && MetasoundEditor->GetBoundsForSelectedNodes(Bounds, 50.0f))
 	{
@@ -375,7 +483,7 @@ UEdGraphNode* FMetasoundGraphSchemaAction_Paste::PerformAction(UEdGraph* ParentG
 {
 	using namespace Metasound::Editor;
 
-	TSharedPtr<FEditor> MetasoundEditor = GetEditorForGraph(*ParentGraph);
+	TSharedPtr<FEditor> MetasoundEditor = SchemaPrivate::GetEditorForGraph(*ParentGraph);
 	if (MetasoundEditor.IsValid())
 	{
 		MetasoundEditor->PasteNodes(&Location);
@@ -408,9 +516,9 @@ void UMetasoundEditorGraphSchema::GetGraphContextActions(FGraphContextMenuBuilde
 	using namespace Metasound::Frontend;
 
 	FActionClassFilters ClassFilters;
-	if (ContextMenuBuilder.FromPin)
+	if (const UEdGraphPin* FromPin = ContextMenuBuilder.FromPin)
 	{
-		if (ContextMenuBuilder.FromPin->Direction == EGPD_Input)
+		if (FromPin->Direction == EGPD_Input)
 		{
 			FConstInputHandle InputHandle = FGraphBuilder::GetConstInputHandleFromPin(ContextMenuBuilder.FromPin);
 			ClassFilters.OutputFilterFunction = [InputHandle](const FMetasoundFrontendClassOutput& InOutput)
@@ -429,11 +537,19 @@ void UMetasoundEditorGraphSchema::GetGraphContextActions(FGraphContextMenuBuilde
 				});
 				return bHasOutputOfType;
 			});
+
+			TSharedPtr<FMetasoundGraphSchemaAction_PromoteToInput> NewNodeAction = MakeShared<FMetasoundGraphSchemaAction_PromoteToInput>(
+				SchemaPrivate::GetContextGroupDisplayName(SchemaPrivate::EPrimaryContextGroup::Inputs),
+				LOCTEXT("PromoteToInputName", "Promote To Graph Input"),
+				LOCTEXT("PromoteToInputTooltip", "Promotes node input to graph input"),
+				static_cast<int32>(SchemaPrivate::EPrimaryContextGroup::Inputs));
+
+			static_cast<FGraphActionMenuBuilder&>(ContextMenuBuilder).AddAction(NewNodeAction);
 		}
 
-		if (ContextMenuBuilder.FromPin->Direction == EGPD_Output)
+		if (FromPin->Direction == EGPD_Output)
 		{
-			Frontend::FConstOutputHandle OutputHandle = Editor::FGraphBuilder::GetConstOutputHandleFromPin(ContextMenuBuilder.FromPin);
+			Frontend::FConstOutputHandle OutputHandle = Editor::FGraphBuilder::GetConstOutputHandleFromPin(FromPin);
 			ClassFilters.InputFilterFunction = [OutputHandle](const FMetasoundFrontendClassInput& InInput)
 			{
 				return InInput.TypeName == OutputHandle->GetDataType();
@@ -454,7 +570,7 @@ void UMetasoundEditorGraphSchema::GetGraphContextActions(FGraphContextMenuBuilde
 	}
 	else
 	{
-		TSharedPtr<Editor::FEditor> MetasoundEditor = Editor::GetEditorForGraph(*ContextMenuBuilder.CurrentGraph);
+		TSharedPtr<FEditor> MetasoundEditor = SchemaPrivate::GetEditorForGraph(*ContextMenuBuilder.CurrentGraph);
 		if (MetasoundEditor.IsValid() && MetasoundEditor->CanPasteNodes())
 		{
 			TSharedPtr<FMetasoundGraphSchemaAction_Paste> NewAction = MakeShared<FMetasoundGraphSchemaAction_Paste>(FText::GetEmpty(), LOCTEXT("PasteHereAction", "Paste here"), FText::GetEmpty(), 0);
@@ -672,7 +788,7 @@ FText UMetasoundEditorGraphSchema::GetPinDisplayName(const UEdGraphPin* Pin) con
 
 FLinearColor UMetasoundEditorGraphSchema::GetPinTypeColor(const FEdGraphPinType& PinType) const
 {
-	return Metasound::Editor::GetPinCategoryColor(PinType);
+	return Metasound::Editor::SchemaPrivate::GetPinCategoryColor(PinType);
 
 }
 
@@ -766,9 +882,6 @@ void UMetasoundEditorGraphSchema::GetDataTypeInputNodeActions(FGraphContextMenuB
 	using namespace Metasound::Editor;
 	using namespace Metasound::Frontend;
 
-	const FText InputMenuName = LOCTEXT("MetasoundActionsInputsMenu", "Inputs");
-	const FText InputTooltipFormat = LOCTEXT("MetasoundTooltipAddInputFormat", "Adds an input node referencing '{0}' to the graph.");
-
 	TArray<FConstNodeHandle> Inputs = InGraphHandle->GetConstInputNodes();
 	for (int32 i = Inputs.Num() - 1; i >= 0; --i)
 	{
@@ -778,7 +891,18 @@ void UMetasoundEditorGraphSchema::GetDataTypeInputNodeActions(FGraphContextMenuB
 			Inputs.RemoveAtSwap(i, 1, false);
 		}
 	}
-	GetDataTypeActions<FMetasoundGraphSchemaAction_NewInput>(ActionMenuBuilder, Inputs, InFilter, InputMenuName, InputTooltipFormat, bShowSelectedActions);
+
+	const SchemaPrivate::FDataTypeActionQuery ActionQuery
+	{
+		ActionMenuBuilder,
+		Inputs,
+		InFilter,
+		SchemaPrivate::EPrimaryContextGroup::Inputs,
+		SchemaPrivate::InputDisplayNameFormat,
+		SchemaPrivate::InputTooltipFormat,
+		bShowSelectedActions
+	};
+	SchemaPrivate::GetDataTypeActions<FMetasoundGraphSchemaAction_NewInput>(ActionQuery);
 }
 
 void UMetasoundEditorGraphSchema::GetDataTypeOutputNodeActions(FGraphContextMenuBuilder& ActionMenuBuilder, Metasound::Frontend::FConstGraphHandle InGraphHandle, Metasound::Editor::FInterfaceNodeFilterFunction InFilter, bool bShowSelectedActions) const
@@ -786,9 +910,6 @@ void UMetasoundEditorGraphSchema::GetDataTypeOutputNodeActions(FGraphContextMenu
 	using namespace Metasound::Editor;
 	using namespace Metasound::Frontend;
 
-	const FText OutputMenuName = LOCTEXT("MetasoundActionsOutputsMenu", "Outputs");
-	const FText OutputTooltipFormat = LOCTEXT("MetasoundTooltipAddOutputFormat", "Adds an output node referencing '{0}' to the graph.");
-
 	TArray<FConstNodeHandle> Outputs = InGraphHandle->GetConstOutputNodes();
 
 	// Prune and only add actions for outputs that are not already represented in the graph
@@ -808,7 +929,17 @@ void UMetasoundEditorGraphSchema::GetDataTypeOutputNodeActions(FGraphContextMenu
 		}
 	}
 
-	GetDataTypeActions<FMetasoundGraphSchemaAction_NewOutput>(ActionMenuBuilder, Outputs, InFilter, OutputMenuName, OutputTooltipFormat, bShowSelectedActions);
+	const SchemaPrivate::FDataTypeActionQuery ActionQuery
+	{
+		ActionMenuBuilder,
+		Outputs,
+		InFilter,
+		SchemaPrivate::EPrimaryContextGroup::Outputs,
+		SchemaPrivate::OutputDisplayNameFormat,
+		SchemaPrivate::OutputTooltipFormat,
+		bShowSelectedActions
+	};
+	SchemaPrivate::GetDataTypeActions<FMetasoundGraphSchemaAction_NewOutput>(ActionQuery);
 }
 
 void UMetasoundEditorGraphSchema::GetFunctionActions(FGraphActionMenuBuilder& ActionMenuBuilder, Metasound::Editor::FActionClassFilters InFilters, bool bShowSelectedActions) const
@@ -860,7 +991,7 @@ void UMetasoundEditorGraphSchema::GetCommentAction(FGraphActionMenuBuilder& Acti
 
 	if (!ActionMenuBuilder.FromPin && CurrentGraph)
 	{
-		TSharedPtr<FEditor> MetasoundEditor = GetEditorForGraph(*CurrentGraph);
+		TSharedPtr<FEditor> MetasoundEditor = SchemaPrivate::GetEditorForGraph(*CurrentGraph);
 		if (MetasoundEditor.IsValid())
 		{
 			const int32 NumSelected = MetasoundEditor->GetNumNodesSelected();
@@ -877,7 +1008,7 @@ int32 UMetasoundEditorGraphSchema::GetNodeSelectionCount(const UEdGraph* Graph)
 {
 	using namespace Metasound::Editor;
 
-	TSharedPtr<FEditor> MetasoundEditor = GetEditorForGraph(*Graph);
+	TSharedPtr<FEditor> MetasoundEditor = SchemaPrivate::GetEditorForGraph(*Graph);
 	if (MetasoundEditor.IsValid())
 	{
 		return MetasoundEditor->GetNumNodesSelected();
diff --git a/Engine/Plugins/Runtime/Metasound/Source/MetasoundEditor/Public/MetasoundEditorGraph.h b/Engine/Plugins/Runtime/Metasound/Source/MetasoundEditor/Public/MetasoundEditorGraph.h
--- a/Engine/Plugins/Runtime/Metasound/Source/MetasoundEditor/Public/MetasoundEditorGraph.h
+++ b/Engine/Plugins/Runtime/Metasound/Source/MetasoundEditor/Public/MetasoundEditorGraph.h
@@ -70,6 +70,10 @@ public:
 	{
 		return EMetasoundFrontendLiteralType::None;
 	}
+
+	virtual void SetFromLiteral(const FMetasoundFrontendLiteral& InLiteral)
+	{
+	}
 };
 
 UCLASS()
diff --git a/Engine/Plugins/Runtime/Metasound/Source/MetasoundEditor/Public/MetasoundEditorGraphInputNodes.h b/Engine/Plugins/Runtime/Metasound/Source/MetasoundEditor/Public/MetasoundEditorGraphInputNodes.h
--- a/Engine/Plugins/Runtime/Metasound/Source/MetasoundEditor/Public/MetasoundEditorGraphInputNodes.h
+++ b/Engine/Plugins/Runtime/Metasound/Source/MetasoundEditor/Public/MetasoundEditorGraphInputNodes.h
@@ -5,6 +5,7 @@
 #include "Algo/Transform.h"
 #include "CoreMinimal.h"
 #include "EdGraph/EdGraphNode.h"
+#include "MetasoundDataReference.h"
 #include "MetasoundEditorGraph.h"
 #include "MetasoundEditorGraphNode.h"
 #include "MetasoundFrontend.h"
@@ -12,6 +13,7 @@
 #include "MetasoundFrontendDocument.h"
 #include "MetasoundFrontendLiteral.h"
 #include "MetasoundFrontendRegistries.h"
+#include "MetasoundPrimitives.h"
 #include "Misc/Guid.h"
 #include "Sound/SoundWave.h"
 #include "UObject/ObjectMacros.h"
@@ -27,11 +29,29 @@ namespace Metasound
 {
 	namespace Editor
 	{
+		// Forward Declarations
 		class FGraphBuilder;
+
+		namespace InputPrivate
+		{
+			template <typename TPODType>
+			void ConvertLiteral(const FMetasoundFrontendLiteral& InLiteral, TPODType& OutValue)
+			{
+				const FName& TypeName = GetMetasoundDataTypeName<TPODType>();
+				OutValue = InLiteral.ToLiteral(TypeName).Value.Get<TPODType>();
+			}
+
+			template <typename TPODType, typename TLiteralType = TPODType>
+			void ConvertLiteralToArray(const FMetasoundFrontendLiteral& InLiteral, TArray<TLiteralType>& OutArray)
+			{
+				const FName& TypeName = *FString(GetMetasoundDataTypeString<TPODType>() + TEXT(":Array"));
+				TArray<TPODType> NewValue = InLiteral.ToLiteral(TypeName).Value.Get<TArray<TPODType>>();
+				Algo::Transform(NewValue, OutArray, [](const TPODType& InValue) { return TLiteralType { InValue }; });
+			}
+		}
 	}
 }
 
-
 UCLASS()
 class METASOUNDEDITOR_API UMetasoundEditorGraphInputNode : public UMetasoundEditorGraphNode
 {
@@ -54,7 +74,7 @@ public:
 		return Super::GetClassName();
 	}
 
-	void UpdatePreviewInstance(const Metasound::FVertexKey& InParameterName, IAudioInstanceTransmitter& InInstanceTransmitter) const
+	virtual void UpdatePreviewInstance(const Metasound::FVertexKey& InParameterName, IAudioInstanceTransmitter& InInstanceTransmitter) const
 	{
 		if (ensure(Input))
 		{
@@ -121,12 +141,17 @@ public:
 		return Literal;
 	}
 
-	virtual EMetasoundFrontendLiteralType GetLiteralType() const
+	virtual EMetasoundFrontendLiteralType GetLiteralType() const override
 	{
 		return EMetasoundFrontendLiteralType::Boolean;
 	}
 
-	void UpdatePreviewInstance(const Metasound::FVertexKey& InParameterName, IAudioInstanceTransmitter& InInstanceTransmitter) const override
+	virtual void SetFromLiteral(const FMetasoundFrontendLiteral& InLiteral) override
+	{
+		Metasound::Editor::InputPrivate::ConvertLiteral<bool>(InLiteral, Default.Value);
+	}
+
+	virtual void UpdatePreviewInstance(const Metasound::FVertexKey& InParameterName, IAudioInstanceTransmitter& InInstanceTransmitter) const override
 	{
 		InInstanceTransmitter.SetParameter(*InParameterName, Default.Value);
 	}
@@ -153,12 +178,17 @@ public:
 		return Literal;
 	}
 
-	virtual EMetasoundFrontendLiteralType GetLiteralType() const
+	virtual EMetasoundFrontendLiteralType GetLiteralType() const override
 	{
 		return EMetasoundFrontendLiteralType::BooleanArray;
 	}
 
-	void UpdatePreviewInstance(const Metasound::FVertexKey& InParameterName, IAudioInstanceTransmitter& InInstanceTransmitter) const override
+	virtual void SetFromLiteral(const FMetasoundFrontendLiteral& InLiteral) override
+	{
+		Metasound::Editor::InputPrivate::ConvertLiteralToArray<bool, FMetasoundEditorGraphInputBoolRef>(InLiteral, Default);
+	}
+
+	virtual void UpdatePreviewInstance(const Metasound::FVertexKey& InParameterName, IAudioInstanceTransmitter& InInstanceTransmitter) const override
 	{
 		TArray<bool> BoolArray;
 		Algo::Transform(Default, BoolArray, [](const FMetasoundEditorGraphInputBoolRef& InValue) { return InValue.Value; });
@@ -195,12 +225,17 @@ public:
 		return Literal;
 	}
 
-	virtual EMetasoundFrontendLiteralType GetLiteralType() const
+	virtual EMetasoundFrontendLiteralType GetLiteralType() const override
 	{
 		return EMetasoundFrontendLiteralType::Integer;
 	}
 
-	void UpdatePreviewInstance(const Metasound::FVertexKey& InParameterName, IAudioInstanceTransmitter& InInstanceTransmitter) const override
+	virtual void SetFromLiteral(const FMetasoundFrontendLiteral& InLiteral) override
+	{
+		Metasound::Editor::InputPrivate::ConvertLiteral<int32>(InLiteral, Default.Value);
+	}
+
+	virtual void UpdatePreviewInstance(const Metasound::FVertexKey& InParameterName, IAudioInstanceTransmitter& InInstanceTransmitter) const override
 	{
 		InInstanceTransmitter.SetParameter(*InParameterName, Default.Value);
 	}
@@ -228,11 +263,16 @@ public:
 		return Literal;
 	}
 
-	virtual EMetasoundFrontendLiteralType GetLiteralType() const
+	virtual EMetasoundFrontendLiteralType GetLiteralType() const override
 	{
 		return EMetasoundFrontendLiteralType::IntegerArray;
 	}
 
+	virtual void SetFromLiteral(const FMetasoundFrontendLiteral& InLiteral) override
+	{
+		Metasound::Editor::InputPrivate::ConvertLiteralToArray<int32, FMetasoundEditorGraphInputIntRef>(InLiteral, Default);
+	}
+
 	void UpdatePreviewInstance(const Metasound::FVertexKey& InParameterName, IAudioInstanceTransmitter& InInstanceTransmitter) const override
 	{
 		TArray<int32> IntArray;
@@ -260,12 +300,17 @@ public:
 		return Literal;
 	}
 
-	virtual EMetasoundFrontendLiteralType GetLiteralType() const
+	virtual EMetasoundFrontendLiteralType GetLiteralType() const override
 	{
 		return EMetasoundFrontendLiteralType::Float;
 	}
 
-	void UpdatePreviewInstance(const Metasound::FVertexKey& InParameterName, IAudioInstanceTransmitter& InInstanceTransmitter) const override
+	virtual void SetFromLiteral(const FMetasoundFrontendLiteral& InLiteral) override
+	{
+		Metasound::Editor::InputPrivate::ConvertLiteral<float>(InLiteral, Default);
+	}
+
+	virtual void UpdatePreviewInstance(const Metasound::FVertexKey& InParameterName, IAudioInstanceTransmitter& InInstanceTransmitter) const override
 	{
 		InInstanceTransmitter.SetParameter(*InParameterName, Default);
 	}
@@ -289,12 +334,17 @@ public:
 		return Literal;
 	}
 
-	virtual EMetasoundFrontendLiteralType GetLiteralType() const
+	virtual EMetasoundFrontendLiteralType GetLiteralType() const override
 	{
 		return EMetasoundFrontendLiteralType::FloatArray;
 	}
 
-	void UpdatePreviewInstance(const Metasound::FVertexKey& InParameterName, IAudioInstanceTransmitter& InInstanceTransmitter) const override
+	virtual void SetFromLiteral(const FMetasoundFrontendLiteral& InLiteral) override
+	{
+		Metasound::Editor::InputPrivate::ConvertLiteralToArray<float>(InLiteral, Default);
+	}
+
+	virtual void UpdatePreviewInstance(const Metasound::FVertexKey& InParameterName, IAudioInstanceTransmitter& InInstanceTransmitter) const override
 	{
 		InInstanceTransmitter.SetParameter(*InParameterName, TArray<float> { Default });
 	}
@@ -318,12 +368,17 @@ public:
 		return Literal;
 	}
 
-	virtual EMetasoundFrontendLiteralType GetLiteralType() const
+	virtual EMetasoundFrontendLiteralType GetLiteralType() const override
 	{
 		return EMetasoundFrontendLiteralType::String;
 	}
 
-	void UpdatePreviewInstance(const Metasound::FVertexKey& InParameterName, IAudioInstanceTransmitter& InInstanceTransmitter) const override
+	virtual void SetFromLiteral(const FMetasoundFrontendLiteral& InLiteral) override
+	{
+		Metasound::Editor::InputPrivate::ConvertLiteral<FString>(InLiteral, Default);
+	}
+
+	virtual void UpdatePreviewInstance(const Metasound::FVertexKey& InParameterName, IAudioInstanceTransmitter& InInstanceTransmitter) const override
 	{
 		InInstanceTransmitter.SetParameter(*InParameterName, FString { Default });
 	}
@@ -347,12 +402,17 @@ public:
 		return Literal;
 	}
 
-	virtual EMetasoundFrontendLiteralType GetLiteralType() const
+	virtual EMetasoundFrontendLiteralType GetLiteralType() const override
 	{
 		return EMetasoundFrontendLiteralType::StringArray;
 	}
 
-	void UpdatePreviewInstance(const Metasound::FVertexKey& InParameterName, IAudioInstanceTransmitter& InInstanceTransmitter) const override
+	virtual void SetFromLiteral(const FMetasoundFrontendLiteral& InLiteral) override
+	{
+		Metasound::Editor::InputPrivate::ConvertLiteralToArray<FString>(InLiteral, Default);
+	}
+
+	virtual void UpdatePreviewInstance(const Metasound::FVertexKey& InParameterName, IAudioInstanceTransmitter& InInstanceTransmitter) const override
 	{
 		InInstanceTransmitter.SetParameter(*InParameterName, TArray<FString> { Default });
 	}
@@ -386,12 +446,17 @@ public:
 		return Literal;
 	}
 
-	virtual EMetasoundFrontendLiteralType GetLiteralType() const
+	virtual EMetasoundFrontendLiteralType GetLiteralType() const override
 	{
 		return EMetasoundFrontendLiteralType::UObject;
 	}
 
-	void UpdatePreviewInstance(const Metasound::FVertexKey& InParameterName, IAudioInstanceTransmitter& InInstanceTransmitter) const override
+	virtual void SetFromLiteral(const FMetasoundFrontendLiteral& InLiteral) override
+	{
+// 		Metasound::Editor::InputPrivate::ConvertLiteral<UObject*>(InLiteral, Default.Object);
+	}
+
+	virtual void UpdatePreviewInstance(const Metasound::FVertexKey& InParameterName, IAudioInstanceTransmitter& InInstanceTransmitter) const override
 	{
 		// TODO. We need proxy object here safely.
 		//InInstanceTransmitter.SetParameter(*InParameterName, Default.Object);
@@ -419,12 +484,17 @@ public:
 		return Literal;
 	}
 
-	virtual EMetasoundFrontendLiteralType GetLiteralType() const
+	virtual EMetasoundFrontendLiteralType GetLiteralType() const override
 	{
 		return EMetasoundFrontendLiteralType::UObjectArray;
 	}
 
-	void UpdatePreviewInstance(const Metasound::FVertexKey& InParameterName, IAudioInstanceTransmitter& InInstanceTransmitter) const override
+	virtual void SetFromLiteral(const FMetasoundFrontendLiteral& InLiteral) override
+	{
+// 		Metasound::Editor::InputPrivate::ConvertLiteralToArray<UObject*, FMetasoundEditorGraphInputObjectRef>(InLiteral, Default);
+	}
+
+	virtual void UpdatePreviewInstance(const Metasound::FVertexKey& InParameterName, IAudioInstanceTransmitter& InInstanceTransmitter) const override
 	{
 		TArray<UObject*> ObjectArray;
 		Algo::Transform(Default, ObjectArray, [](const FMetasoundEditorGraphInputObjectRef& InValue) { return InValue.Object; });
diff --git a/Engine/Plugins/Runtime/Metasound/Source/MetasoundEditor/Public/MetasoundEditorGraphSchema.h b/Engine/Plugins/Runtime/Metasound/Source/MetasoundEditor/Public/MetasoundEditorGraphSchema.h
--- a/Engine/Plugins/Runtime/Metasound/Source/MetasoundEditor/Public/MetasoundEditorGraphSchema.h
+++ b/Engine/Plugins/Runtime/Metasound/Source/MetasoundEditor/Public/MetasoundEditorGraphSchema.h
@@ -92,7 +92,7 @@ namespace Metasound
 	} // namespace Editor
 } // namespace Metasound
 
-/** Action to add an input to the graph */
+/** Action to add an input reference to the graph */
 USTRUCT()
 struct METASOUNDEDITOR_API FMetasoundGraphSchemaAction_NewInput : public FEdGraphSchemaAction
 {
@@ -112,6 +112,23 @@ struct METASOUNDEDITOR_API FMetasoundGraphSchemaAction_NewInput : public FEdGrap
 	//~ End FEdGraphSchemaAction Interface
 };
 
+/** Action to promote a literal to a graph input */
+USTRUCT()
+struct METASOUNDEDITOR_API FMetasoundGraphSchemaAction_PromoteToInput : public FEdGraphSchemaAction
+{
+	GENERATED_USTRUCT_BODY();
+
+	FMetasoundGraphSchemaAction_PromoteToInput()
+		: FEdGraphSchemaAction()
+	{}
+
+	FMetasoundGraphSchemaAction_PromoteToInput(FText InNodeCategory, FText InDisplayName, FText InToolTip, const int32 InGrouping);
+
+	//~ Begin FEdGraphSchemaAction Interface
+	virtual UEdGraphNode* PerformAction(UEdGraph* ParentGraph, UEdGraphPin* FromPin, const FVector2D Location, bool bSelectNewNode = true) override;
+	//~ End FEdGraphSchemaAction Interface
+};
+
 /** Action to add an output to the graph */
 USTRUCT()
 struct METASOUNDEDITOR_API FMetasoundGraphSchemaAction_NewOutput : public FEdGraphSchemaAction
