Array pin updates: display as arrays, don't show default value picker
- Added playtime to be able to determine when the metasound is actually being previewed

diff --git a/Engine/Plugins/Runtime/Metasound/Source/MetasoundEditor/Private/MetasoundEditor.cpp b/Engine/Plugins/Runtime/Metasound/Source/MetasoundEditor/Private/MetasoundEditor.cpp
--- a/Engine/Plugins/Runtime/Metasound/Source/MetasoundEditor/Private/MetasoundEditor.cpp
+++ b/Engine/Plugins/Runtime/Metasound/Source/MetasoundEditor/Private/MetasoundEditor.cpp
@@ -44,6 +44,8 @@
 
 // This needs to be moved to public directory
 #include "../../GraphEditor/Private/GraphActionNode.h"
+#include "../../PropertyEditor/Public/DetailLayoutBuilder.h"
+#include "../../../Runtime/SlateCore/Public/Widgets/SOverlay.h"
 
 
 #define LOCTEXT_NAMESPACE "MetasoundEditor"
@@ -149,9 +151,19 @@ namespace Metasound
 
 			FAssetEditorToolkit::RegisterTabSpawners(InTabManager);
 
-			InTabManager->RegisterTabSpawner(TabFactory::Names::GraphCanvas, FOnSpawnTab::CreateLambda([InMetasoundGraphEditor = MetasoundGraphEditor](const FSpawnTabArgs& Args)
-			{
-				return TabFactory::CreateGraphCanvasTab(InMetasoundGraphEditor, Args);
+			InTabManager->RegisterTabSpawner(TabFactory::Names::GraphCanvas, FOnSpawnTab::CreateLambda([InPlayTimeWidget = PlayTimeWidget, InMetasoundGraphEditor = MetasoundGraphEditor](const FSpawnTabArgs& Args)
+			{
+				return TabFactory::CreateGraphCanvasTab(SNew(SOverlay)
+					+ SOverlay::Slot()
+					[
+						InMetasoundGraphEditor.ToSharedRef()
+					]
+					+ SOverlay::Slot()
+					[
+						InPlayTimeWidget.ToSharedRef()
+					]
+					.Padding(5.0f, 5.0f)
+				, Args);
 			}))
 			.SetDisplayName(LOCTEXT("GraphCanvasTab", "Viewport"))
 			.SetGroup(WorkspaceMenuCategoryRef)
@@ -274,29 +286,29 @@ namespace Metasound
 
 			CreateAnalyzers();
 
-			const TSharedRef<FTabManager::FLayout> StandaloneDefaultLayout = FTabManager::NewLayout("Standalone_MetasoundEditor_Layout_v6")
+			const TSharedRef<FTabManager::FLayout> StandaloneDefaultLayout = FTabManager::NewLayout("Standalone_MetasoundEditor_Layout_v7")
 				->AddArea
 				(
 					FTabManager::NewPrimaryArea()
 					->SetOrientation(Orient_Vertical)
 					->Split(FTabManager::NewSplitter()
 						->SetOrientation(Orient_Horizontal)
-						->SetSizeCoefficient(0.9f)
 						->Split
 						(
 							FTabManager::NewSplitter()
+							->SetSizeCoefficient(0.15f)
 							->SetOrientation(Orient_Vertical)
 							->Split
 							(
 								FTabManager::NewStack()
-								->SetSizeCoefficient(0.5f)
+								->SetSizeCoefficient(0.25f)
 								->SetHideTabWell(false)
 								->AddTab(TabFactory::Names::Metasound, ETabState::OpenedTab)
 							)
 							->Split
 							(
 								FTabManager::NewStack()
-								->SetSizeCoefficient(0.5f)
+								->SetSizeCoefficient(0.50f)
 								->SetHideTabWell(false)
 								->AddTab(TabFactory::Names::Inspector, ETabState::OpenedTab)
 							)
@@ -304,14 +316,14 @@ namespace Metasound
 						->Split
 						(
 							FTabManager::NewStack()
-							->SetSizeCoefficient(0.65f)
+							->SetSizeCoefficient(0.80f)
 							->SetHideTabWell(true)
 							->AddTab(TabFactory::Names::GraphCanvas, ETabState::OpenedTab)
 						)
 						->Split
 						(
 							FTabManager::NewStack()
-							->SetSizeCoefficient(0.125f)
+							->SetSizeCoefficient(0.05f)
 							->SetHideTabWell(true)
 							->AddTab(TabFactory::Names::Analyzers, ETabState::OpenedTab)
 						)
@@ -422,14 +434,12 @@ namespace Metasound
 
 		void FEditor::CreateInternalWidgets()
 		{
-			MetasoundGraphEditor = CreateGraphEditorWidget();
+			CreateGraphEditorWidget();
 
 			FDetailsViewArgs Args;
 			Args.bHideSelectionTip = true;
 			Args.NotifyHook = this;
 
-			FPropertyEditorModule& PropertyModule = FModuleManager::LoadModuleChecked<FPropertyEditorModule>("PropertyEditor");
-
 			SAssignNew(MetasoundInterfaceMenu, SGraphActionMenu, false)
 				.AlphaSortItems(false)
 // 				.OnActionDoubleClicked(this, &FEditor::OnActionDoubleClicked)
@@ -447,8 +457,8 @@ namespace Metasound
 				.OnGetSectionWidget(this, &FEditor::OnGetMenuSectionWidget)
 				.UseSectionStyling(true);
 
+			FPropertyEditorModule& PropertyModule = FModuleManager::LoadModuleChecked<FPropertyEditorModule>("PropertyEditor");
 			MetasoundDetails = PropertyModule.CreateDetailView(Args);
-
 			Palette = SNew(SMetasoundPalette);
 		}
 
@@ -717,12 +727,11 @@ namespace Metasound
 
 			if (USoundBase* MetasoundToPlay = Cast<USoundBase>(Metasound))
 			{
-
-				GEditor->PlayPreviewSound(MetasoundToPlay);
-
 				// Set the send to the audio bus that is used for analyzing the metasound output
-				if (UAudioComponent* PreviewComp = GEditor->GetPreviewAudioComponent())
+				if (UAudioComponent* PreviewComp = GEditor->PlayPreviewSound(MetasoundToPlay))
 				{
+					PlayTime = 0.0;
+
 					if (MetasoundAudioBus.IsValid())
 					{
 						PreviewComp->SetAudioBusSendPostEffect(MetasoundAudioBus.Get(), 1.0f);
@@ -736,9 +745,8 @@ namespace Metasound
 						if(FMetasoundAssetBase* MetasoundAsset = IMetasoundUObjectRegistry::Get().GetObjectAsAssetBase(Metasound))
 						{
 							const uint64 AudioComponentID = PreviewComp->GetAudioComponentID();
-							float SampleRate = PreviewAudioDevice->GetSampleRate();
-
-							MetasoundComm = MetasoundToPlay->CreateInstanceTransmitter(FAudioInstanceTransmitterInitParams{AudioComponentID, SampleRate});
+							const float SampleRate = PreviewAudioDevice->GetSampleRate();
+							MetasoundComm = MetasoundToPlay->CreateInstanceTransmitter(FAudioInstanceTransmitterInitParams { AudioComponentID, SampleRate });
 						}
 					}
 
@@ -749,18 +757,29 @@ namespace Metasound
 				}
 
 				MetasoundGraphEditor->RegisterActiveTimer(0.0f,
-					FWidgetActiveTimerDelegate::CreateLambda([](double InCurrentTime, float InDeltaTime)
+					FWidgetActiveTimerDelegate::CreateLambda([this](double InCurrentTime, float InDeltaTime)
+					{
+						UAudioComponent* PreviewComp = GEditor->GetPreviewAudioComponent();
+						if (PreviewComp && PreviewComp->IsPlaying())
 						{
-							UAudioComponent* PreviewComp = GEditor->GetPreviewAudioComponent();
-							if (PreviewComp && PreviewComp->IsPlaying())
-							{
-								return EActiveTimerReturnType::Continue;
-							}
-							else
+							if (PlayTimeWidget.IsValid())
 							{
-								return EActiveTimerReturnType::Stop;
+								PlayTime += InDeltaTime;
+								FString PlayTimeString = FTimespan::FromSeconds(PlayTime).ToString();
+
+								// Remove leading '+'
+								PlayTimeString.ReplaceInline(TEXT("+"), TEXT(""));
+								PlayTimeWidget->SetText(FText::FromString(PlayTimeString));
 							}
-						})
+							return EActiveTimerReturnType::Continue;
+						}
+						else
+						{
+							PlayTime = 0.0;
+							PlayTimeWidget->SetText(FText::GetEmpty());
+							return EActiveTimerReturnType::Stop;
+						}
+					})
 				);
 			}
 		}
@@ -900,7 +919,7 @@ namespace Metasound
 		{
 		}
 
-		TSharedRef<SGraphEditor> FEditor::CreateGraphEditorWidget()
+		void FEditor::CreateGraphEditorWidget()
 		{
 			if (!GraphEditorCommands.IsValid())
 			{
@@ -992,7 +1011,7 @@ namespace Metasound
 			FMetasoundAssetBase* MetasoundAsset = IMetasoundUObjectRegistry::Get().GetObjectAsAssetBase(Metasound);
 			check(MetasoundAsset);
 
-			return SNew(SGraphEditor)
+			SAssignNew(MetasoundGraphEditor, SGraphEditor)
 				.AdditionalCommands(GraphEditorCommands)
 				.IsEditable(true)
 				.Appearance(AppearanceInfo)
@@ -1000,6 +1019,11 @@ namespace Metasound
 				.GraphEvents(InEvents)
 				.AutoExpandActionMenu(true)
 				.ShowGraphStateOverlay(false);
+
+			SAssignNew(PlayTimeWidget, STextBlock)
+				.Visibility(EVisibility::HitTestInvisible)
+				.TextStyle(FEditorStyle::Get(), "Graph.ZoomText")
+				.ColorAndOpacity(FLinearColor(1, 1, 1, 0.30f));
 		}
 
 		void FEditor::OnSelectedNodesChanged(const TSet<UObject*>& InSelectedNodes)
diff --git a/Engine/Plugins/Runtime/Metasound/Source/MetasoundEditor/Private/MetasoundEditor.h b/Engine/Plugins/Runtime/Metasound/Source/MetasoundEditor/Private/MetasoundEditor.h
--- a/Engine/Plugins/Runtime/Metasound/Source/MetasoundEditor/Private/MetasoundEditor.h
+++ b/Engine/Plugins/Runtime/Metasound/Source/MetasoundEditor/Private/MetasoundEditor.h
@@ -226,7 +226,7 @@ namespace Metasound
 			void OnCreateComment();
 
 			/** Create new graph editor widget */
-			TSharedRef<SGraphEditor> CreateGraphEditorWidget();
+			void CreateGraphEditorWidget();
 
 		private:
 			void EditObjectSettings();
@@ -252,6 +252,10 @@ namespace Metasound
 			/** Palette of Node types */
 			TSharedPtr<SMetasoundPalette> Palette;
 
+			/** Widget showing playtime that overlays the graph when previewing */
+			TSharedPtr<STextBlock> PlayTimeWidget;
+			double PlayTime = 0.0f;
+
 			/** Command list for this editor */
 			TSharedPtr<FUICommandList> GraphEditorCommands;
 
diff --git a/Engine/Plugins/Runtime/Metasound/Source/MetasoundEditor/Private/MetasoundEditorModule.cpp b/Engine/Plugins/Runtime/Metasound/Source/MetasoundEditor/Private/MetasoundEditorModule.cpp
--- a/Engine/Plugins/Runtime/Metasound/Source/MetasoundEditor/Private/MetasoundEditorModule.cpp
+++ b/Engine/Plugins/Runtime/Metasound/Source/MetasoundEditor/Private/MetasoundEditorModule.cpp
@@ -150,7 +150,9 @@ namespace Metasound
 
 			virtual void RegisterDataType(FName InPinCategoryName, FName InPinSubCategoryName, const FDataTypeRegistryInfo& InRegistryInfo) override
 			{
-				FEdGraphPinType PinType(InPinCategoryName, InPinSubCategoryName, nullptr, EPinContainerType::None, false, FEdGraphTerminalType());
+				const bool bIsArray = InRegistryInfo.IsArrayType();
+				const EPinContainerType ContainerType = bIsArray ? EPinContainerType::Array : EPinContainerType::None;
+				FEdGraphPinType PinType(InPinCategoryName, InPinSubCategoryName, nullptr, ContainerType, false, FEdGraphTerminalType());
 				DataTypeInfo.Emplace(InRegistryInfo.DataTypeName, Editor::FEditorDataType(MoveTemp(PinType), InRegistryInfo));
 			}
 
@@ -252,7 +254,9 @@ namespace Metasound
 						}
 					}
 
-					FEdGraphPinType PinType(PinCategory, PinSubCategory, nullptr, EPinContainerType::None, false, FEdGraphTerminalType());
+					const bool bIsArray = RegistryInfo.IsArrayType();
+					const EPinContainerType ContainerType = bIsArray ? EPinContainerType::Array : EPinContainerType::None;
+					FEdGraphPinType PinType(PinCategory, PinSubCategory, nullptr, ContainerType, false, FEdGraphTerminalType());
  					UClass* ClassToUse = FMetasoundFrontendRegistryContainer::Get()->GetLiteralUClassForDataType(DataTypeName);
  					PinType.PinSubCategoryObject = Cast<UObject>(ClassToUse);
 
diff --git a/Engine/Plugins/Runtime/Metasound/Source/MetasoundEditor/Private/MetasoundEditorTabFactory.h b/Engine/Plugins/Runtime/Metasound/Source/MetasoundEditor/Private/MetasoundEditorTabFactory.h
--- a/Engine/Plugins/Runtime/Metasound/Source/MetasoundEditor/Private/MetasoundEditorTabFactory.h
+++ b/Engine/Plugins/Runtime/Metasound/Source/MetasoundEditor/Private/MetasoundEditorTabFactory.h
@@ -38,7 +38,7 @@ namespace Metasound
 					];
 			}
 
-			TSharedRef<SDockTab> CreateGraphCanvasTab(TSharedPtr<SGraphEditor> GraphEditor, const FSpawnTabArgs& Args)
+			TSharedRef<SDockTab> CreateGraphCanvasTab(TSharedPtr<SWidget> GraphEditor, const FSpawnTabArgs& Args)
 			{
 				check(Args.GetTabId() == Names::GraphCanvas);
 
diff --git a/Engine/Plugins/Runtime/Metasound/Source/MetasoundEditor/Private/SMetasoundGraphNode.cpp b/Engine/Plugins/Runtime/Metasound/Source/MetasoundEditor/Private/SMetasoundGraphNode.cpp
--- a/Engine/Plugins/Runtime/Metasound/Source/MetasoundEditor/Private/SMetasoundGraphNode.cpp
+++ b/Engine/Plugins/Runtime/Metasound/Source/MetasoundEditor/Private/SMetasoundGraphNode.cpp
@@ -79,6 +79,12 @@ TSharedPtr<SGraphPin> SMetasoundGraphNode::CreatePinWidget(UEdGraphPin* InPin) c
 
 	if (const UMetasoundEditorGraphSchema* GraphSchema = Cast<const UMetasoundEditorGraphSchema>(InPin->GetSchema()))
 	{
+		// Don't show default value field for container types
+		if (InPin->PinType.ContainerType != EPinContainerType::None)
+		{
+			return SNew(SGraphPin, InPin);
+		}
+
 		if (InPin->PinType.PinCategory == FGraphBuilder::PinCategoryAudio)
 		{
 			return SNew(SGraphPin, InPin);
diff --git a/Engine/Plugins/Runtime/Metasound/Source/MetasoundFrontend/Public/MetasoundFrontendRegistries.h b/Engine/Plugins/Runtime/Metasound/Source/MetasoundFrontend/Public/MetasoundFrontendRegistries.h
--- a/Engine/Plugins/Runtime/Metasound/Source/MetasoundFrontend/Public/MetasoundFrontendRegistries.h
+++ b/Engine/Plugins/Runtime/Metasound/Source/MetasoundFrontend/Public/MetasoundFrontendRegistries.h
@@ -103,6 +103,15 @@ namespace Metasound
 		// If this datatype was registered with a specific UClass to use to filter with, that will be used here:
 		UClass* ProxyGeneratorClass = nullptr;
 
+		FORCEINLINE bool IsArrayType() const
+		{
+			return bIsDefaultArrayParsable
+			|| bIsBoolArrayParsable
+			|| bIsIntArrayParsable
+			|| bIsFloatArrayParsable
+			|| bIsStringArrayParsable
+			|| bIsProxyArrayParsable;
+		}
 	};
 
 	namespace Frontend
diff --git a/Engine/Source/Editor/UnrealEd/Classes/Editor/EditorEngine.h b/Engine/Source/Editor/UnrealEd/Classes/Editor/EditorEngine.h
--- a/Engine/Source/Editor/UnrealEd/Classes/Editor/EditorEngine.h
+++ b/Engine/Source/Editor/UnrealEd/Classes/Editor/EditorEngine.h
@@ -1126,7 +1126,7 @@ public:
 	 * @param	Sound		A sound to attach to the audio component
 	 * @param	SoundNode	A sound node that is attached to the audio component when the sound is NULL
 	 */
-	UAudioComponent* ResetPreviewAudioComponent( class USoundBase* Sound = NULL, class USoundNode* SoundNode = NULL );
+	UAudioComponent* ResetPreviewAudioComponent(USoundBase* Sound = nullptr, USoundNode* SoundNode = nullptr);
 
 	/**
 	 * Plays a preview of a specified sound or node
@@ -1134,7 +1134,7 @@ public:
 	 * @param	Sound		A sound to attach to the audio component
 	 * @param	SoundNode	A sound node that is attached to the audio component when the sound is NULL
 	 */
-	void PlayPreviewSound(USoundBase* Sound, USoundNode* SoundNode = NULL);
+	UAudioComponent* PlayPreviewSound(USoundBase* Sound, USoundNode* SoundNode = nullptr);
 
 	/**
 	 * Clean up any world specific editor components so they can be GC correctly
diff --git a/Engine/Source/Editor/UnrealEd/Private/EditorEngine.cpp b/Engine/Source/Editor/UnrealEd/Private/EditorEngine.cpp
--- a/Engine/Source/Editor/UnrealEd/Private/EditorEngine.cpp
+++ b/Engine/Source/Editor/UnrealEd/Private/EditorEngine.cpp
@@ -2379,10 +2379,9 @@ UAudioComponent* UEditorEngine::ResetPreviewAudioComponent( USoundBase* Sound, U
 	return PreviewAudioComponent;
 }
 
-void UEditorEngine::PlayPreviewSound( USoundBase* Sound,  USoundNode* SoundNode )
+UAudioComponent* UEditorEngine::PlayPreviewSound(USoundBase* Sound,  USoundNode* SoundNode)
 {
-	UAudioComponent* AudioComponent = ResetPreviewAudioComponent(Sound, SoundNode);
-	if(AudioComponent)
+	if(UAudioComponent* AudioComponent = ResetPreviewAudioComponent(Sound, SoundNode))
 	{
 		AudioComponent->bAutoDestroy = false;
 		AudioComponent->bIsUISound = true;
@@ -2390,8 +2389,12 @@ void UEditorEngine::PlayPreviewSound( USoundBase* Sound,  USoundNode* SoundNode
 		AudioComponent->bReverb = false;
 		AudioComponent->bCenterChannelOnly = false;
 		AudioComponent->bIsPreviewSound = true;
-		AudioComponent->Play();	
+		AudioComponent->Play();
+
+		return AudioComponent;
 	}
+
+	return nullptr;
 }
 
 void UEditorEngine::PlayEditorSound( const FString& SoundAssetName )
