Adding missing header guard

diff --git a/Engine/Plugins/Runtime/Metasound/Source/MetasoundStandardNodes/Public/MetasoundParamHelper.h b/Engine/Plugins/Runtime/Metasound/Source/MetasoundStandardNodes/Public/MetasoundParamHelper.h
--- a/Engine/Plugins/Runtime/Metasound/Source/MetasoundStandardNodes/Public/MetasoundParamHelper.h
+++ b/Engine/Plugins/Runtime/Metasound/Source/MetasoundStandardNodes/Public/MetasoundParamHelper.h
@@ -1,5 +1,6 @@
 // Copyright Epic Games, Inc. All Rights Reserved.
 
+#pragma once
 
 /*
 
@@ -51,4 +52,4 @@ Where NUMBER is a number (or whatever) to slot into the format specifier you def
 
 #define METASOUND_GET_VARIABLE_PARAM_NAME(NAME, INDEX) FString::Format(NAME##Name, {INDEX})
 #define METASOUND_GET_VARIABLE_TT(NAME, INDEX)  FText::Format(NAME##Tooltip, INDEX)
-#define METASOUND_GET_VARIABLE_PARAM_NAME_AND_TT(NAME, INDEX)  FString::Format(NAME##Name, {INDEX}), FText::Format(NAME##Tooltip, INDEX)
\ No newline at end of file
+#define METASOUND_GET_VARIABLE_PARAM_NAME_AND_TT(NAME, INDEX)  FString::Format(NAME##Name, {INDEX}), FText::Format(NAME##Tooltip, INDEX)
