Only march Metasound ants when previewing
- Remove redundant graph policy code that was not used (and move implementation that was used to old stub code's .h/.cpp)

diff --git a/Engine/Plugins/Runtime/Metasound/Source/MetasoundEditor/Private/MetasoundEditor.cpp b/Engine/Plugins/Runtime/Metasound/Source/MetasoundEditor/Private/MetasoundEditor.cpp
--- a/Engine/Plugins/Runtime/Metasound/Source/MetasoundEditor/Private/MetasoundEditor.cpp
+++ b/Engine/Plugins/Runtime/Metasound/Source/MetasoundEditor/Private/MetasoundEditor.cpp
@@ -798,6 +798,11 @@ namespace Metasound
 			return true;
 		}
 
+		double FEditor::GetPlayTime() const
+		{
+			return PlayTime;
+		}
+
 		void FEditor::Stop()
 		{
 			check(GEditor);
diff --git a/Engine/Plugins/Runtime/Metasound/Source/MetasoundEditor/Private/MetasoundEditor.h b/Engine/Plugins/Runtime/Metasound/Source/MetasoundEditor/Private/MetasoundEditor.h
--- a/Engine/Plugins/Runtime/Metasound/Source/MetasoundEditor/Private/MetasoundEditor.h
+++ b/Engine/Plugins/Runtime/Metasound/Source/MetasoundEditor/Private/MetasoundEditor.h
@@ -65,6 +65,8 @@ namespace Metasound
 			virtual void RegisterTabSpawners(const TSharedRef<FTabManager>& TabManager) override;
 			virtual void UnregisterTabSpawners(const TSharedRef<FTabManager>& TabManager) override;
 
+			double GetPlayTime() const;
+
 			/** Edits the specified Metasound object */
 			void InitMetasoundEditor(const EToolkitMode::Type Mode, const TSharedPtr<IToolkitHost>& InitToolkitHost, UObject* ObjectToEdit);
 
@@ -254,7 +256,7 @@ namespace Metasound
 
 			/** Widget showing playtime that overlays the graph when previewing */
 			TSharedPtr<STextBlock> PlayTimeWidget;
-			double PlayTime = 0.0f;
+			double PlayTime = 0.0;
 
 			/** Command list for this editor */
 			TSharedPtr<FUICommandList> GraphEditorCommands;
diff --git a/Engine/Plugins/Runtime/Metasound/Source/MetasoundEditor/Private/MetasoundEditorGraphBuilder.cpp b/Engine/Plugins/Runtime/Metasound/Source/MetasoundEditor/Private/MetasoundEditorGraphBuilder.cpp
--- a/Engine/Plugins/Runtime/Metasound/Source/MetasoundEditor/Private/MetasoundEditorGraphBuilder.cpp
+++ b/Engine/Plugins/Runtime/Metasound/Source/MetasoundEditor/Private/MetasoundEditorGraphBuilder.cpp
@@ -9,10 +9,12 @@
 #include "Framework/Notifications/NotificationManager.h"
 #include "Metasound.h"
 #include "MetasoundAssetBase.h"
+#include "MetasoundEditor.h"
 #include "MetasoundEditorGraph.h"
 #include "MetasoundEditorGraphNode.h"
 #include "MetasoundEditorGraphInputNodes.h"
 #include "MetasoundEditorModule.h"
+#include "MetasoundEditorSettings.h"
 #include "MetasoundFrontendQuery.h"
 #include "MetasoundFrontendQuerySteps.h"
 #include "MetasoundFrontendRegistries.h"
@@ -20,6 +22,7 @@
 #include "MetasoundUObjectRegistry.h"
 #include "Modules/ModuleManager.h"
 #include "Templates/Tuple.h"
+#include "Toolkits/ToolkitManager.h"
 #include "Widgets/Notifications/SNotificationList.h"
 
 
@@ -196,6 +199,75 @@ namespace Metasound
 			return NewOutputName;
 		}
 
+		TSharedPtr<FEditor> FGraphBuilder::GetEditorForGraph(const UObject& Metasound)
+		{
+			TSharedPtr<IToolkit> FoundAssetEditor = FToolkitManager::Get().FindEditorForAsset(CastChecked<const UObject>(&Metasound));
+			return StaticCastSharedPtr<FEditor, IToolkit>(FoundAssetEditor);
+		}
+
+		TSharedPtr<FEditor> FGraphBuilder::GetEditorForGraph(const UEdGraph& EdGraph)
+		{
+			const UMetasoundEditorGraph* MetasoundGraph = CastChecked<const UMetasoundEditorGraph>(&EdGraph);
+			return GetEditorForGraph(MetasoundGraph->GetMetasoundChecked());
+		}
+
+		FLinearColor FGraphBuilder::GetPinCategoryColor(const FEdGraphPinType& PinType)
+		{
+			const UMetasoundEditorSettings* Settings = GetDefault<UMetasoundEditorSettings>();
+			check(Settings);
+
+			if (PinType.PinCategory == PinCategoryAudio)
+			{
+				return Settings->AudioPinTypeColor;
+			}
+
+			if (PinType.PinCategory == PinCategoryTrigger)
+			{
+				return Settings->TriggerPinTypeColor;
+			}
+
+			if (PinType.PinCategory == PinCategoryBoolean)
+			{
+				return Settings->BooleanPinTypeColor;
+			}
+
+			if (PinType.PinCategory == PinCategoryFloat)
+			{
+				if (PinType.PinSubCategory == PinSubCategoryTime)
+				{
+					return Settings->TimePinTypeColor;
+				}
+				return Settings->FloatPinTypeColor;
+			}
+
+			if (PinType.PinCategory == PinCategoryInt32)
+			{
+				return Settings->IntPinTypeColor;
+			}
+
+			if (PinType.PinCategory == PinCategoryInt64)
+			{
+				return Settings->Int64PinTypeColor;
+			}
+
+			if (PinType.PinCategory == PinCategoryString)
+			{
+				return Settings->StringPinTypeColor;
+			}
+
+			if (PinType.PinCategory == PinCategoryDouble)
+			{
+				return Settings->DoublePinTypeColor;
+			}
+
+			if (PinType.PinCategory == PinCategoryObject)
+			{
+				return Settings->ObjectPinTypeColor;
+			}
+
+			return Settings->DefaultPinTypeColor;
+		}
+
 		Frontend::FInputHandle FGraphBuilder::GetInputHandleFromPin(const UEdGraphPin* InPin)
 		{
 			using namespace Frontend;
diff --git a/Engine/Plugins/Runtime/Metasound/Source/MetasoundEditor/Private/MetasoundEditorGraphBuilder.h b/Engine/Plugins/Runtime/Metasound/Source/MetasoundEditor/Private/MetasoundEditorGraphBuilder.h
--- a/Engine/Plugins/Runtime/Metasound/Source/MetasoundEditor/Private/MetasoundEditorGraphBuilder.h
+++ b/Engine/Plugins/Runtime/Metasound/Source/MetasoundEditor/Private/MetasoundEditorGraphBuilder.h
@@ -17,6 +17,7 @@ class UEdGraphPin;
 class UMetasound;
 class UMetasoundEditorGraphNode;
 
+struct FEdGraphPinType;
 struct FMetasoundFrontendNodeStyle;
 
 
@@ -27,6 +28,9 @@ namespace Metasound
 
 	namespace Editor
 	{
+		// Forward Declarations
+		class FEditor;
+
 		class FGraphBuilder
 		{
 			static void InitGraphNode(Frontend::FNodeHandle& InNodeHandle, UMetasoundEditorGraphNode* NewGraphNode, UObject& InMetasound);
@@ -51,6 +55,10 @@ namespace Metasound
 			// Adds an EdGraph node to mirror the provided FNodeHandle.
 			static UMetasoundEditorGraphNode* AddNode(UObject& InMetasound, Frontend::FNodeHandle InNodeHandle, FVector2D InLocation, bool bInSelectNewNode = true);
 
+			// Convenience functions for retrieving the editor for the given Metasound/EdGraph
+			static TSharedPtr<FEditor> GetEditorForGraph(const UObject& Metasound);
+			static TSharedPtr<FEditor> GetEditorForGraph(const UEdGraph& EdGraph);
+
 			// Adds a node handle to mirror the provided graph node and binds to it.  Does *NOT* mirror existing EdGraph connections
 			// nor does it remove existing bound Frontend Node (if set) from associated Frontend Graph.
 			static Frontend::FNodeHandle AddNodeHandle(UObject& InMetasound, UMetasoundEditorGraphNode& InGraphNode);
@@ -125,6 +133,9 @@ namespace Metasound
 			// Editor Graph nodes referencing the given variable.
 			static void DeleteVariableNodeHandle(UMetasoundEditorGraphVariable& InVariable);
 
+			// Retrieves the proper pin color for the given PinType
+			static FLinearColor GetPinCategoryColor(const FEdGraphPinType& PinType);
+
 			// Constructs graph with default inputs & outputs.
 			static void ConstructGraph(UObject& InMetasound);
 
diff --git a/Engine/Plugins/Runtime/Metasound/Source/MetasoundEditor/Private/MetasoundEditorGraphConnectionDrawingPolicy.cpp b/Engine/Plugins/Runtime/Metasound/Source/MetasoundEditor/Private/MetasoundEditorGraphConnectionDrawingPolicy.cpp
--- a/Engine/Plugins/Runtime/Metasound/Source/MetasoundEditor/Private/MetasoundEditorGraphConnectionDrawingPolicy.cpp
+++ b/Engine/Plugins/Runtime/Metasound/Source/MetasoundEditor/Private/MetasoundEditorGraphConnectionDrawingPolicy.cpp
@@ -1,83 +1,104 @@
 // Copyright Epic Games, Inc. All Rights Reserved.
-
 #include "MetasoundEditorGraphConnectionDrawingPolicy.h"
 
-#include "Misc/App.h"
-#include "MetasoundEditorGraph.h"
-#include "MetasoundEditorGraphSchema.h"
 #include "Components/AudioComponent.h"
-
 #include "Editor.h"
+#include "MetasoundEditor.h"
+#include "MetasoundEditorGraph.h"
+#include "MetasoundEditorGraphBuilder.h"
+#include "MetasoundEditorGraphSchema.h"
+#include "Misc/App.h"
 
-FConnectionDrawingPolicy* FMetasoundGraphConnectionDrawingPolicyFactory::CreateConnectionPolicy(const class UEdGraphSchema* Schema, int32 InBackLayerID, int32 InFrontLayerID, float ZoomFactor, const class FSlateRect& InClippingRect, class FSlateWindowElementList& InDrawElements, class UEdGraph* InGraphObj) const
-{
-	if (Schema->IsA(UMetasoundEditorGraphSchema::StaticClass()))
-	{
-		return new FMetasoundGraphConnectionDrawingPolicy(InBackLayerID, InFrontLayerID, ZoomFactor, InClippingRect, InDrawElements, InGraphObj);
-	}
-	return nullptr;
-}
-
-FMetasoundGraphConnectionDrawingPolicy::FMetasoundGraphConnectionDrawingPolicy(int32 InBackLayerID, int32 InFrontLayerID, float ZoomFactor, const FSlateRect& InClippingRect, FSlateWindowElementList& InDrawElements, UEdGraph* InGraphObj)
-	: FConnectionDrawingPolicy(InBackLayerID, InFrontLayerID, ZoomFactor, InClippingRect, InDrawElements)
-	, GraphObj(InGraphObj)
-{
-	// Cache off the editor options
-	ActiveColor = Settings->TraceAttackColor;
-	InactiveColor = Settings->TraceReleaseColor;
-
-	ActiveWireThickness = Settings->TraceAttackWireThickness;
-	InactiveWireThickness = Settings->TraceReleaseWireThickness;
-
-	// Don't want to draw ending arrowheads
-	ArrowImage = nullptr;
-	ArrowRadius = FVector2D::ZeroVector;
-}
-
-void FMetasoundGraphConnectionDrawingPolicy::Draw(TMap<TSharedRef<SWidget>, FArrangedWidget>& InPinGeometries, FArrangedChildren& ArrangedNodes)
-{
-	// Draw everything
-	FConnectionDrawingPolicy::Draw(InPinGeometries, ArrangedNodes);
-}
 
-// Give specific editor modes a chance to highlight this connection or darken non-interesting connections
-void FMetasoundGraphConnectionDrawingPolicy::DetermineWiringStyle(UEdGraphPin* OutputPin, UEdGraphPin* InputPin, /*inout*/ FConnectionParams& Params)
+namespace Metasound
 {
-	Params.AssociatedPin1 = OutputPin;
-	Params.AssociatedPin2 = InputPin;
-
-	// Get the schema and grab the default color from it
-	check(OutputPin);
-	check(GraphObj);
-	const UEdGraphSchema* Schema = GraphObj->GetSchema();
-
-	Params.WireColor = Schema->GetPinTypeColor(OutputPin->PinType);
-
-	if (!InputPin)
-	{
-		return;
-	}
-	
-	bool bExecuted = false;
-
-	// Run thru the predecessors, and on
-	if (FExecPairingMap* PredecessorMap = PredecessorNodes.Find(InputPin->GetOwningNode()))
+	namespace Editor
 	{
-		if (FTimePair* Times = PredecessorMap->Find(OutputPin->GetOwningNode()))
+		FConnectionDrawingPolicy* FGraphConnectionDrawingPolicyFactory::CreateConnectionPolicy(
+			const UEdGraphSchema* Schema,
+			int32 InBackLayerID,
+			int32 InFrontLayerID,
+			float ZoomFactor,
+			const FSlateRect& InClippingRect,
+			FSlateWindowElementList& InDrawElements,
+			UEdGraph* InGraphObj) const
 		{
-			bExecuted = true;
+			if (Schema->IsA(UMetasoundEditorGraphSchema::StaticClass()))
+			{
+				return new FGraphConnectionDrawingPolicy(InBackLayerID, InFrontLayerID, ZoomFactor, InClippingRect, InDrawElements, InGraphObj);
+			}
+			return nullptr;
+		}
 
-			Params.WireThickness = ActiveWireThickness;
-			Params.WireColor = ActiveColor;
+		void FGraphConnectionDrawingPolicy::Draw(TMap<TSharedRef<SWidget>, FArrangedWidget>& InPinGeometries, FArrangedChildren& ArrangedNodes)
+		{
+			// Draw everything
+			FConnectionDrawingPolicy::Draw(InPinGeometries, ArrangedNodes);
+		}
 
-			Params.bDrawBubbles = true;
+		void FGraphConnectionDrawingPolicy::DetermineWiringStyle(UEdGraphPin* OutputPin, UEdGraphPin* InputPin, FConnectionParams& OutParams)
+		{
+			using namespace Frontend;
+
+			if (!OutputPin || !InputPin || !GraphObj)
+			{
+				return;
+			}
+
+			OutParams.AssociatedPin1 = InputPin;
+			OutParams.AssociatedPin2 = OutputPin;
+
+			OutParams.WireColor = FGraphBuilder::GetPinCategoryColor(OutputPin->PinType);
+			bool bExecuted = false;
+
+			// Run through the predecessors, and on
+			if (FExecPairingMap* PredecessorMap = PredecessorNodes.Find(OutputPin->GetOwningNode()))
+			{
+				if (FTimePair* Times = PredecessorMap->Find(InputPin->GetOwningNode()))
+				{
+					bExecuted = true;
+
+					OutParams.WireThickness = ActiveWireThickness;
+					OutParams.bDrawBubbles = true;
+				}
+			}
+
+			if (!bExecuted)
+			{
+				if (InputPin->PinType.PinCategory == FGraphBuilder::PinCategoryTrigger)
+				{
+					OutParams.WireThickness = Settings->DefaultExecutionWireThickness;
+				}
+				else if (InputPin->PinType.PinCategory == FGraphBuilder::PinCategoryAudio)
+				{
+					bool bIsActive = false;
+					if (UEdGraphNode* Node = InputPin->GetOwningNode())
+					{
+						if (UEdGraph* Graph = Node->GetGraph())
+						{
+							TSharedPtr<FEditor> MetasoundEditor = FGraphBuilder::GetEditorForGraph(*Graph);
+							if (MetasoundEditor.IsValid())
+							{
+								bIsActive = MetasoundEditor->GetPlayTime() > 0.0;
+							}
+						}
+					}
+
+					OutParams.bDrawBubbles = bIsActive;
+				}
+				else
+				{
+					OutParams.WireThickness = InactiveWireThickness;
+				}
+			}
 		}
-	}
 
-	if (!bExecuted)
-	{
-		// It's not followed, fade it and keep it thin
-		Params.WireColor = InactiveColor;
-		Params.WireThickness = InactiveWireThickness;
-	}
-}
+		FGraphConnectionDrawingPolicy::FGraphConnectionDrawingPolicy(int32 InBackLayerID, int32 InFrontLayerID, float InZoomFactor, const FSlateRect& InClippingRect, FSlateWindowElementList& InDrawElements, UEdGraph* InGraphObj)
+			: FConnectionDrawingPolicy(InBackLayerID, InFrontLayerID, InZoomFactor, InClippingRect, InDrawElements)
+			, GraphObj(InGraphObj)
+		{
+			ActiveWireThickness = Settings->TraceAttackWireThickness;
+			InactiveWireThickness = Settings->TraceReleaseWireThickness;
+		}
+	} // namespace Editor
+} // namespace Metasound
diff --git a/Engine/Plugins/Runtime/Metasound/Source/MetasoundEditor/Private/MetasoundEditorGraphConnectionDrawingPolicy.h b/Engine/Plugins/Runtime/Metasound/Source/MetasoundEditor/Private/MetasoundEditorGraphConnectionDrawingPolicy.h
--- a/Engine/Plugins/Runtime/Metasound/Source/MetasoundEditor/Private/MetasoundEditorGraphConnectionDrawingPolicy.h
+++ b/Engine/Plugins/Runtime/Metasound/Source/MetasoundEditor/Private/MetasoundEditorGraphConnectionDrawingPolicy.h
@@ -1,5 +1,4 @@
 // Copyright Epic Games, Inc. All Rights Reserved.
-
 #pragma once
 
 #include "CoreMinimal.h"
@@ -8,58 +7,67 @@
 #include "EdGraphUtilities.h"
 #include "ConnectionDrawingPolicy.h"
 
+
+// Forward Declarations
 class FSlateRect;
 class FSlateWindowElementList;
 class UEdGraph;
 class UEdGraphSchema;
 
-struct FMetasoundGraphConnectionDrawingPolicyFactory : public FGraphPanelPinConnectionFactory
+namespace Metasound
 {
-public:
-	virtual ~FMetasoundGraphConnectionDrawingPolicyFactory() = default;
-
-	// FGraphPanelPinConnectionFactory
-	virtual class FConnectionDrawingPolicy* CreateConnectionPolicy(const UEdGraphSchema* Schema, int32 InBackLayerID, int32 InFrontLayerID, float ZoomFactor, const FSlateRect& InClippingRect, FSlateWindowElementList& InDrawElements, class UEdGraph* InGraphObj) const override;
-	// ~FGraphPanelPinConnectionFactory
-};
-
-
-// This class draws the connections for an UEdGraph using a Metasound schema
-class FMetasoundGraphConnectionDrawingPolicy : public FConnectionDrawingPolicy
-{
-protected:
-	// Times for one execution pair within the current graph
-	struct FTimePair
+	namespace Editor
 	{
-		double PredExecTime;
-		double ThisExecTime;
+		struct FGraphConnectionDrawingPolicyFactory : public FGraphPanelPinConnectionFactory
+		{
+		public:
+			virtual ~FGraphConnectionDrawingPolicyFactory() = default;
 
-		FTimePair()
-			: PredExecTime(0.0)
-			, ThisExecTime(0.0)
+			// FGraphPanelPinConnectionFactory
+			virtual class FConnectionDrawingPolicy* CreateConnectionPolicy(
+				const UEdGraphSchema* Schema,
+				int32 InBackLayerID,
+				int32 InFrontLayerID,
+				float ZoomFactor,
+				const FSlateRect& InClippingRect,
+				FSlateWindowElementList& InDrawElements,
+				UEdGraph* InGraphObj) const override;
+			// ~FGraphPanelPinConnectionFactory
+		};
+
+		// This class draws the connections for an UEdGraph using a SoundCue schema
+		class FGraphConnectionDrawingPolicy : public FConnectionDrawingPolicy
 		{
-		}
-	};
+		protected:
+			// Times for one execution pair within the current graph
+			struct FTimePair
+			{
+				double PredExecTime = 0.0;
+				double ThisExecTime = 0.0;
+			};
+
+			// Map of pairings
+			using FExecPairingMap = TMap<UEdGraphNode*, FTimePair>;
 
-	// Map of pairings
-	using FExecPairingMap = TMap<UEdGraphNode*, FTimePair>;
+			// Map of nodes that preceded before a given node in the execution sequence (one entry for each pairing)
+			TMap<UEdGraphNode*, FExecPairingMap> PredecessorNodes;
 
-	// Map of nodes that preceded before a given node in the execution sequence (one entry for each pairing)
-	TMap<UEdGraphNode*, FExecPairingMap> PredecessorNodes;
+			UEdGraph* GraphObj = nullptr;
 
-	UEdGraph* GraphObj;
+			FLinearColor ActiveColor;
+			FLinearColor InactiveColor;
 
-	FLinearColor ActiveColor;
-	FLinearColor InactiveColor;
+			float ActiveWireThickness;
+			float InactiveWireThickness;
 
-	float ActiveWireThickness;
-	float InactiveWireThickness;
+			void Draw(TMap<TSharedRef<SWidget>, FArrangedWidget>& InPinGeometries, FArrangedChildren& ArrangedNodes);
 
-public:
-	FMetasoundGraphConnectionDrawingPolicy(int32 InBackLayerID, int32 InFrontLayerID, float ZoomFactor, const FSlateRect& InClippingRect, FSlateWindowElementList& InDrawElements, UEdGraph* InGraphObj);
+		public:
+			FGraphConnectionDrawingPolicy(int32 InBackLayerID, int32 InFrontLayerID, float ZoomFactor, const FSlateRect& InClippingRect, FSlateWindowElementList& InDrawElements, UEdGraph* InGraphObj);
 
-	// FConnectionDrawingPolicy interface
-	virtual void DetermineWiringStyle(UEdGraphPin* OutputPin, UEdGraphPin* InputPin, /*inout*/ FConnectionParams& Params) override;
-	virtual void Draw(TMap<TSharedRef<SWidget>, FArrangedWidget>& PinGeometries, FArrangedChildren& ArrangedNodes) override;
-	// End of FConnectionDrawingPolicy interface
-};
+			// FConnectionDrawingPolicy interface
+			virtual void DetermineWiringStyle(UEdGraphPin* OutputPin, UEdGraphPin* InputPin, /*inout*/ FConnectionParams& Params) override;
+			// End of FConnectionDrawingPolicy interface
+		};
+	} // namespace Editor
+} // namespace Metasound
diff --git a/Engine/Plugins/Runtime/Metasound/Source/MetasoundEditor/Private/MetasoundEditorGraphSchema.cpp b/Engine/Plugins/Runtime/Metasound/Source/MetasoundEditor/Private/MetasoundEditorGraphSchema.cpp
--- a/Engine/Plugins/Runtime/Metasound/Source/MetasoundEditor/Private/MetasoundEditorGraphSchema.cpp
+++ b/Engine/Plugins/Runtime/Metasound/Source/MetasoundEditor/Private/MetasoundEditorGraphSchema.cpp
@@ -127,75 +127,6 @@ namespace Metasound
 				return false;
 			}
 
-			TSharedPtr<FEditor> GetEditorForGraph(const UObject& Metasound)
-			{
-				TSharedPtr<IToolkit> FoundAssetEditor = FToolkitManager::Get().FindEditorForAsset(CastChecked<const UObject>(&Metasound));
-				return StaticCastSharedPtr<FEditor, IToolkit>(FoundAssetEditor);
-			}
-
-			TSharedPtr<FEditor> GetEditorForGraph(const UEdGraph& EdGraph)
-			{
-				const UMetasoundEditorGraph* MetasoundGraph = CastChecked<const UMetasoundEditorGraph>(&EdGraph);
-				return GetEditorForGraph(MetasoundGraph->GetMetasoundChecked());
-			}
-
-			FLinearColor GetPinCategoryColor(const FEdGraphPinType& PinType)
-			{
-				const UMetasoundEditorSettings* Settings = GetDefault<UMetasoundEditorSettings>();
-				check(Settings);
-
-				if (PinType.PinCategory == FGraphBuilder::PinCategoryAudio)
-				{
-					return Settings->AudioPinTypeColor;
-				}
-
-				if (PinType.PinCategory == FGraphBuilder::PinCategoryTrigger)
-				{
-					return Settings->TriggerPinTypeColor;
-				}
-
-				if (PinType.PinCategory == FGraphBuilder::PinCategoryBoolean)
-				{
-					return Settings->BooleanPinTypeColor;
-				}
-
-				if (PinType.PinCategory == FGraphBuilder::PinCategoryFloat)
-				{
-					if (PinType.PinSubCategory == FGraphBuilder::PinSubCategoryTime)
-					{
-						return Settings->TimePinTypeColor;
-					}
-					return Settings->FloatPinTypeColor;
-				}
-
-				if (PinType.PinCategory == FGraphBuilder::PinCategoryInt32)
-				{
-					return Settings->IntPinTypeColor;
-				}
-
-				if (PinType.PinCategory == FGraphBuilder::PinCategoryInt64)
-				{
-					return Settings->Int64PinTypeColor;
-				}
-
-				if (PinType.PinCategory == FGraphBuilder::PinCategoryString)
-				{
-					return Settings->StringPinTypeColor;
-				}
-
-				if (PinType.PinCategory == FGraphBuilder::PinCategoryDouble)
-				{
-					return Settings->DoublePinTypeColor;
-				}
-
-				if (PinType.PinCategory == FGraphBuilder::PinCategoryObject)
-				{
-					return Settings->ObjectPinTypeColor;
-				}
-
-				return Settings->DefaultPinTypeColor;
-			}
-
 			struct FDataTypeActionQuery
 			{
 				FGraphActionMenuBuilder& ActionMenuBuilder;
@@ -227,74 +158,6 @@ namespace Metasound
 				}
 			}
 		}
-
-		FConnectionDrawingPolicy* FGraphConnectionDrawingPolicyFactory::CreateConnectionPolicy(
-			const UEdGraphSchema* InSchema,
-			int32 InBackLayerID,
-			int32 InFrontLayerID,
-			float InZoomFactor,
-			const FSlateRect& InClippingRect,
-			FSlateWindowElementList& InDrawElements,
-			UEdGraph* InGraphObj) const
-		{
-			if (InSchema->IsA(UMetasoundEditorGraphSchema::StaticClass()))
-			{
-				return new FGraphConnectionDrawingPolicy(InBackLayerID, InFrontLayerID, InZoomFactor, InClippingRect, InDrawElements, InGraphObj);
-			}
-			return nullptr;
-		}
-
-		FGraphConnectionDrawingPolicy::FGraphConnectionDrawingPolicy(int32 InBackLayerID, int32 InFrontLayerID, float InZoomFactor, const FSlateRect& InClippingRect, FSlateWindowElementList& InDrawElements, UEdGraph* InGraphObj)
-			: FConnectionDrawingPolicy(InBackLayerID, InFrontLayerID, InZoomFactor, InClippingRect, InDrawElements)
-			, GraphObj(InGraphObj)
-		{
-			ActiveWireThickness = Settings->TraceAttackWireThickness;
-			InactiveWireThickness = Settings->TraceReleaseWireThickness;
-		}
-
-		void FGraphConnectionDrawingPolicy::DetermineWiringStyle(UEdGraphPin* OutputPin, UEdGraphPin* InputPin, FConnectionParams& OutParams)
-		{
-			using namespace Frontend;
-
-			if (!OutputPin || !InputPin || !GraphObj)
-			{
-				return;
-			}
-
-			OutParams.AssociatedPin1 = InputPin;
-			OutParams.AssociatedPin2 = OutputPin;
-
-			OutParams.WireColor = SchemaPrivate::GetPinCategoryColor(OutputPin->PinType);
-			bool bExecuted = false;
-
-			// Run through the predecessors, and on
-			if (FExecPairingMap* PredecessorMap = PredecessorNodes.Find(OutputPin->GetOwningNode()))
-			{
-				if (FTimePair* Times = PredecessorMap->Find(InputPin->GetOwningNode()))
-				{
-					bExecuted = true;
-
-					OutParams.WireThickness = ActiveWireThickness;
-					OutParams.bDrawBubbles = true;
-				}
-			}
-
-			if (!bExecuted)
-			{
-				if (InputPin->PinType.PinCategory == FGraphBuilder::PinCategoryTrigger)
-				{
-					OutParams.WireThickness = Settings->DefaultExecutionWireThickness;
-				}
-				else if (InputPin->PinType.PinCategory == FGraphBuilder::PinCategoryAudio)
-				{
-					OutParams.bDrawBubbles = true;
-				}
-				else
-				{
-					OutParams.WireThickness = InactiveWireThickness;
-				}
-			}
-		}
 	} // namespace Editor
 } // namespace Metasound
 
@@ -398,7 +261,7 @@ UEdGraphNode* FMetasoundGraphSchemaAction_PromoteToInput::PerformAction(UEdGraph
 
 				if (ensure(SchemaPrivate::TryConnectNewNodeToPin(*EdGraphNode, FromPin)))
 				{
-					TSharedPtr<FEditor> MetasoundEditor = SchemaPrivate::GetEditorForGraph(*ParentGraph);
+					TSharedPtr<FEditor> MetasoundEditor = FGraphBuilder::GetEditorForGraph(*ParentGraph);
 					if (MetasoundEditor.IsValid())
 					{
 						MetasoundEditor->OnInputNameChanged(NewGraphNode->GetNodeID());
@@ -467,7 +330,7 @@ UEdGraphNode* FMetasoundGraphSchemaAction_NewComment::PerformAction(UEdGraph* Pa
 
 	FSlateRect Bounds;
 	FVector2D SpawnLocation = Location;
-	TSharedPtr<FEditor> MetasoundEditor = SchemaPrivate::GetEditorForGraph(*ParentGraph);
+	TSharedPtr<FEditor> MetasoundEditor = FGraphBuilder::GetEditorForGraph(*ParentGraph);
 
 	if (MetasoundEditor.IsValid() && MetasoundEditor->GetBoundsForSelectedNodes(Bounds, 50.0f))
 	{
@@ -483,7 +346,7 @@ UEdGraphNode* FMetasoundGraphSchemaAction_Paste::PerformAction(UEdGraph* ParentG
 {
 	using namespace Metasound::Editor;
 
-	TSharedPtr<FEditor> MetasoundEditor = SchemaPrivate::GetEditorForGraph(*ParentGraph);
+	TSharedPtr<FEditor> MetasoundEditor = FGraphBuilder::GetEditorForGraph(*ParentGraph);
 	if (MetasoundEditor.IsValid())
 	{
 		MetasoundEditor->PasteNodes(&Location);
@@ -570,7 +433,7 @@ void UMetasoundEditorGraphSchema::GetGraphContextActions(FGraphContextMenuBuilde
 	}
 	else
 	{
-		TSharedPtr<FEditor> MetasoundEditor = SchemaPrivate::GetEditorForGraph(*ContextMenuBuilder.CurrentGraph);
+		TSharedPtr<FEditor> MetasoundEditor = FGraphBuilder::GetEditorForGraph(*ContextMenuBuilder.CurrentGraph);
 		if (MetasoundEditor.IsValid() && MetasoundEditor->CanPasteNodes())
 		{
 			TSharedPtr<FMetasoundGraphSchemaAction_Paste> NewAction = MakeShared<FMetasoundGraphSchemaAction_Paste>(FText::GetEmpty(), LOCTEXT("PasteHereAction", "Paste here"), FText::GetEmpty(), 0);
@@ -788,7 +651,7 @@ FText UMetasoundEditorGraphSchema::GetPinDisplayName(const UEdGraphPin* Pin) con
 
 FLinearColor UMetasoundEditorGraphSchema::GetPinTypeColor(const FEdGraphPinType& PinType) const
 {
-	return Metasound::Editor::SchemaPrivate::GetPinCategoryColor(PinType);
+	return Metasound::Editor::FGraphBuilder::GetPinCategoryColor(PinType);
 
 }
 
@@ -991,7 +854,7 @@ void UMetasoundEditorGraphSchema::GetCommentAction(FGraphActionMenuBuilder& Acti
 
 	if (!ActionMenuBuilder.FromPin && CurrentGraph)
 	{
-		TSharedPtr<FEditor> MetasoundEditor = SchemaPrivate::GetEditorForGraph(*CurrentGraph);
+		TSharedPtr<FEditor> MetasoundEditor = FGraphBuilder::GetEditorForGraph(*CurrentGraph);
 		if (MetasoundEditor.IsValid())
 		{
 			const int32 NumSelected = MetasoundEditor->GetNumNodesSelected();
@@ -1008,7 +871,7 @@ int32 UMetasoundEditorGraphSchema::GetNodeSelectionCount(const UEdGraph* Graph)
 {
 	using namespace Metasound::Editor;
 
-	TSharedPtr<FEditor> MetasoundEditor = SchemaPrivate::GetEditorForGraph(*Graph);
+	TSharedPtr<FEditor> MetasoundEditor = FGraphBuilder::GetEditorForGraph(*Graph);
 	if (MetasoundEditor.IsValid())
 	{
 		return MetasoundEditor->GetNumNodesSelected();
diff --git a/Engine/Plugins/Runtime/Metasound/Source/MetasoundEditor/Private/MetasoundEditorModule.cpp b/Engine/Plugins/Runtime/Metasound/Source/MetasoundEditor/Private/MetasoundEditorModule.cpp
--- a/Engine/Plugins/Runtime/Metasound/Source/MetasoundEditor/Private/MetasoundEditorModule.cpp
+++ b/Engine/Plugins/Runtime/Metasound/Source/MetasoundEditor/Private/MetasoundEditorModule.cpp
@@ -16,6 +16,7 @@
 #include "MetasoundDetailCustomization.h"
 #include "MetasoundEditorGraph.h"
 #include "MetasoundEditorGraphBuilder.h"
+#include "MetasoundEditorGraphConnectionDrawingPolicy.h"
 #include "MetasoundEditorGraphNodeFactory.h"
 #include "MetasoundEditorGraphSchema.h"
 #include "MetasoundEditorSettings.h"
diff --git a/Engine/Plugins/Runtime/Metasound/Source/MetasoundEditor/Public/MetasoundEditorGraphSchema.h b/Engine/Plugins/Runtime/Metasound/Source/MetasoundEditor/Public/MetasoundEditorGraphSchema.h
--- a/Engine/Plugins/Runtime/Metasound/Source/MetasoundEditor/Public/MetasoundEditorGraphSchema.h
+++ b/Engine/Plugins/Runtime/Metasound/Source/MetasoundEditor/Public/MetasoundEditorGraphSchema.h
@@ -36,59 +36,6 @@ namespace Metasound
 			FInputFilterFunction InputFilterFunction;
 			FOutputFilterFunction OutputFilterFunction;
 		};
-
-		struct FGraphConnectionDrawingPolicyFactory : public FGraphPanelPinConnectionFactory
-		{
-		public:
-			virtual ~FGraphConnectionDrawingPolicyFactory() = default;
-
-			// FGraphPanelPinConnectionFactory
-			virtual class FConnectionDrawingPolicy* CreateConnectionPolicy(
-				const UEdGraphSchema* Schema,
-				int32 InBackLayerID,
-				int32 InFrontLayerID,
-				float ZoomFactor,
-				const FSlateRect& InClippingRect,
-				FSlateWindowElementList& InDrawElements,
-				UEdGraph* InGraphObj) const override;
-			// ~FGraphPanelPinConnectionFactory
-		};
-
-		// This class draws the connections for an UEdGraph using a SoundCue schema
-		class FGraphConnectionDrawingPolicy : public FConnectionDrawingPolicy
-		{
-		protected:
-			// Times for one execution pair within the current graph
-			struct FTimePair
-			{
-				double PredExecTime;
-				double ThisExecTime;
-
-				FTimePair()
-					: PredExecTime(0.0)
-					, ThisExecTime(0.0)
-				{
-				}
-			};
-
-			// Map of pairings
-			using FExecPairingMap = TMap<UEdGraphNode*, FTimePair>;
-
-			// Map of nodes that preceded before a given node in the execution sequence (one entry for each pairing)
-			TMap<UEdGraphNode*, FExecPairingMap> PredecessorNodes;
-
-			UEdGraph* GraphObj = nullptr;
-
-			float ActiveWireThickness;
-			float InactiveWireThickness;
-
-		public:
-			FGraphConnectionDrawingPolicy(int32 InBackLayerID, int32 InFrontLayerID, float ZoomFactor, const FSlateRect& InClippingRect, FSlateWindowElementList& InDrawElements, UEdGraph* InGraphObj);
-
-			// FConnectionDrawingPolicy interface
-			virtual void DetermineWiringStyle(UEdGraphPin* OutputPin, UEdGraphPin* InputPin, /*inout*/ FConnectionParams& Params) override;
-			// End of FConnectionDrawingPolicy interface
-		};
 	} // namespace Editor
 } // namespace Metasound
 
