- WavePlayer Clean kabaam & Jira Jam
    - Only report wave format errors (i.e. when there is no valid wave or the wave is not mono/stereo)
    - Apply equal-power for stereo upmixing mono files
    - Make headerless
    - Fix log spam when swapping input type to WaveAsset (or any UObject proxy for that matter)

diff --git a/Engine/Plugins/Runtime/Metasound/Source/MetasoundEngine/Private/MetasoundWavePlayerNode.cpp b/Engine/Plugins/Runtime/Metasound/Source/MetasoundEngine/Private/MetasoundWavePlayerNode.cpp
--- a/Engine/Plugins/Runtime/Metasound/Source/MetasoundEngine/Private/MetasoundWavePlayerNode.cpp
+++ b/Engine/Plugins/Runtime/Metasound/Source/MetasoundEngine/Private/MetasoundWavePlayerNode.cpp
@@ -1,51 +1,48 @@
 // Copyright Epic Games, Inc. All Rights Reserved.
-#include "MetasoundWavePlayerNode.h"
 
 #include "AudioResampler.h"
+#include "CoreMinimal.h"
+#include "DSP/BufferVectorOperations.h"
 #include "MetasoundBuildError.h"
 #include "MetasoundExecutableOperator.h"
+#include "MetasoundLog.h"
 #include "MetasoundNodeRegistrationMacro.h"
 #include "MetasoundPrimitives.h"
 #include "MetasoundWave.h"
 #include "MetasoundTrigger.h"
 #include "MetasoundEngineNodesNames.h"
 
-
 #define LOCTEXT_NAMESPACE "MetasoundWaveNode"
 
-// static const int32 SMOOTH = -1;
-// static int32 NoDiscontinuities(const float* start, int32 numframes, const float thresh = 0.3f)
-// {
-// 	for (int32 i = 0; i < numframes - 1; ++i)
-// 	{
-// 		float delta = FMath::Abs(start[i] - start[i + 1]);
-// 		if (delta > thresh)
-// 		{
-// 			return i;
-// 		}
-// 	}
-// 	return SMOOTH; // all good!
-// }
 
 namespace Metasound
 {
-	// WavePlayer custom error 
-	class FWavePlayerError : public FBuildErrorBase
+	class FWavePlayerNode : public FNode
 	{
-	public:
-		FWavePlayerError(const FWavePlayerNode& InNode, FText InErrorDescription)
-			: FBuildErrorBase(ErrorType, InErrorDescription)
+		class FOperatorFactory : public IOperatorFactory
 		{
-			AddNode(InNode);
-		}
+			virtual TUniquePtr<IOperator> CreateOperator(const FCreateOperatorParams& InParams, FBuildErrorArray& OutErrors) override;
+		};
+
+	public:
+		static FVertexInterface DeclareVertexInterface();
+		static const FNodeClassMetadata& GetNodeInfo();
 
-		virtual ~FWavePlayerError() = default;
+		FWavePlayerNode(const FString& InName, const FGuid& InInstanceID);
+		FWavePlayerNode(const FNodeInitData& InInitData);
+		virtual ~FWavePlayerNode() = default;
 
-		static const FName ErrorType;
+		virtual FOperatorFactorySharedRef GetDefaultOperatorFactory() const override;
+		virtual const FVertexInterface& GetVertexInterface() const override;
+		virtual bool SetVertexInterface(const FVertexInterface& InInterface) override;
+		virtual bool IsVertexInterfaceSupported(const FVertexInterface& InInterface) const override;
+
+	private:
+		FOperatorFactorySharedRef Factory;
+		FVertexInterface Interface;
 	};
 
-	const FName FWavePlayerError::ErrorType = FName(TEXT("WavePlayerError"));
-	
+
 	class FWavePlayerOperator : public TExecutableOperator<FWavePlayerOperator>
 	{
 	public:
@@ -77,7 +74,7 @@ namespace Metasound
 			check(OutputSampleRate);
 			check(AudioBufferL->Num() == OutputBlockSizeInFrames && AudioBufferR->Num() == OutputBlockSizeInFrames);
 
-			if (Wave->IsSoundWaveValid())
+			if (IsWaveValid())
 			{
 				CurrentSoundWaveName = (*Wave)->GetFName();
 			}
@@ -106,23 +103,20 @@ namespace Metasound
 			return OutputDataReferences;
 		}
 
-
 		void Execute()
 		{
 			TrigggerOnDone->AdvanceBlock();
 			TrigggerOnLooped->AdvanceBlock();
 
-			// see if we have a new soundwave input
-			FName NewSoundWaveName = FName();
+			// Check if soundwave input is new
 			if (Wave->IsSoundWaveValid())
 			{
-				NewSoundWaveName = (*Wave)->GetFName();
-			}
-
-			if (NewSoundWaveName != CurrentSoundWaveName)
-			{
-				ResetDecoder();
-				CurrentSoundWaveName = NewSoundWaveName;
+				const FName SoundWaveName = (*Wave)->GetFName();
+				if (SoundWaveName != CurrentSoundWaveName)
+				{
+					ResetDecoder();
+					CurrentSoundWaveName = SoundWaveName;
+				}
 			}
 
 			// zero output buffers
@@ -207,7 +201,6 @@ namespace Metasound
 				if (CurrAudioFrame == NextSeekFrame)
 				{
 					ExecuteSeekRequest();
-
 					++SeekTrigIndex;
 				}
 
@@ -229,7 +222,7 @@ namespace Metasound
 
 		void StartPlaying()
 		{
-			ResetDecoder();
+			ResetDecoder(true /* bLogFailures */);
 
 			if (!bIsPlaying)
 			{
@@ -239,17 +232,42 @@ namespace Metasound
 
 		void ExecuteSeekRequest()
 		{
-			// TODO: get this to work w/o full decoder reset
-			// Decoder.SeekToTime(FMath::Max(0.f, *SeekTimeSeconds));
-
-			// in the mean time, using this instead:
+			// TODO: Don't do full decoder reset (as performed below) and instead seek decoder.
+			// ex: Decoder.SeekToTime(FMath::Max(0.f, *SeekTimeSeconds));
 			if (!bIsPlaying)
 			{
 				StartPlaying();
 			}
 		}
 
-		bool ResetDecoder()
+		bool IsWaveValid(bool bReportToLog = false) const
+		{
+			if (!Wave->IsSoundWaveValid())
+			{
+				if (bReportToLog)
+				{
+					UE_LOG(LogMetasound, Error, TEXT("Failed to initialize SoundWave decoder. WavePlayerNode references invalid or missing Wave."));
+				}
+				return false;
+			}
+
+			const FSoundWaveProxyPtr& WaveProxy = Wave->GetSoundWaveProxy();
+
+			// WavePlayerNode currently only supports mono or stereo inputs
+			const int32 NumChannels = WaveProxy->GetNumChannels();
+			if (NumChannels > 2)
+			{
+				if (bReportToLog)
+				{
+					UE_LOG(LogMetasound, Error, TEXT("Failed to initialize SoundWave decoder. WavePlayerNode only supports 2 channels max [%s: %d Channels])"), *WaveProxy->GetFullName(), NumChannels);
+				}
+				return false;
+			}
+
+			return true;
+		}
+
+		bool ResetDecoder(bool bLogFailures = false)
 		{
 			Audio::FSimpleDecoderWrapper::InitParams Params;
 			Params.OutputBlockSizeInFrames = OutputBlockSizeInFrames;
@@ -258,25 +276,25 @@ namespace Metasound
 			Params.InitialPitchShiftSemitones = *PitchShiftSemiTones;
 			Params.StartTimeSeconds = FMath::Max(0.f, *SeekTimeSeconds);
 
-			if (false == Wave->IsSoundWaveValid())
+			if (IsWaveValid(bLogFailures))
 			{
-				return false;
+				const FSoundWaveProxyPtr& WaveProxy = Wave->GetSoundWaveProxy();
+				return Decoder.Initialize(Params, WaveProxy);
 			}
 
-			return Decoder.Initialize(Params, Wave->GetSoundWaveProxy());
+			return false;
 		}
 
 
 		int32 ExecuteInternal(int32 StartFrame, int32 EndFrame)
 		{
-			bool bCanDecodeWave = Wave->IsSoundWaveValid() && ((*Wave)->GetNumChannels() <= 2); // only support mono or stereo inputs
-
 			// note: output is hard-coded to stereo (dual-mono)
 			float* FinalOutputLeft = AudioBufferL->GetData() + StartFrame;
 			float* FinalOutputRight = AudioBufferR->GetData() + StartFrame;
 			const int32 NumOutputFrames = (EndFrame - StartFrame);
 			int32 NumFramesDecoded = 0;
 
+			const bool bCanDecodeWave = IsWaveValid();
 			if (bCanDecodeWave)
 			{
 				ensure(Decoder.CanGenerateAudio());
@@ -285,7 +303,7 @@ namespace Metasound
 				const bool bNeedsUpmix = (NumInputChannels == 1);
 				const bool bNeedsDeinterleave = !bNeedsUpmix;
 
-				const int32 NumSamplesToGenerate = NumOutputFrames * NumInputChannels; // (stereo output)
+				const int32 NumSamplesToGenerate = NumOutputFrames * NumInputChannels;
 
 				PostSrcBuffer.Reset(NumSamplesToGenerate);
 				PostSrcBuffer.AddZeroed(NumSamplesToGenerate);
@@ -308,16 +326,15 @@ namespace Metasound
 
 				if (bNeedsUpmix)
 				{
-					// TODO: attenuate by -3 dB?
-					// copy to Left & Right output buffers
-					FMemory::Memcpy(FinalOutputLeft, PostSrcBufferPtr, sizeof(float) * NumFramesDecoded);
-					FMemory::Memcpy(FinalOutputRight, PostSrcBufferPtr, sizeof(float) * NumFramesDecoded);
+					static const float Sqrt2Over2 = 0.707106781f;
+					Audio::BufferMultiplyByConstant(PostSrcBufferPtr, Sqrt2Over2, FinalOutputLeft, NumFramesDecoded);
+					Audio::BufferMultiplyByConstant(PostSrcBufferPtr, Sqrt2Over2, FinalOutputRight, NumFramesDecoded);
 				}
 				else if (bNeedsDeinterleave)
 				{
 					for (int32 i = 0; i < NumFramesDecoded; ++i)
 					{
-						// de-interleave each stereo frame into output buffers
+						// deinterleave each stereo frame into output buffers
 						FinalOutputLeft[i] = PostSrcBufferPtr[(i << 1)];
 						FinalOutputRight[i] = PostSrcBufferPtr[(i << 1) + 1];
 					}
@@ -350,7 +367,7 @@ namespace Metasound
 		FTriggerWriteRef TrigggerOnDone;
 
 		// source decode
-		TArray<float> PostSrcBuffer;
+		Audio::FAlignedFloatBuffer PostSrcBuffer;
 		Audio::FSimpleDecoderWrapper Decoder;
 
 		FName CurrentSoundWaveName{ };
@@ -379,16 +396,6 @@ namespace Metasound
 		FFloatReadRef PitchShiftSemiTones = InputDataRefs.GetDataReadReferenceOrConstruct<float>(TEXT("PitchShift"));
 		FBoolReadRef IsLooping = InputDataRefs.GetDataReadReferenceOrConstruct<bool>(TEXT("Loop"), false);
 
-		if (!Wave->IsSoundWaveValid())
-		{
-			AddBuildError<FWavePlayerError>(OutErrors, WaveNode, LOCTEXT("NoSoundWave", "No Sound Wave"));
-
-		}
-		else if ((*Wave)->GetNumChannels() > 2)
-		{
-			AddBuildError<FWavePlayerError>(OutErrors, WaveNode, LOCTEXT("WavePlayerCurrentlyOnlySuportsMonoOrStereoAssets", "Wave Player Currently Only Supports Mono Or Stereo Assets"));
-		}
-
 		return MakeUnique<FWavePlayerOperator>(
 			  InParams.OperatorSettings
 			, Wave
diff --git a/Engine/Plugins/Runtime/Metasound/Source/MetasoundEngine/Public/MetasoundWavePlayerNode.h b/Engine/Plugins/Runtime/Metasound/Source/MetasoundEngine/Public/MetasoundWavePlayerNode.h
deleted file mode 100644
diff --git a/Engine/Plugins/Runtime/Metasound/Source/MetasoundFrontend/Private/MetasoundFrontendLiteral.cpp b/Engine/Plugins/Runtime/Metasound/Source/MetasoundFrontend/Private/MetasoundFrontendLiteral.cpp
--- a/Engine/Plugins/Runtime/Metasound/Source/MetasoundFrontend/Private/MetasoundFrontendLiteral.cpp
+++ b/Engine/Plugins/Runtime/Metasound/Source/MetasoundFrontend/Private/MetasoundFrontendLiteral.cpp
@@ -224,82 +224,91 @@ void FMetasoundFrontendLiteral::SetFromLiteral(const Metasound::FLiteral& InLite
 	switch (InLiteral.GetType())
 	{
 		case ELiteralType::None:
-
+		{
 			Set(FMetasoundFrontendLiteral::FDefault{});
-
-			break;
+		}
+		break;
 
 		case ELiteralType::Boolean:
-
+		{
 			Set(InLiteral.Value.Get<bool>());
-
-			break;
+		}
+		break;
 
 		case ELiteralType::Float:
-
+		{
 			Set(InLiteral.Value.Get<float>());
-
-			break;
+		}
+		break;
 
 		case ELiteralType::Integer:
-		
+		{
 			Set(InLiteral.Value.Get<int32>());
-	
-			break;
+		}
+		break;
 
 		case ELiteralType::String:
-		
+		{
 			Set(InLiteral.Value.Get<FString>());
-		
-			break;
+		}
+		break;
 
 		case ELiteralType::UObjectProxy:
+		{
+			// Only error if attempting to retrieve valid UObject from ProxyDataPtr
+			// as this function can safely is used to initialize from defaults (which
+			// is valid as a null proxy can safely correspond to a null UObject ptr).
+			if (InLiteral.Value.Get<Audio::IProxyDataPtr>().IsValid())
 			{
 				UE_LOG(LogMetasound, Error, TEXT("Cannot set UObjectProxy from Metasound::FLiteral"));
-				Set(static_cast<UObject*>(nullptr));
 			}
-
-			break;
+			Set(static_cast<UObject*>(nullptr));
+		}
+		break;
 
 		case ELiteralType::NoneArray:
-			{
-				int32 Num = InLiteral.Value.Get<TArray<FLiteral::FNone>>().Num(); 
-				Set(FMetasoundFrontendLiteral::FDefaultArray{ Num });
-			}
-
-			break;
+		{
+			int32 Num = InLiteral.Value.Get<TArray<FLiteral::FNone>>().Num(); 
+			Set(FMetasoundFrontendLiteral::FDefaultArray{ Num });
+		}
+		break;
 
 		case ELiteralType::BooleanArray:
-
+		{
 			Set(InLiteral.Value.Get<TArray<bool>>());
-
-			break;
+		}
+		break;
 
 		case ELiteralType::IntegerArray:
-
+		{
 			Set(InLiteral.Value.Get<TArray<int32>>());
-
-			break;
+		}
+		break;
 
 		case ELiteralType::FloatArray:
-
+		{
 			Set(InLiteral.Value.Get<TArray<float>>());
-
-			break;
+		}
+		break;
 
 		case ELiteralType::StringArray:
-
+		{
 			Set(InLiteral.Value.Get<TArray<FString>>());
-
-			break;
+		}
+		break;
 
 		case ELiteralType::UObjectProxyArray:
+		{
+			// Only error if attempting to retrieve valid UObject from ProxyDataPtr
+			// as this function can safely is used to initialize from defaults (which
+			// is valid as a null proxy can safely correspond to a null UObject ptr).
+			if (!InLiteral.Value.Get<TArray<Audio::IProxyDataPtr>>().IsEmpty())
 			{
 				UE_LOG(LogMetasound, Error, TEXT("Cannot set UObjectProxy from Metasound::FLiteral"));
-				Set(TArray<UObject*>());
 			}
-
-			break;
+			Set(TArray<UObject*>());
+		}
+		break;
 
 		case ELiteralType::Invalid:
 		default:
diff --git a/Engine/Source/Runtime/SignalProcessing/Private/BufferVectorOperations.cpp b/Engine/Source/Runtime/SignalProcessing/Private/BufferVectorOperations.cpp
--- a/Engine/Source/Runtime/SignalProcessing/Private/BufferVectorOperations.cpp
+++ b/Engine/Source/Runtime/SignalProcessing/Private/BufferVectorOperations.cpp
@@ -75,8 +75,6 @@ namespace Audio
 
 	void BufferMultiplyByConstant(const FAlignedFloatBuffer& InFloatBuffer, float InValue, FAlignedFloatBuffer& OutFloatBuffer)
 	{
-		check(InFloatBuffer.Num() >= 4);
-
 		// Prepare output buffer
 		OutFloatBuffer.Reset();
 		OutFloatBuffer.AddUninitialized(InFloatBuffer.Num());
@@ -94,15 +92,13 @@ namespace Audio
 
 	void BufferMultiplyByConstant(const float* RESTRICT InFloatBuffer, float InValue, float* RESTRICT OutFloatBuffer, const int32 InNumSamples)
 	{
-		check(InNumSamples >= 4);
-		RestrictedPtrAliasCheck(InFloatBuffer, OutFloatBuffer, InNumSamples);
-
 #if !AUDIO_USE_SIMD
 		for (int32 i = 0; i < InNumSamples; ++i)
 		{
 			OutFloatBuffer[i] = InValue * InFloatBuffer[i];
 		}
 #else
+		RestrictedPtrAliasCheck(InFloatBuffer, OutFloatBuffer, InNumSamples);
 
 		// Can only SIMD on multiple of 4 buffers, we'll do normal multiples on last bit
 		const int32 NumSamplesRemaining = InNumSamples % 4;
