IWYU Fixes
#rnx

diff --git a/Engine/Plugins/Runtime/Metasound/Source/MetasoundEditor/Private/MetasoundNodeDetailCustomization.cpp b/Engine/Plugins/Runtime/Metasound/Source/MetasoundEditor/Private/MetasoundNodeDetailCustomization.cpp
--- a/Engine/Plugins/Runtime/Metasound/Source/MetasoundEditor/Private/MetasoundNodeDetailCustomization.cpp
+++ b/Engine/Plugins/Runtime/Metasound/Source/MetasoundEditor/Private/MetasoundNodeDetailCustomization.cpp
@@ -15,6 +15,7 @@
 #include "MetasoundEditorGraphBuilder.h"
 #include "MetasoundEditorGraphNode.h"
 #include "MetasoundEditorGraphInputNodes.h"
+#include "MetasoundEditorModule.h"
 #include "MetasoundFrontend.h"
 #include "MetasoundFrontendController.h"
 #include "MetasoundFrontendRegistries.h"
@@ -28,6 +29,8 @@
 #include "Templates/SharedPointer.h"
 #include "UObject/WeakObjectPtr.h"
 #include "UObject/WeakObjectPtrTemplates.h"
+#include "Widgets/Images/SImage.h"
+#include "Widgets/Input/SButton.h"
 #include "Widgets/Input/SCheckBox.h"
 #include "Widgets/Input/SMultiLineEditableTextBox.h"
 #include "Widgets/Input/STextComboBox.h"
@@ -35,7 +38,6 @@
 #include "Widgets/SToolTip.h"
 #include "Widgets/SWidget.h"
 #include "Widgets/Text/STextBlock.h"
-#include "MetasoundEditorModule.h"
 
 
 #define LOCTEXT_NAMESPACE "MetasoundEditor"
