Fix for when deleting an item in the Metasound Action Menu, the inspector doesn't refresh
Fix for missing literal on WaveAsset inputs

diff --git a/Engine/Plugins/Runtime/Metasound/Source/MetasoundEditor/Private/MetasoundEditor.cpp b/Engine/Plugins/Runtime/Metasound/Source/MetasoundEditor/Private/MetasoundEditor.cpp
--- a/Engine/Plugins/Runtime/Metasound/Source/MetasoundEditor/Private/MetasoundEditor.cpp
+++ b/Engine/Plugins/Runtime/Metasound/Source/MetasoundEditor/Private/MetasoundEditor.cpp
@@ -5,6 +5,7 @@
 #include "AudioDevice.h"
 #include "AudioMeterStyle.h"
 #include "Components/AudioComponent.h"
+#include "DetailLayoutBuilder.h"
 #include "EdGraph/EdGraphNode.h"
 #include "EdGraphUtilities.h"
 #include "Editor.h"
@@ -38,14 +39,13 @@
 #include "Widgets/Docking/SDockTab.h"
 #include "Widgets/Images/SImage.h"
 #include "Widgets/Input/SButton.h"
+#include "Widgets/SOverlay.h"
 #include "Widgets/Notifications/SNotificationList.h"
 #include "Widgets/SBoxPanel.h"
 #include "Templates/SharedPointer.h"
 
 // This needs to be moved to public directory
 #include "../../GraphEditor/Private/GraphActionNode.h"
-#include "../../PropertyEditor/Public/DetailLayoutBuilder.h"
-#include "../../../Runtime/SlateCore/Public/Widgets/SOverlay.h"
 
 
 #define LOCTEXT_NAMESPACE "MetasoundEditor"
@@ -1067,6 +1067,85 @@ namespace Metasound
 			}
 		}
 
+		void FEditor::DeleteInterfaceItem(TSharedPtr<FMetasoundGraphNodeSchemaAction> ActionToDelete, UMetasoundEditorGraph* Graph)
+		{
+			using namespace Metasound::Frontend;
+
+			UMetasoundEditorGraphVariable* Variable = ActionToDelete->GetVariable();
+
+			FNodeHandle VariableHandle = Variable->GetNodeHandle();
+			const FGuid IDToDelete = VariableHandle->GetID();
+
+			struct FNameNodeIDPair
+			{
+				FName Name;
+				FGuid ID;
+			};
+
+			auto GetNameToSelect = [IDToDelete](const TArray<FConstNodeHandle>& Handles)
+			{
+				int32 IndexToDelete = -1;
+				bool bGetNextValidName = false;
+				for (int32 i = 0; i < Handles.Num(); ++i)
+				{
+					if (Handles[i]->GetID() == IDToDelete)
+					{
+						IndexToDelete = i;
+						break;
+					}
+				}
+
+				if (IndexToDelete >= 0)
+				{
+					for (int32 i = IndexToDelete + 1; i < Handles.Num(); ++i)
+					{
+						if (Handles[i]->GetNodeStyle().Display.Visibility != EMetasoundFrontendNodeStyleDisplayVisibility::Hidden)
+						{
+							const FName NameToSelect(*(Handles[i]->GetDisplayName().ToString()));
+							return FNameNodeIDPair { NameToSelect, Handles[i]->GetID() };
+						}
+					}
+
+					for (int32 i = IndexToDelete - 1; i >= 0; --i)
+					{
+						if (Handles[i]->GetNodeStyle().Display.Visibility != EMetasoundFrontendNodeStyleDisplayVisibility::Hidden)
+						{
+							FName NameToSelect = (*(Handles[i]->GetDisplayName().ToString()));
+							return FNameNodeIDPair { NameToSelect, Handles[i]->GetID() };
+						}
+					}
+				}
+
+				return FNameNodeIDPair();
+			};
+
+			const TArray<FConstNodeHandle> InputHandles = VariableHandle->GetOwningGraph()->GetConstInputNodes();
+			FNameNodeIDPair NameIDPair = GetNameToSelect(InputHandles);
+			int32 SectionId = static_cast<int32>(ENodeSection::Inputs);
+			if (NameIDPair.Name.IsNone())
+			{
+				SectionId = static_cast<int32>(ENodeSection::Outputs);
+				const TArray<FConstNodeHandle> OutputHandles = VariableHandle->GetOwningGraph()->GetConstOutputNodes();
+				NameIDPair = GetNameToSelect(OutputHandles);
+			}
+
+			FGraphBuilder::DeleteVariableNodeHandle(*Variable);
+			Graph->RemoveVariable(*Variable);
+			MetasoundInterfaceMenu->RefreshAllActions(true /* bPreserveExpansion */);
+
+			if (!NameIDPair.Name.IsNone())
+			{
+				if (MetasoundInterfaceMenu->SelectItemByName(NameIDPair.Name, ESelectInfo::Direct, SectionId))
+				{
+					if (UMetasoundEditorGraphVariable* VariableToSelect = Graph->FindVariable(NameIDPair.ID))
+					{
+						const TArray<UObject*> VariablesToSelect { VariableToSelect };
+						SetSelection(VariablesToSelect);
+					}
+				}
+			}
+		}
+
 		void FEditor::DeleteSelected()
 		{
 			using namespace Frontend;
@@ -1115,64 +1194,7 @@ namespace Metasound
 
 					if (ActionToDelete.IsValid())
 					{
-						UMetasoundEditorGraphVariable* Variable = ActionToDelete->GetVariable();
-
-						FNodeHandle VariableHandle = Variable->GetNodeHandle();
-						const FGuid IDToDelete = VariableHandle->GetID();
-
-						auto GetNameToSelect = [IDToDelete](const TArray<FConstNodeHandle>& Handles)
-						{
-							int32 IndexToDelete = -1;
-							bool bGetNextValidName = false;
-							for (int32 i = 0; i < Handles.Num(); ++i)
-							{
-								if (Handles[i]->GetID() == IDToDelete)
-								{
-									IndexToDelete = i;
-									break;
-								}
-							}
-
-							if (IndexToDelete >= 0)
-							{
-								for (int32 i = IndexToDelete + 1; i < Handles.Num(); ++i)
-								{
-									if (Handles[i]->GetNodeStyle().Display.Visibility != EMetasoundFrontendNodeStyleDisplayVisibility::Hidden)
-									{
-										return FName(*(Handles[i]->GetDisplayName().ToString()));
-									}
-								}
-
-								for (int32 i = IndexToDelete - 1; i >= 0; --i)
-								{
-									if (Handles[i]->GetNodeStyle().Display.Visibility != EMetasoundFrontendNodeStyleDisplayVisibility::Hidden)
-									{
-										return FName(*(Handles[i]->GetDisplayName().ToString()));
-									}
-								}
-							}
-
-							return FName();
-						};
-
-						const TArray<FConstNodeHandle> InputHandles = VariableHandle->GetOwningGraph()->GetConstInputNodes();
-						FName NameToSelect = GetNameToSelect(InputHandles);
-						int32 SectionId = static_cast<int32>(ENodeSection::Inputs);
-						if (NameToSelect.IsNone())
-						{
-							SectionId = static_cast<int32>(ENodeSection::Outputs);
-							const TArray<FConstNodeHandle> OutputHandles = VariableHandle->GetOwningGraph()->GetConstOutputNodes();
-							NameToSelect = GetNameToSelect(OutputHandles);
-						}
-
-						FGraphBuilder::DeleteVariableNodeHandle(*Variable);
-						Graph->RemoveVariable(*Variable);
-						MetasoundInterfaceMenu->RefreshAllActions(true /* bPreserveExpansion */);
-
-						if (!NameToSelect.IsNone())
-						{
-							MetasoundInterfaceMenu->SelectItemByName(NameToSelect, ESelectInfo::Direct, SectionId);
-						}
+						DeleteInterfaceItem(ActionToDelete, Graph);
 					}
 					return;
 				}
diff --git a/Engine/Plugins/Runtime/Metasound/Source/MetasoundEditor/Private/MetasoundEditor.h b/Engine/Plugins/Runtime/Metasound/Source/MetasoundEditor/Private/MetasoundEditor.h
--- a/Engine/Plugins/Runtime/Metasound/Source/MetasoundEditor/Private/MetasoundEditor.h
+++ b/Engine/Plugins/Runtime/Metasound/Source/MetasoundEditor/Private/MetasoundEditor.h
@@ -27,25 +27,30 @@
 
 
 // Forward Declarations
-class IDetailsView;
-class IToolkitHost;
+class FTabManager;
 class SDockableTab;
 class SGraphEditor;
 class SMetasoundPalette;
 class FSlateRect;
-class FTabManager;
+class IDetailsView;
+class IToolkitHost;
 class SVerticalBox;
 class UEdGraphNode;
 class UMetasound;
+class UMetasoundEditorGraph;
 
 struct FGraphActionNode;
 struct FMeterResults;
 struct FPropertyChangedEvent;
 
+
 namespace Metasound
 {
 	namespace Editor
 	{
+		// Forward Declarations
+		class FMetasoundGraphNodeSchemaAction;
+
 		/* Enums to use when grouping the members in the list panel. Enum order dictates visible order. */
 		enum class ENodeSection : uint8
 		{
@@ -142,6 +147,8 @@ namespace Metasound
 			/** Deletes from the Metasound Menu (i.e. input or output) if in focus, or the currently selected nodes if the graph editor is in focus. */
 			void DeleteSelected();
 
+			void DeleteInterfaceItem(TSharedPtr<FMetasoundGraphNodeSchemaAction> ActionToDelete, UMetasoundEditorGraph* Graph);
+
 			/** Delete the currently selected nodes */
 			void DeleteSelectedNodes();
 
diff --git a/Engine/Plugins/Runtime/Metasound/Source/MetasoundEditor/Private/MetasoundEditorGraphBuilder.cpp b/Engine/Plugins/Runtime/Metasound/Source/MetasoundEditor/Private/MetasoundEditorGraphBuilder.cpp
--- a/Engine/Plugins/Runtime/Metasound/Source/MetasoundEditor/Private/MetasoundEditorGraphBuilder.cpp
+++ b/Engine/Plugins/Runtime/Metasound/Source/MetasoundEditor/Private/MetasoundEditorGraphBuilder.cpp
@@ -469,19 +469,26 @@ namespace Metasound
 						if (UClass* Class = FrontendRegistry->GetLiteralUClassForDataType(TypeName))
 						{
 							FAssetRegistryModule& AssetRegistryModule = FModuleManager::LoadModuleChecked<FAssetRegistryModule>("AssetRegistry");
+
 							FARFilter Filter;
 							Filter.bRecursiveClasses = false;
 							Filter.ObjectPaths.Add(FName(*InInputPin.DefaultValue));
 							Filter.ClassNames.Add(Class->GetFName());
+
 							TArray<FAssetData> AssetData;
 							AssetRegistryModule.Get().GetAssets(Filter, AssetData);
-
-							if (AssetData.Num() > 0)
+							if (!AssetData.IsEmpty())
 							{
 								OutDefaultLiteral.Set(AssetData[0].GetAsset());
+								bObjectFound = true;
 							}
 						}
 					}
+					
+					if (!bObjectFound)
+					{
+						OutDefaultLiteral.Set(static_cast<UObject*>(nullptr));
+					}
 				}
 				break;
 
