A couple of conversion nodes (float to time, time to float, float to audio, audio to float)

diff --git a/Engine/Plugins/Runtime/Metasound/Source/MetasoundStandardNodes/Private/MetasoundConversionNodes.cpp b/Engine/Plugins/Runtime/Metasound/Source/MetasoundStandardNodes/Private/MetasoundConversionNodes.cpp
new file mode 100644
--- /dev/null
+++ b/Engine/Plugins/Runtime/Metasound/Source/MetasoundStandardNodes/Private/MetasoundConversionNodes.cpp
@@ -0,0 +1,284 @@
+// Copyright Epic Games, Inc. All Rights Reserved.
+
+#include "CoreMinimal.h"
+#include "MetasoundNodeRegistrationMacro.h"
+#include "MetasoundFacade.h"
+#include "MetasoundExecutableOperator.h"
+#include "MetasoundPrimitives.h"
+#include "MetasoundStandardNodesNames.h"
+#include "MetasoundTime.h"
+#include "MetasoundAudioBuffer.h"
+#include "MetasoundStandardNodesCategories.h"
+#include "Internationalization/Text.h"
+#include "MetasoundParamHelper.h"
+
+#define LOCTEXT_NAMESPACE "MetasoundStandardNodes"
+
+namespace Metasound
+{
+	namespace ConversionNodeVertexNames
+	{
+		METASOUND_PARAM(InputValue, "In", "Input value A.");
+		METASOUND_PARAM(OutputValue, "Out", "The converted value.");
+	}
+
+	namespace MetasoundConversionNodePrivate
+	{
+		template<typename FromType, typename ToType>
+		struct TConversionNodeSpecialization
+		{
+			bool bSupported = false;
+		};
+
+		template<>
+		struct TConversionNodeSpecialization<float, FTime>
+		{
+			static TDataReadReference<float> CreateInputRef(const FVertexInterface& Interface, const FCreateOperatorParams& InParams)
+			{
+				using namespace ConversionNodeVertexNames;
+				const FDataReferenceCollection& InputCollection = InParams.InputDataReferences;
+				const FInputVertexInterface& InputInterface = Interface.GetInputInterface();
+				return InputCollection.GetDataReadReferenceOrConstructWithVertexDefault<float>(InputInterface, METASOUND_GET_PARAM_NAME(InputValue), InParams.OperatorSettings);
+			}
+
+			static void GetConvertedValue(float InValue, FTime& OutValue)
+			{
+				OutValue = FTime(InValue);
+			}
+		};
+
+		template<>
+		struct TConversionNodeSpecialization<float, FAudioBuffer>
+		{
+			static TDataReadReference<float> CreateInputRef(const FVertexInterface& Interface, const FCreateOperatorParams& InParams)
+			{
+				using namespace ConversionNodeVertexNames;
+				const FDataReferenceCollection& InputCollection = InParams.InputDataReferences;
+				const FInputVertexInterface& InputInterface = Interface.GetInputInterface();
+				return InputCollection.GetDataReadReferenceOrConstructWithVertexDefault<float>(InputInterface, METASOUND_GET_PARAM_NAME(InputValue), InParams.OperatorSettings);
+			}
+
+			static void GetConvertedValue(float InValue, FAudioBuffer& OutValue)
+			{
+				int32 NumSamples = OutValue.Num();
+				float* OutValuePtr = OutValue.GetData();
+				for (int32 i = 0; i < NumSamples; ++i)
+				{
+					OutValuePtr[i] = InValue;
+				}
+			}
+		};
+
+		template<>
+		struct TConversionNodeSpecialization<int32, FTime>
+		{
+			static TDataReadReference<int32> CreateInputRef(const FVertexInterface& Interface, const FCreateOperatorParams& InParams)
+			{
+				using namespace ConversionNodeVertexNames;
+				const FDataReferenceCollection& InputCollection = InParams.InputDataReferences;
+				const FInputVertexInterface& InputInterface = Interface.GetInputInterface();
+				return InputCollection.GetDataReadReferenceOrConstructWithVertexDefault<int32>(InputInterface, METASOUND_GET_PARAM_NAME(InputValue), InParams.OperatorSettings);
+			}
+
+			static void GetConvertedValue(int32 InValue, FTime& OutValue)
+			{
+				OutValue = FTime((float)InValue);
+			}
+		};
+
+		template<>
+		struct TConversionNodeSpecialization<FTime, float>
+		{
+			static TDataReadReference<FTime> CreateInputRef(const FVertexInterface& Interface, const FCreateOperatorParams& InParams)
+			{
+				using namespace ConversionNodeVertexNames;
+				const FDataReferenceCollection& InputCollection = InParams.InputDataReferences;
+				const FInputVertexInterface& InputInterface = Interface.GetInputInterface();
+				return InputCollection.GetDataReadReferenceOrConstructWithVertexDefault<FTime>(InputInterface, METASOUND_GET_PARAM_NAME(InputValue), InParams.OperatorSettings);
+			}
+
+			static void GetConvertedValue(const FTime& InValue, float& OutValue)
+			{
+				OutValue = InValue.GetSeconds();
+			}
+		};
+
+		template<>
+		struct TConversionNodeSpecialization<FTime, int32>
+		{
+			static TDataReadReference<FTime> CreateInputRef(const FVertexInterface& Interface, const FCreateOperatorParams& InParams)
+			{
+				using namespace ConversionNodeVertexNames;
+				const FDataReferenceCollection& InputCollection = InParams.InputDataReferences;
+				const FInputVertexInterface& InputInterface = Interface.GetInputInterface();
+				return InputCollection.GetDataReadReferenceOrConstructWithVertexDefault<FTime>(InputInterface, METASOUND_GET_PARAM_NAME(InputValue), InParams.OperatorSettings);
+			}
+
+			static void GetConvertedValue(const FTime& InValue, int32& OutValue)
+			{
+				OutValue = (int32)InValue.GetSeconds();
+			}
+		};
+
+		template<>
+		struct TConversionNodeSpecialization<FAudioBuffer, float>
+		{
+			static TDataReadReference<FAudioBuffer> CreateInputRef(const FVertexInterface& Interface, const FCreateOperatorParams& InParams)
+			{
+				using namespace ConversionNodeVertexNames;
+				const FDataReferenceCollection& InputCollection = InParams.InputDataReferences;
+				return InputCollection.GetDataReadReferenceOrConstruct<FAudioBuffer>(METASOUND_GET_PARAM_NAME(InputValue), InParams.OperatorSettings);
+			}
+
+			static void GetConvertedValue(const FAudioBuffer& InValue, float& OutValue)
+			{
+				OutValue = 0.0f;
+				int32 NumSamples = InValue.Num();
+				const float* InValuePtr = InValue.GetData();
+				for (int32 i = 0; i < NumSamples; ++i)
+				{
+					OutValue += InValuePtr[i];
+				}
+				OutValue /= NumSamples;
+			}
+		};
+	}
+
+	template<typename FromType, typename ToType>
+	class TConversionOperator : public TExecutableOperator<TConversionOperator<FromType, ToType>>
+	{
+	public:
+		static const FVertexInterface& GetDefaultInterface()
+		{
+			using namespace ConversionNodeVertexNames;
+
+			static FText InputDesc = FText::Format(LOCTEXT("ConvDisplayNamePatternFrom", "Input {0} value."), FText::FromString(GetMetasoundDataTypeString<FromType>()));
+			static FText OutputDesc = FText::Format(LOCTEXT("ConvDisplayNamePatternTo", "Ouput {0} value."), FText::FromString(GetMetasoundDataTypeString<ToType>()));
+
+			static const FVertexInterface DefaultInterface(
+				FInputVertexInterface(
+					TInputDataVertexModel<FromType>(METASOUND_GET_PARAM_NAME(InputValue), InputDesc)
+				),
+				FOutputVertexInterface(
+					TOutputDataVertexModel<ToType>(METASOUND_GET_PARAM_NAME(OutputValue), OutputDesc)
+				)
+			);
+
+			return DefaultInterface;
+		}
+
+		static const FNodeClassMetadata& GetNodeInfo()
+		{
+			auto InitNodeInfo = []() -> FNodeClassMetadata
+			{
+				FString FromTypeString = GetMetasoundDataTypeString<FromType>();
+				FString ToTypeString = GetMetasoundDataTypeString<ToType>();
+				FName ClassName = *FString::Format(TEXT("Conversion{0}To{1}"), { FromTypeString, ToTypeString });
+				FText NodeDisplayName = FText::Format(LOCTEXT("ConverterNodeDisplayName", "{0} To {1}"), FText::FromString(FromTypeString), FText::FromString(ToTypeString));
+				FText NodeDescription = FText::Format(LOCTEXT("ConverterNodeDesc", "Converts from {0} to {1}."), FText::FromString(FromTypeString), FText::FromString(ToTypeString));
+
+				FNodeClassMetadata Info;
+				Info.ClassName = { Metasound::StandardNodes::Namespace, ClassName, TEXT("") };
+				Info.MajorVersion = 1;
+				Info.MinorVersion = 0;
+				Info.DisplayName = NodeDisplayName;
+				Info.Description = NodeDescription;
+				Info.Author = PluginAuthor;
+				Info.PromptIfMissing = PluginNodeMissingPrompt;
+				Info.DefaultInterface = GetDefaultInterface();
+				Info.CategoryHierarchy.Emplace(StandardNodes::Conversions);
+
+				return Info;
+			};
+
+			static const FNodeClassMetadata Info = InitNodeInfo();
+			return Info;
+		}
+
+		static TUniquePtr<IOperator> CreateOperator(const FCreateOperatorParams& InParams, TArray<TUniquePtr<IOperatorBuildError>>& OutErrors)
+		{
+			using namespace MetasoundConversionNodePrivate;
+			TDataReadReference<FromType> InputValue = TConversionNodeSpecialization<FromType, ToType>::CreateInputRef(GetDefaultInterface(), InParams);
+			return MakeUnique<TConversionOperator<FromType, ToType>>(InParams.OperatorSettings, InputValue);
+		}
+
+
+		TConversionOperator(const FOperatorSettings& InSettings, TDataReadReference<FromType> InInputValue)
+			: InputValue(InInputValue)
+			, OutputValue(TDataWriteReferenceFactory<ToType>::CreateAny(InSettings))
+		{
+		}
+
+		virtual ~TConversionOperator() = default;
+
+
+		virtual FDataReferenceCollection GetInputs() const override
+		{
+			using namespace ConversionNodeVertexNames;
+
+			FDataReferenceCollection Inputs;
+			Inputs.AddDataReadReference(METASOUND_GET_PARAM_NAME(InputValue), InputValue);
+
+			return Inputs;
+		}
+
+		virtual FDataReferenceCollection GetOutputs() const override
+		{
+			using namespace ConversionNodeVertexNames;
+
+			FDataReferenceCollection Outputs;
+			Outputs.AddDataReadReference(METASOUND_GET_PARAM_NAME(OutputValue), OutputValue);
+
+			return Outputs;
+		}
+
+		void Execute()
+		{
+			using namespace MetasoundConversionNodePrivate;
+			TConversionNodeSpecialization<FromType, ToType>::GetConvertedValue(*InputValue, *OutputValue);
+		}
+
+	private:
+
+		TDataReadReference<FromType> InputValue;
+		TDataWriteReference<ToType> OutputValue;
+	};
+
+	/** TConversionNode
+	 *
+	 *  Generates a random float value when triggered.
+	 */
+	template<typename FromType, typename ToType>
+	class METASOUNDSTANDARDNODES_API TConversionNode : public FNodeFacade
+	{
+	public:
+		/**
+		 * Constructor used by the Metasound Frontend.
+		 */
+		TConversionNode(const FNodeInitData& InInitData)
+			: FNodeFacade(InInitData.InstanceName, InInitData.InstanceID, TFacadeOperatorClass<TConversionOperator<FromType, ToType>>())
+		{}
+
+		virtual ~TConversionNode() = default;
+	};
+
+	using FConversionFloatToTime = TConversionNode<float, FTime>;
+	METASOUND_REGISTER_NODE(FConversionFloatToTime)
+
+	using FConversionTimeToFloat = TConversionNode<FTime, float>;
+	METASOUND_REGISTER_NODE(FConversionTimeToFloat)
+
+	using FConversionInt32ToTime = TConversionNode<int32, FTime>;
+	METASOUND_REGISTER_NODE(FConversionInt32ToTime)
+
+	using FConversionTimeToInt32 = TConversionNode<FTime, int32>;
+	METASOUND_REGISTER_NODE(FConversionTimeToInt32)
+
+	using FConversionFloatToAudio = TConversionNode<float, FAudioBuffer>;
+	METASOUND_REGISTER_NODE(FConversionFloatToAudio)
+
+	using FConversionAudioToFloat = TConversionNode<FAudioBuffer, float>;
+	METASOUND_REGISTER_NODE(FConversionAudioToFloat)
+}
+
+
diff --git a/Engine/Plugins/Runtime/Metasound/Source/MetasoundStandardNodes/Private/MetasoundStandardNodesCategories.cpp b/Engine/Plugins/Runtime/Metasound/Source/MetasoundStandardNodes/Private/MetasoundStandardNodesCategories.cpp
--- a/Engine/Plugins/Runtime/Metasound/Source/MetasoundStandardNodes/Private/MetasoundStandardNodesCategories.cpp
+++ b/Engine/Plugins/Runtime/Metasound/Source/MetasoundStandardNodes/Private/MetasoundStandardNodesCategories.cpp
@@ -12,6 +12,7 @@ namespace Metasound
 		const FText Filters = { LOCTEXT("Metasound_FiltersCategory", "Filters") };
 		const FText TriggerUtils = { LOCTEXT("Metasound_TriggerUtilsCategory", "Trigger Utilities") };
 		const FText RandomUtils = { LOCTEXT("Metasound_RandomUtilsCategory", "Random Utilities") };
+		const FText Conversions  = { LOCTEXT("Metasound_ConversionsCategory", "Conversions") };
 	}
 }
 
diff --git a/Engine/Plugins/Runtime/Metasound/Source/MetasoundStandardNodes/Public/MetasoundStandardNodesCategories.h b/Engine/Plugins/Runtime/Metasound/Source/MetasoundStandardNodes/Public/MetasoundStandardNodesCategories.h
--- a/Engine/Plugins/Runtime/Metasound/Source/MetasoundStandardNodes/Public/MetasoundStandardNodesCategories.h
+++ b/Engine/Plugins/Runtime/Metasound/Source/MetasoundStandardNodes/Public/MetasoundStandardNodesCategories.h
@@ -12,5 +12,6 @@ namespace Metasound
 		extern const FText Filters;
 		extern const FText TriggerUtils;
 		extern const FText RandomUtils;
+		extern const FText Conversions;
 	}
 }
