Fix Copy/Pasting object default values not propagating in Metasounds Editor

diff --git a/Engine/Plugins/Runtime/Metasound/Source/MetasoundEditor/Private/MetasoundEditorGraphBuilder.cpp b/Engine/Plugins/Runtime/Metasound/Source/MetasoundEditor/Private/MetasoundEditorGraphBuilder.cpp
--- a/Engine/Plugins/Runtime/Metasound/Source/MetasoundEditor/Private/MetasoundEditorGraphBuilder.cpp
+++ b/Engine/Plugins/Runtime/Metasound/Source/MetasoundEditor/Private/MetasoundEditorGraphBuilder.cpp
@@ -470,9 +470,18 @@ namespace Metasound
 						{
 							FAssetRegistryModule& AssetRegistryModule = FModuleManager::LoadModuleChecked<FAssetRegistryModule>("AssetRegistry");
 
+							const FString ClassName = Class->GetName();
+							FString ObjectPath = InInputPin.DefaultValue;
+
+							// Remove class prefix if included in default value path
+							if (ObjectPath.StartsWith(ClassName + TEXT(" ")))
+							{
+								ObjectPath.RemoveFromStart(ClassName + TEXT(" "));
+							}
+
 							FARFilter Filter;
 							Filter.bRecursiveClasses = false;
-							Filter.ObjectPaths.Add(FName(*InInputPin.DefaultValue));
+							Filter.ObjectPaths.Add(*ObjectPath);
 							Filter.ClassNames.Add(Class->GetFName());
 
 							TArray<FAssetData> AssetData;
