Random get node clean

- fixed issue with weights array

diff --git a/Engine/Plugins/Runtime/Metasound/Source/MetasoundFrontend/Private/MetasoundArrayRandomNode.cpp b/Engine/Plugins/Runtime/Metasound/Source/MetasoundFrontend/Private/MetasoundArrayRandomNode.cpp
--- a/Engine/Plugins/Runtime/Metasound/Source/MetasoundFrontend/Private/MetasoundArrayRandomNode.cpp
+++ b/Engine/Plugins/Runtime/Metasound/Source/MetasoundFrontend/Private/MetasoundArrayRandomNode.cpp
@@ -1,85 +1,39 @@
 // Copyright Epic Games, Inc. All Rights Reserved.
 
 #include "MetasoundArrayRandomNode.h"
-
-#define LOCTEXT_NAMESPACE "MetasoundFrontend"
+#include "CoreMinimal.h"
+#include "Internationalization/Text.h"
+#include "MetasoundBuilderInterface.h"
+#include "MetasoundDataFactory.h"
+#include "MetasoundExecutableOperator.h"
+#include "MetasoundFacade.h"
+#include "MetasoundLog.h"
+#include "MetasoundNodeInterface.h"
+#include "MetasoundNodeRegistrationMacro.h"
+#include "MetasoundOperatorInterface.h"
+#include "MetasoundPrimitives.h"
+#include "MetasoundTrigger.h"
+#include "Containers/CircularQueue.h"
+#include "MetasoundArrayNodes.h"
+
+#include <type_traits>
+
+#define LOCTEXT_NAMESPACE "MetasoundFrontend_RandomArrayGet"
 
 namespace Metasound
 {
-	namespace ArrayNodeRandomGetVertexNames
-	{
-		const FString& GetInputTriggerNextName()
-		{
-			static const FString Name = TEXT("Next");
-			return Name;
-		}
-
-		const FString& GetInputTriggerResetName()
-		{
-			static const FString Name = TEXT("Reset");
-			return Name;
-		}
-
-		const FString& GetInputRandomArrayName()
-		{
-			static const FString Name = TEXT("In Array");
-			return Name;
-		}
-
-		const FString& GetInputWeightsName()
-		{
-			static const FString Name = TEXT("Weights");
-			return Name;
-		}
-
-		const FString& GetInputSeedName()
-		{
-			static const FString Name = TEXT("Seed");
-			return Name;
-		}
-
-		const FString& GetInputNoRepeatOrderName()
-		{
-			static const FString Name = TEXT("No Repeats");
-			return Name;
-		}
-
-		const FString& GetInputEnableSharedStateName()
-		{
-			static const FString Name = TEXT("Enable Shared State");
-			return Name;
-		}
-
-		const FString& GetOutputTriggerOnNextName()
-		{
-			static const FString Name = TEXT("On Next");
-			return Name;
-		}
-
-		const FString& GetOutputTriggerOnResetName()
-		{
-			static const FString Name = TEXT("On Reset");
-			return Name;
-		}
-
-		const FString& GetOutputValueName()
-		{
-			static const FString Name = TEXT("Value");
-			return Name;
-		}
-	}
-
-	FArrayRandomGet::FArrayRandomGet(int32 InSeed, int32 InMaxIndex, int32 InNoRepeatOrder)
+	FArrayRandomGet::FArrayRandomGet(int32 InSeed, int32 InMaxIndex, const TArray<float>& InWeights, int32 InNoRepeatOrder)
 	{
-		Init(InSeed, InMaxIndex, InNoRepeatOrder);
+		Init(InSeed, InMaxIndex, InWeights, InNoRepeatOrder);
 	}
 
-	void FArrayRandomGet::Init(int32 InSeed, int32 InMaxIndex, int32 InNoRepeatOrder)
+	void FArrayRandomGet::Init(int32 InSeed, int32 InMaxIndex, const TArray<float>& InWeights, int32 InNoRepeatOrder)
 	{
 		SetSeed(InSeed);
 		MaxIndex = InMaxIndex;
 		SetNoRepeatOrder(InNoRepeatOrder);
 		check(!InNoRepeatOrder || !PreviousIndicesQueue->IsFull());
+		RandomWeights = InWeights;
 	}
 
 	void FArrayRandomGet::SetSeed(int32 InSeed)
@@ -109,15 +63,18 @@ namespace Metasound
 		}
 	}
 
+	void FArrayRandomGet::SetRandomWeights(const TArray<float>& InRandomWeights)
+	{
+		RandomWeights = InRandomWeights;
+	}
+
 	void FArrayRandomGet::ResetSeed()
 	{
 		RandomStream.Reset();
 	}
 
-	// Returns the next random weighted value in the array indices. 
-	int32 FArrayRandomGet::NextValue()
+	float FArrayRandomGet::ComputeTotalWeight()
 	{
-		// First compute the total size of the weights
 		float TotalWeight = 0.0f;
 		if (RandomWeights.Num() > 0)
 		{
@@ -134,12 +91,33 @@ namespace Metasound
 				TotalWeight += RandomWeights[i % RandomWeights.Num()];
 			}
 		}
-		else
+		return TotalWeight;
+	}
+
+	// Returns the next random weighted value in the array indices. 
+	int32 FArrayRandomGet::NextValue()
+	{
+		// First compute the total size of the weights
+		float TotalWeight = ComputeTotalWeight();
+		bool bNoWeights = false;
+		if (TotalWeight == 0.0f)
 		{
-			// If we don't have a random weights array, everything is equal weight so we only need to consider the total number of un-picked indices
-			TotalWeight = (float)(FMath::Max(MaxIndex - PreviousIndices.Num(), 1));
+			// Reset the previous indices if we ran out of choices left after weighting
+			if (PreviousIndices.Num() > 0)
+			{
+				PreviousIndices.Reset();
+				PreviousIndicesQueue->Empty();
+				TotalWeight = ComputeTotalWeight();
+			}
+			else
+			{
+				// If we don't have a random weights array, everything is equal weight so we only need to consider the total number of un-picked indices
+				TotalWeight = (float)(FMath::Max(MaxIndex - PreviousIndices.Num(), 1));
+				bNoWeights = true;
+			}
 		}
 
+
 		// Make a random choice based on the total weight
 		float Choice = RandomStream.FRandRange(0.0f, TotalWeight);
 
@@ -154,13 +132,13 @@ namespace Metasound
 			}
 
 			float NextTotalWeight = TotalWeight;
-			if (RandomWeights.Num() > 0)
+			if (bNoWeights)
 			{
-				NextTotalWeight += RandomWeights[i % RandomWeights.Num()];
+				NextTotalWeight += 1.0f;
 			}
 			else
 			{
-				NextTotalWeight += 1.0f;
+				NextTotalWeight += RandomWeights[i % RandomWeights.Num()];
 			}
 
 			if (Choice >= TotalWeight && Choice < NextTotalWeight)
@@ -203,13 +181,13 @@ namespace Metasound
 		return RGM;
 	}
 
-	void FSharedStateRandomGetManager::InitSharedState(uint32 InSharedStateId, int32 InSeed, int32 InNumElements, int32 InNoRepeatOrder)
+	void FSharedStateRandomGetManager::InitSharedState(uint32 InSharedStateId, int32 InSeed, int32 InNumElements, const TArray<float>& InWeights, int32 InNoRepeatOrder)
 	{
 		FScopeLock Lock(&CritSect);
 
 		if (!RandomGets.Contains(InSharedStateId))
 		{
-			RandomGets.Add(InSharedStateId, MakeUnique<FArrayRandomGet>(InSeed, InNumElements, InNoRepeatOrder));
+			RandomGets.Add(InSharedStateId, MakeUnique<FArrayRandomGet>(InSeed, InNumElements, InWeights, InNoRepeatOrder));
 		}
 	}
 
diff --git a/Engine/Plugins/Runtime/Metasound/Source/MetasoundFrontend/Public/MetasoundArrayRandomNode.h b/Engine/Plugins/Runtime/Metasound/Source/MetasoundFrontend/Public/MetasoundArrayRandomNode.h
--- a/Engine/Plugins/Runtime/Metasound/Source/MetasoundFrontend/Public/MetasoundArrayRandomNode.h
+++ b/Engine/Plugins/Runtime/Metasound/Source/MetasoundFrontend/Public/MetasoundArrayRandomNode.h
@@ -3,55 +3,58 @@
 #pragma once
 
 #include "CoreMinimal.h"
-#include "Internationalization/Text.h"
-#include "MetasoundBuilderInterface.h"
-#include "MetasoundDataFactory.h"
-#include "MetasoundExecutableOperator.h"
 #include "MetasoundFacade.h"
-#include "MetasoundLog.h"
 #include "MetasoundNodeInterface.h"
-#include "MetasoundNodeRegistrationMacro.h"
-#include "MetasoundOperatorInterface.h"
-#include "MetasoundPrimitives.h"
-#include "MetasoundTrigger.h"
-#include "Containers/CircularQueue.h"
+#include "MetasoundExecutableOperator.h"
 #include "MetasoundArrayNodes.h"
+#include "MetasoundTrigger.h"
 
-#include <type_traits>
-
-#define LOCTEXT_NAMESPACE "MetasoundFrontend"
-
+#define LOCTEXT_NAMESPACE "MetasoundFrontend_RandomArrayGet"
 
 namespace Metasound
 {
-	/** Random Get Node Vertex Names */
 	namespace ArrayNodeRandomGetVertexNames
 	{
-		METASOUNDFRONTEND_API const FString& GetInputTriggerNextName();
-		METASOUNDFRONTEND_API const FString& GetInputTriggerResetName();
-		METASOUNDFRONTEND_API const FString& GetInputRandomArrayName();
-		METASOUNDFRONTEND_API const FString& GetInputWeightsName();
-		METASOUNDFRONTEND_API const FString& GetInputSeedName();
-		METASOUNDFRONTEND_API const FString& GetInputNoRepeatOrderName();
-		METASOUNDFRONTEND_API const FString& GetInputEnableSharedStateName();
-
-		METASOUNDFRONTEND_API const FString& GetOutputTriggerOnNextName();
-		METASOUNDFRONTEND_API const FString& GetOutputTriggerOnResetName();
-		METASOUNDFRONTEND_API const FString& GetOutputValueName();
+		static const TCHAR* InputTriggerNextName = TEXT("Next");
+		static const TCHAR* InputTriggerResetName = TEXT("Reset");
+		static const TCHAR* InputRandomArrayName = TEXT("In Array");
+		static const TCHAR* InputWeightsName = TEXT("Weights");
+		static const TCHAR* InputSeedName = TEXT("Seed");
+		static const TCHAR* InputNoRepeatOrderName = TEXT("No Repeats");
+		static const TCHAR* InputEnableSharedStateName = TEXT("Enable Shared State");
+		static const TCHAR* OutputTriggerOnNextName = TEXT("On Next");
+		static const TCHAR* OutputTriggerOnResetName = TEXT("On Reset");
+		static const TCHAR* OutputValueName = TEXT("Value");
+
+		static const FText InputTriggerNextTooltip = LOCTEXT("TriggerNextTooltip", "Trigger to get the next value in the randomized array.");
+		static const FText InputTriggerResetTooltip = LOCTEXT("TriggerResetTooltip", "Trigger to reset the seed for the randomized array.");
+		static const FText InputRandomArrayTooltip = LOCTEXT("RandomArrayTooltip", "Input array to randomized.");
+		static const FText InputWeightsTooltip = LOCTEXT("WeightsTooltip", "Input array of weights to use for random selection. Will repeat if this array is shorter than the input array to select from.");
+		static const FText InputSeedTooltip = LOCTEXT("SeedTooltip", "Seed to use for the the random shuffle.");
+		static const FText InputNoRepeatOrderTooltip = LOCTEXT("NoRepeatOrderTooltip", "The number of elements to track to avoid repeating in a row.");
+		static const FText InputEnableSharedStateTooltip = LOCTEXT("EnableSharedStateTooltip", "Set to enabled shared state across instances of this metasound.");
+		static const FText OutputTriggerOnNextTooltip = LOCTEXT("TriggerOnNextTooltip", "Triggers when the \"Next\" input is triggered.");
+		static const FText OutputTriggerOnResetTooltip = LOCTEXT("TriggerOnResetTooltip", "Triggers when the \"Shuffle\" input is triggered or if the array is auto-shuffled.");
+		static const FText OutputValueTooltip = LOCTEXT("ValueTooltip", "Value of the current shuffled element.");
 	}
 
 	class METASOUNDFRONTEND_API FArrayRandomGet
 	{
 	public:
 		FArrayRandomGet() = default;
-		FArrayRandomGet(int32 InSeed, int32 InMaxIndex, int32 InNoRepeatOrder);
-		void Init(int32 InSeed, int32 InMaxIndex, int32 InNoRepeatOrder);
+		FArrayRandomGet(int32 InSeed, int32 InMaxIndex, const TArray<float>& InWeights, int32 InNoRepeatOrder);
+		~FArrayRandomGet() = default;
+
+		void Init(int32 InSeed, int32 InMaxIndex, const TArray<float>& InWeights, int32 InNoRepeatOrder);
 		void SetSeed(int32 InSeed);
 		void SetNoRepeatOrder(int32 InNoRepeatOrder);
+		void SetRandomWeights(const TArray<float>& InRandomWeights);
 		void ResetSeed();
 		int32 NextValue();
 
 	private:
+		float ComputeTotalWeight();
+
 		// The current index into the array of indicies (wraps between 0 and ShuffleIndices.Num())
 		TArray<int32> PreviousIndices;
 		TUniquePtr<TCircularQueue<int32>> PreviousIndicesQueue;
@@ -70,7 +73,7 @@ namespace Metasound
 	public:
 		static FSharedStateRandomGetManager& Get();
 
-		void InitSharedState(uint32 InSharedStateId, int32 InSeed, int32 InNumElements, int32 InNoRepeatOrder);
+		void InitSharedState(uint32 InSharedStateId, int32 InSeed, int32 InNumElements, const TArray<float>& InWeights, int32 InNoRepeatOrder);
 		int32 NextValue(uint32 InSharedStateId);
 		void SetSeed(uint32 InSharedStateId, int32 InSeed);
 		void SetNoRepeatOrder(uint32 InSharedStateId, int32 InNoRepeatOrder);
@@ -85,15 +88,15 @@ namespace Metasound
 		TMap<uint32, TUniquePtr<FArrayRandomGet>> RandomGets;
 	};
 
-	/** TArrayShuffleOperator shuffles an array on trigger and outputs values sequentially on "next". It avoids repeating shuffled elements and supports auto-shuffling.*/
 	template<typename ArrayType>
 	class TArrayRandomGetOperator : public TExecutableOperator<TArrayRandomGetOperator<ArrayType>>
 	{
 	public:
 		using FArrayDataReadReference = TDataReadReference<ArrayType>;
+		using FArrayWeightReadReference = TDataReadReference<TArray<float>>;
+		using WeightArrayType = TArray<float>;
 		using ElementType = typename MetasoundArrayNodesPrivate::TArrayElementType<ArrayType>::Type;
 		using FElementTypeWriteReference = TDataWriteReference<ElementType>;
-		using WeightArrayType = TArray<float>;
 
 		static const FVertexInterface& GetDefaultInterface()
 		{
@@ -101,18 +104,18 @@ namespace Metasound
 
 			static const FVertexInterface DefaultInterface(
 				FInputVertexInterface(
-					TInputDataVertexModel<FTrigger>(GetInputTriggerNextName(), LOCTEXT("RandomGetArrayOpInputTriggerNextTT", "Trigger to get the next value in the randomized array.")),
-					TInputDataVertexModel<FTrigger>(GetInputTriggerResetName(), LOCTEXT("RandomGetArrayOpInputTriggerResetTT", "Trigger to reset the seed for the randomized array.")),
-					TInputDataVertexModel<ArrayType>(GetInputRandomArrayName(), LOCTEXT("RandomGetArrayOpArrayTT", "Input array to randomized.")),
-					TInputDataVertexModel<WeightArrayType>(GetInputWeightsName(), LOCTEXT("RandomGetOpInputSeedTT", "Input array of weights to use for random selection. Will repeat if this array is shorter than the input array to select from.")),
-					TInputDataVertexModel<int32>(GetInputSeedName(), LOCTEXT("RandomGetOpInputSeedTT", "Seed to use for the the random shuffle."), -1),
-					TInputDataVertexModel<int32>(GetInputNoRepeatOrderName(), LOCTEXT("RandomGetOpInputNoRepeatOrderTT", "The number of elements to track to avoid repeating in a row."), 1),
-					TInputDataVertexModel<bool>(GetInputEnableSharedStateName(), LOCTEXT("RandomGetOpEnableSharedStateNameTT", "Set to enabled shared state across instances of this metasound."), false)
+					TInputDataVertexModel<FTrigger>(InputTriggerNextName, InputTriggerNextTooltip),
+					TInputDataVertexModel<FTrigger>(InputTriggerResetName, InputTriggerResetTooltip),
+					TInputDataVertexModel<ArrayType>(InputRandomArrayName, InputRandomArrayTooltip),
+					TInputDataVertexModel<WeightArrayType>(InputWeightsName, InputWeightsTooltip),
+					TInputDataVertexModel<int32>(InputSeedName, InputSeedTooltip, -1),
+					TInputDataVertexModel<int32>(InputNoRepeatOrderName, InputNoRepeatOrderTooltip, 1),
+					TInputDataVertexModel<bool>(InputEnableSharedStateName, InputEnableSharedStateTooltip, false)
 				),
 				FOutputVertexInterface(
-					TOutputDataVertexModel<FTrigger>(GetOutputTriggerOnNextName(), LOCTEXT("RandomGetOpOutputTriggerOnNextNameTT", "Triggers when the \"Next\" input is triggered.")),
-					TOutputDataVertexModel<FTrigger>(GetOutputTriggerOnResetName(), LOCTEXT("RandomGetOpOutputTriggerOnShuffleNameTT", "Triggers when the \"Shuffle\" input is triggered or if the array is auto-shuffled.")),
-					TOutputDataVertexModel<ElementType>(GetOutputValueName(), LOCTEXT("RandomGetOpOutputValueTT", "Value of the current shuffled element."))
+					TOutputDataVertexModel<FTrigger>(OutputTriggerOnNextName, OutputTriggerOnNextTooltip),
+					TOutputDataVertexModel<FTrigger>(OutputTriggerOnResetName, OutputTriggerOnResetTooltip),
+					TOutputDataVertexModel<ElementType>(OutputValueName, OutputValueTooltip)
 				)
 			);
 
@@ -144,13 +147,13 @@ namespace Metasound
 
 			const FInputVertexInterface& Inputs = InParams.Node.GetVertexInterface().GetInputInterface();
 
-			FTriggerReadRef InTriggerNext = InParams.InputDataReferences.GetDataReadReferenceOrConstructWithVertexDefault<FTrigger>(Inputs, GetInputTriggerNextName(), InParams.OperatorSettings);
-			FTriggerReadRef InTriggerReset = InParams.InputDataReferences.GetDataReadReferenceOrConstructWithVertexDefault<FTrigger>(Inputs, GetInputTriggerResetName(), InParams.OperatorSettings);
-			FArrayDataReadReference InInputArray = InParams.InputDataReferences.GetDataReadReferenceOrConstructWithVertexDefault<ArrayType>(Inputs, GetInputRandomArrayName(), InParams.OperatorSettings);
-			TDataReadReference<WeightArrayType> InInputWeightsArray = InParams.InputDataReferences.GetDataReadReferenceOrConstructWithVertexDefault<WeightArrayType>(Inputs, GetInputRandomArrayName(), InParams.OperatorSettings);
-			FInt32ReadRef InSeedValue = InParams.InputDataReferences.GetDataReadReferenceOrConstructWithVertexDefault<int32>(Inputs, GetInputSeedName(), InParams.OperatorSettings);
-			FInt32ReadRef InNoRepeatOrder = InParams.InputDataReferences.GetDataReadReferenceOrConstructWithVertexDefault<int32>(Inputs, GetInputNoRepeatOrderName(), InParams.OperatorSettings);
-			FBoolReadRef bInEnableSharedState = InParams.InputDataReferences.GetDataReadReferenceOrConstructWithVertexDefault<bool>(Inputs, GetInputEnableSharedStateName(), InParams.OperatorSettings);
+			FTriggerReadRef InTriggerNext = InParams.InputDataReferences.GetDataReadReferenceOrConstructWithVertexDefault<FTrigger>(Inputs, InputTriggerNextName, InParams.OperatorSettings);
+			FTriggerReadRef InTriggerReset = InParams.InputDataReferences.GetDataReadReferenceOrConstructWithVertexDefault<FTrigger>(Inputs, InputTriggerResetName, InParams.OperatorSettings);
+			FArrayDataReadReference InInputArray = InParams.InputDataReferences.GetDataReadReferenceOrConstructWithVertexDefault<ArrayType>(Inputs, InputRandomArrayName, InParams.OperatorSettings);
+			FArrayWeightReadReference InInputWeightsArray = InParams.InputDataReferences.GetDataReadReferenceOrConstructWithVertexDefault<WeightArrayType>(Inputs, InputWeightsName, InParams.OperatorSettings);
+			FInt32ReadRef InSeedValue = InParams.InputDataReferences.GetDataReadReferenceOrConstructWithVertexDefault<int32>(Inputs, InputSeedName, InParams.OperatorSettings);
+			FInt32ReadRef InNoRepeatOrder = InParams.InputDataReferences.GetDataReadReferenceOrConstructWithVertexDefault<int32>(Inputs, InputNoRepeatOrderName, InParams.OperatorSettings);
+			FBoolReadRef bInEnableSharedState = InParams.InputDataReferences.GetDataReadReferenceOrConstructWithVertexDefault<bool>(Inputs, InputEnableSharedStateName, InParams.OperatorSettings);
 
 			return MakeUnique<TArrayRandomGetOperator>(InParams, InTriggerNext, InTriggerReset, InInputArray, InInputWeightsArray, InSeedValue, InNoRepeatOrder, bInEnableSharedState);
 		}
@@ -183,6 +186,8 @@ namespace Metasound
 			const ArrayType& InputArrayRef = *InputArray;
 			int32 ArraySize = InputArrayRef.Num();
 
+			const TArray<float>& InputWeightsArrayRef = *InputWeightsArray;
+
 			if (ArraySize > 0)
 			{
 				if (*bEnableSharedState)
@@ -192,11 +197,11 @@ namespace Metasound
 					check(SharedStateUniqueId != INDEX_NONE);
 
 					FSharedStateRandomGetManager& RGM = FSharedStateRandomGetManager::Get();
-					RGM.InitSharedState(SharedStateUniqueId, PrevSeedValue, ArraySize, PrevNoRepeatOrder);
+					RGM.InitSharedState(SharedStateUniqueId, PrevSeedValue, ArraySize, InputWeightsArrayRef, PrevNoRepeatOrder);
 				}
 				else
 				{
-					ArrayRandomGet = MakeUnique<FArrayRandomGet>(PrevSeedValue, ArraySize, PrevNoRepeatOrder);
+					ArrayRandomGet = MakeUnique<FArrayRandomGet>(PrevSeedValue, ArraySize, InputWeightsArrayRef, PrevNoRepeatOrder);
 				}
 			}
 			else
@@ -212,13 +217,13 @@ namespace Metasound
 			using namespace ArrayNodeRandomGetVertexNames;
 
 			FDataReferenceCollection Inputs;
-			Inputs.AddDataReadReference(GetInputTriggerNextName(), TriggerNext);
-			Inputs.AddDataReadReference(GetInputTriggerResetName(), TriggerReset);
-			Inputs.AddDataReadReference(GetInputRandomArrayName(), InputArray);
-			Inputs.AddDataReadReference(GetInputWeightsName(), InputWeightsArray);
-			Inputs.AddDataReadReference(GetInputSeedName(), SeedValue);
-			Inputs.AddDataReadReference(GetInputNoRepeatOrderName(), NoRepeatOrder);
-			Inputs.AddDataReadReference(GetInputEnableSharedStateName(), bEnableSharedState);
+			Inputs.AddDataReadReference(InputTriggerNextName, TriggerNext);
+			Inputs.AddDataReadReference(InputTriggerResetName, TriggerReset);
+			Inputs.AddDataReadReference(InputRandomArrayName, InputArray);
+			Inputs.AddDataReadReference(InputWeightsName, InputWeightsArray);
+			Inputs.AddDataReadReference(InputSeedName, SeedValue);
+			Inputs.AddDataReadReference(InputNoRepeatOrderName, NoRepeatOrder);
+			Inputs.AddDataReadReference(InputEnableSharedStateName, bEnableSharedState);
 
 			return Inputs;
 		}
@@ -228,9 +233,9 @@ namespace Metasound
 			using namespace ArrayNodeRandomGetVertexNames;
 
 			FDataReferenceCollection Outputs;
-			Outputs.AddDataReadReference(GetOutputTriggerOnNextName(), TriggerOnNext);
-			Outputs.AddDataReadReference(GetOutputTriggerOnResetName(), TriggerOnReset);
-			Outputs.AddDataReadReference(GetOutputValueName(), OutValue);
+			Outputs.AddDataReadReference(OutputTriggerOnNextName, TriggerOnNext);
+			Outputs.AddDataReadReference(OutputTriggerOnResetName, TriggerOnReset);
+			Outputs.AddDataReadReference(OutputValueName, OutValue);
 
 			return Outputs;
 		}
@@ -239,12 +244,12 @@ namespace Metasound
 		{
 			TriggerOnNext->AdvanceBlock();
 			TriggerOnReset->AdvanceBlock();
- 
- 			const ArrayType& InputArrayRef = *InputArray;
- 
- 			// Check for a seed change
- 			if (PrevSeedValue != *SeedValue)
- 			{
+
+			const ArrayType& InputArrayRef = *InputArray;
+
+			// Check for a seed change
+			if (PrevSeedValue != *SeedValue)
+			{
 				PrevSeedValue = *SeedValue;
 
 				if (SharedStateUniqueId != INDEX_NONE)
@@ -260,7 +265,7 @@ namespace Metasound
 			}
 
 			if (PrevNoRepeatOrder != *NoRepeatOrder)
-			{	
+			{
 				PrevNoRepeatOrder = *NoRepeatOrder;
 				if (SharedStateUniqueId != INDEX_NONE)
 				{
@@ -368,4 +373,4 @@ namespace Metasound
 
 }
 
-#undef LOCTEXT_NAMESPACE
+#undef LOCTEXT_NAMESPACE
\ No newline at end of file
