Fixing case where you would have no weights array and a no repeat order

diff --git a/Engine/Plugins/Runtime/Metasound/Source/MetasoundFrontend/Private/MetasoundArrayRandomNode.cpp b/Engine/Plugins/Runtime/Metasound/Source/MetasoundFrontend/Private/MetasoundArrayRandomNode.cpp
--- a/Engine/Plugins/Runtime/Metasound/Source/MetasoundFrontend/Private/MetasoundArrayRandomNode.cpp
+++ b/Engine/Plugins/Runtime/Metasound/Source/MetasoundFrontend/Private/MetasoundArrayRandomNode.cpp
@@ -98,29 +98,31 @@ namespace Metasound
 	int32 FArrayRandomGet::NextValue()
 	{
 		// First compute the total size of the weights
-		float TotalWeight = ComputeTotalWeight();
-		bool bNoWeights = false;
-		if (TotalWeight == 0.0f)
+		bool bHasWeights = RandomWeights.Num() > 0;
+		float TotalWeight = 0.0f;
+		if (bHasWeights)
 		{
-			// Reset the previous indices if we ran out of choices left after weighting
-			if (PreviousIndices.Num() > 0)
+			TotalWeight = ComputeTotalWeight();
+			if (TotalWeight == 0.0f && PreviousIndices.Num() > 0)
 			{
 				PreviousIndices.Reset();
 				PreviousIndicesQueue->Empty();
 				TotalWeight = ComputeTotalWeight();
 			}
-			else
-			{
-				// If we don't have a random weights array, everything is equal weight so we only need to consider the total number of un-picked indices
-				TotalWeight = (float)(FMath::Max(MaxIndex - PreviousIndices.Num(), 1));
-				bNoWeights = true;
-			}
+
+			// Weights might have been set with all 0.0s. If that's the case, we treat as if there were no weights set.
+			bHasWeights = (TotalWeight > 0.0f);
 		}
 
+		if (!bHasWeights)
+		{
+			TotalWeight = (float)(FMath::Max(MaxIndex - PreviousIndices.Num(), 1));
+		}
+		check(TotalWeight > 0.0f);
+	
 
 		// Make a random choice based on the total weight
 		float Choice = RandomStream.FRandRange(0.0f, TotalWeight);
-
 		// Now find the index this choice matches up to
 		TotalWeight = 0.0f;
 		int32 ChosenIndex = INDEX_NONE;
@@ -132,13 +134,14 @@ namespace Metasound
 			}
 
 			float NextTotalWeight = TotalWeight;
-			if (bNoWeights)
+			if (bHasWeights)
 			{
-				NextTotalWeight += 1.0f;
+				check(RandomWeights.Num() > 0);
+				NextTotalWeight += RandomWeights[i % RandomWeights.Num()];
 			}
 			else
 			{
-				NextTotalWeight += RandomWeights[i % RandomWeights.Num()];
+				NextTotalWeight += 1.0f;
 			}
 
 			if (Choice >= TotalWeight && Choice < NextTotalWeight)
