Deleting old ADSR node, added new one.

- Cleaned up AD envelope

diff --git a/Engine/Plugins/Runtime/Metasound/Source/MetasoundStandardNodes/Private/MetasoundADEnvelopeNode.cpp b/Engine/Plugins/Runtime/Metasound/Source/MetasoundStandardNodes/Private/MetasoundADEnvelopeNode.cpp
--- a/Engine/Plugins/Runtime/Metasound/Source/MetasoundStandardNodes/Private/MetasoundADEnvelopeNode.cpp
+++ b/Engine/Plugins/Runtime/Metasound/Source/MetasoundStandardNodes/Private/MetasoundADEnvelopeNode.cpp
@@ -30,7 +30,7 @@ namespace Metasound
 
 		METASOUND_PARAM(OutputOnTrigger, "On Trigger", "Triggers when the envelope is triggered.");
 		METASOUND_PARAM(OutputOnDone, "On Done", "Triggers when the envelope finishes or loops back if looping is enabled.");
-		METASOUND_PARAM(OutpuTADEnvelopeValue, "Out Envelope", "The output value of the envelope.");
+		METASOUND_PARAM(OutputEnvelopeValue, "Out Envelope", "The output value of the envelope.");
 	}
 
 	namespace ADEnvelopeNodePrivate
@@ -54,7 +54,7 @@ namespace Metasound
 
 			// Where the envelope value was when it was triggered
 			float StartingEnvelopeValue = 0.0f;
-			float CurrenTADEnvelopeValue = 0.0f;
+			float CurrentEnvelopeValue = 0.0f;
 
 			bool bLooping = false;
 		};
@@ -67,42 +67,45 @@ namespace Metasound
 		template<>
 		struct TADEnvelope<float>
 		{
-			static void GetNexTADEnvelopeOutput(FEnvState& InState, int32 StartFrame, int32 EndFrame, TArray<int32>& OuFinishedFrames, float& OutADEnvelopeValue)
+			static void GetNextEnvelopeOutput(FEnvState& InState, int32 StartFrame, int32 EndFrame, TArray<int32>& OutFinishedFrames, float& OutEnvelopeValue)
 			{
 				// Don't need to do anything if we're not generating the envelope at the top of the block since this is a block-rate envelope
 				if (StartFrame > 0 || InState.CurrentSampleIndex == INDEX_NONE)
 				{
-					OutADEnvelopeValue = 0.0f;
+					OutEnvelopeValue = 0.0f;
 					return;
 				}
 
-				int32 TotalEnvSampleCount = (InState.AttackSampleCount + InState.DecaySampleCount);
-
 				// We are in attack
 				if (InState.CurrentSampleIndex < InState.AttackSampleCount)
 				{
 					float AttackFraction = (float)InState.CurrentSampleIndex++ / InState.AttackSampleCount;
-					OutADEnvelopeValue = InState.StartingEnvelopeValue + (1.0f - InState.StartingEnvelopeValue) * FMath::Pow(AttackFraction, InState.AttackCurveFactor);
-				}
-				// We are in Decay
-				else if (InState.CurrentSampleIndex < TotalEnvSampleCount)
-				{
-					int32 SampleCountInDecayState = InState.CurrentSampleIndex++ - InState.AttackSampleCount;
-					float DecayFracton = (float)SampleCountInDecayState / InState.DecaySampleCount;
-					OutADEnvelopeValue = 1.0f - FMath::Pow(DecayFracton, InState.DecayCurveFactor);
-				}
-				// We are looping so reset the sample index
-				else if (InState.bLooping)
-				{
-					InState.CurrentSampleIndex = 0;
-					OuFinishedFrames.Add(0);
+					OutEnvelopeValue = InState.StartingEnvelopeValue + (1.0f - InState.StartingEnvelopeValue) * FMath::Pow(AttackFraction, InState.AttackCurveFactor);
 				}
 				else
 				{
-					// Envelope is done
-					InState.CurrentSampleIndex = INDEX_NONE;
-					OutADEnvelopeValue = 0.0f;
-					OuFinishedFrames.Add(0);
+					int32 TotalEnvSampleCount = (InState.AttackSampleCount + InState.DecaySampleCount);
+					
+					// We are in Decay
+					if (InState.CurrentSampleIndex < TotalEnvSampleCount)
+					{
+						int32 SampleCountInDecayState = InState.CurrentSampleIndex++ - InState.AttackSampleCount;
+						float DecayFraction = (float)SampleCountInDecayState / InState.DecaySampleCount;
+						OutEnvelopeValue = 1.0f - FMath::Pow(DecayFraction, InState.DecayCurveFactor);
+					}
+					// We are looping so reset the sample index
+					else if (InState.bLooping)
+					{
+						InState.CurrentSampleIndex = 0;
+						OutFinishedFrames.Add(0);
+					}
+					else
+					{
+						// Envelope is done
+						InState.CurrentSampleIndex = INDEX_NONE;
+						OutEnvelopeValue = 0.0f;
+						OutFinishedFrames.Add(0);
+					}
 				}
 			}
 
@@ -112,17 +115,16 @@ namespace Metasound
 		template<>
 		struct TADEnvelope<FAudioBuffer>
 		{
-			static void GetNexTADEnvelopeOutput(FEnvState& InState, int32 StartFrame, int32 EndFrame, TArray<int32>& OuFinishedFrames, FAudioBuffer& OutADEnvelopeValue)
+			static void GetNextEnvelopeOutput(FEnvState& InState, int32 StartFrame, int32 EndFrame, TArray<int32>& OutFinishedFrames, FAudioBuffer& OutEnvelopeValue)
 			{
 				// If we are not active zero the buffer and early exit
 				if (InState.CurrentSampleIndex == INDEX_NONE)
 				{
-					OutADEnvelopeValue.Zero();
+					OutEnvelopeValue.Zero();
 					return;
 				}
 
-				float* OutEnvPtr = OutADEnvelopeValue.GetData();
-				int32 TotalEnvSampleCount = (InState.AttackSampleCount + InState.DecaySampleCount);
+				float* OutEnvPtr = OutEnvelopeValue.GetData();
 				for (int32 i = StartFrame; i < EndFrame; ++i)
 				{
 					// We are in attack
@@ -130,46 +132,51 @@ namespace Metasound
 					{
 						float AttackFraction = (float)InState.CurrentSampleIndex++ / InState.AttackSampleCount;
 						float EnvValue = FMath::Pow(AttackFraction, InState.AttackCurveFactor);
-						float TargeTADEnvelopeValue = InState.StartingEnvelopeValue + (1.0f - InState.StartingEnvelopeValue) * EnvValue;
-						InState.EnvEase.SetValue(TargeTADEnvelopeValue);
-						InState.CurrenTADEnvelopeValue = InState.EnvEase.GetNextValue();
-						OutEnvPtr[i] = InState.CurrenTADEnvelopeValue;
+						float TargetEnvelopeValue = InState.StartingEnvelopeValue + (1.0f - InState.StartingEnvelopeValue) * EnvValue;
+						InState.EnvEase.SetValue(TargetEnvelopeValue);
+						InState.CurrentEnvelopeValue = InState.EnvEase.GetNextValue();
+						OutEnvPtr[i] = InState.CurrentEnvelopeValue;
 					}
-					// We are in Decay
-					else if (InState.CurrentSampleIndex < TotalEnvSampleCount)
+					else 
 					{
-						int32 SampleCountInDecayState = InState.CurrentSampleIndex++ - InState.AttackSampleCount;
-						float DecayFracton = (float)SampleCountInDecayState / InState.DecaySampleCount;
-						float TargeTADEnvelopeValue = 1.0f - FMath::Pow(DecayFracton, InState.DecayCurveFactor);
-						InState.EnvEase.SetValue(TargeTADEnvelopeValue);
-						InState.CurrenTADEnvelopeValue = InState.EnvEase.GetNextValue();
-						OutEnvPtr[i] = InState.CurrenTADEnvelopeValue;
-					}
-					// We are looping so reset the sample index
-					else if (InState.bLooping)
-					{
-						InState.StartingEnvelopeValue = 0.0f;
-						InState.CurrenTADEnvelopeValue = 0.0f;
-						InState.CurrentSampleIndex = 0;
-						OuFinishedFrames.Add(i);
-					}
-					else
-					{
-						InState.CurrenTADEnvelopeValue = InState.EnvEase.GetNextValue();
-						OutEnvPtr[i] = InState.CurrenTADEnvelopeValue;
+						int32 TotalEnvSampleCount = (InState.AttackSampleCount + InState.DecaySampleCount);
 
-						// Envelope is done
-						if (InState.EnvEase.IsDone())
+						// We are in Decay
+						if (InState.CurrentSampleIndex < TotalEnvSampleCount)
+						{
+							int32 SampleCountInDecayState = InState.CurrentSampleIndex++ - InState.AttackSampleCount;
+							float DecayFracton = (float)SampleCountInDecayState / InState.DecaySampleCount;
+							float TargetEnvelopeValue = 1.0f - FMath::Pow(DecayFracton, InState.DecayCurveFactor);
+							InState.EnvEase.SetValue(TargetEnvelopeValue);
+							InState.CurrentEnvelopeValue = InState.EnvEase.GetNextValue();
+							OutEnvPtr[i] = InState.CurrentEnvelopeValue;
+						}
+						// We are looping so reset the sample index
+						else if (InState.bLooping)
+						{
+							InState.StartingEnvelopeValue = 0.0f;
+							InState.CurrentEnvelopeValue = 0.0f;
+							InState.CurrentSampleIndex = 0;
+							OutFinishedFrames.Add(i);
+						}
+						else
 						{
-							// Zero out the rest of the envelope
-							int32 NumSamplesLeft = EndFrame - i - 1;
-							if (NumSamplesLeft > 0)
+							InState.CurrentEnvelopeValue = InState.EnvEase.GetNextValue();
+							OutEnvPtr[i] = InState.CurrentEnvelopeValue;
+
+							// Envelope is done
+							if (InState.EnvEase.IsDone())
 							{
-								FMemory::Memzero(&OutEnvPtr[i + 1], sizeof(float) * NumSamplesLeft);
+								// Zero out the rest of the envelope
+								int32 NumSamplesLeft = EndFrame - i - 1;
+								if (NumSamplesLeft > 0)
+								{
+									FMemory::Memzero(&OutEnvPtr[i + 1], sizeof(float) * NumSamplesLeft);
+								}
+								InState.CurrentSampleIndex = INDEX_NONE;
+								OutFinishedFrames.Add(i);
+								break;
 							}
-							InState.CurrentSampleIndex = INDEX_NONE;
-							OuFinishedFrames.Add(i);
-							break;
 						}
 					}
 				}
@@ -199,7 +206,7 @@ namespace Metasound
 				FOutputVertexInterface(
 					TOutputDataVertexModel<FTrigger>(METASOUND_GET_PARAM_NAME_AND_TT(OutputOnTrigger)),
 					TOutputDataVertexModel<FTrigger>(METASOUND_GET_PARAM_NAME_AND_TT(OutputOnDone)),
-					TOutputDataVertexModel<ValueType>(METASOUND_GET_PARAM_NAME_AND_TT(OutpuTADEnvelopeValue))
+					TOutputDataVertexModel<ValueType>(METASOUND_GET_PARAM_NAME_AND_TT(OutputEnvelopeValue))
 				)
 			);
 
@@ -261,15 +268,15 @@ namespace Metasound
 			const FFloatReadRef& InAttackCurveFactor,
 			const FFloatReadRef& InDecayeCurveFactor,
 			const FBoolReadRef& bInLooping)
-			: TriggerIn(InTriggerIn)
+			: TriggerAttackIn(InTriggerIn)
 			, AttackTime(InAttackTime)
 			, DecayTime(InDecayTime)
 			, AttackCurveFactor(InAttackCurveFactor)
 			, DecayCurveFactor(InDecayeCurveFactor)
 			, bLooping(bInLooping)
-			, OnEnvelopeGen(TDataWriteReferenceFactory<FTrigger>::CreateAny(InSettings))
+			, OnAttackTrigger(TDataWriteReferenceFactory<FTrigger>::CreateAny(InSettings))
 			, OnDone(TDataWriteReferenceFactory<FTrigger>::CreateAny(InSettings))
-			, OutpuTADEnvelope(TDataWriteReferenceFactory<ValueType>::CreateAny(InSettings))
+			, OutputEnvelope(TDataWriteReferenceFactory<ValueType>::CreateAny(InSettings))
 		{
 			NumFramesPerBlock = InSettings.GetNumFramesPerBlock();
 
@@ -292,7 +299,7 @@ namespace Metasound
 			using namespace ADEnvelopeVertexNames;
 
 			FDataReferenceCollection Inputs;
-			Inputs.AddDataReadReference(METASOUND_GET_PARAM_NAME(InputTrigger), TriggerIn);
+			Inputs.AddDataReadReference(METASOUND_GET_PARAM_NAME(InputTrigger), TriggerAttackIn);
 			Inputs.AddDataReadReference(METASOUND_GET_PARAM_NAME(InputAttackTime), AttackTime);
 			Inputs.AddDataReadReference(METASOUND_GET_PARAM_NAME(InputDecayTime), DecayTime);
 			Inputs.AddDataReadReference(METASOUND_GET_PARAM_NAME(InputAttackCurve), AttackCurveFactor);
@@ -307,9 +314,9 @@ namespace Metasound
 			using namespace ADEnvelopeVertexNames;
 
 			FDataReferenceCollection Outputs;
-			Outputs.AddDataReadReference(METASOUND_GET_PARAM_NAME(OutputOnTrigger), OnEnvelopeGen);
+			Outputs.AddDataReadReference(METASOUND_GET_PARAM_NAME(OutputOnTrigger), OnAttackTrigger);
 			Outputs.AddDataReadReference(METASOUND_GET_PARAM_NAME(OutputOnDone), OnDone);
-			Outputs.AddDataReadReference(METASOUND_GET_PARAM_NAME(OutpuTADEnvelopeValue), OutpuTADEnvelope);
+			Outputs.AddDataReadReference(METASOUND_GET_PARAM_NAME(OutputEnvelopeValue), OutputEnvelope);
 
 			return Outputs;
 		}
@@ -330,18 +337,18 @@ namespace Metasound
 		{
 			using namespace ADEnvelopeNodePrivate;
 
-			OnEnvelopeGen->AdvanceBlock();
+			OnAttackTrigger->AdvanceBlock();
 			OnDone->AdvanceBlock();
 
 			// check for any updates to input params
 			UpdateParams();
 
-			TriggerIn->ExecuteBlock(
+			TriggerAttackIn->ExecuteBlock(
 				// OnPreTrigger
 				[&](int32 StartFrame, int32 EndFrame)
 				{
 					TArray<int32> FinishedFrames;
-					TADEnvelope<ValueType>::GetNexTADEnvelopeOutput(EnvState, StartFrame, EndFrame, FinishedFrames, *OutpuTADEnvelope);
+					TADEnvelope<ValueType>::GetNextEnvelopeOutput(EnvState, StartFrame, EndFrame, FinishedFrames, *OutputEnvelope);
 
 					for (int32 FrameFinished : FinishedFrames)
 					{
@@ -356,34 +363,34 @@ namespace Metasound
 
 					// Set the sample index to the top of the envelope
 					EnvState.CurrentSampleIndex = 0;
-					EnvState.StartingEnvelopeValue = EnvState.CurrenTADEnvelopeValue;
+					EnvState.StartingEnvelopeValue = EnvState.CurrentEnvelopeValue;
 
 					// Generate the output (this will no-op if we're block rate)
 					TArray<int32> FinishedFrames;
-					TADEnvelope<ValueType>::GetNexTADEnvelopeOutput(EnvState, StartFrame, EndFrame, FinishedFrames, *OutpuTADEnvelope);
+					TADEnvelope<ValueType>::GetNextEnvelopeOutput(EnvState, StartFrame, EndFrame, FinishedFrames, *OutputEnvelope);
 					for (int32 FrameFinished : FinishedFrames)
 					{
 						OnDone->TriggerFrame(FrameFinished);
 					}
 
 					// Forward the trigger
-					OnEnvelopeGen->TriggerFrame(StartFrame);
+					OnAttackTrigger->TriggerFrame(StartFrame);
 				}
 			);
 		}
 
 	private:
 
-		FTriggerReadRef TriggerIn;
+		FTriggerReadRef TriggerAttackIn;
 		FTimeReadRef AttackTime;
 		FTimeReadRef DecayTime;
 		FFloatReadRef AttackCurveFactor;
 		FFloatReadRef DecayCurveFactor;
 		FBoolReadRef bLooping;
 
-		FTriggerWriteRef OnEnvelopeGen;
+		FTriggerWriteRef OnAttackTrigger;
 		FTriggerWriteRef OnDone;
-		TDataWriteReference<ValueType> OutpuTADEnvelope;
+		TDataWriteReference<ValueType> OutputEnvelope;
 
 		// This will either be the block rate or sample rate depending on if this is block-rate or audio-rate envelope
 		float SampleRate = 0.0f;
@@ -413,8 +420,8 @@ namespace Metasound
 	using FADEnvelopeNodeFloat = TADEnvelopeNode<float>;
 	METASOUND_REGISTER_NODE(FADEnvelopeNodeFloat)
 
-	using FEnvelopeAudioBuffer = TADEnvelopeNode<FAudioBuffer>;
-	METASOUND_REGISTER_NODE(FEnvelopeAudioBuffer)
+	using FADEnvelopeAudioBuffer = TADEnvelopeNode<FAudioBuffer>;
+	METASOUND_REGISTER_NODE(FADEnvelopeAudioBuffer)
 }
 
 #undef LOCTEXT_NAMESPACE
diff --git a/Engine/Plugins/Runtime/Metasound/Source/MetasoundStandardNodes/Private/MetasoundADSREnvelopeNode.cpp b/Engine/Plugins/Runtime/Metasound/Source/MetasoundStandardNodes/Private/MetasoundADSREnvelopeNode.cpp
new file mode 100644
--- /dev/null
+++ b/Engine/Plugins/Runtime/Metasound/Source/MetasoundStandardNodes/Private/MetasoundADSREnvelopeNode.cpp
@@ -0,0 +1,584 @@
+// Copyright Epic Games, Inc. All Rights Reserved.
+
+#include "CoreMinimal.h"
+#include "Internationalization/Text.h"
+#include "MetasoundExecutableOperator.h"
+#include "MetasoundNodeRegistrationMacro.h"
+#include "MetasoundDataTypeRegistrationMacro.h"
+#include "MetasoundPrimitives.h"
+#include "MetasoundStandardNodesNames.h"
+#include "MetasoundTrigger.h"
+#include "MetasoundTime.h"
+#include "MetasoundAudioBuffer.h"
+#include "Internationalization/Text.h"
+#include "DSP/BufferVectorOperations.h"
+#include "DSP/Dsp.h"
+#include "MetasoundParamHelper.h"
+
+#define LOCTEXT_NAMESPACE "MetasoundStandardNodes_ADSR"
+
+namespace Metasound
+{
+	namespace ADSREnvelopeVertexNames
+	{
+		METASOUND_PARAM(InputAttackTrigger, "Trigger Attack", "Trigger to start the attack phase of the envelope generator.");
+		METASOUND_PARAM(InputReleaseTrigger, "Trigger Release", "Trigger to start the release phase of the envelope generator.");
+
+		METASOUND_PARAM(InputAttackTime, "Attack Time", "Attack time of the envelope.");
+		METASOUND_PARAM(InputDecayTime, "Decay Time", "Decay time of the envelope.");
+		METASOUND_PARAM(InputSustainLevel, "Sustain Level", "The sustain level.");
+		METASOUND_PARAM(InputReleaseTime, "Release Time", "Decay time of the envelope.");
+
+		METASOUND_PARAM(InputAttackCurve, "Attack Curve", "The exponential curve factor of the attack. 1.0 = linear growth, < 1.0 logorithmic growth, > 1.0 exponential growth.");
+		METASOUND_PARAM(InputDecayCurve, "Decay Curve", "The exponential curve factor of the decay. 1.0 = linear decay, < 1.0 exponential decay, > 1.0 logorithmic decay.");
+		METASOUND_PARAM(InputReleaseCurve, "Release Curve", "The exponential curve factor of the release. 1.0 = linear release, < 1.0 exponential release, > 1.0 logorithmic release.");
+
+		METASOUND_PARAM(OutputOnAttackTrigger, "On Attack Triggered", "Triggers when the envelope attack is triggered.");
+		METASOUND_PARAM(OutputOnReleaseTrigger, "On Release Triggered", "Triggers when the envelope attack is triggered.");
+		METASOUND_PARAM(OutputOnDone, "On Done", "Triggers when the envelope finishes.");
+		METASOUND_PARAM(OutputEnvelopeValue, "Out Envelope", "The output value of the envelope.");
+	}
+
+	namespace ADSREnvelopeNodePrivate
+	{
+// 		enum class EEnvState : uint8
+// 		{
+// 			Attack,
+// 			Decay,
+// 			Sustain,
+// 			Release
+// 		};
+// 
+// 		class FEnvSegment
+// 		{		
+// 		public:
+// 			FEnvSegment(EEnvState InState, int32 InSampleCount, float InCurveValue, float InTargetValue)
+// 				: State(InState)
+// 				, SampleCount(InSampleCount)
+// 				, CurveValue(InCurveValue)
+// 				, TargetValue(InTargetValue)
+// 			{}
+// 
+// 			void Update(float InSampleCount, float InTargetValue)
+// 			{
+// 				SampleCount = InSampleCount;
+// 				TargetValue = InTargetValue;
+// 			}
+// 
+// 			void Start()
+// 			{
+// 				CurrentSampleCount = 0;
+// 			}
+// 
+// 			// return true if done in this state
+// 			bool GetNextValue(float& OutValue)
+// 			{
+// 				if ()
+// 			}
+// 
+// 			EEnvState State;
+// 			int32 SampleCount = 0;
+// 			float TargetValue = 0.0f;
+// 			int32 CurrentSampleCount = 0;
+// 
+// 		};
+
+		struct FEnvState
+		{
+			// Where the envelope is. If INDEX_NONE, then the envelope is not triggered
+			int32 CurrentSampleIndex = INDEX_NONE;
+
+			// Number of samples for attack
+			int32 AttackSampleCount = 0;
+
+			// Number of samples for Decay
+			int32 DecaySampleCount = 0;
+
+			// Number of samples for Decay
+			int32 ReleaseSampleCount = 0;
+			
+			// Sustain leve
+			float SustainLevel = 0.0f;
+
+			// Curve factors for attack/decay/release
+			float AttackCurveFactor = 0.0f;
+			float DecayCurveFactor = 0.0f;
+			float ReleaseCurveFactor = 0.0f;
+
+			Audio::FExponentialEase EnvEase;
+
+			// Where the envelope value was when it was triggered
+			float StartingEnvelopeValue = 0.0f;
+			float CurrentEnvelopeValue = 0.0f;
+			float EnvelopeValueAtReleaseStart = 0.0f;
+
+			// If this is set, we are in release mode
+			bool bIsInRelease = false;
+		};
+
+		template<typename ValueType>
+		struct TADSREnvelope
+		{
+		};
+
+		template<>
+		struct TADSREnvelope<float>
+		{
+			static void GetNextEnvelopeOutput(FEnvState& InState, int32 StartFrame, int32 EndFrame, TArray<int32>& OutFinishedFrames, float& OutEnvelopeValue)
+			{
+				// Don't need to do anything if we're not generating the envelope at the top of the block since this is a block-rate envelope
+				if (StartFrame > 0 || InState.CurrentSampleIndex == INDEX_NONE)
+				{
+					OutEnvelopeValue = 0.0f;
+					return;
+				}
+
+				// Sample count to the end of the decay phase
+				int32 DecayEnvSampleCount = (InState.AttackSampleCount + InState.DecaySampleCount);
+
+				// If we are in the release state, jump forward in our sample count
+				if (InState.bIsInRelease)
+				{
+					int32 SampleStartOfRelease = InState.AttackSampleCount + InState.DecaySampleCount;
+					if (InState.CurrentSampleIndex < SampleStartOfRelease)
+					{
+						InState.EnvelopeValueAtReleaseStart = InState.CurrentEnvelopeValue;
+						InState.CurrentSampleIndex = InState.AttackSampleCount + InState.DecaySampleCount;
+					}
+				}
+
+				// We are in attack
+				if (InState.CurrentSampleIndex < InState.AttackSampleCount)
+				{
+					float AttackFraction = (float)InState.CurrentSampleIndex++ / InState.AttackSampleCount;
+					float EnvValue = FMath::Pow(AttackFraction, InState.AttackCurveFactor);
+					float TargeEnvelopeValue = InState.StartingEnvelopeValue + (1.0f - InState.StartingEnvelopeValue) * EnvValue;
+					InState.EnvEase.SetValue(TargeEnvelopeValue);
+					InState.CurrentEnvelopeValue = InState.EnvEase.GetNextValue();
+					OutEnvelopeValue = InState.CurrentEnvelopeValue;
+				}
+				// We are in decay
+				else if (InState.CurrentSampleIndex < DecayEnvSampleCount)
+				{
+					int32 SampleCountInDecayState = InState.CurrentSampleIndex++ - InState.AttackSampleCount;
+					float DecayFracton = (float)SampleCountInDecayState / InState.DecaySampleCount;
+					float TargetEnvelopeValue = 1.0f - (1.0f - InState.SustainLevel) * FMath::Pow(DecayFracton, InState.DecayCurveFactor);
+					InState.EnvEase.SetValue(TargetEnvelopeValue);
+					InState.CurrentEnvelopeValue = InState.EnvEase.GetNextValue();
+					OutEnvelopeValue = InState.CurrentEnvelopeValue;
+				}
+				// We are in sustain
+				else if (!InState.bIsInRelease)
+				{
+					InState.EnvEase.SetValue(InState.SustainLevel);
+					InState.CurrentEnvelopeValue = InState.EnvEase.GetNextValue();
+					OutEnvelopeValue = InState.CurrentEnvelopeValue;
+					InState.EnvelopeValueAtReleaseStart = OutEnvelopeValue;
+				}
+				// We are in release mode or finished
+				else 			
+				{
+					int32 ReleaseEnvSampleCount = (InState.AttackSampleCount + InState.DecaySampleCount + InState.ReleaseSampleCount);
+					// We are in release
+					if (InState.CurrentSampleIndex < ReleaseEnvSampleCount)
+					{
+						int32 SampleCountInReleaseState = InState.CurrentSampleIndex++ - InState.DecaySampleCount - InState.AttackSampleCount;
+						float ReleaseFraction = (float)SampleCountInReleaseState / InState.ReleaseSampleCount;
+						float TargetEnvelopeValue = InState.EnvelopeValueAtReleaseStart * (1.0f - FMath::Pow(ReleaseFraction, InState.ReleaseCurveFactor));
+						InState.EnvEase.SetValue(TargetEnvelopeValue);
+						InState.CurrentEnvelopeValue = InState.EnvEase.GetNextValue();
+						OutEnvelopeValue = InState.CurrentEnvelopeValue;
+					}
+					// We are done
+					else
+					{
+						InState.CurrentSampleIndex = INDEX_NONE;
+						OutEnvelopeValue = 0.0f;
+						OutFinishedFrames.Add(0);
+					}
+				}
+			}
+
+			static bool IsAudio() { return false; }
+		};
+
+		template<>
+		struct TADSREnvelope<FAudioBuffer>
+		{
+			static void GetNextEnvelopeOutput(FEnvState& InState, int32 StartFrame, int32 EndFrame, TArray<int32>& OutFinishedFrames, FAudioBuffer& OutEnvelopeValue)
+			{
+				// If we are not active zero the buffer and early exit
+				if (InState.CurrentSampleIndex == INDEX_NONE)
+				{
+					OutEnvelopeValue.Zero();
+					return;
+				}
+
+				// If we are in the release state, jump forward in our sample count
+				if (InState.bIsInRelease)
+				{
+					int32 SampleStartOfRelease = InState.AttackSampleCount + InState.DecaySampleCount;
+					if (InState.CurrentSampleIndex <= SampleStartOfRelease)
+					{
+						InState.EnvelopeValueAtReleaseStart = InState.CurrentEnvelopeValue;
+						InState.CurrentSampleIndex = SampleStartOfRelease;
+					}
+				}
+
+				float* OutEnvPtr = OutEnvelopeValue.GetData();
+				for (int32 i = StartFrame; i < EndFrame; ++i)
+				{
+					// We are in attack
+					if (InState.CurrentSampleIndex <= InState.AttackSampleCount)
+					{
+						float AttackFraction = (float)InState.CurrentSampleIndex++ / InState.AttackSampleCount;
+						float EnvValue = FMath::Pow(AttackFraction, InState.AttackCurveFactor);
+						float TargeEnvelopeValue = InState.StartingEnvelopeValue + (1.0f - InState.StartingEnvelopeValue) * EnvValue;
+						InState.EnvEase.SetValue(TargeEnvelopeValue);
+						InState.CurrentEnvelopeValue = InState.EnvEase.GetNextValue();
+						OutEnvPtr[i] = InState.CurrentEnvelopeValue;
+					}
+					else
+					{
+						int32 DecayEnvSampleCount = (InState.AttackSampleCount + InState.DecaySampleCount);
+						// We are in decay
+						if (InState.CurrentSampleIndex < DecayEnvSampleCount)
+						{
+							int32 SampleCountInDecayState = InState.CurrentSampleIndex++ - InState.AttackSampleCount;
+							float DecayFracton = (float)SampleCountInDecayState / InState.DecaySampleCount;
+							float TargetEnvelopeValue = 1.0f - (1.0f - InState.SustainLevel) * FMath::Pow(DecayFracton, InState.DecayCurveFactor);
+							InState.EnvEase.SetValue(TargetEnvelopeValue);
+							InState.CurrentEnvelopeValue = InState.EnvEase.GetNextValue();
+							OutEnvPtr[i] = InState.CurrentEnvelopeValue;
+						}
+						// We are in sustain
+						else if (!InState.bIsInRelease)
+						{
+							InState.EnvEase.SetValue(InState.SustainLevel);
+							InState.CurrentEnvelopeValue = InState.EnvEase.GetNextValue();
+							InState.EnvelopeValueAtReleaseStart = InState.CurrentEnvelopeValue;
+							OutEnvPtr[i] = InState.CurrentEnvelopeValue;
+						}
+						else 
+						{
+							int32 ReleaseEnvSampleCount = (InState.AttackSampleCount + InState.DecaySampleCount + InState.ReleaseSampleCount);
+
+							// We are in release
+							if (InState.CurrentSampleIndex < ReleaseEnvSampleCount)
+							{
+								int32 SampleCountInReleaseState = InState.CurrentSampleIndex++ - InState.DecaySampleCount - InState.AttackSampleCount;
+								float ReleaseFraction = (float)SampleCountInReleaseState / InState.ReleaseSampleCount;
+								float TargetEnvelopeValue = InState.EnvelopeValueAtReleaseStart * (1.0 - FMath::Pow(ReleaseFraction, InState.ReleaseCurveFactor));
+								InState.EnvEase.SetValue(TargetEnvelopeValue);
+								InState.CurrentEnvelopeValue = InState.EnvEase.GetNextValue();
+								OutEnvPtr[i] = InState.CurrentEnvelopeValue;
+							}
+							// We're done
+							else
+							{
+								InState.CurrentEnvelopeValue = InState.EnvEase.GetNextValue();
+								OutEnvPtr[i] = InState.CurrentEnvelopeValue;
+
+								// Envelope is done
+								if (InState.EnvEase.IsDone())
+								{
+									// Zero out the rest of the envelope
+									int32 NumSamplesLeft = EndFrame - i - 1;
+									if (NumSamplesLeft > 0)
+									{
+										FMemory::Memzero(&OutEnvPtr[i + 1], sizeof(float) * NumSamplesLeft);
+									}
+									InState.CurrentSampleIndex = INDEX_NONE;
+									OutFinishedFrames.Add(i);
+									break;
+								}
+							}
+						}
+					}
+				}
+			}
+
+			static bool IsAudio() { return true; }
+		};
+
+	}
+
+	template<typename ValueType>
+	class TADSREnvelopeNodeOperator : public TExecutableOperator<TADSREnvelopeNodeOperator<ValueType>>
+	{
+	public:
+		static const FVertexInterface& GetDefaultInterface()
+		{
+			using namespace ADSREnvelopeVertexNames;
+
+			static const FVertexInterface DefaultInterface(
+				FInputVertexInterface(
+					TInputDataVertexModel<FTrigger>(METASOUND_GET_PARAM_NAME_AND_TT(InputAttackTrigger)),
+					TInputDataVertexModel<FTrigger>(METASOUND_GET_PARAM_NAME_AND_TT(InputReleaseTrigger)),
+					TInputDataVertexModel<FTime>(METASOUND_GET_PARAM_NAME_AND_TT(InputAttackTime), 0.01f),
+					TInputDataVertexModel<FTime>(METASOUND_GET_PARAM_NAME_AND_TT(InputDecayTime), 0.2f),
+					TInputDataVertexModel<float>(METASOUND_GET_PARAM_NAME_AND_TT(InputSustainLevel), 0.5f),
+					TInputDataVertexModel<FTime>(METASOUND_GET_PARAM_NAME_AND_TT(InputReleaseTime), 1.0f),
+					TInputDataVertexModel<float>(METASOUND_GET_PARAM_NAME_AND_TT(InputAttackCurve), 1.0f),
+					TInputDataVertexModel<float>(METASOUND_GET_PARAM_NAME_AND_TT(InputDecayCurve), 1.0f),
+					TInputDataVertexModel<float>(METASOUND_GET_PARAM_NAME_AND_TT(InputReleaseCurve), 1.0f)
+				),
+				FOutputVertexInterface(
+					TOutputDataVertexModel<FTrigger>(METASOUND_GET_PARAM_NAME_AND_TT(OutputOnAttackTrigger)),
+					TOutputDataVertexModel<FTrigger>(METASOUND_GET_PARAM_NAME_AND_TT(OutputOnReleaseTrigger)),
+					TOutputDataVertexModel<FTrigger>(METASOUND_GET_PARAM_NAME_AND_TT(OutputOnDone)),
+					TOutputDataVertexModel<ValueType>(METASOUND_GET_PARAM_NAME_AND_TT(OutputEnvelopeValue))
+				)
+			);
+
+			return DefaultInterface;
+		}
+
+		static const FNodeClassMetadata& GetNodeInfo()
+		{
+			auto CreateNodeClassMetadata = []() -> FNodeClassMetadata
+			{
+				FName DataTypeName = GetMetasoundDataTypeName<ValueType>();
+				FName OperatorName = TEXT("ADSR Envelope");
+				FText NodeDisplayName = FText::Format(LOCTEXT("ADSREnvelopeDisplayNamePattern", "ADSR Envelope ({0})"), FText::FromString(GetMetasoundDataTypeString<ValueType>()));
+				FText NodeDescription = LOCTEXT("ADSREnevelopeDesc", "Generates an attack-decay-sustain-release envelope value output when triggered.");
+				FVertexInterface NodeInterface = GetDefaultInterface();
+
+				FNodeClassMetadata Metadata
+				{
+					FNodeClassName{FName("ADSR Envelope"), OperatorName, DataTypeName},
+					1, // Major Version
+					0, // Minor Version
+					NodeDisplayName,
+					NodeDescription,
+					PluginAuthor,
+					PluginNodeMissingPrompt,
+					NodeInterface,
+					{LOCTEXT("EnvelopeCat", "Envelopes")},
+					{TEXT("Envelope")},
+					FNodeDisplayStyle{}
+				};
+
+				return Metadata;
+			};
+
+			static const FNodeClassMetadata Metadata = CreateNodeClassMetadata();
+			return Metadata;
+		}
+
+		static TUniquePtr<IOperator> CreateOperator(const FCreateOperatorParams& InParams, TArray<TUniquePtr<IOperatorBuildError>>& OutErrors)
+		{
+			using namespace ADSREnvelopeVertexNames;
+
+			const FInputVertexInterface& InputInterface = GetDefaultInterface().GetInputInterface();
+
+			FTriggerReadRef TriggerAttackIn = InParams.InputDataReferences.GetDataReadReferenceOrConstruct<FTrigger>(METASOUND_GET_PARAM_NAME(InputAttackTrigger), InParams.OperatorSettings);
+			FTriggerReadRef TriggerReleaseIn = InParams.InputDataReferences.GetDataReadReferenceOrConstruct<FTrigger>(METASOUND_GET_PARAM_NAME(InputReleaseTrigger), InParams.OperatorSettings);
+
+			FTimeReadRef AttackTime = InParams.InputDataReferences.GetDataReadReferenceOrConstructWithVertexDefault<FTime>(InputInterface, METASOUND_GET_PARAM_NAME(InputAttackTime), InParams.OperatorSettings);
+			FTimeReadRef DecayTime = InParams.InputDataReferences.GetDataReadReferenceOrConstructWithVertexDefault<FTime>(InputInterface, METASOUND_GET_PARAM_NAME(InputDecayTime), InParams.OperatorSettings);
+			FFloatReadRef SustainLevel = InParams.InputDataReferences.GetDataReadReferenceOrConstructWithVertexDefault<float>(InputInterface, METASOUND_GET_PARAM_NAME(InputSustainLevel), InParams.OperatorSettings);
+			FTimeReadRef ReleaseTime = InParams.InputDataReferences.GetDataReadReferenceOrConstructWithVertexDefault<FTime>(InputInterface, METASOUND_GET_PARAM_NAME(InputReleaseTime), InParams.OperatorSettings);
+
+			FFloatReadRef AttackCurveFactor = InParams.InputDataReferences.GetDataReadReferenceOrConstructWithVertexDefault<float>(InputInterface, METASOUND_GET_PARAM_NAME(InputAttackCurve), InParams.OperatorSettings);
+			FFloatReadRef DecayCurveFactor = InParams.InputDataReferences.GetDataReadReferenceOrConstructWithVertexDefault<float>(InputInterface, METASOUND_GET_PARAM_NAME(InputDecayCurve), InParams.OperatorSettings);
+			FFloatReadRef ReleaseCurveFactor = InParams.InputDataReferences.GetDataReadReferenceOrConstructWithVertexDefault<float>(InputInterface, METASOUND_GET_PARAM_NAME(InputReleaseCurve), InParams.OperatorSettings);
+
+			return MakeUnique<TADSREnvelopeNodeOperator<ValueType>>(InParams.OperatorSettings, TriggerAttackIn, TriggerReleaseIn, AttackTime, DecayTime, SustainLevel, ReleaseTime, AttackCurveFactor, DecayCurveFactor, ReleaseCurveFactor);
+		}
+
+		TADSREnvelopeNodeOperator(const FOperatorSettings& InSettings,
+			const FTriggerReadRef& InTriggerAttackIn,
+			const FTriggerReadRef& InTriggerReleaseIn,
+			const FTimeReadRef& InAttackTime,
+			const FTimeReadRef& InDecayTime,
+			const FFloatReadRef& InSustainLevel,
+			const FTimeReadRef& InReleaseTime,
+			const FFloatReadRef& InAttackCurveFactor,
+			const FFloatReadRef& InDecayCurveFactor,
+			const FFloatReadRef& InReleaseCurveFactor)
+			: TriggerAttackIn(InTriggerAttackIn)
+			, TriggerReleaseIn(InTriggerReleaseIn)
+			, AttackTime(InAttackTime)
+			, DecayTime(InDecayTime)
+			, SustainLevel(InSustainLevel)
+			, ReleaseTime(InReleaseTime)
+			, AttackCurveFactor(InAttackCurveFactor)
+			, DecayCurveFactor(InDecayCurveFactor)
+			, ReleaseCurveFactor(InReleaseCurveFactor)
+			, OnAttackTrigger(TDataWriteReferenceFactory<FTrigger>::CreateAny(InSettings))
+			, OnReleaseTrigger(TDataWriteReferenceFactory<FTrigger>::CreateAny(InSettings))
+			, OnDone(TDataWriteReferenceFactory<FTrigger>::CreateAny(InSettings))
+			, OutputEnvelope(TDataWriteReferenceFactory<ValueType>::CreateAny(InSettings))
+		{
+			NumFramesPerBlock = InSettings.GetNumFramesPerBlock();
+
+			EnvState.EnvEase.SetEaseFactor(0.01f);
+
+			if (ADSREnvelopeNodePrivate::TADSREnvelope<ValueType>::IsAudio())
+			{
+				SampleRate = InSettings.GetSampleRate();
+			}
+			else
+			{
+				SampleRate = InSettings.GetActualBlockRate();
+			}
+		}
+
+		virtual ~TADSREnvelopeNodeOperator() = default;
+
+		virtual FDataReferenceCollection GetInputs() const override
+		{
+			using namespace ADSREnvelopeVertexNames;
+
+			FDataReferenceCollection Inputs;
+			Inputs.AddDataReadReference(METASOUND_GET_PARAM_NAME(InputAttackTrigger), TriggerAttackIn);
+			Inputs.AddDataReadReference(METASOUND_GET_PARAM_NAME(InputAttackTrigger), TriggerReleaseIn);
+			Inputs.AddDataReadReference(METASOUND_GET_PARAM_NAME(InputAttackTime), AttackTime);
+			Inputs.AddDataReadReference(METASOUND_GET_PARAM_NAME(InputDecayTime), DecayTime);
+			Inputs.AddDataReadReference(METASOUND_GET_PARAM_NAME(InputSustainLevel), SustainLevel);
+			Inputs.AddDataReadReference(METASOUND_GET_PARAM_NAME(InputReleaseTime), ReleaseTime);
+			Inputs.AddDataReadReference(METASOUND_GET_PARAM_NAME(InputAttackCurve), AttackCurveFactor);
+			Inputs.AddDataReadReference(METASOUND_GET_PARAM_NAME(InputDecayCurve), DecayCurveFactor);
+			Inputs.AddDataReadReference(METASOUND_GET_PARAM_NAME(InputReleaseCurve), ReleaseCurveFactor);
+
+			return Inputs;
+		}
+
+		virtual FDataReferenceCollection GetOutputs() const override
+		{
+			using namespace ADSREnvelopeVertexNames;
+
+			FDataReferenceCollection Outputs;
+			Outputs.AddDataReadReference(METASOUND_GET_PARAM_NAME(OutputOnAttackTrigger), OnAttackTrigger);
+			Outputs.AddDataReadReference(METASOUND_GET_PARAM_NAME(OutputOnReleaseTrigger), OnReleaseTrigger);
+			Outputs.AddDataReadReference(METASOUND_GET_PARAM_NAME(OutputOnDone), OnDone);
+			Outputs.AddDataReadReference(METASOUND_GET_PARAM_NAME(OutputEnvelopeValue), OutputEnvelope);
+
+			return Outputs;
+		}
+
+		void UpdateParams()
+		{
+			float AttackTimeSeconds = AttackTime->GetSeconds();
+			float DecayTimeSeconds = DecayTime->GetSeconds();
+			float ReleaseTimeSeconds = ReleaseTime->GetSeconds();
+			EnvState.AttackSampleCount = SampleRate * FMath::Max(0.0f, AttackTimeSeconds);
+			EnvState.DecaySampleCount = SampleRate * FMath::Max(0.0f, DecayTimeSeconds);
+			EnvState.SustainLevel = FMath::Max(0.0f, *SustainLevel);
+			EnvState.ReleaseSampleCount = SampleRate * FMath::Max(0.0f, ReleaseTimeSeconds);
+			EnvState.AttackCurveFactor = FMath::Max(KINDA_SMALL_NUMBER, *AttackCurveFactor);
+			EnvState.DecayCurveFactor = FMath::Max(KINDA_SMALL_NUMBER, *DecayCurveFactor);
+			EnvState.ReleaseCurveFactor = FMath::Max(KINDA_SMALL_NUMBER, *ReleaseCurveFactor);
+		}
+
+
+		void Execute()
+		{
+			using namespace ADSREnvelopeNodePrivate;
+
+			OnAttackTrigger->AdvanceBlock();
+			OnReleaseTrigger->AdvanceBlock();
+			OnDone->AdvanceBlock();
+ 
+			// check for any updates to input params
+			UpdateParams();
+ 		
+			TriggerReleaseIn->ExecuteBlock(
+				[&](int32 StartFrame, int32 EndFrame)
+				{
+				},
+				// OnTrigger
+				[&](int32 StartFrame, int32 EndFrame)
+				{
+					// This schedules a release sample count
+					EnvState.bIsInRelease = true;
+					OnReleaseTrigger->TriggerFrame(StartFrame);
+				}
+			);
+
+			TriggerAttackIn->ExecuteBlock(
+				// OnPreTrigger
+				[&](int32 StartFrame, int32 EndFrame)
+				{
+					TArray<int32> FinishedFrames;
+					TADSREnvelope<ValueType>::GetNextEnvelopeOutput(EnvState, StartFrame, EndFrame, FinishedFrames, *OutputEnvelope);
+
+					for (int32 FrameFinished : FinishedFrames)
+					{
+						OnDone->TriggerFrame(FrameFinished);
+					}
+				},
+				// OnTrigger
+				[&](int32 StartFrame, int32 EndFrame)
+				{
+					// Get latest params
+					UpdateParams();
+
+					// Set the sample index to the top of the envelope
+					EnvState.CurrentSampleIndex = 0;
+					EnvState.StartingEnvelopeValue = EnvState.CurrentEnvelopeValue;
+					EnvState.bIsInRelease = false;
+					// Generate the output (this will no-op if we're block rate)
+					TArray<int32> FinishedFrames;
+					TADSREnvelope<ValueType>::GetNextEnvelopeOutput(EnvState, StartFrame, EndFrame, FinishedFrames, *OutputEnvelope);
+					for (int32 FrameFinished : FinishedFrames)
+					{
+						OnDone->TriggerFrame(FrameFinished);
+					}
+
+					// Forward the trigger
+					OnAttackTrigger->TriggerFrame(StartFrame);
+				}
+			);
+		}
+
+	private:
+		FTriggerReadRef TriggerAttackIn;
+		FTriggerReadRef TriggerReleaseIn;
+		FTimeReadRef AttackTime;
+		FTimeReadRef DecayTime;
+		FFloatReadRef SustainLevel;
+		FTimeReadRef ReleaseTime;
+		FFloatReadRef AttackCurveFactor;
+		FFloatReadRef DecayCurveFactor;
+		FFloatReadRef ReleaseCurveFactor;
+
+		FTriggerWriteRef OnAttackTrigger;
+		FTriggerWriteRef OnReleaseTrigger;
+		FTriggerWriteRef OnDone;
+		TDataWriteReference<ValueType> OutputEnvelope;
+
+		// This will either be the block rate or sample rate depending on if this is block-rate or audio-rate envelope
+		float SampleRate = 0.0f;
+		int32 NumFramesPerBlock = 0;
+
+		ADSREnvelopeNodePrivate::FEnvState EnvState;
+	};
+
+	/** TADSREnvelopeNode
+	 *
+	 *  Creates an Attack/Decay envelope node.
+	 */
+	template<typename ValueType>
+	class METASOUNDSTANDARDNODES_API TADSREnvelopeNode : public FNodeFacade
+	{
+	public:
+		/**
+		 * Constructor used by the Metasound Frontend.
+		 */
+		TADSREnvelopeNode(const FNodeInitData& InInitData)
+			: FNodeFacade(InInitData.InstanceName, InInitData.InstanceID, TFacadeOperatorClass<TADSREnvelopeNodeOperator<ValueType>>())
+		{}
+
+		virtual ~TADSREnvelopeNode() = default;
+	};
+
+	using FADSREnvelopeNodeFloat = TADSREnvelopeNode<float>;
+	METASOUND_REGISTER_NODE(FADSREnvelopeNodeFloat)
+
+	using FADSREnvelopeAudioBuffer = TADSREnvelopeNode<FAudioBuffer>;
+	METASOUND_REGISTER_NODE(FADSREnvelopeAudioBuffer)
+}
+
+#undef LOCTEXT_NAMESPACE
diff --git a/Engine/Plugins/Runtime/Metasound/Source/MetasoundStandardNodes/Private/MetasoundADSRNode.cpp b/Engine/Plugins/Runtime/Metasound/Source/MetasoundStandardNodes/Private/MetasoundADSRNode.cpp
deleted file mode 100644
diff --git a/Engine/Plugins/Runtime/Metasound/Source/MetasoundStandardNodes/Public/MetasoundADSRNode.h b/Engine/Plugins/Runtime/Metasound/Source/MetasoundStandardNodes/Public/MetasoundADSRNode.h
deleted file mode 100644
