Metasound Wave Writer Node

diff --git a/Engine/Plugins/Runtime/Metasound/Source/MetasoundStandardNodes/Private/MetasoundDebugLogNode.cpp b/Engine/Plugins/Runtime/Metasound/Source/MetasoundStandardNodes/Private/MetasoundDebugLogNode.cpp
--- a/Engine/Plugins/Runtime/Metasound/Source/MetasoundStandardNodes/Private/MetasoundDebugLogNode.cpp
+++ b/Engine/Plugins/Runtime/Metasound/Source/MetasoundStandardNodes/Private/MetasoundDebugLogNode.cpp
@@ -3,6 +3,7 @@
 #include "MetasoundDebugLogNode.h"
 
 #include "MetasoundNodeRegistrationMacro.h"
+#include "MetasoundStandardNodesCategories.h"
 
 #define LOCTEXT_NAMESPACE "MetasoundStandardNodes"
 
@@ -22,7 +23,7 @@ namespace Metasound
 				PluginAuthor,
 				PluginNodeMissingPrompt,
 				InDefaultInterface,
-				{LOCTEXT("DebugLogCategory", "Debug")},
+				{ StandardNodes::DebugUtils },
 				{TEXT("DebugLog")},
 				FNodeDisplayStyle{}
 			};
diff --git a/Engine/Plugins/Runtime/Metasound/Source/MetasoundStandardNodes/Private/MetasoundStandardNodesCategories.cpp b/Engine/Plugins/Runtime/Metasound/Source/MetasoundStandardNodes/Private/MetasoundStandardNodesCategories.cpp
--- a/Engine/Plugins/Runtime/Metasound/Source/MetasoundStandardNodes/Private/MetasoundStandardNodesCategories.cpp
+++ b/Engine/Plugins/Runtime/Metasound/Source/MetasoundStandardNodes/Private/MetasoundStandardNodesCategories.cpp
@@ -13,6 +13,8 @@ namespace Metasound
 		const FText TriggerUtils = { LOCTEXT("Metasound_TriggerUtilsCategory", "Trigger Utilities") };
 		const FText RandomUtils = { LOCTEXT("Metasound_RandomUtilsCategory", "Random Utilities") };
 		const FText Conversions  = { LOCTEXT("Metasound_ConversionsCategory", "Conversions") };
+		const FText DebugUtils = { LOCTEXT("Metasound_DebugUtilsCategory", "Debug Utilities") };
+		const FText Io = { LOCTEXT("Metasound_IoCategory", "File IO") };
 	}
 }
 
diff --git a/Engine/Plugins/Runtime/Metasound/Source/MetasoundStandardNodes/Private/MetasoundWaveWriterNode.cpp b/Engine/Plugins/Runtime/Metasound/Source/MetasoundStandardNodes/Private/MetasoundWaveWriterNode.cpp
new file mode 100644
--- /dev/null
+++ b/Engine/Plugins/Runtime/Metasound/Source/MetasoundStandardNodes/Private/MetasoundWaveWriterNode.cpp
@@ -0,0 +1,310 @@
+// Copyright Epic Games, Inc. All Rights Reserved.
+
+#include "MetasoundWaveWriterNode.h"
+
+#include "MetasoundNodeRegistrationMacro.h"
+#include "MetasoundPrimitives.h"
+#include "MetasoundStandardNodesNames.h"
+#include "MetasoundStandardNodesCategories.h"
+#include "MetasoundExecutableOperator.h"
+#include "MetasoundBuildError.h"
+#include "Misc/ScopeLock.h"
+
+#define LOCTEXT_NAMESPACE "MetasoundStandardNodes"
+
+namespace Metasound
+{	
+	// Incremental wave writer class.
+	class FWaveWriter
+	{		
+		// Local definition so we don't depend on platform includes.
+		enum EFormatType { IEEE_FLOAT = 0x3 }; // WAVE_FORMAT_IEEE_FLOAT
+		struct FWaveFormatEx
+		{
+			uint16	FormatTag;
+			uint16	NumChannels;
+			uint32	NumSamplesPerSec;
+			uint32	AverageBytesPerSec;
+			uint16	BlockAlign;
+			uint16	NumBitsPerSample;
+			uint16	Size;
+		};
+		
+	public:
+		FWaveWriter(TUniquePtr<FArchive>&& InOutputStream, int32 InSampleRate, int32 InNumChannels, bool bInUpdateHeaderAfterEveryWrite)
+			: OutputStream{ MoveTemp(InOutputStream) }
+			, bUpdateHeaderAfterEveryWrite{ bInUpdateHeaderAfterEveryWrite }
+		{
+			WriteHeader(InSampleRate, InNumChannels);
+		}
+
+		~FWaveWriter()
+		{
+			UpdateHeader();
+		}
+
+		void Write(TArrayView<const float> InBuffer)
+		{
+			OutputStream->Serialize((void*)InBuffer.GetData(), InBuffer.GetTypeSize()*InBuffer.Num());
+			
+			if (bUpdateHeaderAfterEveryWrite)
+			{
+				UpdateHeader();
+			}
+		}
+
+	private:
+		void UpdateHeader()
+		{
+			// RIFF/fmt/data. (bytes per chunk)
+			static const int32 HeaderSize = sizeof(FWaveFormatEx) + sizeof(int32) + sizeof(int32) + sizeof(int32) + sizeof(int32) + sizeof(int32);
+
+			int32 WritePos = OutputStream->Tell();
+
+			// update data chunk size
+			OutputStream->Seek(DataSizePos);
+			int32 DataSize = WritePos - DataSizePos - 4;
+			*OutputStream << DataSize;
+
+			// update top riff size
+			OutputStream->Seek(RiffSizePos);
+			int32 RiffSize = HeaderSize + DataSize - 4;
+			*OutputStream << RiffSize;
+
+			OutputStream->Seek(WritePos);
+		}
+
+		TUniquePtr<FArchive> OutputStream;
+		int32 RiffSizePos = 0;
+		int32 DataSizePos = 0;
+		bool bUpdateHeaderAfterEveryWrite = false;
+
+		void WriteHeader(int32 InSampleRate, int32 InNumChannels)
+		{	
+			FWaveFormatEx Fmt = { 0 };
+			Fmt.NumChannels = InNumChannels;
+			Fmt.NumSamplesPerSec = InSampleRate;
+			Fmt.NumBitsPerSample = sizeof(float) * 8;
+			Fmt.BlockAlign = (Fmt.NumBitsPerSample * InNumChannels) / 8;
+			Fmt.AverageBytesPerSec = Fmt.BlockAlign * InSampleRate;
+			Fmt.FormatTag = EFormatType::IEEE_FLOAT;// WAVE_FORMAT_IEEE_FLOAT;
+		
+			int32 ID = 'FFIR';
+			*OutputStream << ID;
+			RiffSizePos = OutputStream->Tell();
+			int32 RiffChunkSize = 0;
+			*OutputStream << RiffChunkSize;
+
+			ID = 'EVAW';
+			*OutputStream << ID;
+
+			ID = ' tmf';
+			*OutputStream << ID;
+			int32 FmtSize = sizeof(Fmt);
+			*OutputStream << FmtSize;
+			OutputStream->Serialize((void*)&Fmt, FmtSize);
+
+			ID = 'atad';
+			*OutputStream << ID;
+			DataSizePos = OutputStream->Tell();
+			int32 DataChunkSize = 0;
+			*OutputStream << DataChunkSize;
+		}
+	};
+
+	class FFileWriteError : public FBuildErrorBase
+	{
+	public:
+		static const FName ErrorType;
+
+		virtual ~FFileWriteError() = default;
+
+		FFileWriteError(const FNode& InNode, const FString& InFilename)
+			: FBuildErrorBase(ErrorType, FText::FromString(FString::Printf(TEXT("File Writer Error while trying to write '%s'"), *InFilename)))
+		{
+			AddNode(InNode);
+		}
+	};
+	const FName FFileWriteError::ErrorType = FName(TEXT("MetasoundFileWriterError"));
+
+	class FNumberedFileCache
+	{
+	public:
+		static const FString Seperator;
+
+		FNumberedFileCache(const FString& InRootPath, const FString& InExt)
+			: RootPath{ InRootPath }, FileExtention{ InExt }
+		{
+			CacheFilenames();
+		}
+
+		FString GenerateNextNumberedFilename(const FName& InPrefix)
+		{
+			FScopeLock Lock{ &Cs };
+			uint32& CurrentMax = FileIndexMap.FindOrAdd(InPrefix);
+			FString Filename{ *InPrefix.GetPlainNameString() };
+			Filename.Append(*Seperator);
+			Filename.AppendInt(++CurrentMax);
+			Filename.Append(*FileExtention);
+			return RootPath / Filename;
+		}
+	private:
+
+		// Slow directory search of the root path for filenames.
+		void CacheFilenames()
+		{
+			FScopeLock Lock{ &Cs };
+			
+			// Find all files, split filenames into prefix + number, saving max number we find.
+			TArray<FString> Files;
+			IFileManager::Get().FindFiles(Files , *RootPath, *FileExtention);
+			for (const FString& i : Files)
+			{
+				FString Prefix, Postfix;
+				if (i.Split(Seperator, &Prefix, &Postfix, ESearchCase::IgnoreCase, ESearchDir::FromEnd))
+				{
+					FString NumberString = FPaths::GetBaseFilename(Postfix);
+					if (FCString::IsNumeric(*NumberString))
+					{
+						int32 Number = FCString::Atoi(*NumberString);
+						if (Number >= 0)
+						{
+							uint32& CurrentMax = FileIndexMap.FindOrAdd(*Prefix);
+							if (static_cast<uint32>(Number) > CurrentMax)
+							{
+								CurrentMax = static_cast<uint32>(Number);
+							}
+						}
+					}
+				}
+			}
+		}
+		FCriticalSection Cs;
+		FString RootPath;
+		FString FileExtention;
+		TMap<FName, uint32> FileIndexMap;
+	};
+	const FString FNumberedFileCache::Seperator{ TEXT("_") };
+
+	class FWaveWriterOperator : public TExecutableOperator<FWaveWriterOperator>
+	{
+	public:
+		FWaveWriterOperator(const FOperatorSettings& InSettings, FAudioBufferReadRef&& InAudioBuffer, FBoolReadRef&& InEnabled, TUniquePtr<FWaveWriter>&& InStream)
+			: AudioInput{ MoveTemp(InAudioBuffer) }
+			, Enabled{ MoveTemp(InEnabled) }
+			, Writer{ MoveTemp(InStream) }
+		{}
+
+		FDataReferenceCollection GetInputs() const override
+		{
+			FDataReferenceCollection InputDataReferences;
+			InputDataReferences.AddDataReadReference(AudioInputPinName, AudioInput);
+			return InputDataReferences;
+		}
+		FDataReferenceCollection GetOutputs() const override
+		{
+			FDataReferenceCollection OutputDataReferences;
+			return OutputDataReferences;
+		}
+
+		static FVertexInterface DeclareVertexInterface() 
+		{
+			static const FVertexInterface Interface(
+				FInputVertexInterface(
+					TInputDataVertexModel<FString>(FilenamePrefixPinName, LOCTEXT("WaveWriterFilenamePrefixDescriptionX", "Filename Prefix of file you are writing."), FString(DefaultFileName)),
+					TInputDataVertexModel<bool>(EnabledPinName, LOCTEXT("WaveWriterEnabledDescription", "If this wave writer is enabled or not. File will remain open if disabled until graph termination."), true),
+					TInputDataVertexModel<FAudioBuffer>(AudioInputPinName, LOCTEXT("WaveWriterAudioInputDescription", "Audio input that you want serialized."))
+				),
+				FOutputVertexInterface()
+			);
+			return Interface;
+		}
+
+		static const FNodeClassMetadata& GetNodeInfo()
+		{
+			auto InitNodeInfo = []() -> FNodeClassMetadata
+			{
+				FNodeClassMetadata Info;
+				Info.ClassName = { Metasound::StandardNodes::Namespace, TEXT("WaveWriter"), Metasound::StandardNodes::AudioVariant };
+				Info.MajorVersion = 1;
+				Info.MinorVersion = 0;
+				Info.DisplayName = LOCTEXT("Metasound_WaveWriterNodeDisplayName", "Wave Writer");
+				Info.Description = LOCTEXT("Metasound_WaveWriterNodeDescription", "Write a the incoming audio signal to disk");
+				Info.Author = PluginAuthor;
+				Info.PromptIfMissing = PluginNodeMissingPrompt;
+				Info.DefaultInterface = DeclareVertexInterface();
+				Info.CategoryHierarchy.Emplace(StandardNodes::Io);
+				return Info;
+			};
+			static const FNodeClassMetadata Info = InitNodeInfo();
+			return Info;
+		}
+
+		static TUniquePtr<IOperator> CreateOperator(const FCreateOperatorParams& InParams, FBuildErrorArray& OutErrors)
+		{
+			const FWaveWriterNode& Node = static_cast<const FWaveWriterNode&>(InParams.Node);
+			const FDataReferenceCollection& InputCol = InParams.InputDataReferences;
+			const FOperatorSettings& Settings = InParams.OperatorSettings;
+			const FInputVertexInterface& InputInterface = DeclareVertexInterface().GetInputInterface();
+			
+			static const TCHAR* WaveExt = TEXT(".wav");
+
+			// Build cache of numbered files (do this once only).
+			static TUniquePtr<FNumberedFileCache> NumberedFileCache = MakeUnique<FNumberedFileCache>(*FPaths::AudioCaptureDir(), WaveExt);
+			
+			FStringReadRef FilenamePrefix = InputCol.GetDataReadReferenceOrConstructWithVertexDefault<FString>(InputInterface, FilenamePrefixPinName, Settings);
+			FName Prefix = FName{ **FilenamePrefix };
+			FString Filename = NumberedFileCache->GenerateNextNumberedFilename(Prefix);
+
+			if (TUniquePtr<FArchive> FileWriter = TUniquePtr<FArchive>(IFileManager::Get().CreateFileWriter(*Filename, IO_WRITE)))
+			{
+				return MakeUnique<FWaveWriterOperator>(
+					Settings,
+					InputCol.GetDataReadReferenceOrConstruct<FAudioBuffer>(AudioInputPinName, Settings),
+					InputCol.GetDataReadReferenceOrConstructWithVertexDefault<bool>(InputInterface, EnabledPinName, Settings),
+					MakeUnique<FWaveWriter>(MoveTemp(FileWriter), Settings.GetSampleRate(), 1, true)
+				);
+			}
+			
+			// Failed to open a writer object. Log an error.
+			OutErrors.Emplace(MakeUnique<FFileWriteError>(Node, Filename));
+			return {};
+		}
+
+		void Execute()
+		{
+			if (Writer && *Enabled)
+			{
+				Writer->Write(MakeArrayView(AudioInput->GetData(), AudioInput->Num()));
+			}
+		}
+
+	protected:
+		FAudioBufferReadRef AudioInput;
+		FBoolReadRef Enabled;
+		TUniquePtr<FArchive> Output;
+		TUniquePtr<FWaveWriter> Writer;
+		static constexpr const TCHAR* DefaultFileName = TEXT("Output");
+		static constexpr const TCHAR* AudioInputPinName = TEXT("In");
+		static constexpr const TCHAR* EnabledPinName = TEXT("Enabled");
+		static constexpr const TCHAR* FilenamePrefixPinName = TEXT("Filename Prefix");
+	};
+
+	// Linkage for constexpr on older clang.
+	constexpr const TCHAR* FWaveWriterOperator::DefaultFileName;
+	constexpr const TCHAR* FWaveWriterOperator::AudioInputPinName;
+	constexpr const TCHAR* FWaveWriterOperator::EnabledPinName;
+	constexpr const TCHAR* FWaveWriterOperator::FilenamePrefixPinName;
+
+	FWaveWriterNode::FWaveWriterNode(const FString& InName, const FGuid& InInstanceID)
+		: FNodeFacade{ InName, InInstanceID, TFacadeOperatorClass<FWaveWriterOperator>() }
+	{}
+
+	FWaveWriterNode::FWaveWriterNode(const FNodeInitData& InInitData)
+		: FWaveWriterNode{ InInitData.InstanceName, InInitData.InstanceID }
+	{}
+
+	METASOUND_REGISTER_NODE(FWaveWriterNode)
+}
+
+#undef LOCTEXT_NAMESPACE
diff --git a/Engine/Plugins/Runtime/Metasound/Source/MetasoundStandardNodes/Public/MetasoundStandardNodesCategories.h b/Engine/Plugins/Runtime/Metasound/Source/MetasoundStandardNodes/Public/MetasoundStandardNodesCategories.h
--- a/Engine/Plugins/Runtime/Metasound/Source/MetasoundStandardNodes/Public/MetasoundStandardNodesCategories.h
+++ b/Engine/Plugins/Runtime/Metasound/Source/MetasoundStandardNodes/Public/MetasoundStandardNodesCategories.h
@@ -13,5 +13,7 @@ namespace Metasound
 		extern const FText TriggerUtils;
 		extern const FText RandomUtils;
 		extern const FText Conversions;
+		extern const FText DebugUtils;
+		extern const FText Io;
 	}
 }
diff --git a/Engine/Plugins/Runtime/Metasound/Source/MetasoundStandardNodes/Public/MetasoundWaveWriterNode.h b/Engine/Plugins/Runtime/Metasound/Source/MetasoundStandardNodes/Public/MetasoundWaveWriterNode.h
new file mode 100644
--- /dev/null
+++ b/Engine/Plugins/Runtime/Metasound/Source/MetasoundStandardNodes/Public/MetasoundWaveWriterNode.h
@@ -0,0 +1,19 @@
+// Copyright Epic Games, Inc. All Rights Reserved.
+
+#pragma once
+
+#include "CoreMinimal.h"
+#include "MetasoundFacade.h"
+
+#define LOCTEXT_NAMESPACE "MetasoundStandardNodes"
+
+namespace Metasound
+{
+	class METASOUNDSTANDARDNODES_API FWaveWriterNode : public FNodeFacade
+	{
+	public:
+		FWaveWriterNode(const FString& InName, const FGuid& InInstanceID);
+		FWaveWriterNode(const FNodeInitData& InInitData);
+	private:
+	};
+} // namespace Metasound
\ No newline at end of file
diff --git a/Engine/Source/Runtime/Core/Private/Misc/Paths.cpp b/Engine/Source/Runtime/Core/Private/Misc/Paths.cpp
--- a/Engine/Source/Runtime/Core/Private/Misc/Paths.cpp
+++ b/Engine/Source/Runtime/Core/Private/Misc/Paths.cpp
@@ -417,6 +417,11 @@ FString FPaths::VideoCaptureDir()
 	return FPaths::ProjectSavedDir() + TEXT("VideoCaptures/");
 }
 
+FString FPaths::AudioCaptureDir()
+{
+	return FPaths::ProjectSavedDir() + TEXT("AudioCaptures/");
+}
+
 FString FPaths::ProjectLogDir()
 {
 #if PLATFORM_PS4
diff --git a/Engine/Source/Runtime/Core/Public/Misc/Paths.h b/Engine/Source/Runtime/Core/Public/Misc/Paths.h
--- a/Engine/Source/Runtime/Core/Public/Misc/Paths.h
+++ b/Engine/Source/Runtime/Core/Public/Misc/Paths.h
@@ -276,6 +276,13 @@ public:
 	 * @return Video capture directory
 	 */
 	static FString VideoCaptureDir();
+
+	/**
+	 * Returns the directory the engine uses to output user requested audio capture files.
+	 *
+	 * @return Audio capture directory
+	 */
+	static FString AudioCaptureDir();
 	
 	/**
 	 * Returns the directory the engine uses to output logs. This currently can't 
