Metasound Frontend Perf Tuneup
- Added frontend search engine to cache frontend queries
- Allowing frontend queries to be updated in place.
- Tracking node registration transactions to allow incremental runtime checks of newly registered nodes.
- Create public RegistryContainer interface and hide implementation to allow for further optimizations without changing public headers

diff --git a/Engine/Plugins/Runtime/Metasound/Source/MetasoundEngine/Private/MetasoundSource.cpp b/Engine/Plugins/Runtime/Metasound/Source/MetasoundEngine/Private/MetasoundSource.cpp
--- a/Engine/Plugins/Runtime/Metasound/Source/MetasoundEngine/Private/MetasoundSource.cpp
+++ b/Engine/Plugins/Runtime/Metasound/Source/MetasoundEngine/Private/MetasoundSource.cpp
@@ -10,6 +10,7 @@
 #include "MetasoundFrontendController.h"
 #include "MetasoundFrontendQuery.h"
 #include "MetasoundFrontendQuerySteps.h"
+#include "MetasoundFrontendSearchEngine.h"
 #include "MetasoundGenerator.h"
 #include "MetasoundInstanceTransmitter.h"
 #include "MetasoundLog.h"
@@ -185,19 +186,11 @@ bool UMetasoundSource::GetReceiveNodeMetadataForDataType(const FName& InTypeName
 {
 	using namespace Metasound;
 
-	FFrontendQuery Query;
+	TArray<FMetasoundFrontendClass> ReceiverNodeClasses = Metasound::Frontend::ISearchEngine::Get().FindClassesWithClassName(FReceiveNodeNames::GetClassNameForDataType(InTypeName));
 
-	Query.AddStep<FGenerateAllAvailableNodeClasses>()
-		.AddStep<FFilterClassesByClassName>(FReceiveNodeNames::GetClassNameForDataType(InTypeName));
-	
-	FFrontendQuerySelectionView Result = Query.ExecuteQuery();
-
-	TArrayView<const FFrontendQueryEntry* const> Selection = Result.GetSelection();
-
-	if (Selection.Num() > 0)
+	if (ReceiverNodeClasses.Num() > 0)
 	{
-		check(Selection[0]->Value.IsType<FMetasoundFrontendClass>());
-		OutMetadata = Selection[0]->Value.Get<FMetasoundFrontendClass>().Metadata;
+		OutMetadata = ReceiverNodeClasses[0].Metadata;
 		return true;
 	}
 
@@ -327,7 +320,23 @@ TArray<FString> UMetasoundSource::GetTransmittableInputVertexNames() const
 			FDataTypeRegistryInfo TypeInfo;
 			if (Frontend::GetTraitsForDataType(ClassInput->TypeName, TypeInfo))
 			{
-				return TypeInfo.bIsTransmittable;
+				if (TypeInfo.bIsTransmittable)
+				{
+					// TODO: Currently values set directly on node pins are represented
+					// as input nodes in the graph. They should not be used for transmission
+					// as the number of input nodes increases quickly as more nodes
+					// are added to a graph. Connecting these input nodes to the 
+					// transmission system is relatively expensive. These undesirable input nodes are
+					// filtered out by ignoring input nodes which are not "Visible". 
+					Frontend::FConstNodeHandle InputNode = RootGraph->GetNodeWithID(ClassInput->NodeID);
+					if (InputNode->IsValid())
+					{
+						if (EMetasoundFrontendNodeStyleDisplayVisibility::Visible == InputNode->GetNodeStyle().Display.Visibility)
+						{
+							return true;
+						}
+					}
+				}
 			}
 		}
 		return false;
diff --git a/Engine/Plugins/Runtime/Metasound/Source/MetasoundEngine/Private/MetasoundWave.cpp b/Engine/Plugins/Runtime/Metasound/Source/MetasoundEngine/Private/MetasoundWave.cpp
--- a/Engine/Plugins/Runtime/Metasound/Source/MetasoundEngine/Private/MetasoundWave.cpp
+++ b/Engine/Plugins/Runtime/Metasound/Source/MetasoundEngine/Private/MetasoundWave.cpp
@@ -94,8 +94,11 @@ namespace Audio
 				bDecoderHasLooped = DecodeResult == Audio::IDecoder::EDecodeResult::Looped;
 
 				// perform linear pitch shift into OutputCircularBuffer
-				const TArrayView<float> BufferToPitchShift(PostSrcBuffer.GetData(), NumResamplerOutputFrames * NumChannels);
-				PitchShifter.ProcessAudio(BufferToPitchShift, OutputCircularBuffer);
+				if (NumResamplerOutputFrames > 0)
+				{
+					const TArrayView<float> BufferToPitchShift(PostSrcBuffer.GetData(), NumResamplerOutputFrames * NumChannels);
+					PitchShifter.ProcessAudio(BufferToPitchShift, OutputCircularBuffer);
+				}
 			}
 		}
 
diff --git a/Engine/Plugins/Runtime/Metasound/Source/MetasoundFrontend/Private/MetasoundFrontend.cpp b/Engine/Plugins/Runtime/Metasound/Source/MetasoundFrontend/Private/MetasoundFrontend.cpp
--- a/Engine/Plugins/Runtime/Metasound/Source/MetasoundFrontend/Private/MetasoundFrontend.cpp
+++ b/Engine/Plugins/Runtime/Metasound/Source/MetasoundFrontend/Private/MetasoundFrontend.cpp
@@ -1,5 +1,6 @@
 // Copyright Epic Games, Inc. All Rights Reserved.
 
+
 #include "MetasoundFrontend.h"
 
 #include "HAL/FileManager.h"
@@ -9,6 +10,7 @@
 #include "MetasoundFrontendDocument.h"
 #include "MetasoundFrontendGraph.h"
 #include "MetasoundFrontendRegistries.h"
+#include "MetasoundFrontendRegistryTransaction.h"
 #include "MetasoundJsonBackend.h"
 #include "MetasoundOperatorBuilder.h"
 #include "MetasoundPrimitives.h"
@@ -30,23 +32,39 @@ namespace Metasound
 {
 	namespace Frontend
 	{
-		TArray<FNodeClassInfo> GetAllAvailableNodeClasses()
+		TArray<FNodeClassInfo> GetAllAvailableNodeClasses(FRegistryTransactionID* OutCurrentRegistryTransactionID)
 		{
-			TArray<FNodeClassInfo> OutClasses;
-
-			auto& Registry = GetExternalNodeRegistry();
-			for (auto& NodeClassTuple : Registry)
+			FMetasoundFrontendRegistryContainer* Registry = FMetasoundFrontendRegistryContainer::Get();
+			if (ensure(nullptr != Registry))
 			{
-				FNodeClassInfo ClassInfo;
-				ClassInfo.NodeType = EMetasoundFrontendClassType::External;
-				ClassInfo.LookupKey = NodeClassTuple.Key;
-
-				OutClasses.Add(ClassInfo);
+				return Registry->GetAllAvailableNodeClasses(OutCurrentRegistryTransactionID);
+			}
+			else
+			{
+				if (nullptr != OutCurrentRegistryTransactionID)
+				{
+					*OutCurrentRegistryTransactionID = GetOriginRegistryTransactionID();
+				}
+				return TArray<FNodeClassInfo>();
 			}
-
-			return OutClasses;
 		}
 
+		TArray<FNodeClassInfo> GetNodeClassesRegisteredSince(FRegistryTransactionID InTransactionID, FRegistryTransactionID* OutCurrentTransactionID)
+		{
+			FMetasoundFrontendRegistryContainer* Registry = FMetasoundFrontendRegistryContainer::Get();
+			if (ensure(nullptr != Registry))
+			{
+				return Registry->GetNodeClassesRegisteredSince(InTransactionID, OutCurrentTransactionID);
+			}
+			else
+			{
+				if (nullptr != OutCurrentTransactionID)
+				{
+					*OutCurrentTransactionID = GetOriginRegistryTransactionID();
+				}
+				return TArray<FNodeClassInfo>();
+			}
+		}
 
 		// gets all metadata (name, description, author, what to say if it's missing) for a given node.
 		FMetasoundFrontendClassMetadata GenerateClassMetadata(const FNodeClassInfo& InInfo)
@@ -131,18 +149,19 @@ namespace Metasound
 			return ClassDescription;
 		}
 
-		FMetasoundFrontendClass GenerateClassDescription(const FNodeClassInfo& InInfo)
+		FMetasoundFrontendClass GenerateClassDescription(const FNodeClassInfo& InNodeInfo)
 		{
-			auto& Registry = GetExternalNodeRegistry();
-			if (Registry.Contains(InInfo.LookupKey))
-			{
-				return Registry[InInfo.LookupKey].CreateFrontendClass();
-			}
-			else
+			FMetasoundFrontendClass OutClass;
+
+			FMetasoundFrontendRegistryContainer* Registry = FMetasoundFrontendRegistryContainer::Get();
+
+			if (ensure(nullptr != Registry))
 			{
-				ensureAlwaysMsgf(false, TEXT("Cannot generate description of unregistered node [NodeClassName:%s]"), *InInfo.LookupKey.NodeClassFullName.ToString());
-				return FMetasoundFrontendClass();
+				bool bSuccess = Registry->FindFrontendClassFromRegistered(InNodeInfo, OutClass);
+				ensureAlwaysMsgf(bSuccess, TEXT("Cannot generate description of unregistered node [NodeClassName:%s]"), *InNodeInfo.LookupKey.NodeClassFullName.ToString());
 			}
+
+			return OutClass;
 		}
 
 		TArray<FName> GetAllAvailableDataTypes()
@@ -190,7 +209,11 @@ class FMetasoundFrontendModule : public IModuleInterface
 {
 	virtual void StartupModule() override
 	{
-		Metasound::Frontend::RegisterPendingNodes();
+		FMetasoundFrontendRegistryContainer* Registry = FMetasoundFrontendRegistryContainer::Get();
+		if (ensure(nullptr != Registry))
+		{
+			Registry->RegisterPendingNodes();
+		}
 	}
 };
 
diff --git a/Engine/Plugins/Runtime/Metasound/Source/MetasoundFrontend/Private/MetasoundFrontendBaseClasses.cpp b/Engine/Plugins/Runtime/Metasound/Source/MetasoundFrontend/Private/MetasoundFrontendBaseClasses.cpp
--- a/Engine/Plugins/Runtime/Metasound/Source/MetasoundFrontend/Private/MetasoundFrontendBaseClasses.cpp
+++ b/Engine/Plugins/Runtime/Metasound/Source/MetasoundFrontend/Private/MetasoundFrontendBaseClasses.cpp
@@ -11,11 +11,6 @@ namespace Metasound
 {
 	namespace Frontend
 	{
-		TMap<FNodeRegistryKey, FNodeRegistryElement>& GetExternalNodeRegistry()
-		{
-			return FMetasoundFrontendRegistryContainer::Get()->GetExternalNodeRegistry();
-		}
-
 		TUniquePtr<INode> ConstructInputNode(const FName& InInputType, FInputNodeConstructorParams&& InParams)
 		{
 			return FMetasoundFrontendRegistryContainer::Get()->ConstructInputNode(InInputType, MoveTemp(InParams));
@@ -26,680 +21,6 @@ namespace Metasound
 			return FMetasoundFrontendRegistryContainer::Get()->ConstructOutputNode(InOutputType, InParams);
 		}
 
-		TUniquePtr<INode> ConstructExternalNode(const FName& InNodeType, uint32 InNodeHash, const FNodeInitData& InInitData)
-		{
-			return FMetasoundFrontendRegistryContainer::Get()->ConstructExternalNode(InNodeType, InNodeHash, InInitData);
-		}
-
-
-		ITransactable::ITransactable(uint32 InUndoLimit, TWeakObjectPtr<UObject> InOwningAsset /* = nullptr */)
-			: UndoLimit(InUndoLimit)
-			, OwningAsset(InOwningAsset)
-		{
-		}
-
-		bool ITransactable::Undo()
-		{
-			if (UndoTransactableStack.Num() == 0)
-			{
-				return false;
-			}
-
-			if (TSharedPtr<ITransactable> Owner = OwningTransactable.Pin())
-			{
-				Owner->DiscardUndoFromOwnedTransactable(AsShared());
-			}
-
-			TSharedPtr<ITransactable> Transactable = UndoTransactableStack.Pop().Pin();
-			TSharedPtr<ITransactable> PreviousTransactable = AsShared();
-
-			RedoTransactableStack.Push(Transactable);
-
-			while(Transactable)
-			{
-				if (Transactable == PreviousTransactable)
-				{
-					bool bUndoSucceeded = Transactable->PerformLocalUndo();
-					if (bUndoSucceeded && OwningAsset.IsValid())
-					{
-						OwningAsset->MarkPackageDirty();
-					}
-
-					return bUndoSucceeded;
-				}
-
-				PreviousTransactable = Transactable;
-				Transactable = Transactable->UndoTransactableStack.Pop().Pin();
-				PreviousTransactable->RedoTransactableStack.Push(Transactable);
-			}
-
-			return false;
-		}
-
-		bool ITransactable::Redo()
-		{
-			if (RedoTransactableStack.Num() == 0)
-			{
-				return false;
-			}
-
-			if (TSharedPtr<ITransactable> Owner = OwningTransactable.Pin())
-			{
-				Owner->DiscardRedoFromOwnedTransactable(AsShared());
-			}
-
-			TSharedPtr<ITransactable> Transactable = RedoTransactableStack.Pop().Pin();
-			TSharedPtr<ITransactable> PreviousTransactable = AsShared();
-
-			UndoTransactableStack.Push(Transactable);
-
-			while (Transactable)
-			{
-				if (Transactable == PreviousTransactable)
-				{
-					const bool bRedoSucceeded = Transactable->PerformLocalRedo();
-					if (bRedoSucceeded && OwningAsset.IsValid())
-					{
-						OwningAsset->MarkPackageDirty();
-					}
-
-					return bRedoSucceeded;
-				}
-
-				PreviousTransactable = Transactable;
-				Transactable = Transactable->RedoTransactableStack.Pop().Pin();
-				PreviousTransactable->UndoTransactableStack.Push(Transactable);
-			}
-
-			return false;
-		}
-
-		void ITransactable::CommitTransaction(FReversibleTransaction&& InTransactionDescription)
-		{
-			if (((uint32)LocalUndoTransactionStack.Num()) >= UndoLimit)
-			{
-				// Discard the oldest undo action.
-				LocalUndoTransactionStack.RemoveAt(0);
-			}
-
-			LocalUndoTransactionStack.Push(MoveTemp(InTransactionDescription));
-			RedoTransactableStack.Reset();
-
-			TWeakPtr<ITransactable> WeakThisPtr = AsShared();
-			UndoTransactableStack.Push(WeakThisPtr);
-
-			if (TSharedPtr<ITransactable> Owner = OwningTransactable.Pin())
-			{
-				Owner->PushUndoFromOwnedTransactable(WeakThisPtr);
-			}
-
-			if (OwningAsset.IsValid())
-			{
-				OwningAsset->MarkPackageDirty();
-			}
-		}
-
-		bool ITransactable::RegisterOwningTransactable(ITransactable& InOwningTransactable)
-		{
-			// Sanity check that this will not create a cycle.
-			while (TSharedPtr<ITransactable> Owner = InOwningTransactable.OwningTransactable.Pin())
-			{
-				if (Owner.Get() == this)
-				{
-					return false;
-				}
-			}
-
-			OwningTransactable = InOwningTransactable.AsShared();
-			return true;
-		}
-
-
-		bool ITransactable::PerformLocalUndo()
-		{
-			if (LocalUndoTransactionStack.Num() == 0)
-			{
-				return false;
-			}
-
-			FReversibleTransaction Transaction = LocalUndoTransactionStack.Pop();
-			bool bResult = Transaction.UndoTransaction();
-			if (bResult)
-			{
-				LocalRedoTransactionStack.Push(MoveTemp(Transaction));
-			}
-
-			return bResult;
-		}
-
-		bool ITransactable::PerformLocalRedo()
-		{
-			if (LocalRedoTransactionStack.Num() == 0)
-			{
-				return false;
-			}
-
-			FReversibleTransaction Transaction = LocalRedoTransactionStack.Pop();
-			bool bResult = Transaction.RedoTransaction();
-			if (bResult)
-			{
-				LocalUndoTransactionStack.Push(MoveTemp(Transaction));
-			}
-
-			return bResult;
-		}
-
-		bool ITransactable::DiscardUndoFromOwnedTransactable(TWeakPtr<ITransactable> InOwnedTransactable)
-		{
-			// NOTE: This relies on TArray.Push adding an element to the end of the array.
-			// Start at the end of the array and work backwards until we find the most recent undo operation for this transactable.
-			for (int32 Index = UndoTransactableStack.Num() - 1; Index >= 0; Index--)
-			{
-				if (UndoTransactableStack[Index] == InOwnedTransactable)
-				{
-					// Discard, but preserve the stack order.
-					UndoTransactableStack.RemoveAt(Index);
-					return true;
-				}
-			}
-
-			return false;
-		}
-
-		bool ITransactable::DiscardRedoFromOwnedTransactable(TWeakPtr<ITransactable> InOwnedTransactable)
-		{
-			// NOTE: This relies on TArray.Push adding an element to the end of the array.
-			// Start at the end of the array and work backwards until we find the most recent undo operation for this transactable.
-			for (int32 Index = RedoTransactableStack.Num() - 1; Index >= 0; Index--)
-			{
-				if (RedoTransactableStack[Index] == InOwnedTransactable)
-				{
-					// Discard, but preserve the stack order.
-					RedoTransactableStack.RemoveAt(Index);
-					return true;
-				}
-			}
-
-			return false;
-		}
-
-		bool ITransactable::PushUndoFromOwnedTransactable(TWeakPtr<ITransactable> InOwnedTransactable)
-		{
-			UndoTransactableStack.Push(MoveTemp(InOwnedTransactable));
-			RedoTransactableStack.Reset();
-
-			return true;
-		}
-
-
-
-		bool ITransactable::PushRedoFromOwnedTransactable(TWeakPtr<ITransactable> InOwnedTransactable)
-		{
-			RedoTransactableStack.Push(MoveTemp(InOwnedTransactable));
-			DiscardUndoFromOwnedTransactable(InOwnedTransactable);
-
-			return true;
-		}
-
-#if 0
-
-		FDescriptionAccessPoint::FDescriptionAccessPoint(TAccessPtr<FMetasoundFrontendDocument> InRootDocumentPtr)
-			: RootDocumentPtr(InRootDocumentPtr)
-		{
-		}
-
-		FMetasoundFrontendDocument& FDescriptionAccessPoint::GetRootChecked() const
-		{
-			check(RootDocumentPtr.IsValid());
-
-			return *RootDocumentPtr;
-		}
-
-		bool FDescriptionAccessPoint::GetDescriptionPtrFromPath(const FDescPath& InPathFromRoot, FMetasoundDescriptionPtr& OutPtr) const
-		{
-			if (!RootDocumentPtr.IsValid())
-			{
-				return false;
-			}
-
-			FDescPath CurrentPath = InPathFromRoot;
-			
-			OutPtr.Set<FMetasoundFrontendDocument*>(RootDocumentPtr.Get());
-			FDescriptionUnwindStep CurrentStep = { OutPtr, Path::EDescType::Document };
-
-			while (CurrentPath.Num() != 0 && CurrentStep.Type != Path::EDescType::Invalid)
-			{
-				CurrentStep = GoToNext(CurrentPath, CurrentStep);
-			}
-
-			OutPtr = CurrentStep.DescriptionStructPtr;
-
-			return CurrentStep.Type != Path::EDescType::Invalid;
-		}
-
-
-		FDescriptionAccessPoint::FDescriptionUnwindStep FDescriptionAccessPoint::GoToNextFromDocument(FMetasoundFrontendDocument& InDocument, FDescPath& InPath, const Path::FElement& InNext) const
-		{
-			check(RootDocumentPtr.IsValid());
-			switch (InNext.CurrentDescType)
-			{
-				case Path::EDescType::Document:
-				{
-					FDescriptionUnwindStep UnwindStep;
-					UnwindStep.DescriptionStructPtr.Set<FMetasoundFrontendDocument*>(&InDocument);
-					UnwindStep.Type = Path::EDescType::Document;
-					return UnwindStep;
-				}
-				case Path::EDescType::Class:
-				{
-					FDescriptionUnwindStep UnwindStep;
-					UnwindStep.DescriptionStructPtr.Set<FMetasoundFrontendClass*>(&(InDocument.RootGraph));
-					UnwindStep.Type = Path::EDescType::Class;
-					return UnwindStep;
-				}
-				case Path::EDescType::DocDependencies:
-				{
-					// The next element in a path after Dependencies will always be the name of the dependency.
-					Path::FElement DependencyElement = InPath.Path[0];
-					InPath.Path.RemoveAt(0);
-					if (!ensureAlwaysMsgf(DependencyElement.CurrentDescType == Path::EDescType::Class, TEXT("Invalid path set up.")))
-					{
-						return FDescriptionUnwindStep::CreateInvalid();
-					}
-
-					FString& DependencyName = DependencyElement.LookupName;
-
-					int32 DependencyID = DependencyElement.LookupID;
-					
-					if (!ensureAlwaysMsgf(DependencyID != INDEX_NONE || DependencyName.Len() > 0, TEXT("Path to a dependency did not include a valid ID or dependency name.")))
-					{
-						return FDescriptionUnwindStep::CreateInvalid();
-					}
-
-					if (DependencyID == FMetasoundFrontendClass::RootClassID)
-					{
-						FDescriptionUnwindStep UnwindStep;
-						UnwindStep.DescriptionStructPtr.Set<FMetasoundFrontendClass*>(&(InDocument.RootGraph));
-						UnwindStep.Type = Path::EDescType::Class;
-						return UnwindStep;
-					}
-
-					TArray<FMetasoundFrontendClass>& DependenciesList = InDocument.Dependencies;
-
-					// Dependencies can be looked up by ID or by name.
-					if (DependencyID != INDEX_NONE)
-					{
-						// Scan the dependencies list for the matching lookup ID.
-						for (FMetasoundFrontendClass& Dependency : DependenciesList)
-						{
-							if (Dependency.UniqueID == DependencyID)
-							{
-								FDescriptionUnwindStep UnwindStep;
-								UnwindStep.DescriptionStructPtr.Set<FMetasoundFrontendClass*>(&Dependency);
-								UnwindStep.Type = Path::EDescType::Class;
-								return UnwindStep;
-							}
-						}
-					}
-					else
-					{
-						// TODO: remove this chunk of code in the "else{}" block. Not sure if it has to be here. 
-						checkNoEntry();
-						// fall back to scanning the dependencies list for the matching lookup name.
-						for (FMetasoundFrontendClass& Dependency : DependenciesList)
-						{
-							if (Dependency.Metadata.NodeName == DependencyName)
-							{
-								FDescriptionUnwindStep UnwindStep;
-								UnwindStep.DescriptionStructPtr.Set<FMetasoundFrontendClass*>(&Dependency);
-								UnwindStep.Type = Path::EDescType::Class;
-								return UnwindStep;
-							}
-						}
-					}
-					
-
-					// If we reached the end of the Dependencies list and didn't find a match, ensure.
-					ensureAlwaysMsgf(false, TEXT("Couldn't find dependency %s in path."), *DependencyName);
-					return FDescriptionUnwindStep::CreateInvalid();
-				}
-			}
-
-			checkNoEntry();
-			return FDescriptionUnwindStep::CreateInvalid();
-		}
-
-		FDescriptionAccessPoint::FDescriptionUnwindStep FDescriptionAccessPoint::GoToNextFromClass(FMetasoundFrontendClass& InClassDescription, FDescPath& InPath, const Path::FElement& InNext) const
-		{
-			switch (InNext.CurrentDescType)
-			{
-				case Path::EDescType::Graph:
-				{
-					FDescriptionUnwindStep UnwindStep;
-					UnwindStep.DescriptionStructPtr.Set<FMetasoundFrontendGraph*>(&(InClassDescription.Graph));
-					UnwindStep.Type = Path::EDescType::Graph;
-					return UnwindStep;
-				}
-				case Path::EDescType::ClassDependencies:
-				{
-					FDescriptionUnwindStep UnwindStep;
-					UnwindStep.DescriptionStructPtr.Set<FClassDependencyIDs*>(&(InClassDescription.DependencyIDs));
-					UnwindStep.Type = Path::EDescType::ClassDependencies;
-					return UnwindStep;
-				}
-				case Path::EDescType::Inputs:
-				{
-					// The next element after an Inputs element should always be the name of an input.
-					Path::FElement InputElement = InPath.Path[0];
-					InPath.Path.RemoveAt(0);
-					if (!ensureAlwaysMsgf(InputElement.CurrentDescType == Path::EDescType::Input, TEXT("Invalid path set up.")))
-					{
-						return FDescriptionUnwindStep::CreateInvalid();
-					}
-
-					FString& InputName = InputElement.LookupName;
-
-					// Scan the inputs list for the lookup name.
-					TArray<FMetasoundFrontendClassInput>& InputsList = InClassDescription.Inputs;
-					for (FMetasoundFrontendClassInput& Input : InputsList)
-					{
-						if (Input.Name == InputName)
-						{
-							FDescriptionUnwindStep UnwindStep;
-							UnwindStep.DescriptionStructPtr.Set<FMetasoundFrontendClassInput*>(&Input);
-							UnwindStep.Type = Path::EDescType::Input;
-							return UnwindStep;
-						}
-					}
-
-					// If we reached the end of the Inputs list and didn't find a match, ensure.
-					ensureAlwaysMsgf(false, TEXT("Couldn't find input %s in path."), *InputName);
-					return FDescriptionUnwindStep::CreateInvalid();
-				}
-				case Path::EDescType::Outputs:
-				{
-					// The next element after an Outputs element should always be the name of an output.
-					Path::FElement OutputElement = InPath.Path[0];
-					InPath.Path.RemoveAt(0);
-					if (!ensureAlwaysMsgf(OutputElement.CurrentDescType == Path::EDescType::Output, TEXT("Invalid path set up.")))
-					{
-						return FDescriptionUnwindStep::CreateInvalid();
-					}
-				
-					FString& OutputName = OutputElement.LookupName;
-
-					// Scan the outputs list for the lookup name.
-					TArray<FMetasoundFrontendClassOutput>& OutputsList = InClassDescription.Outputs;
-					for (FMetasoundFrontendClassOutput& Output : OutputsList)
-					{
-						if (Output.Name == OutputElement.LookupName)
-						{
-							FDescriptionUnwindStep UnwindStep;
-							UnwindStep.DescriptionStructPtr.Set<FMetasoundFrontendClassOutput*>(&Output);
-							UnwindStep.Type = Path::EDescType::Output;
-							return UnwindStep;
-						}
-					}
-
-					// If we reached the end of the Inputs list and didn't find a match, ensure.
-					ensureAlwaysMsgf(false, TEXT("Couldn't find output %s in path."), *InNext.LookupName);
-					return FDescriptionUnwindStep::CreateInvalid();
-				}
-				case Path::EDescType::Metadata:
-				{
-					FDescriptionUnwindStep UnwindStep;
-					UnwindStep.DescriptionStructPtr.Set<FMetasoundFrontendClassMetadata*>(&(InClassDescription.Metadata));
-					UnwindStep.Type = Path::EDescType::Metadata;
-
-					return UnwindStep;
-				}
-				default:
-				{
-					ensureAlwaysMsgf(false, TEXT("Invalid path- Tried to path directly from a Class Description to a type that wasn't a direct memember of the Class"));
-					return FDescriptionUnwindStep::CreateInvalid();
-				}
-			}
-		}
-
-
-		FDescriptionAccessPoint::FDescriptionUnwindStep FDescriptionAccessPoint::GoToNext(FDescPath& InPath, FDescriptionUnwindStep InElement) const
-		{
-			if (!ensureMsgf(InPath.Path.Num() != 0, TEXT("Attempted to unwind an empty path.")))
-			{
-				return FDescriptionUnwindStep::CreateInvalid();
-			}
-
-			Path::FElement NextStep = InPath.Path[0];
-
-			InPath.Path.RemoveAt(0);
-
-			switch (InElement.Type)
-			{
-				case Path::EDescType::Document:
-				{
-					FMetasoundFrontendDocument* Document = InElement.DescriptionStructPtr.Get<FMetasoundFrontendDocument*>();
-					return GoToNextFromDocument(*Document, InPath, NextStep);
-					break;
-				}
-				case Path::EDescType::Class:
-				{
-					FMetasoundFrontendClass* ClassDescription = InElement.DescriptionStructPtr.Get<FMetasoundFrontendClass*>();
-					return GoToNextFromClass(*ClassDescription, InPath, NextStep);
-					break;
-				}
-				case Path::EDescType::Graph:
-				{
-					if (!ensureAlwaysMsgf(NextStep.CurrentDescType == Path::EDescType::Nodes, TEXT("Invalid path. the Graph description only contains the Nodes list.")))
-					{
-						return FDescriptionUnwindStep::CreateInvalid();
-					}
-
-					FMetasoundFrontendGraph* GraphDescription = InElement.DescriptionStructPtr.Get<FMetasoundFrontendGraph*>();
-				
-					if (!ensureAlwaysMsgf(InPath.Path.Num() != 0, TEXT("Incomplete path! path stopped at Nodes list without specifying a node ID.")))
-					{
-						return FDescriptionUnwindStep::CreateInvalid();
-					}
-
-					Path::FElement NodeElement = InPath.Path[0];
-					InPath.Path.RemoveAt(0);
-
-					if (!ensureAlwaysMsgf(NodeElement.CurrentDescType == Path::EDescType::Node, TEXT("Invalid path! a Nodes element must always be followed by a Node ID.")))
-					{
-						return FDescriptionUnwindStep::CreateInvalid();
-					}
-				
-					int32 NodeID = NodeElement.LookupID;
-
-					TArray<FMetasoundFrontendNode>& NodeList = GraphDescription->Nodes;
-					for (FMetasoundFrontendNode& Node : NodeList)
-					{
-						if (Node.UniqueID == NodeID)
-						{
-							FDescriptionUnwindStep UnwindStep;
-							UnwindStep.DescriptionStructPtr.Set<FMetasoundFrontendNode*>(&Node);
-							UnwindStep.Type = Path::EDescType::Node;
-							return UnwindStep;
-						}
-					}
-					break;
-				}
-				case Path::EDescType::Inputs:
-				{
-					ensureAlwaysMsgf(false, TEXT("Invalid path. Inputs will always be pathed directly after a Class."));
-					return FDescriptionUnwindStep::CreateInvalid();
-				}
-				case Path::EDescType::DocDependencies:
-				{
-					ensureAlwaysMsgf(false, TEXT("Invalid path. Document dependencies should always follow after a Document and should always list a specific dependency by ID or name."));
-					return FDescriptionUnwindStep::CreateInvalid();
-				}
-				case Path::EDescType::Input:
-				case Path::EDescType::Output:
-				case Path::EDescType::Node:
-				case Path::EDescType::Metadata:
-				default:
-				{
-					ensureAlwaysMsgf(false, TEXT("Invalid path. Input, Output, Node, and Metadata don't have any child elements."));
-					return FDescriptionUnwindStep::CreateInvalid();
-				}
-			}
-
-			// if we ever hit this, we likely missed a return on one of these branches.
-			checkNoEntry();
-			return FDescriptionUnwindStep::CreateInvalid();
-		}
-
-		FDescPath Path::GetPathToClassForNode(FDescPath InPathForNode, FString& InNodeName)
-		{
-			return FDescPath()[EFromDocument::ToDependencies][*InNodeName];
-		}
-
-
-		FDescPath Path::GetOwningClassDescription(FDescPath InPathForGraph)
-		{
-			// Backtrack from the end of the path until we find a Class element.
-			for (int32 Level = InPathForGraph.Path.Num() - 1; Level >= 0; Level--)
-			{
-				if (InPathForGraph.Path[Level].CurrentDescType == Path::EDescType::Class)
-				{
-					return InPathForGraph;
-				}
-				else
-				{
-					InPathForGraph.Path.RemoveAt(InPathForGraph.Path.Num() - 1);
-				}
-			}
-
-			return InPathForGraph;
-		}
-
-		FDescPath Path::GetDependencyPath(int32 InDependencyID)
-		{
-			return FDescPath()[Path::EFromDocument::ToDependencies][InDependencyID];
-		}
-
-		FDescPath Path::GetInputDescriptionPath(FDescPath InPathForInputNode, const FString& InputName)
-		{
-			return (InPathForInputNode << 3)[Path::EFromClass::ToInputs][*InputName];
-		}
-
-		FDescPath Path::GetOutputDescriptionPath(FDescPath InPathForOutputNode, const FString& OutputName)
-		{
-			return (InPathForOutputNode << 3)[Path::EFromClass::ToOutputs][*OutputName];
-		}
-
-		FDescPath Path::GetOuterGraphPath(FDescPath InPath)
-		{
-			// Unwind element by element until we hit a graph.
-			while (InPath.Path.Num() > 1 && InPath.Path.Last().CurrentDescType != Path::EDescType::Graph)
-			{
-				InPath.Path.Pop();
-			}
-
-			return InPath;
-		}
-
-		FString Path::GetPrintableString(FDescPath InPath)
-		{
-			FString OutString = FString(TEXT("//"));
-			for (FElement& PathElement : InPath.Path)
-			{
-				switch (PathElement.CurrentDescType)
-				{
-					case EDescType::Document:
-					{
-						OutString += TEXT("Document/");
-						break;
-					}
-					case EDescType::Class:
-					{
-						OutString += TEXT("Class/");
-						break;
-					}
-					case EDescType::DocDependencies:
-					{
-						OutString += TEXT("Dependencies(");
-
-						if (PathElement.LookupID != INDEX_NONE)
-						{
-							OutString.AppendInt(PathElement.LookupID);
-						}
-						else
-						{
-							OutString += PathElement.LookupName;
-						}
-						
-						OutString += TEXT(")/");
-						break;
-					}
-					case EDescType::ClassDependencies:
-					{
-						OutString += TEXT("Dependencies(");
-						OutString.AppendInt(PathElement.LookupID);
-						OutString += TEXT(")/");
-						break;
-					}
-					case EDescType::Graph:
-					{
-						OutString += TEXT("Graph/");
-						break;
-					}
-					case EDescType::Inputs:
-					{
-						OutString += TEXT("Inputs/");
-						break;
-					}
-					case EDescType::Input:
-					{
-						OutString += TEXT("Input(");
-						OutString += PathElement.LookupName;
-						OutString += TEXT(")/");
-						break;
-					}
-					case EDescType::Metadata:
-					{
-						OutString += TEXT("Metadata/");
-						break;
-					}
-					case EDescType::Nodes:
-					{
-						OutString += TEXT("Nodes/");
-						break;
-					}
-					case EDescType::Node:
-					{
-						OutString += TEXT("Node(");
-						OutString.AppendInt(PathElement.LookupID);
-						OutString += TEXT(")/");
-						break;
-					}
-					case EDescType::Outputs:
-					{
-						OutString += TEXT("Outputs/");
-						break;
-					}
-					case EDescType::Output:
-					{
-						OutString += TEXT("Output(");
-						OutString += PathElement.LookupName;
-						OutString += TEXT(")/");
-						break;
-					}
-					default:
-					{
-						OutString += TEXT("Unknown/");
-						break;
-					}
-				}
-			}
-
-			return OutString;
-		}
-#endif
-
-		
-
 		void RegisterPendingNodes()
 		{
 			FMetasoundFrontendRegistryContainer::Get()->RegisterPendingNodes();
diff --git a/Engine/Plugins/Runtime/Metasound/Source/MetasoundFrontend/Private/MetasoundFrontendGraph.cpp b/Engine/Plugins/Runtime/Metasound/Source/MetasoundFrontend/Private/MetasoundFrontendGraph.cpp
--- a/Engine/Plugins/Runtime/Metasound/Source/MetasoundFrontend/Private/MetasoundFrontendGraph.cpp
+++ b/Engine/Plugins/Runtime/Metasound/Source/MetasoundFrontend/Private/MetasoundFrontendGraph.cpp
@@ -238,7 +238,7 @@ namespace Metasound
 
 		Metasound::Frontend::FNodeRegistryKey Key = FMetasoundFrontendRegistryContainer::GetRegistryKey(InClass.Metadata);
 
-		return FMetasoundFrontendRegistryContainer::Get()->ConstructExternalNode(Key.NodeClassFullName, Key.NodeHash, InitData);
+		return FMetasoundFrontendRegistryContainer::Get()->ConstructExternalNode(Key, InitData);
 	}
 
 
diff --git a/Engine/Plugins/Runtime/Metasound/Source/MetasoundFrontend/Private/MetasoundFrontendQuery.cpp b/Engine/Plugins/Runtime/Metasound/Source/MetasoundFrontend/Private/MetasoundFrontendQuery.cpp
--- a/Engine/Plugins/Runtime/Metasound/Source/MetasoundFrontend/Private/MetasoundFrontendQuery.cpp
+++ b/Engine/Plugins/Runtime/Metasound/Source/MetasoundFrontend/Private/MetasoundFrontendQuery.cpp
@@ -370,6 +370,33 @@ namespace Metasound
 		return MakeArrayView(Selection);
 	}
 
+	void FFrontendQuerySelection::AppendToStorageAndSelection(const FFrontendQuerySelection& InSelection)
+	{
+		TArrayView<const FFrontendQueryEntry* const> InEntries = InSelection.GetSelection();
+			
+		int32 StorageIndex = Storage.Num();
+		int32 SelectionIndex = Selection.Num();
+
+		for (const FFrontendQueryEntry* Entry : InEntries)
+		{
+			Storage.Add(*Entry);
+		}
+
+		const int32 Num = Storage.Num() - StorageIndex;
+
+		if (Num > 0)
+		{
+			Selection.AddZeroed(Num);
+
+			for (int32 i = 0; i < Num; i++)
+			{
+				Selection[SelectionIndex] = &Storage[StorageIndex];
+				SelectionIndex++;
+				StorageIndex++;
+			}
+		}
+	}
+
 	void FFrontendQuerySelection::AppendToStorageAndSelection(TArrayView<const FFrontendQueryEntry> InEntries)
 	{
 		const int32 Num = InEntries.Num();
@@ -428,13 +455,9 @@ namespace Metasound
 
 
 
-	FFrontendQuerySelectionView::FFrontendQuerySelectionView(TUniquePtr<FFrontendQuerySelection>&& InResult)
-	:	Result(MoveTemp(InResult))
+	FFrontendQuerySelectionView::FFrontendQuerySelectionView(TSharedRef<const FFrontendQuerySelection, ESPMode::ThreadSafe> InResult)
+	:	Result(InResult)
 	{
-		if (!Result.IsValid())
-		{
-			Result = MakeUnique<FFrontendQuerySelection>();
-		}
 	}
 
 	TArrayView<const FFrontendQueryEntry> FFrontendQuerySelectionView::GetStorage() const
@@ -525,6 +548,11 @@ namespace Metasound
 		}
 	}
 
+	FFrontendQuery::FFrontendQuery()
+	: Result(MakeShared<FFrontendQuerySelection, ESPMode::ThreadSafe>())
+	{
+	}
+
 	const TArray<TUniquePtr<FFrontendQueryStep>>& FFrontendQuery::GetSteps() const
 	{
 		return Steps;
@@ -573,8 +601,38 @@ namespace Metasound
 
 	FFrontendQuerySelectionView FFrontendQuery::ExecuteQuery()
 	{
-		Result = MakeUnique<FFrontendQuerySelection>();
+		Reset();
+
+		Execute(*Result);
+
+		return FFrontendQuerySelectionView(Result);
+	}
+
+	FFrontendQuerySelectionView FFrontendQuery::Reset()
+	{
+		Result->ResetStorageAndSelection();
+		
+		return FFrontendQuerySelectionView(Result);
+	}
+
+	FFrontendQuerySelectionView FFrontendQuery::ExecuteQueryAndAppend()
+	{
+		FFrontendQuerySelection NewSelection;
+
+		Execute(NewSelection);
+
+		Result->AppendToStorageAndSelection(NewSelection);
+
+		return FFrontendQuerySelectionView(Result);
+	}
+
+	FFrontendQuerySelectionView FFrontendQuery::GetSelection() const
+	{
+		return FFrontendQuerySelectionView(Result);
+	}
 
+	void FFrontendQuery::Execute(FFrontendQuerySelection& OutSelection)
+	{
 		for (int32 StepIndex = 0; StepIndex < Steps.Num(); StepIndex++)
 		{
 			if (!Steps[StepIndex].IsValid())
@@ -582,10 +640,8 @@ namespace Metasound
 				continue;
 			}
 
-			Steps[StepIndex]->ExecuteStep(*Result);
+			Steps[StepIndex]->ExecuteStep(OutSelection);
 		}
-
-		return FFrontendQuerySelectionView(MoveTemp(Result));
 	}
 }
 
diff --git a/Engine/Plugins/Runtime/Metasound/Source/MetasoundFrontend/Private/MetasoundFrontendQuerySteps.cpp b/Engine/Plugins/Runtime/Metasound/Source/MetasoundFrontend/Private/MetasoundFrontendQuerySteps.cpp
--- a/Engine/Plugins/Runtime/Metasound/Source/MetasoundFrontend/Private/MetasoundFrontendQuerySteps.cpp
+++ b/Engine/Plugins/Runtime/Metasound/Source/MetasoundFrontend/Private/MetasoundFrontendQuerySteps.cpp
@@ -5,6 +5,7 @@
 #include "CoreMinimal.h"
 #include "MetasoundFrontend.h"
 #include "MetasoundFrontendDocument.h"
+#include "MetasoundFrontendRegistryTransaction.h"
 
 namespace Metasound
 {
@@ -18,6 +19,21 @@ namespace Metasound
 		}
 	}
 
+	FGenerateNewlyAvailableNodeClasses::FGenerateNewlyAvailableNodeClasses()
+	: CurrentTransactionID(Frontend::GetOriginRegistryTransactionID())
+	{
+	}
+
+	void FGenerateNewlyAvailableNodeClasses::Generate(TArray<FFrontendQueryEntry>& OutEntries) const
+	{
+		const TArray<Frontend::FNodeClassInfo> ClassInfos = Frontend::GetNodeClassesRegisteredSince(CurrentTransactionID, &CurrentTransactionID);
+
+		for (const Frontend::FNodeClassInfo& ClassInfo : ClassInfos)
+		{
+			OutEntries.Emplace(TInPlaceType<FMetasoundFrontendClass>(), Frontend::GenerateClassDescription(ClassInfo));
+		}
+	}
+
 	FFilterClassesByInputVertexDataType::FFilterClassesByInputVertexDataType(const FName& InTypeName)
 	:	InputVertexTypeName(InTypeName)
 	{
diff --git a/Engine/Plugins/Runtime/Metasound/Source/MetasoundFrontend/Private/MetasoundFrontendRegistries.cpp b/Engine/Plugins/Runtime/Metasound/Source/MetasoundFrontend/Private/MetasoundFrontendRegistries.cpp
--- a/Engine/Plugins/Runtime/Metasound/Source/MetasoundFrontend/Private/MetasoundFrontendRegistries.cpp
+++ b/Engine/Plugins/Runtime/Metasound/Source/MetasoundFrontend/Private/MetasoundFrontendRegistries.cpp
@@ -3,6 +3,7 @@
 #include "MetasoundFrontendRegistries.h"
 
 #include "CoreMinimal.h"
+#include "MetasoundFrontendRegistryTransaction.h"
 #include "MetasoundLog.h"
 #include "MetasoundRouter.h"
 #include "Misc/ScopeLock.h"
@@ -12,444 +13,731 @@
 #define WITH_METASOUND_FRONTEND 0
 #endif
 
+
 namespace Metasound
 {
-	namespace MetasoundFrontendRegistryPrivate
+	namespace Frontend
 	{
-		// All registry keys should be created through this function to ensure consistency.
-		Frontend::FNodeRegistryKey GetRegistryKey(const FName& InClassName, int32 InMajorVersion)
+		namespace MetasoundFrontendRegistryPrivate
 		{
-			Frontend::FNodeRegistryKey Key;
+			// All registry keys should be created through this function to ensure consistency.
+			Frontend::FNodeRegistryKey GetRegistryKey(const FName& InClassName, int32 InMajorVersion)
+			{
+				Frontend::FNodeRegistryKey Key;
 
-			Key.NodeClassFullName = InClassName;
+				Key.NodeClassFullName = InClassName;
 
-			// NodeHash is hash of node name and major version.
-			Key.NodeHash = FCrc::StrCrc32(*Key.NodeClassFullName.ToString());
-			Key.NodeHash = HashCombine(Key.NodeHash, FCrc::TypeCrc32(InMajorVersion));
+				// NodeHash is hash of node name and major version.
+				Key.NodeHash = FCrc::StrCrc32(*Key.NodeClassFullName.ToString());
+				Key.NodeHash = HashCombine(Key.NodeHash, FCrc::TypeCrc32(InMajorVersion));
 
-			return Key;
-		}
-
-		// Return the compatible literal with the most descriptive type.
-		// TODO: Currently TIsParsable<> allows for implicit conversion of
-		// constructor arguments of integral types which can cause some confusion
-		// here when trying to match a literal type to a constructor. For example:
-		//
-		// struct FBoolConstructibleType
-		// {
-		// 	FBoolConstructibleType(bool InValue);
-		// };
-		//
-		// static_assert(TIsParsable<FBoolConstructible, double>::Value); 
-		//
-		// Implicit conversions are currently allowed in TIsParsable because this
-		// is perfectly legal syntax.
-		//
-		// double Value = 10.0;
-		// FBoolConstructibleType BoolConstructible = Value;
-		//
-		// There are some tricks to possibly disable implicit conversions when
-		// checking for specific constructors, but they are yet to be implemented 
-		// and are untested. Here's the basic idea.
-		//
-		// template<DataType, DesiredIntegralArgType>
-		// struct TOnlyConvertIfIsSame
-		// {
-		// 		// Implicit conversion only defined if types match.
-		// 		template<typename SuppliedIntegralArgType, std::enable_if<std::is_same<std::decay<SuppliedIntegralArgType>::type, DesiredIntegralArgType>::value, int> = 0>
-		// 		operator DesiredIntegralArgType()
-		// 		{
-		// 			return DesiredIntegralArgType{};
-		// 		}
-		// };
-		//
-		// static_assert(false == std::is_constructible<FBoolConstructibleType, TOnlyConvertIfSame<double>>::value);
-		// static_assert(true == std::is_constructible<FBoolConstructibleType, TOnlyConvertIfSame<bool>>::value);
-		ELiteralType GetMostDescriptiveLiteralForDataType(const FDataTypeRegistryInfo& InDataTypeInfo)
-		{
-			if (InDataTypeInfo.bIsProxyArrayParsable)
-			{
-				return ELiteralType::UObjectProxyArray;
-			}
-			else if (InDataTypeInfo.bIsProxyParsable)
-			{
-				return ELiteralType::UObjectProxy;
+				return Key;
 			}
-			else if (InDataTypeInfo.bIsEnum && InDataTypeInfo.bIsIntParsable)
+
+			// Return the compatible literal with the most descriptive type.
+			// TODO: Currently TIsParsable<> allows for implicit conversion of
+			// constructor arguments of integral types which can cause some confusion
+			// here when trying to match a literal type to a constructor. For example:
+			//
+			// struct FBoolConstructibleType
+			// {
+			// 	FBoolConstructibleType(bool InValue);
+			// };
+			//
+			// static_assert(TIsParsable<FBoolConstructible, double>::Value); 
+			//
+			// Implicit conversions are currently allowed in TIsParsable because this
+			// is perfectly legal syntax.
+			//
+			// double Value = 10.0;
+			// FBoolConstructibleType BoolConstructible = Value;
+			//
+			// There are some tricks to possibly disable implicit conversions when
+			// checking for specific constructors, but they are yet to be implemented 
+			// and are untested. Here's the basic idea.
+			//
+			// template<DataType, DesiredIntegralArgType>
+			// struct TOnlyConvertIfIsSame
+			// {
+			// 		// Implicit conversion only defined if types match.
+			// 		template<typename SuppliedIntegralArgType, std::enable_if<std::is_same<std::decay<SuppliedIntegralArgType>::type, DesiredIntegralArgType>::value, int> = 0>
+			// 		operator DesiredIntegralArgType()
+			// 		{
+			// 			return DesiredIntegralArgType{};
+			// 		}
+			// };
+			//
+			// static_assert(false == std::is_constructible<FBoolConstructibleType, TOnlyConvertIfSame<double>>::value);
+			// static_assert(true == std::is_constructible<FBoolConstructibleType, TOnlyConvertIfSame<bool>>::value);
+			ELiteralType GetMostDescriptiveLiteralForDataType(const FDataTypeRegistryInfo& InDataTypeInfo)
 			{
-				return ELiteralType::Integer;
+				if (InDataTypeInfo.bIsProxyArrayParsable)
+				{
+					return ELiteralType::UObjectProxyArray;
+				}
+				else if (InDataTypeInfo.bIsProxyParsable)
+				{
+					return ELiteralType::UObjectProxy;
+				}
+				else if (InDataTypeInfo.bIsEnum && InDataTypeInfo.bIsIntParsable)
+				{
+					return ELiteralType::Integer;
+				}
+				else if (InDataTypeInfo.bIsStringArrayParsable)
+				{
+					return ELiteralType::StringArray;
+				}
+				else if (InDataTypeInfo.bIsFloatArrayParsable)
+				{
+					return ELiteralType::FloatArray;
+				}
+				else if (InDataTypeInfo.bIsIntArrayParsable)
+				{
+					return ELiteralType::IntegerArray;
+				}
+				else if (InDataTypeInfo.bIsBoolArrayParsable)
+				{
+					return ELiteralType::BooleanArray;
+				}
+				else if (InDataTypeInfo.bIsStringParsable)
+				{
+					return ELiteralType::String;
+				}
+				else if (InDataTypeInfo.bIsFloatParsable)
+				{
+					return ELiteralType::Float;
+				}
+				else if (InDataTypeInfo.bIsIntParsable)
+				{
+					return ELiteralType::Integer;
+				}
+				else if (InDataTypeInfo.bIsBoolParsable)
+				{
+					return ELiteralType::Boolean;
+				}
+				else if (InDataTypeInfo.bIsDefaultArrayParsable)
+				{
+					return ELiteralType::NoneArray; 
+				}
+				else if (InDataTypeInfo.bIsDefaultParsable)
+				{
+					return ELiteralType::None;
+				}
+				else
+				{
+					// if we ever hit this, something has gone wrong with the REGISTER_METASOUND_DATATYPE macro.
+					// we should have failed to compile if any of these are false.
+					checkNoEntry();
+					return ELiteralType::Invalid;
+				}
 			}
-			else if (InDataTypeInfo.bIsStringArrayParsable)
+
+			// Registry container private implementation.
+			class FRegistryContainerImpl : public FMetasoundFrontendRegistryContainer
 			{
-				return ELiteralType::StringArray;
-			}
-			else if (InDataTypeInfo.bIsFloatArrayParsable)
+
+			public:
+				using FConverterNodeRegistryKey = ::Metasound::Frontend::FConverterNodeRegistryKey;
+				using FConverterNodeRegistryValue = ::Metasound::Frontend::FConverterNodeRegistryValue;
+				using FConverterNodeInfo = ::Metasound::Frontend::FConverterNodeInfo;
+
+				using FNodeRegistryKey = Metasound::Frontend::FNodeRegistryKey;
+				using FNodeRegistryElement = Metasound::Frontend::FNodeRegistryElement;
+
+				using FDataTypeRegistryInfo = Metasound::FDataTypeRegistryInfo;
+				using FDataTypeConstructorCallbacks = ::Metasound::FDataTypeConstructorCallbacks;
+				using FNodeClassMetadata = Metasound::FNodeClassMetadata;
+				using IEnumDataTypeInterface = Metasound::Frontend::IEnumDataTypeInterface;
+
+
+				FRegistryContainerImpl() = default;
+
+				FRegistryContainerImpl(const FRegistryContainerImpl&) = delete;
+				FRegistryContainerImpl& operator=(const FRegistryContainerImpl&) = delete;
+
+				virtual ~FRegistryContainerImpl() = default;
+
+				// Add a function to the init command array.
+				bool EnqueueInitCommand(TUniqueFunction<void()>&& InFunc) override;
+
+				// This is called on module startup. This invokes any registration commands enqueued by our registration macros.
+				void RegisterPendingNodes() override;
+
+				/** Register external node with the frontend.
+				 *
+				 * @param InCreateNode - Function for creating node from FNodeInitData.
+				 * @param InCreateDescription - Function for creating a FMetasoundFrontendClass.
+				 *
+				 * @return True on success.
+				 */
+				bool RegisterExternalNode(Metasound::FCreateMetasoundNodeFunction&& InCreateNode, Metasound::FCreateMetasoundFrontendClassFunction&& InCreateDescription) override;
+				bool RegisterConversionNode(const FConverterNodeRegistryKey& InNodeKey, const FConverterNodeInfo& InNodeInfo) override;
+				bool IsNodeRegistered(const FNodeRegistryKey& InKey) const override;
+
+				// Get all available nodes
+				TArray<Frontend::FNodeClassInfo> GetAllAvailableNodeClasses(Frontend::FRegistryTransactionID* OutCurrentRegistryTransactionID) const override;
+				TArray<Metasound::Frontend::FNodeClassInfo> GetNodeClassesRegisteredSince(Metasound::Frontend::FRegistryTransactionID InSince, Metasound::Frontend::FRegistryTransactionID* OutCurrentRegistryTransactionID) const override;
+
+				// Return any data types that can be used as a metasound input type or output type.
+				TArray<FName> GetAllValidDataTypes() override;
+
+				// Find Frontend Document data.
+				bool FindFrontendClassFromRegistered(const FNodeClassInfo& InClassInfo, FMetasoundFrontendClass& OutClass) override;
+				bool FindFrontendClassFromRegistered(const FMetasoundFrontendClassMetadata& InMetadata, FMetasoundFrontendClass& OutClass) override;
+				bool FindInputNodeClassMetadataForDataType(const FName& InDataTypeName, FMetasoundFrontendClassMetadata& OutMetadata) override;
+				bool FindOutputNodeClassMetadataForDataType(const FName& InDataTypeName, FMetasoundFrontendClassMetadata& OutMetadata) override;
+
+				// Create a new instance of a C++ implemented node from the registry.
+				TUniquePtr<Metasound::INode> ConstructInputNode(const FName& InInputType, Metasound::FInputNodeConstructorParams&& InParams) override;
+				TUniquePtr<Metasound::INode> ConstructOutputNode(const FName& InOutputType, const Metasound::FOutputNodeConstructorParams& InParams) override;
+				TUniquePtr<Metasound::INode> ConstructExternalNode(const FNodeRegistryKey& InKey, const Metasound::FNodeInitData& InInitData) override;
+
+				// Returns a list of possible nodes to use to convert from FromDataType to ToDataType.
+				// Returns an empty array if none are available.
+				TArray<FConverterNodeInfo> GetPossibleConverterNodes(const FName& FromDataType, const FName& ToDataType) override;
+
+				bool RegisterDataType(const FDataTypeRegistryInfo& InDataInfo, const FDataTypeConstructorCallbacks& InCallbacks) override;
+
+				bool RegisterEnumDataInterface(FName InDataType, TSharedPtr<IEnumDataTypeInterface>&& InInterface) override;
+
+				// Get the desired kind of literal for a given data type. Returns EConstructorArgType::Invalid if the data type couldn't be found.
+				Metasound::ELiteralType GetDesiredLiteralTypeForDataType(FName InDataType) const override;
+
+				// Get whether we can build a literal of this specific type for InDataType.
+				bool DoesDataTypeSupportLiteralType(FName InDataType, Metasound::ELiteralType InLiteralType) const override;
+
+				// Handle uobjects and literals
+				UClass* GetLiteralUClassForDataType(FName InDataType) const override;
+				Metasound::FLiteral GenerateLiteralForUObject(const FName& InDataType, UObject* InObject) override;
+				Metasound::FLiteral GenerateLiteralForUObjectArray(const FName& InDataType, TArray<UObject*> InObjectArray) override;
+
+				// Get info about a specific data type (what kind of literals we can use, etc.)
+				// @returns false if InDataType wasn't found in the registry. 
+				bool GetInfoForDataType(FName InDataType, FDataTypeRegistryInfo& OutInfo) override;
+
+				TSharedPtr<const Metasound::Frontend::IEnumDataTypeInterface> GetEnumInterfaceForDataType(FName InDataType) const override;
+
+				TSharedPtr<Metasound::IDataChannel, ESPMode::ThreadSafe> CreateDataChannelForDataType(const FName& InDataType, const Metasound::FOperatorSettings& InOperatorSettings) const override;
+			private:
+
+
+				// This buffer is used to enqueue nodes and datatypes to register when the module has been initialized,
+				// in order to avoid bad behavior with ensures, logs, etc. on static initialization.
+				// The bad news is that TInlineAllocator is the safest allocator to use on static init.
+				// The good news is that none of these lambdas typically have captures, so this should have low memory overhead.
+				static constexpr int32 MaxNumNodesAndDatatypesToInitialize = 8192;
+				TArray<TUniqueFunction<void()>, TInlineAllocator<MaxNumNodesAndDatatypesToInitialize>> LazyInitCommands;
+				
+				FCriticalSection LazyInitCommandCritSection;
+
+				// Registry in which we keep all information about nodes implemented in C++.
+				TMap<FNodeRegistryKey, FNodeRegistryElement> ExternalNodeRegistry;
+
+				// Registry in which we keep lists of possible nodes to use to convert between two datatypes
+				TMap<FConverterNodeRegistryKey, FConverterNodeRegistryValue> ConverterNodeRegistry;
+
+				struct FDataTypeRegistryElement
+				{
+					Metasound::FDataTypeConstructorCallbacks Callbacks;
+
+					Metasound::FDataTypeRegistryInfo Info;
+
+					TSharedPtr<const Metasound::Frontend::IEnumDataTypeInterface> EnumInterface;
+				};
+
+				TMap<FName, FDataTypeRegistryElement> DataTypeRegistry;
+				TMap<FNodeRegistryKey, FDataTypeRegistryElement> DataTypeNodeRegistry;
+				
+				FRegistryTransactionHistory RegistryTransactionHistory;
+			};
+
+			void FRegistryContainerImpl::RegisterPendingNodes()
 			{
-				return ELiteralType::FloatArray;
+				FScopeLock ScopeLock(&LazyInitCommandCritSection);
+
+				UE_LOG(LogTemp, Display, TEXT("Processing %i Metasounds Frontend Registration Requests."), LazyInitCommands.Num());
+				uint64 CurrentTime = FPlatformTime::Cycles64();
+
+				for (TUniqueFunction<void()>& Command : LazyInitCommands)
+				{
+					Command();
+				}
+
+				LazyInitCommands.Empty();
+
+				uint64 CyclesUsed = FPlatformTime::Cycles64() - CurrentTime;
+				UE_LOG(LogTemp, Display, TEXT("Initializing Metasounds Frontend took %f seconds."), FPlatformTime::ToSeconds64(CyclesUsed));
 			}
-			else if (InDataTypeInfo.bIsIntArrayParsable)
+
+			bool FRegistryContainerImpl::EnqueueInitCommand(TUniqueFunction<void()>&& InFunc)
 			{
-				return ELiteralType::IntegerArray;
+				FScopeLock ScopeLock(&LazyInitCommandCritSection);
+				if (LazyInitCommands.Num() >= MaxNumNodesAndDatatypesToInitialize)
+				{
+					UE_LOG(LogTemp, Warning, TEXT("Registering more that %d nodes and datatypes for metasounds! Consider increasing MetasoundFrontendRegistryContainer::MaxNumNodesAndDatatypesToInitialize."));
+				}
+
+				LazyInitCommands.Add(MoveTemp(InFunc));
+				return true;
 			}
-			else if (InDataTypeInfo.bIsBoolArrayParsable)
+
+			TArray<FNodeClassInfo> FRegistryContainerImpl::GetAllAvailableNodeClasses(FRegistryTransactionID* OutCurrentRegistryTransactionID) const 
 			{
-				return ELiteralType::BooleanArray;
+				TArray<FNodeClassInfo> OutClasses;
+
+				TArray<const IRegistryTransaction*> Transactions = RegistryTransactionHistory.GetTransactions(GetOriginRegistryTransactionID(), OutCurrentRegistryTransactionID);
+
+				for (const IRegistryTransaction* Transaction : Transactions)
+				{
+					if (const FNodeClassInfo* Info = Transaction->GetNodeClassInfo())
+					{
+						OutClasses.Add(*Info);
+					}
+				}
+
+				return OutClasses;
 			}
-			else if (InDataTypeInfo.bIsStringParsable)
+
+			TArray<Metasound::Frontend::FNodeClassInfo> FRegistryContainerImpl::GetNodeClassesRegisteredSince(Metasound::Frontend::FRegistryTransactionID InSince, Metasound::Frontend::FRegistryTransactionID* OutCurrentRegistryTransactionID) const 
 			{
-				return ELiteralType::String;
+				TArray<FNodeClassInfo> OutClasses;
+
+				TArray<const IRegistryTransaction*> Transactions = RegistryTransactionHistory.GetTransactions(InSince, OutCurrentRegistryTransactionID);
+
+				for (const IRegistryTransaction* Transaction : Transactions)
+				{
+					if (const FNodeClassInfo* Info = Transaction->GetNodeClassInfo())
+					{
+						OutClasses.Add(*Info);
+					}
+				}
+
+				return OutClasses;
 			}
-			else if (InDataTypeInfo.bIsFloatParsable)
+
+
+			TUniquePtr<Metasound::INode> FRegistryContainerImpl::ConstructInputNode(const FName& InInputType, Metasound::FInputNodeConstructorParams&& InParams)
 			{
-				return ELiteralType::Float;
+				if (ensureAlwaysMsgf(DataTypeRegistry.Contains(InInputType), TEXT("Couldn't find data type %s!"), *InInputType.ToString()))
+				{
+					return DataTypeRegistry[InInputType].Callbacks.CreateInputNode(MoveTemp(InParams));
+				}
+				else
+				{
+					return nullptr;
+				}
 			}
-			else if (InDataTypeInfo.bIsIntParsable)
+
+			TUniquePtr<Metasound::INode> FRegistryContainerImpl::ConstructOutputNode(const FName& InOutputType, const Metasound::FOutputNodeConstructorParams& InParams)
 			{
-				return ELiteralType::Integer;
+				if (ensureAlwaysMsgf(DataTypeRegistry.Contains(InOutputType), TEXT("Couldn't find data type %s!"), *InOutputType.ToString()))
+				{
+					return DataTypeRegistry[InOutputType].Callbacks.CreateOutputNode(InParams);
+				}
+				else
+				{
+					return nullptr;
+				}
 			}
-			else if (InDataTypeInfo.bIsBoolParsable)
+
+			Metasound::FLiteral FRegistryContainerImpl::GenerateLiteralForUObject(const FName& InDataType, UObject* InObject)
 			{
-				return ELiteralType::Boolean;
+				if (ensureAlwaysMsgf(DataTypeRegistry.Contains(InDataType), TEXT("Couldn't find data type %s!"), *InDataType.ToString()))
+				{
+					 Audio::IProxyDataPtr ProxyPtr = DataTypeRegistry[InDataType].Callbacks.CreateAudioProxy(InObject);
+					 if (ensureAlwaysMsgf(ProxyPtr.IsValid(), TEXT("UObject failed to create a valid proxy!")))
+					 {
+						 return Metasound::FLiteral(MoveTemp(ProxyPtr));
+					 }
+					 else
+					 {
+						 return Metasound::FLiteral();
+					 }
+				}
+				else
+				{
+					return Metasound::FLiteral();
+				}
 			}
-			else if (InDataTypeInfo.bIsDefaultArrayParsable)
+
+			Metasound::FLiteral FRegistryContainerImpl::GenerateLiteralForUObjectArray(const FName& InDataType, TArray<UObject*> InObjectArray)
 			{
-				return ELiteralType::NoneArray; 
+				if (ensureAlwaysMsgf(DataTypeRegistry.Contains(InDataType), TEXT("Couldn't find data type %s!"), *InDataType.ToString()))
+				{
+					TArray<Audio::IProxyDataPtr> ProxyArray;
+
+					for (UObject* InObject : InObjectArray)
+					{
+						if (InObject)
+						{
+							Audio::IProxyDataPtr ProxyPtr = DataTypeRegistry[InDataType].Callbacks.CreateAudioProxy(InObject);
+							ensureAlwaysMsgf(ProxyPtr.IsValid(), TEXT("UObject failed to create a valid proxy!"));
+							ProxyArray.Add(MoveTemp(ProxyPtr));
+						}
+					}
+
+					return Metasound::FLiteral(MoveTemp(ProxyArray));
+				}
+				else
+				{
+					return Metasound::FLiteral();
+				}
 			}
-			else if (InDataTypeInfo.bIsDefaultParsable)
+
+			TUniquePtr<Metasound::INode> FRegistryContainerImpl::ConstructExternalNode(const FNodeRegistryKey& InKey, const Metasound::FNodeInitData& InInitData)
 			{
-				return ELiteralType::None;
+				if (!ExternalNodeRegistry.Contains(InKey))
+				{
+					return nullptr;
+				}
+				else
+				{
+					return ExternalNodeRegistry[InKey].CreateNode(InInitData);
+				}
 			}
-			else
+
+			TArray<::Metasound::Frontend::FConverterNodeInfo> FRegistryContainerImpl::GetPossibleConverterNodes(const FName& FromDataType, const FName& ToDataType)
 			{
-				// if we ever hit this, something has gone wrong with the REGISTER_METASOUND_DATATYPE macro.
-				// we should have failed to compile if any of these are false.
-				checkNoEntry();
-				return ELiteralType::Invalid;
+				FConverterNodeRegistryKey InKey = { FromDataType, ToDataType };
+				if (!ConverterNodeRegistry.Contains(InKey))
+				{
+					return TArray<FConverterNodeInfo>();
+				}
+				else
+				{
+					return ConverterNodeRegistry[InKey].PotentialConverterNodes;
+				}
 			}
-		}
-	}
-}
 
-FMetasoundFrontendRegistryContainer::FMetasoundFrontendRegistryContainer()
-{
-}
 
-FMetasoundFrontendRegistryContainer* FMetasoundFrontendRegistryContainer::LazySingleton = nullptr;
-
-FMetasoundFrontendRegistryContainer* FMetasoundFrontendRegistryContainer::Get()
-{
-	if (!LazySingleton)
-	{
-		LazySingleton = new FMetasoundFrontendRegistryContainer();
-	}
 
-	return LazySingleton;
-}
+			Metasound::ELiteralType FRegistryContainerImpl::GetDesiredLiteralTypeForDataType(FName InDataType) const
+			{
+				if (!DataTypeRegistry.Contains(InDataType))
+				{
+					return Metasound::ELiteralType::Invalid;
+				}
 
-void FMetasoundFrontendRegistryContainer::ShutdownMetasoundFrontend()
-{
-	if (LazySingleton)
-	{
-		delete LazySingleton;
-		LazySingleton = nullptr;
-	}
-}
+				const FDataTypeRegistryElement& DataTypeInfo = DataTypeRegistry[InDataType];
 
-void FMetasoundFrontendRegistryContainer::RegisterPendingNodes()
-{
-	FScopeLock ScopeLock(&LazyInitCommandCritSection);
+				// If there's a designated preferred literal type for this datatype, use that.
+				if (DataTypeInfo.Info.PreferredLiteralType != Metasound::ELiteralType::Invalid)
+				{
+					return DataTypeInfo.Info.PreferredLiteralType;
+				}
 
-	UE_LOG(LogTemp, Display, TEXT("Processing %i Metasounds Frontend Registration Requests."), LazyInitCommands.Num());
-	uint64 CurrentTime = FPlatformTime::Cycles64();
+				// Otherwise, we opt for the highest precision construction option available.
+				return Metasound::Frontend::MetasoundFrontendRegistryPrivate::GetMostDescriptiveLiteralForDataType(DataTypeInfo.Info);
+			}
 
-	for (TUniqueFunction<void()>& Command : LazyInitCommands)
-	{
-		Command();
-	}
+			UClass* FRegistryContainerImpl::GetLiteralUClassForDataType(FName InDataType) const
+			{
+				if (!DataTypeRegistry.Contains(InDataType))
+				{
+					ensureAlwaysMsgf(false, TEXT("DataType %s not registered."), *InDataType.ToString());
+					return nullptr;
+				}
+				else
+				{
+					return DataTypeRegistry[InDataType].Info.ProxyGeneratorClass;
+				}
+			}
 
-	LazyInitCommands.Empty();
+			bool FRegistryContainerImpl::DoesDataTypeSupportLiteralType(FName InDataType, Metasound::ELiteralType InLiteralType) const
+			{
+				if (!DataTypeRegistry.Contains(InDataType))
+				{
+					ensureAlwaysMsgf(false, TEXT("DataType %s not registered."), *InDataType.ToString());
+					return false;
+				}
+
+				const FDataTypeRegistryElement& DataTypeInfo = DataTypeRegistry[InDataType];
+				
+				switch (InLiteralType)
+				{
+					case Metasound::ELiteralType::Boolean:
+					{
+						return DataTypeInfo.Info.bIsBoolParsable;
+					}
+					case Metasound::ELiteralType::BooleanArray:
+					{
+						return DataTypeInfo.Info.bIsBoolArrayParsable;
+					}
+
+					case Metasound::ELiteralType::Integer:
+					{
+						return DataTypeInfo.Info.bIsIntParsable;
+					}
+					case Metasound::ELiteralType::IntegerArray:
+					{
+						return DataTypeInfo.Info.bIsIntArrayParsable;
+					}
+
+					case Metasound::ELiteralType::Float:
+					{
+						return DataTypeInfo.Info.bIsFloatParsable;
+					}
+					case Metasound::ELiteralType::FloatArray:
+					{
+						return DataTypeInfo.Info.bIsFloatArrayParsable;
+					}
+
+					case Metasound::ELiteralType::String:
+					{
+						return DataTypeInfo.Info.bIsStringParsable;
+					}
+					case Metasound::ELiteralType::StringArray:
+					{
+						return DataTypeInfo.Info.bIsStringArrayParsable;
+					}
+
+					case Metasound::ELiteralType::UObjectProxy:
+					{
+						return DataTypeInfo.Info.bIsProxyParsable;
+					}
+					case Metasound::ELiteralType::UObjectProxyArray:
+					{
+						return DataTypeInfo.Info.bIsProxyArrayParsable;
+					}
+
+					case Metasound::ELiteralType::None:
+					{
+						return DataTypeInfo.Info.bIsDefaultParsable;
+					}
+					case Metasound::ELiteralType::NoneArray:
+					{
+						return DataTypeInfo.Info.bIsDefaultArrayParsable;
+					}
+
+					case Metasound::ELiteralType::Invalid:
+					default:
+					{
+						static_assert(static_cast<int32>(Metasound::ELiteralType::COUNT) == 13, "Possible missing case coverage for ELiteralType");
+						return false;
+					}
+				}
+			}
 
-	uint64 CyclesUsed = FPlatformTime::Cycles64() - CurrentTime;
-	UE_LOG(LogTemp, Display, TEXT("Initializing Metasounds Frontend took %f seconds."), FPlatformTime::ToSeconds64(CyclesUsed));
-}
+			bool FRegistryContainerImpl::RegisterDataType(const ::Metasound::FDataTypeRegistryInfo& InDataInfo, const ::Metasound::FDataTypeConstructorCallbacks& InCallbacks)
+			{
+				if (!ensureAlwaysMsgf(!DataTypeRegistry.Contains(InDataInfo.DataTypeName),
+					TEXT("Name collision when trying to register Metasound Data Type %s! DataType must have "
+						"unique name and REGISTER_METASOUND_DATATYPE cannot be called in a public header."),
+						*InDataInfo.DataTypeName.ToString()))
+				{
+					// todo: capture callstack for previous declaration for non-shipping builds to help clarify who already registered this name for a type.
+					return false;
+				}
+				else
+				{
+					FDataTypeRegistryElement InElement = { InCallbacks, InDataInfo };
+
+					DataTypeRegistry.Add(InDataInfo.DataTypeName, InElement);
+
+					Metasound::Frontend::FNodeRegistryKey InputNodeRegistryKey = GetRegistryKey(InCallbacks.CreateFrontendInputClass().Metadata);
+					DataTypeNodeRegistry.Add(InputNodeRegistryKey, InElement);
+
+					Metasound::Frontend::FNodeRegistryKey OutputNodeRegistryKey = GetRegistryKey(InCallbacks.CreateFrontendOutputClass().Metadata);
+					DataTypeNodeRegistry.Add(OutputNodeRegistryKey, InElement);
+
+					UE_LOG(LogMetasound, Display, TEXT("Registered Metasound Datatype %s."), *InDataInfo.DataTypeName.ToString());
+					return true;
+				}
+			}
 
-bool FMetasoundFrontendRegistryContainer::EnqueueInitCommand(TUniqueFunction<void()>&& InFunc)
-{
-	FScopeLock ScopeLock(&LazyInitCommandCritSection);
-	if (LazyInitCommands.Num() >= MaxNumNodesAndDatatypesToInitialize)
-	{
-		UE_LOG(LogTemp, Warning, TEXT("Registering more that %d nodes and datatypes for metasounds! Consider increasing MetasoundFrontendRegistryContainer::MaxNumNodesAndDatatypesToInitialize."));
-	}
+			bool FRegistryContainerImpl::RegisterEnumDataInterface(FName InDataType, TSharedPtr<IEnumDataTypeInterface>&& InInterface)
+			{
+				if (FDataTypeRegistryElement* Element = DataTypeRegistry.Find(InDataType))
+				{
+					Element->EnumInterface = MoveTemp(InInterface);
+					return true;
+				}
+				return false;
+			}
 
-	LazyInitCommands.Add(MoveTemp(InFunc));
-	return true;
-}
+			bool FRegistryContainerImpl::RegisterExternalNode(Metasound::FCreateMetasoundNodeFunction&& InCreateNode, Metasound::FCreateMetasoundFrontendClassFunction&& InCreateDescription)
+			{
+				FNodeRegistryElement RegistryElement = FNodeRegistryElement(MoveTemp(InCreateNode), MoveTemp(InCreateDescription));
 
-TMap<FMetasoundFrontendRegistryContainer::FNodeRegistryKey, FMetasoundFrontendRegistryContainer::FNodeRegistryElement>& FMetasoundFrontendRegistryContainer::GetExternalNodeRegistry()
-{
-	return ExternalNodeRegistry;
-}
+				FNodeRegistryKey Key;
 
-TUniquePtr<Metasound::INode> FMetasoundFrontendRegistryContainer::ConstructInputNode(const FName& InInputType, Metasound::FInputNodeConstructorParams&& InParams)
-{
-	if (ensureAlwaysMsgf(DataTypeRegistry.Contains(InInputType), TEXT("Couldn't find data type %s!"), *InInputType.ToString()))
-	{
-		return DataTypeRegistry[InInputType].Callbacks.CreateInputNode(MoveTemp(InParams));
-	}
-	else
-	{
-		return nullptr;
-	}
-}
+				if (ensure(FMetasoundFrontendRegistryContainer::GetRegistryKey(RegistryElement, Key)))
+				{
+					// check to see if an identical node was already registered, and log
+					ensureAlwaysMsgf(!ExternalNodeRegistry.Contains(Key), TEXT("Node with identical name, inputs and outputs to node %s was already registered. The previously registered node will be overwritten. This could also happen because METASOUND_REGISTER_NODE is in a public header."), *Key.NodeClassFullName.ToString());
 
-TUniquePtr<Metasound::INode> FMetasoundFrontendRegistryContainer::ConstructOutputNode(const FName& InOutputType, const Metasound::FOutputNodeConstructorParams& InParams)
-{
-	if (ensureAlwaysMsgf(DataTypeRegistry.Contains(InOutputType), TEXT("Couldn't find data type %s!"), *InOutputType.ToString()))
-	{
-		return DataTypeRegistry[InOutputType].Callbacks.CreateOutputNode(InParams);
-	}
-	else
-	{
-		return nullptr;
-	}
-}
+					// Store update to newly registered node in history so nodes
+					// can be queried by transaction ID
+					FNodeClassInfo ClassInfo = {EMetasoundFrontendClassType::External, Key};
+					RegistryTransactionHistory.Add(MakeNodeRegistrationTransaction(ClassInfo));
 
-Metasound::FLiteral FMetasoundFrontendRegistryContainer::GenerateLiteralForUObject(const FName& InDataType, UObject* InObject)
-{
-	if (ensureAlwaysMsgf(DataTypeRegistry.Contains(InDataType), TEXT("Couldn't find data type %s!"), *InDataType.ToString()))
-	{
-		 Audio::IProxyDataPtr ProxyPtr = DataTypeRegistry[InDataType].Callbacks.CreateAudioProxy(InObject);
-		 if (ensureAlwaysMsgf(ProxyPtr.IsValid(), TEXT("UObject failed to create a valid proxy!")))
-		 {
-			 return Metasound::FLiteral(MoveTemp(ProxyPtr));
-		 }
-		 else
-		 {
-			 return Metasound::FLiteral();
-		 }
-	}
-	else
-	{
-		return Metasound::FLiteral();
-	}
-}
+					// Store registry elements in map so nodes can be queried using registry key. 
+					ExternalNodeRegistry.Add(MoveTemp(Key), MoveTemp(RegistryElement));
 
-Metasound::FLiteral FMetasoundFrontendRegistryContainer::GenerateLiteralForUObjectArray(const FName& InDataType, TArray<UObject*> InObjectArray)
-{
-	if (ensureAlwaysMsgf(DataTypeRegistry.Contains(InDataType), TEXT("Couldn't find data type %s!"), *InDataType.ToString()))
-	{
-		TArray<Audio::IProxyDataPtr> ProxyArray;
+					return true;
+				}
 
-		for (UObject* InObject : InObjectArray)
-		{
-			if (InObject)
-			{
-				Audio::IProxyDataPtr ProxyPtr = DataTypeRegistry[InDataType].Callbacks.CreateAudioProxy(InObject);
-				ensureAlwaysMsgf(ProxyPtr.IsValid(), TEXT("UObject failed to create a valid proxy!"));
-				ProxyArray.Add(MoveTemp(ProxyPtr));
+				return false;
 			}
-		}
 
-		return Metasound::FLiteral(MoveTemp(ProxyArray));
-	}
-	else
-	{
-		return Metasound::FLiteral();
-	}
-}
-
-TUniquePtr<Metasound::INode> FMetasoundFrontendRegistryContainer::ConstructExternalNode(const FName& InNodeClassFullName, uint32 InNodeHash, const Metasound::FNodeInitData& InInitData)
-{
-	Metasound::Frontend::FNodeRegistryKey RegistryKey;
-	RegistryKey.NodeClassFullName = InNodeClassFullName;
-	RegistryKey.NodeHash = InNodeHash;
-
-	if (!ExternalNodeRegistry.Contains(RegistryKey))
-	{
-		return nullptr;
-	}
-	else
-	{
-		return ExternalNodeRegistry[RegistryKey].CreateNode(InInitData);
-	}
-}
 
-TArray<::Metasound::Frontend::FConverterNodeInfo> FMetasoundFrontendRegistryContainer::GetPossibleConverterNodes(const FName& FromDataType, const FName& ToDataType)
-{
-	FConverterNodeRegistryKey InKey = { FromDataType, ToDataType };
-	if (!ConverterNodeRegistry.Contains(InKey))
-	{
-		return TArray<FConverterNodeInfo>();
-	}
-	else
-	{
-		return ConverterNodeRegistry[InKey].PotentialConverterNodes;
-	}
-}
+			bool FRegistryContainerImpl::RegisterConversionNode(const FConverterNodeRegistryKey& InNodeKey, const FConverterNodeInfo& InNodeInfo)
+			{
+				if (!ConverterNodeRegistry.Contains(InNodeKey))
+				{
+					ConverterNodeRegistry.Add(InNodeKey);
+				}
+
+				FConverterNodeRegistryValue& ConverterNodeList = ConverterNodeRegistry[InNodeKey];
+
+				if (ensureAlways(!ConverterNodeList.PotentialConverterNodes.Contains(InNodeInfo)))
+				{
+					ConverterNodeList.PotentialConverterNodes.Add(InNodeInfo);
+					return true;
+				}
+				else
+				{
+					// If we hit this, someone attempted to add the same converter node to our list multiple times.
+					return false;
+				}
+			}
 
+			bool FRegistryContainerImpl::IsNodeRegistered(const FNodeRegistryKey& InKey) const
+			{
+				return ExternalNodeRegistry.Contains(InKey);
+			}
 
+			TArray<FName> FRegistryContainerImpl::GetAllValidDataTypes()
+			{
+				TArray<FName> OutDataTypes;
 
-Metasound::ELiteralType FMetasoundFrontendRegistryContainer::GetDesiredLiteralTypeForDataType(FName InDataType) const
-{
-	if (!DataTypeRegistry.Contains(InDataType))
-	{
-		return Metasound::ELiteralType::Invalid;
-	}
+				for (auto& DataTypeTuple : DataTypeRegistry)
+				{
+					OutDataTypes.Add(DataTypeTuple.Key);
+				}
 
-	const FDataTypeRegistryElement& DataTypeInfo = DataTypeRegistry[InDataType];
+				return OutDataTypes;
+			}
 
-	// If there's a designated preferred literal type for this datatype, use that.
-	if (DataTypeInfo.Info.PreferredLiteralType != Metasound::ELiteralType::Invalid)
-	{
-		return DataTypeInfo.Info.PreferredLiteralType;
-	}
+			bool FRegistryContainerImpl::GetInfoForDataType(FName InDataType, Metasound::FDataTypeRegistryInfo& OutInfo)
+			{
+				if (!DataTypeRegistry.Contains(InDataType))
+				{
+					return false;
+				}
+				else
+				{
+					OutInfo = DataTypeRegistry[InDataType].Info;
+					return true;
+				}
+			}
 
-	// Otherwise, we opt for the highest precision construction option available.
-	return Metasound::MetasoundFrontendRegistryPrivate::GetMostDescriptiveLiteralForDataType(DataTypeInfo.Info);
-}
 
-UClass* FMetasoundFrontendRegistryContainer::GetLiteralUClassForDataType(FName InDataType) const
-{
-	if (!DataTypeRegistry.Contains(InDataType))
-	{
-		ensureAlwaysMsgf(false, TEXT("DataType %s not registered."), *InDataType.ToString());
-		return nullptr;
-	}
-	else
-	{
-		return DataTypeRegistry[InDataType].Info.ProxyGeneratorClass;
-	}
-}
+			TSharedPtr<const Metasound::Frontend::IEnumDataTypeInterface>
+			FRegistryContainerImpl::GetEnumInterfaceForDataType(FName InDataType) const
+			{
+				if (const FDataTypeRegistryElement* Element = DataTypeRegistry.Find(InDataType))
+				{
+					return Element->EnumInterface;
+				}
+				return nullptr;
+			}
 
-bool FMetasoundFrontendRegistryContainer::DoesDataTypeSupportLiteralType(FName InDataType, Metasound::ELiteralType InLiteralType) const
-{
-	if (!DataTypeRegistry.Contains(InDataType))
-	{
-		ensureAlwaysMsgf(false, TEXT("DataType %s not registered."), *InDataType.ToString());
-		return false;
-	}
 
-	const FDataTypeRegistryElement& DataTypeInfo = DataTypeRegistry[InDataType];
-	
-	switch (InLiteralType)
-	{
-		case Metasound::ELiteralType::Boolean:
-		{
-			return DataTypeInfo.Info.bIsBoolParsable;
-		}
-		case Metasound::ELiteralType::BooleanArray:
-		{
-			return DataTypeInfo.Info.bIsBoolArrayParsable;
-		}
+			TSharedPtr<Metasound::IDataChannel, ESPMode::ThreadSafe> FRegistryContainerImpl::CreateDataChannelForDataType(const FName& InDataType, const Metasound::FOperatorSettings& InSettings) const
+			{
+				if (const FDataTypeRegistryElement* Element = DataTypeRegistry.Find(InDataType))
+				{
+					if (Element->Callbacks.CreateDataChannel)
+					{
+						return Element->Callbacks.CreateDataChannel(InSettings);
+					}
+				}
+				return nullptr;
+			}
 
-		case Metasound::ELiteralType::Integer:
-		{
-			return DataTypeInfo.Info.bIsIntParsable;
-		}
-		case Metasound::ELiteralType::IntegerArray:
-		{
-			return DataTypeInfo.Info.bIsIntArrayParsable;
-		}
+			bool FRegistryContainerImpl::FindFrontendClassFromRegistered(const FNodeClassInfo& InClassInfo, FMetasoundFrontendClass& OutClass)
+			{
+				if (InClassInfo.NodeType == EMetasoundFrontendClassType::External)
+				{
+					if (ExternalNodeRegistry.Contains(InClassInfo.LookupKey))
+					{
+						OutClass = ExternalNodeRegistry[InClassInfo.LookupKey].CreateFrontendClass();
+						return true;
+					}
+				}
+				else if (InClassInfo.NodeType == EMetasoundFrontendClassType::Input)
+				{
+					if (DataTypeNodeRegistry.Contains(InClassInfo.LookupKey))
+					{
+						OutClass = DataTypeNodeRegistry[InClassInfo.LookupKey].Callbacks.CreateFrontendInputClass();
+						return true;
+					}
+				}
+				else if (InClassInfo.NodeType == EMetasoundFrontendClassType::Output)
+				{
+					if (DataTypeNodeRegistry.Contains(InClassInfo.LookupKey))
+					{
+						OutClass = DataTypeNodeRegistry[InClassInfo.LookupKey].Callbacks.CreateFrontendOutputClass();
+						return true;
+					}
+				}
+
+				return false;
+			}
 
-		case Metasound::ELiteralType::Float:
-		{
-			return DataTypeInfo.Info.bIsFloatParsable;
-		}
-		case Metasound::ELiteralType::FloatArray:
-		{
-			return DataTypeInfo.Info.bIsFloatArrayParsable;
-		}
+			bool FRegistryContainerImpl::FindFrontendClassFromRegistered(const FMetasoundFrontendClassMetadata& InMetadata, FMetasoundFrontendClass& OutClass)
+			{
+				FNodeClassInfo Info;
+				Info.NodeType = InMetadata.Type;
+				Info.LookupKey = GetRegistryKey(InMetadata);
 
-		case Metasound::ELiteralType::String:
-		{
-			return DataTypeInfo.Info.bIsStringParsable;
-		}
-		case Metasound::ELiteralType::StringArray:
-		{
-			return DataTypeInfo.Info.bIsStringArrayParsable;
-		}
+				return FindFrontendClassFromRegistered(Info, OutClass);
 
-		case Metasound::ELiteralType::UObjectProxy:
-		{
-			return DataTypeInfo.Info.bIsProxyParsable;
-		}
-		case Metasound::ELiteralType::UObjectProxyArray:
-		{
-			return DataTypeInfo.Info.bIsProxyArrayParsable;
-		}
+			}
 
-		case Metasound::ELiteralType::None:
-		{
-			return DataTypeInfo.Info.bIsDefaultParsable;
-		}
-		case Metasound::ELiteralType::NoneArray:
-		{
-			return DataTypeInfo.Info.bIsDefaultArrayParsable;
-		}
+			bool FRegistryContainerImpl::FindInputNodeClassMetadataForDataType(const FName& InDataTypeName, FMetasoundFrontendClassMetadata& OutMetadata)
+			{
+				if (DataTypeRegistry.Contains(InDataTypeName))
+				{
+					OutMetadata = DataTypeRegistry[InDataTypeName].Callbacks.CreateFrontendInputClass().Metadata;
+					return true;
+				}
+				return false;
+			}
 
-		case Metasound::ELiteralType::Invalid:
-		default:
-		{
-			static_assert(static_cast<int32>(Metasound::ELiteralType::COUNT) == 13, "Possible missing case coverage for ELiteralType");
-			return false;
+			bool FRegistryContainerImpl::FindOutputNodeClassMetadataForDataType(const FName& InDataTypeName, FMetasoundFrontendClassMetadata& OutMetadata)
+			{
+				if (DataTypeRegistry.Contains(InDataTypeName))
+				{
+					OutMetadata = DataTypeRegistry[InDataTypeName].Callbacks.CreateFrontendOutputClass().Metadata;
+					return true;
+				}
+				return false;
+			}
 		}
 	}
 }
 
-bool FMetasoundFrontendRegistryContainer::RegisterDataType(const ::Metasound::FDataTypeRegistryInfo& InDataInfo, const ::Metasound::FDataTypeConstructorCallbacks& InCallbacks)
-{
-	if (!ensureAlwaysMsgf(!DataTypeRegistry.Contains(InDataInfo.DataTypeName),
-		TEXT("Name collision when trying to register Metasound Data Type %s! DataType must have "
-			"unique name and REGISTER_METASOUND_DATATYPE cannot be called in a public header."),
-			*InDataInfo.DataTypeName.ToString()))
-	{
-		// todo: capture callstack for previous declaration for non-shipping builds to help clarify who already registered this name for a type.
-		return false;
-	}
-	else
-	{
-		FDataTypeRegistryElement InElement = { InCallbacks, InDataInfo };
-
-		DataTypeRegistry.Add(InDataInfo.DataTypeName, InElement);
-
-		Metasound::Frontend::FNodeRegistryKey InputNodeRegistryKey = GetRegistryKey(InCallbacks.CreateFrontendInputClass().Metadata);
-		DataTypeNodeRegistry.Add(InputNodeRegistryKey, InElement);
-
-		Metasound::Frontend::FNodeRegistryKey OutputNodeRegistryKey = GetRegistryKey(InCallbacks.CreateFrontendOutputClass().Metadata);
-		DataTypeNodeRegistry.Add(OutputNodeRegistryKey, InElement);
-
-		UE_LOG(LogMetasound, Display, TEXT("Registered Metasound Datatype %s."), *InDataInfo.DataTypeName.ToString());
-		return true;
-	}
-}
+FMetasoundFrontendRegistryContainer* FMetasoundFrontendRegistryContainer::LazySingleton = nullptr;
 
-bool FMetasoundFrontendRegistryContainer::RegisterEnumDataInterface(FName InDataType, TSharedPtr<IEnumDataTypeInterface>&& InInterface)
+FMetasoundFrontendRegistryContainer* FMetasoundFrontendRegistryContainer::Get()
 {
-	if (FDataTypeRegistryElement* Element = DataTypeRegistry.Find(InDataType))
+	if (!LazySingleton)
 	{
-		Element->EnumInterface = MoveTemp(InInterface);
-		return true;
+		LazySingleton = new Metasound::Frontend::MetasoundFrontendRegistryPrivate::FRegistryContainerImpl();
 	}
-	return false;
+
+	return LazySingleton;
 }
 
-bool FMetasoundFrontendRegistryContainer::RegisterExternalNode(Metasound::FCreateMetasoundNodeFunction&& InCreateNode, Metasound::FCreateMetasoundFrontendClassFunction&& InCreateDescription)
+void FMetasoundFrontendRegistryContainer::ShutdownMetasoundFrontend()
 {
-	FNodeRegistryElement RegistryElement = FNodeRegistryElement(MoveTemp(InCreateNode), MoveTemp(InCreateDescription));
-
-	FNodeRegistryKey Key;
-
-	if (ensure(FMetasoundFrontendRegistryContainer::GetRegistryKey(RegistryElement, Key)))
+	if (LazySingleton)
 	{
-		// check to see if an identical node was already registered, and log
-		ensureAlwaysMsgf(!ExternalNodeRegistry.Contains(Key), TEXT("Node with identical name, inputs and outputs to node %s was already registered. The previously registered node will be overwritten. This could also happen because METASOUND_REGISTER_NODE is in a public header."), *Key.NodeClassFullName.ToString());
-
-		ExternalNodeRegistry.Add(MoveTemp(Key), MoveTemp(RegistryElement));
-
-		return true;
+		delete LazySingleton;
+		LazySingleton = nullptr;
 	}
-
-	return false;
 }
 
 bool FMetasoundFrontendRegistryContainer::GetRegistryKey(const Metasound::Frontend::FNodeRegistryElement& InElement, Metasound::Frontend::FNodeRegistryKey& OutKey)
@@ -468,12 +756,12 @@ bool FMetasoundFrontendRegistryContainer::GetRegistryKey(const Metasound::Fronte
 
 Metasound::Frontend::FNodeRegistryKey FMetasoundFrontendRegistryContainer::GetRegistryKey(const FNodeClassMetadata& InNodeMetadata)
 {
-	return Metasound::MetasoundFrontendRegistryPrivate::GetRegistryKey(InNodeMetadata.ClassName.GetFullName(), InNodeMetadata.MajorVersion);
+	return Metasound::Frontend::MetasoundFrontendRegistryPrivate::GetRegistryKey(InNodeMetadata.ClassName.GetFullName(), InNodeMetadata.MajorVersion);
 }
 
 Metasound::Frontend::FNodeRegistryKey FMetasoundFrontendRegistryContainer::GetRegistryKey(const FMetasoundFrontendClassMetadata& InNodeMetadata)
 {
-	return Metasound::MetasoundFrontendRegistryPrivate::GetRegistryKey(InNodeMetadata.ClassName.GetFullName(), InNodeMetadata.Version.Major);
+	return Metasound::Frontend::MetasoundFrontendRegistryPrivate::GetRegistryKey(InNodeMetadata.ClassName.GetFullName(), InNodeMetadata.Version.Major);
 }
 
 bool FMetasoundFrontendRegistryContainer::GetFrontendClassFromRegistered(const FMetasoundFrontendClassMetadata& InMetadata, FMetasoundFrontendClass& OutClass)
@@ -482,32 +770,7 @@ bool FMetasoundFrontendRegistryContainer::GetFrontendClassFromRegistered(const F
 
 	if (ensure(nullptr != Registry))
 	{
-		Metasound::Frontend::FNodeRegistryKey RegistryKey = GetRegistryKey(InMetadata);
-
-		if (InMetadata.Type == EMetasoundFrontendClassType::External)
-		{
-			if (Registry->ExternalNodeRegistry.Contains(RegistryKey))
-			{
-				OutClass = Registry->ExternalNodeRegistry[RegistryKey].CreateFrontendClass();
-				return true;
-			}
-		}
-		else if (InMetadata.Type == EMetasoundFrontendClassType::Input)
-		{
-			if (Registry->DataTypeNodeRegistry.Contains(RegistryKey))
-			{
-				OutClass = Registry->DataTypeNodeRegistry[RegistryKey].Callbacks.CreateFrontendInputClass();
-				return true;
-			}
-		}
-		else if (InMetadata.Type == EMetasoundFrontendClassType::Output)
-		{
-			if (Registry->DataTypeNodeRegistry.Contains(RegistryKey))
-			{
-				OutClass = Registry->DataTypeNodeRegistry[RegistryKey].Callbacks.CreateFrontendOutputClass();
-				return true;
-			}
-		}
+		return Registry->FindFrontendClassFromRegistered(InMetadata, OutClass);
 	}
 
 	return false;
@@ -517,11 +780,7 @@ bool FMetasoundFrontendRegistryContainer::GetInputNodeClassMetadataForDataType(c
 {
 	if (FMetasoundFrontendRegistryContainer* Registry = FMetasoundFrontendRegistryContainer::Get())
 	{
-		if (Registry->DataTypeRegistry.Contains(InDataTypeName))
-		{
-			OutMetadata = Registry->DataTypeRegistry[InDataTypeName].Callbacks.CreateFrontendInputClass().Metadata;
-			return true;
-		}
+		return Registry->FindInputNodeClassMetadataForDataType(InDataTypeName, OutMetadata);
 	}
 	return false;
 }
@@ -530,82 +789,8 @@ bool FMetasoundFrontendRegistryContainer::GetOutputNodeClassMetadataForDataType(
 {
 	if (FMetasoundFrontendRegistryContainer* Registry = FMetasoundFrontendRegistryContainer::Get())
 	{
-		if (Registry->DataTypeRegistry.Contains(InDataTypeName))
-		{
-			OutMetadata = Registry->DataTypeRegistry[InDataTypeName].Callbacks.CreateFrontendOutputClass().Metadata;
-			return true;
-		}
+		return Registry->FindOutputNodeClassMetadataForDataType(InDataTypeName, OutMetadata);
 	}
 	return false;
 }
 
-bool FMetasoundFrontendRegistryContainer::RegisterConversionNode(const FConverterNodeRegistryKey& InNodeKey, const FConverterNodeInfo& InNodeInfo)
-{
-	if (!ConverterNodeRegistry.Contains(InNodeKey))
-	{
-		ConverterNodeRegistry.Add(InNodeKey);
-	}
-
-	FConverterNodeRegistryValue& ConverterNodeList = ConverterNodeRegistry[InNodeKey];
-
-	if (ensureAlways(!ConverterNodeList.PotentialConverterNodes.Contains(InNodeInfo)))
-	{
-		ConverterNodeList.PotentialConverterNodes.Add(InNodeInfo);
-		return true;
-	}
-	else
-	{
-		// If we hit this, someone attempted to add the same converter node to our list multiple times.
-		return false;
-	}
-}
-
-TArray<FName> FMetasoundFrontendRegistryContainer::GetAllValidDataTypes()
-{
-	TArray<FName> OutDataTypes;
-
-	for (auto& DataTypeTuple : DataTypeRegistry)
-	{
-		OutDataTypes.Add(DataTypeTuple.Key);
-	}
-
-	return OutDataTypes;
-}
-
-bool FMetasoundFrontendRegistryContainer::GetInfoForDataType(FName InDataType, Metasound::FDataTypeRegistryInfo& OutInfo)
-{
-	if (!DataTypeRegistry.Contains(InDataType))
-	{
-		return false;
-	}
-	else
-	{
-		OutInfo = DataTypeRegistry[InDataType].Info;
-		return true;
-	}
-}
-
-
-TSharedPtr<const Metasound::Frontend::IEnumDataTypeInterface>
-FMetasoundFrontendRegistryContainer::GetEnumInterfaceForDataType(FName InDataType) const
-{
-	if (const FDataTypeRegistryElement* Element = DataTypeRegistry.Find(InDataType))
-	{
-		return Element->EnumInterface;
-	}
-	return nullptr;
-}
-
-
-TSharedPtr<Metasound::IDataChannel, ESPMode::ThreadSafe> FMetasoundFrontendRegistryContainer::CreateDataChannelForDataType(const FName& InDataType, const Metasound::FOperatorSettings& InSettings) const
-{
-	if (const FDataTypeRegistryElement* Element = DataTypeRegistry.Find(InDataType))
-	{
-		if (Element->Callbacks.CreateDataChannel)
-		{
-			return Element->Callbacks.CreateDataChannel(InSettings);
-		}
-	}
-	return nullptr;
-}
-
diff --git a/Engine/Plugins/Runtime/Metasound/Source/MetasoundFrontend/Private/MetasoundFrontendRegistryTransaction.cpp b/Engine/Plugins/Runtime/Metasound/Source/MetasoundFrontend/Private/MetasoundFrontendRegistryTransaction.cpp
new file mode 100644
--- /dev/null
+++ b/Engine/Plugins/Runtime/Metasound/Source/MetasoundFrontend/Private/MetasoundFrontendRegistryTransaction.cpp
@@ -0,0 +1,126 @@
+// Copyright Epic Games, Inc. All Rights Reserved.
+
+
+#include "MetasoundFrontendRegistryTransaction.h"
+
+#include "CoreMinimal.h"
+
+namespace Metasound
+{
+	namespace Frontend
+	{
+		FRegistryTransactionID GetOriginRegistryTransactionID()
+		{
+			return 0;
+		}
+
+		FRegistryTransactionHistory::FRegistryTransactionHistory()
+		: Current(GetOriginRegistryTransactionID())
+		{
+		}
+
+		FRegistryTransactionID FRegistryTransactionHistory::Add(TUniquePtr<IRegistryTransaction>&& InRegistryTransaction)
+		{
+			FScopeLock Lock(&RegistryTransactionMutex);
+
+			if (ensure(InRegistryTransaction.IsValid()))
+			{
+				Current++;
+
+				IRegistryTransaction* RegistryTransactionPointer = InRegistryTransaction.Get();
+				int32 Index = RegistryTransactions.Num();
+
+				RegistryTransactions.Add(MoveTemp(InRegistryTransaction));
+				RegistryTransactionIndexMap.Add(Current, Index);
+			}
+
+			return Current;
+		}
+
+		FRegistryTransactionID FRegistryTransactionHistory::Add(const IRegistryTransaction& InRegistryTransaction)
+		{
+			return Add(InRegistryTransaction.Clone());
+		}
+
+		FRegistryTransactionID FRegistryTransactionHistory::GetCurrent() const
+		{
+			FScopeLock Lock(&RegistryTransactionMutex);
+			{
+				return Current;
+			}
+		}
+
+		TArray<const IRegistryTransaction*> FRegistryTransactionHistory::GetTransactions(FRegistryTransactionID InSince, FRegistryTransactionID* OutCurrent) const
+		{
+			FScopeLock Lock(&RegistryTransactionMutex);
+			{
+				TArray<const IRegistryTransaction*> Result;
+
+				if (nullptr != OutCurrent)
+				{
+					*OutCurrent = Current;
+				}
+
+				int32 Start = INDEX_NONE;
+				
+				if (GetOriginRegistryTransactionID() == InSince)
+				{
+					Start = 0;
+				}
+				else if (const int32* Pos = RegistryTransactionIndexMap.Find(InSince))
+				{
+					Start = *Pos + 1;
+				}
+				
+				if (INDEX_NONE != Start)
+				{
+					const int32 Num = RegistryTransactions.Num();
+
+					if (ensure(Start <= Num))
+					{
+						const int32 OutNum = Num - Start;
+						if (OutNum > 0)
+						{
+							Result.Reserve(OutNum);
+							for (int32 i = Start; i < Num; i++)
+							{
+								Result.Add(RegistryTransactions[i].Get());
+							}
+						}
+					}
+				}
+
+				return Result;
+			}
+			
+		}
+
+		FRegistryTransactionPtr MakeNodeRegistrationTransaction(const FNodeClassInfo& InInfo)
+		{
+			class FNodeRegistrationTransaction : public IRegistryTransaction
+			{
+			public:
+				FNodeRegistrationTransaction(const FNodeClassInfo& InNodeClassInfo)
+				: NodeClassInfo(InNodeClassInfo)
+				{
+				}
+
+				TUniquePtr<IRegistryTransaction> Clone() const override
+				{
+					return MakeUnique<FNodeRegistrationTransaction>(*this);
+				}
+
+				const FNodeClassInfo* GetNodeClassInfo() const 
+				{
+					return &NodeClassInfo;
+				}
+
+			private:
+
+				FNodeClassInfo NodeClassInfo;
+			};
+
+			return MakeUnique<FNodeRegistrationTransaction>(InInfo);
+		}
+	}
+}
diff --git a/Engine/Plugins/Runtime/Metasound/Source/MetasoundFrontend/Private/MetasoundFrontendRegistryTransaction.h b/Engine/Plugins/Runtime/Metasound/Source/MetasoundFrontend/Private/MetasoundFrontendRegistryTransaction.h
new file mode 100644
--- /dev/null
+++ b/Engine/Plugins/Runtime/Metasound/Source/MetasoundFrontend/Private/MetasoundFrontendRegistryTransaction.h
@@ -0,0 +1,75 @@
+// Copyright Epic Games, Inc. All Rights Reserved.
+
+#pragma once
+
+#include "CoreMinimal.h"
+#include "MetasoundFrontend.h"
+
+namespace Metasound
+{
+	namespace Frontend
+	{
+		using FRegistryTransactionID = int32;
+
+		/** Returns an ID representing the beginning of the transaction history. */
+		FRegistryTransactionID GetOriginRegistryTransactionID();
+
+		/** Interface for a registry transaction. */
+		class IRegistryTransaction
+		{
+		public:
+			virtual ~IRegistryTransaction() = default;
+			virtual TUniquePtr<IRegistryTransaction> Clone() const = 0;
+
+			/** If a FNodeClassInfo added during the transaction, this will return
+			 * a non-null pointer to the FNodeClassInfo. */
+			virtual const FNodeClassInfo* GetNodeClassInfo() const = 0;
+		};
+
+		using FRegistryTransactionPtr = TUniquePtr<IRegistryTransaction>;
+
+		/** Create a node registration transaction. */
+		FRegistryTransactionPtr MakeNodeRegistrationTransaction(const FNodeClassInfo& InInfo);
+
+		/** Maintains a history of IRegistryTransactions. Calls are threadsafe (excluding
+		 * the constructor and destructor.)
+		 */
+		class FRegistryTransactionHistory
+		{
+		public:
+
+			FRegistryTransactionHistory();
+
+			/** Add a transaction to the history. Threadsafe. 
+			 *
+			 * @return The transaction ID associated with the action. */
+			FRegistryTransactionID Add(FRegistryTransactionPtr&& InRegistryTransaction);
+
+			/** Add a transaction to the history. Threadsafe. 
+			 *
+			 * @return The transaction ID associated with the action. */
+			FRegistryTransactionID Add(const IRegistryTransaction& InRegistryTransaction);
+
+			/** Gets the transaction ID of the most recent transaction. Threadsafe. */
+			FRegistryTransactionID GetCurrent() const;
+
+			/** Returns an array of transaction pointers. Threadsafe. 
+			 *
+			 * @param InSince - All transactions occuring after this transaction ID will be returned.
+			 * @param OutCurrent - If not null, will be set to the most recent transaction ID returned. 
+			 */
+			TArray<const IRegistryTransaction*> GetTransactions(FRegistryTransactionID InSince, FRegistryTransactionID* OutCurrent) const;
+
+		private:
+
+			mutable FCriticalSection RegistryTransactionMutex;
+
+			FRegistryTransactionID Current;
+
+			TMap<FRegistryTransactionID, int32> RegistryTransactionIndexMap;
+			TArray<FRegistryTransactionPtr> RegistryTransactions;
+		};
+	}
+}
+
+
diff --git a/Engine/Plugins/Runtime/Metasound/Source/MetasoundFrontend/Private/MetasoundFrontendSearchEngine.cpp b/Engine/Plugins/Runtime/Metasound/Source/MetasoundFrontend/Private/MetasoundFrontendSearchEngine.cpp
new file mode 100644
--- /dev/null
+++ b/Engine/Plugins/Runtime/Metasound/Source/MetasoundFrontend/Private/MetasoundFrontendSearchEngine.cpp
@@ -0,0 +1,101 @@
+// Copyright Epic Games, Inc. All Rights Reserved.
+
+#include "MetasoundFrontendSearchEngine.h"
+
+#include "MetasoundFrontendQuery.h"
+#include "MetasoundFrontendQuerySteps.h"
+
+namespace Metasound
+{
+	namespace Frontend
+	{
+		class FSearchEngine : public ISearchEngine
+		{
+				FSearchEngine(const FSearchEngine&) = delete;
+				FSearchEngine(FSearchEngine&&) = delete;
+				FSearchEngine& operator=(const FSearchEngine&) = delete;
+				FSearchEngine& operator=(FSearchEngine&&) = delete;
+
+			public:
+				FSearchEngine() = default;
+
+				virtual ~FSearchEngine() = default;
+
+				TArray<FMetasoundFrontendClass> FindClassesWithClassName(const FNodeClassName& InName) override;
+
+			private:
+
+				FFrontendQuery* FindQuery(const FString& InQueryName);
+				FFrontendQuery* AddQuery(const FString& InQueryName, TUniquePtr<FFrontendQuery>&& InQuery);
+
+				// Store queries. Index them by unique strings. 
+				TMap<FString, TUniquePtr<FFrontendQuery>> QueryCache;
+		};
+
+		ISearchEngine& ISearchEngine::Get()
+		{
+			static FSearchEngine SearchEngine;
+			return SearchEngine;
+		}
+
+		TArray<FMetasoundFrontendClass> FSearchEngine::FindClassesWithClassName(const FNodeClassName& InName) 
+		{
+			// Unique string for query with given parameter name. 
+			FString QueryName = FString(TEXT("FindClassesWithClassName")) + InName.GetFullName().ToString();
+
+			FFrontendQuery* Query = FindQuery(QueryName);
+
+			if (nullptr == Query)
+			{
+				// If the query doesn't exist, it needs to be created. 
+				TUniquePtr<FFrontendQuery> NewQuery = MakeUnique<FFrontendQuery>();
+
+				// Note: Initial query step is FGenerateNewlyAvailableNodeClasses, which does not run query against already examined node classes.
+				NewQuery->AddStep<FGenerateNewlyAvailableNodeClasses>()
+					.AddStep<FFilterClassesByClassName>(InName);
+
+				Query = AddQuery(QueryName, MoveTemp(NewQuery));
+			}
+
+			TArray<FMetasoundFrontendClass> Result;
+
+			if (ensure(nullptr != Query))
+			{
+				// Execute the query again and append the results. (Note: Initial query 
+				// step is FGenerateNewlyAvailableNodeClasses, which does not run query 
+				// against already examined node classes.)
+				FFrontendQuerySelectionView Selection = Query->ExecuteQueryAndAppend();
+
+				TArrayView<const FFrontendQueryEntry* const> Entries = Selection.GetSelection();
+
+				// Copy result to array
+				for (const FFrontendQueryEntry* Entry : Selection.GetSelection())
+				{
+					check(Entry->Value.IsType<FMetasoundFrontendClass>());
+					Result.Add(Entry->Value.Get<FMetasoundFrontendClass>());
+				}
+			}
+
+			return Result;
+		}
+
+		FFrontendQuery* FSearchEngine::FindQuery(const FString& InQueryName)
+		{
+			if (TUniquePtr<FFrontendQuery>* QueryUPtr = QueryCache.Find(InQueryName))
+			{
+				return QueryUPtr->Get();
+			}
+
+			return nullptr;
+		}
+
+		FFrontendQuery* FSearchEngine::AddQuery(const FString& InQueryName, TUniquePtr<FFrontendQuery>&& InQuery)
+		{
+			FFrontendQuery* QueryPtr = InQuery.Get();
+
+			QueryCache.Add(InQueryName, MoveTemp(InQuery));
+
+			return QueryPtr;
+		}
+	}
+}
diff --git a/Engine/Plugins/Runtime/Metasound/Source/MetasoundFrontend/Public/MetasoundFrontend.h b/Engine/Plugins/Runtime/Metasound/Source/MetasoundFrontend/Public/MetasoundFrontend.h
--- a/Engine/Plugins/Runtime/Metasound/Source/MetasoundFrontend/Public/MetasoundFrontend.h
+++ b/Engine/Plugins/Runtime/Metasound/Source/MetasoundFrontend/Public/MetasoundFrontend.h
@@ -1,10 +1,10 @@
 // Copyright Epic Games, Inc. All Rights Reserved.
-
+ 
 #pragma once
 
 #include "CoreMinimal.h"
 #include "MetasoundGraph.h"
-#include "MetasoundFrontendBaseClasses.h"
+#include "MetasoundFrontendRegistries.h"
 #include "MetasoundFrontendDocument.h"
 #include "MetasoundBuilderInterface.h"
 #include "UObject/WeakObjectPtrTemplates.h"
@@ -16,21 +16,14 @@ namespace Metasound
 {	
 	namespace Frontend
 	{
+		using FRegistryTransactionID = int32;
 
-		// Struct with the basics of a node class' information,
-		// used to look up that node from our node browser functions,
-		// and also used in FGraphHandle::AddNewNode.
-		struct METASOUNDFRONTEND_API FNodeClassInfo
-		{
-			// The type for this node.
-			EMetasoundFrontendClassType NodeType = EMetasoundFrontendClassType::Invalid;
-
-			// The lookup key used for the internal node registry.
-			FNodeRegistryKey LookupKey;
-		};
+		METASOUNDFRONTEND_API FRegistryTransactionID GetCurrentRegistryTransactionID();
 
 		// Get all available nodes of any type.
-		METASOUNDFRONTEND_API TArray<FNodeClassInfo> GetAllAvailableNodeClasses();
+		METASOUNDFRONTEND_API TArray<FNodeClassInfo> GetAllAvailableNodeClasses(FRegistryTransactionID* OutCurrentTransactionID=nullptr);
+
+		METASOUNDFRONTEND_API TArray<FNodeClassInfo> GetNodeClassesRegisteredSince(FRegistryTransactionID InTransactionID, FRegistryTransactionID* OutCurrentTransactionID=nullptr);
 
 		/** Returns all metadata (name, description, author, what to say if it's missing) for a given node.
 		 *
@@ -54,7 +47,7 @@ namespace Metasound
 		 *
 		 * @return Class description for external node.
 		 */
-		METASOUNDFRONTEND_API FMetasoundFrontendClass GenerateClassDescription(const FNodeClassInfo& InInfo);
+		METASOUNDFRONTEND_API FMetasoundFrontendClass GenerateClassDescription(const Metasound::Frontend::FNodeClassInfo& InInfo);
 
 		/** Generates a new FMetasoundFrontendClass from Node init data
 		 *
diff --git a/Engine/Plugins/Runtime/Metasound/Source/MetasoundFrontend/Public/MetasoundFrontendBaseClasses.h b/Engine/Plugins/Runtime/Metasound/Source/MetasoundFrontend/Public/MetasoundFrontendBaseClasses.h
--- a/Engine/Plugins/Runtime/Metasound/Source/MetasoundFrontend/Public/MetasoundFrontendBaseClasses.h
+++ b/Engine/Plugins/Runtime/Metasound/Source/MetasoundFrontend/Public/MetasoundFrontendBaseClasses.h
@@ -15,103 +15,9 @@ namespace Metasound
 		// This is called by FMetasoundFrontendModule, and flushes any node or datatype registration that was done prior to boot.
 		void RegisterPendingNodes();
 
-		/** Use this to get a reference to the static registry of all nodes implemented in C++. */
-		METASOUNDFRONTEND_API TMap<FNodeRegistryKey, FNodeRegistryElement>& GetExternalNodeRegistry();
-
 		// Convenience functions to create an INode corresponding to a specific input or output for a metasound graph.
 		// @returns nullptr if the type given wasn't found.
 		METASOUNDFRONTEND_API TUniquePtr<INode> ConstructInputNode(const FName& InInputType, FInputNodeConstructorParams&& InParams);
 		METASOUNDFRONTEND_API TUniquePtr<INode> ConstructOutputNode(const FName& InOutputType, const FOutputNodeConstructorParams& InParams);
-
-		// Convenience functions to create an INodeB corresponding to a specific externally declared node type.
-		// InNodeType and InNodeHash can be retrieved from the FNodeClassInfo generated from the node registry queries in the metasound frontend (GetAllAvailableNodeClasses, GetAllNodeClassesInNamespace, etc.)
-		// @returns nullptr if the type given wasn't found.
-		METASOUNDFRONTEND_API TUniquePtr<INode> ConstructExternalNode(const FName& InNodeClassFullName, uint32 InNodeHash, const FNodeInitData& InInitData);
-
-		// Utility base class for classes that want to support naive multi-level undo/redo.
-		class METASOUNDFRONTEND_API ITransactable : public TSharedFromThis<ITransactable>
-		{
-		public:
-
-			ITransactable() = delete;
-			ITransactable(uint32 InUndoLimit, TWeakObjectPtr<UObject> InOwningAsset = nullptr);
-
-			/*
-			 * Undo a single operation for this or any of it's owned transactables.
-			 * @returns true on success.
-			 */
-			bool Undo();
-
-			/*
-			 * Redo a single operation for this or any of it's owned transactables.
-			 * @returns true on success.
-			 */
-			bool Redo();
-
-		protected:
-			struct FReversibleTransaction
-			{
-				// Function that will be executed when Undo is called.
-				// This should return true on success.
-				TFunction<bool()> UndoTransaction;
-
-				// Function that will be executed when Redo is called.
-				// This should return true on success.
-				TFunction<bool()> RedoTransaction;
-
-				FReversibleTransaction(TFunction<bool()>&& InUndo, TFunction<bool()>&& InRedo)
-					: UndoTransaction(MoveTemp(InUndo))
-					, RedoTransaction(MoveTemp(InRedo))
-				{
-				}
-			};
-
-			// Commits an undoable/redoable transaction for this object.
-			void CommitTransaction(FReversibleTransaction&& InTransactionDescription);
-			
-			// Use this function to denote an ITransactable that should consider this ITransactable part of it's undo roll.
-			// For example, if an undoable transaction is committed to a FNodeHandle,
-			// and you want FGraphHandle::Undo to be able to undo that transaction,
-			// you would call FInputHandle::RegisterOwningTransactable(OwningGraphHandle). 
-			// 
-			// @returns true on success, false if registering this as an owning transaction would create a cycle.
-			bool RegisterOwningTransactable(ITransactable& InOwningTransactable);
-
-		private:
-
-			// Execute the topmost transaction on the undo stack,
-			// push it to redo stack,
-			// and notify the OnwingTransactable.
-			bool PerformLocalUndo();
-
-			// Execute the topmost transaction on the redo stack,
-			// push it to the undo stack,
-			// and notify the OwningTransactable.
-			bool PerformLocalRedo();
-
-			// Called from a owned transactable to denote that it is popping an undo or redo action.
-			// When this occurs, we discard the topmost call to the owned transactable from the UndoTransactableStack.
-			bool DiscardUndoFromOwnedTransactable(TWeakPtr<ITransactable> InOwnedTransactable);
-			bool DiscardRedoFromOwnedTransactable(TWeakPtr<ITransactable> InOwnedTransactable);
-
-			bool PushUndoFromOwnedTransactable(TWeakPtr<ITransactable> InOwnedTransactable);
-			bool PushRedoFromOwnedTransactable(TWeakPtr<ITransactable> InOwnedTransactable);
-
-			// TODO: Make these ring buffers to avoid constantly moving elements around.
-			TArray<TWeakPtr<ITransactable>> UndoTransactableStack;
-			TArray<TWeakPtr<ITransactable>> RedoTransactableStack;
-			TArray<FReversibleTransaction> LocalUndoTransactionStack;
-			TArray<FReversibleTransaction> LocalRedoTransactionStack;
-
-			// When CommitTransaction is called,
-			// this ITransactable is pushed to the ITransactable::TransactableStack of this ITransactable object.
-			TWeakPtr<ITransactable> OwningTransactable;
-
-		protected:
-			uint32 UndoLimit;
-
-			// If this transactable is operating on a UAsset, this can be filled to mark it's corresponding package as dirty.
-			TWeakObjectPtr<UObject> OwningAsset;
-		};
 	}
 }
diff --git a/Engine/Plugins/Runtime/Metasound/Source/MetasoundFrontend/Public/MetasoundFrontendQuery.h b/Engine/Plugins/Runtime/Metasound/Source/MetasoundFrontend/Public/MetasoundFrontendQuery.h
--- a/Engine/Plugins/Runtime/Metasound/Source/MetasoundFrontend/Public/MetasoundFrontendQuery.h
+++ b/Engine/Plugins/Runtime/Metasound/Source/MetasoundFrontend/Public/MetasoundFrontendQuery.h
@@ -85,6 +85,9 @@ namespace Metasound
 		/** Get an array of all entries selected. */
 		TArrayView<const FFrontendQueryEntry * const> GetSelection() const;
 
+		/** Add entries to both storage and current selection. */
+		void AppendToStorageAndSelection(const FFrontendQuerySelection& InSelection);
+
 		/** Add entries to both storage and current selection. */
 		void AppendToStorageAndSelection(TArrayView<const FFrontendQueryEntry> InEntries);
 
@@ -120,7 +123,7 @@ namespace Metasound
 		/** The FFrontendQuerySelectionView is constructed with an existing
 		 * FFrontendQuerySelection.
 		 */
-		FFrontendQuerySelectionView(TUniquePtr<FFrontendQuerySelection>&& InResult);
+		FFrontendQuerySelectionView(TSharedRef<const FFrontendQuerySelection, ESPMode::ThreadSafe> InResult);
 
 		/** Get all the stored results. */
 		TArrayView<const FFrontendQueryEntry> GetStorage() const;
@@ -129,7 +132,7 @@ namespace Metasound
 		TArrayView<const FFrontendQueryEntry* const> GetSelection() const;
 
 	private:
-		TUniquePtr<FFrontendQuerySelection> Result;
+		TSharedRef<const FFrontendQuerySelection, ESPMode::ThreadSafe> Result;
 	};
 
 	/** Interface for an individual step in a query */
@@ -261,7 +264,10 @@ namespace Metasound
 		using FSortFunction = FFrontendQueryStep::FSortFunction;
 		using FLimitFunction = FFrontendQueryStep::FLimitFunction;
 
-		FFrontendQuery() = default;
+		FFrontendQuery();
+
+		FFrontendQuery(const FFrontendQuery&) = delete;
+		FFrontendQuery& operator=(const FFrontendQuery&) = delete;
 
 		/** Return all steps in a query. */
 		const TArray<TUniquePtr<FFrontendQueryStep>>& GetSteps() const;
@@ -293,9 +299,20 @@ namespace Metasound
 		/** Calls all steps in the query and returns the selection. */
 		FFrontendQuerySelectionView ExecuteQuery();
 
+		/** Resets the query result by removing all entries. */
+		FFrontendQuerySelectionView Reset();
+
+		/** Executes the query again and appends the results to the existing results. */
+		FFrontendQuerySelectionView ExecuteQueryAndAppend();
+
+		/** Returns the current result. */
+		FFrontendQuerySelectionView GetSelection() const;
+
 	private:
 
-		TUniquePtr<FFrontendQuerySelection> Result;
+		void Execute(FFrontendQuerySelection& OutSelection);
+
+		TSharedRef<FFrontendQuerySelection, ESPMode::ThreadSafe> Result;
 
 		TArray<TUniquePtr<FFrontendQueryStep>> Steps;
 	};
diff --git a/Engine/Plugins/Runtime/Metasound/Source/MetasoundFrontend/Public/MetasoundFrontendQuerySteps.h b/Engine/Plugins/Runtime/Metasound/Source/MetasoundFrontend/Public/MetasoundFrontendQuerySteps.h
--- a/Engine/Plugins/Runtime/Metasound/Source/MetasoundFrontend/Public/MetasoundFrontendQuerySteps.h
+++ b/Engine/Plugins/Runtime/Metasound/Source/MetasoundFrontend/Public/MetasoundFrontendQuerySteps.h
@@ -9,6 +9,7 @@
 
 namespace Metasound
 {
+	/** Generates all registered node classes */
 	class METASOUNDFRONTEND_API FGenerateAllAvailableNodeClasses : public IFrontendQueryGenerateStep
 	{
 	public:
@@ -16,6 +17,19 @@ namespace Metasound
 		void Generate(TArray<FFrontendQueryEntry>& OutEntries) const override;
 	};
 
+	/** Generates node classes that have been newly registered since last call
+	 * to this FGenerateNewlyAvailableClasses::Generate()
+	 */
+	class METASOUNDFRONTEND_API FGenerateNewlyAvailableNodeClasses : public IFrontendQueryGenerateStep
+	{
+	public:
+		FGenerateNewlyAvailableNodeClasses();
+		void Generate(TArray<FFrontendQueryEntry>& OutEntries) const override;
+
+	private:
+		mutable Frontend::FRegistryTransactionID CurrentTransactionID;
+	};
+
 	class METASOUNDFRONTEND_API FFilterClassesByInputVertexDataType : public IFrontendQueryFilterStep
 	{
 	public: 
diff --git a/Engine/Plugins/Runtime/Metasound/Source/MetasoundFrontend/Public/MetasoundFrontendRegistries.h b/Engine/Plugins/Runtime/Metasound/Source/MetasoundFrontend/Public/MetasoundFrontendRegistries.h
--- a/Engine/Plugins/Runtime/Metasound/Source/MetasoundFrontend/Public/MetasoundFrontendRegistries.h
+++ b/Engine/Plugins/Runtime/Metasound/Source/MetasoundFrontend/Public/MetasoundFrontendRegistries.h
@@ -14,6 +14,11 @@
 
 namespace Metasound
 {
+	namespace Frontend
+	{
+		using FRegistryTransactionID = int32;
+	}
+
 	// This struct is passed to FInputNodeConstructorCallback and FOutputNodeConstructorCallback.
 	struct FInputNodeConstructorParams
 	{
@@ -116,6 +121,7 @@ namespace Metasound
 
 	namespace Frontend
 	{
+
 		struct METASOUNDFRONTEND_API FNodeRegistryKey
 		{
 			// The class name for the node.
@@ -135,6 +141,18 @@ namespace Metasound
 			}
 		};
 
+		// Struct with the basics of a node class' information,
+		// used to look up that node from our node browser functions,
+		// and also used in FGraphHandle::AddNewNode.
+		struct METASOUNDFRONTEND_API FNodeClassInfo
+		{
+			// The type for this node.
+			EMetasoundFrontendClassType NodeType = EMetasoundFrontendClassType::Invalid;
+
+			// The lookup key used for the internal node registry.
+			FNodeRegistryKey LookupKey;
+		};
+
 
 		struct METASOUNDFRONTEND_API FNodeRegistryElement
 		{
@@ -262,6 +280,8 @@ namespace Metasound
  */
 class METASOUNDFRONTEND_API FMetasoundFrontendRegistryContainer
 {
+
+public:
 	using FConverterNodeRegistryKey = ::Metasound::Frontend::FConverterNodeRegistryKey;
 	using FConverterNodeRegistryValue = ::Metasound::Frontend::FConverterNodeRegistryValue;
 	using FConverterNodeInfo = ::Metasound::Frontend::FConverterNodeInfo;
@@ -274,37 +294,75 @@ class METASOUNDFRONTEND_API FMetasoundFrontendRegistryContainer
 	using FNodeClassMetadata = Metasound::FNodeClassMetadata;
 	using IEnumDataTypeInterface = Metasound::Frontend::IEnumDataTypeInterface;
 
-public:
+
 	static FMetasoundFrontendRegistryContainer* Get();
 	static void ShutdownMetasoundFrontend();
 
+	static FNodeRegistryKey GetRegistryKey(const FNodeClassMetadata& InNodeMetadata);
+	static FNodeRegistryKey GetRegistryKey(const FMetasoundFrontendClassMetadata& InNodeMetadata);
+	static bool GetRegistryKey(const FNodeRegistryElement& InElement, FNodeRegistryKey& OutKey);
+	static bool GetFrontendClassFromRegistered(const FMetasoundFrontendClassMetadata& InMetadata, FMetasoundFrontendClass& OutClass);
+	static bool GetInputNodeClassMetadataForDataType(const FName& InDataTypeName, FMetasoundFrontendClassMetadata& OutMetadata);
+	static bool GetOutputNodeClassMetadataForDataType(const FName& InDataTypeName, FMetasoundFrontendClassMetadata& OutMetadata);
+
+
+	FMetasoundFrontendRegistryContainer() = default;
+	virtual ~FMetasoundFrontendRegistryContainer() = default;
+
 	FMetasoundFrontendRegistryContainer(const FMetasoundFrontendRegistryContainer&) = delete;
+	FMetasoundFrontendRegistryContainer& operator=(const FMetasoundFrontendRegistryContainer&) = delete;
+
+	// Enqueu and command for registering a node or data type.
+	// The command queue will be processed on module init or when calling `RegisterPendingNodes()`
+	virtual bool EnqueueInitCommand(TUniqueFunction<void()>&& InFunc) = 0;
 
 	// This is called on module startup. This invokes any registration commands enqueued by our registration macros.
-	void RegisterPendingNodes();
+	virtual void RegisterPendingNodes() = 0;
+
+	/** Register external node with the frontend.
+	 *
+	 * @param InCreateNode - Function for creating node from FNodeInitData.
+	 * @param InCreateDescription - Function for creating a FMetasoundFrontendClass.
+	 *
+	 * @return True on success.
+	 */
+	virtual bool RegisterExternalNode(Metasound::FCreateMetasoundNodeFunction&& InCreateNode, Metasound::FCreateMetasoundFrontendClassFunction&& InCreateDescription) = 0;
+	virtual bool RegisterConversionNode(const FConverterNodeRegistryKey& InNodeKey, const FConverterNodeInfo& InNodeInfo) = 0;
+	virtual bool IsNodeRegistered(const FNodeRegistryKey& InKey) const = 0;
 
 
-	bool EnqueueInitCommand(TUniqueFunction<void()>&& InFunc);
+	/** Return all node classes registered. */
+	virtual TArray<Metasound::Frontend::FNodeClassInfo> GetAllAvailableNodeClasses(Metasound::Frontend::FRegistryTransactionID* OutCurrentRegistryTransactionID) const = 0;
 
-	TMap<FNodeRegistryKey, FNodeRegistryElement>& GetExternalNodeRegistry();
+	/** Return node classes registered since a given transaction id. */
+	virtual TArray<Metasound::Frontend::FNodeClassInfo> GetNodeClassesRegisteredSince(Metasound::Frontend::FRegistryTransactionID InSince, Metasound::Frontend::FRegistryTransactionID* OutCurrentRegistryTransactionID) const = 0;
 
-	TUniquePtr<Metasound::INode> ConstructInputNode(const FName& InInputType, Metasound::FInputNodeConstructorParams&& InParams);
-	TUniquePtr<Metasound::INode> ConstructOutputNode(const FName& InOutputType, const Metasound::FOutputNodeConstructorParams& InParams);
+	// Return any data types that can be used as a metasound input type or output type.
+	virtual TArray<FName> GetAllValidDataTypes() = 0;
 
-	Metasound::FLiteral GenerateLiteralForUObject(const FName& InDataType, UObject* InObject);
-	Metasound::FLiteral GenerateLiteralForUObjectArray(const FName& InDataType, TArray<UObject*> InObjectArray);
+	// Query for MetaSound Frontend document objects.
+	virtual bool FindFrontendClassFromRegistered(const Metasound::Frontend::FNodeClassInfo& InClassInfo, FMetasoundFrontendClass& OutClass) = 0;
+	virtual bool FindFrontendClassFromRegistered(const FMetasoundFrontendClassMetadata& InMetadata, FMetasoundFrontendClass& OutClass) = 0;
+	virtual bool FindInputNodeClassMetadataForDataType(const FName& InDataTypeName, FMetasoundFrontendClassMetadata& OutMetadata) = 0;
+	virtual bool FindOutputNodeClassMetadataForDataType(const FName& InDataTypeName, FMetasoundFrontendClassMetadata& OutMetadata) = 0;
 
 	// Create a new instance of a C++ implemented node from the registry.
-	TUniquePtr<Metasound::INode> ConstructExternalNode(const FName& InNodeType, uint32 InNodeHash, const Metasound::FNodeInitData& InInitData);
+	virtual TUniquePtr<Metasound::INode> ConstructInputNode(const FName& InInputType, Metasound::FInputNodeConstructorParams&& InParams) = 0;
+	virtual TUniquePtr<Metasound::INode> ConstructOutputNode(const FName& InOutputType, const Metasound::FOutputNodeConstructorParams& InParams) = 0;
+	virtual TUniquePtr<Metasound::INode> ConstructExternalNode(const Metasound::Frontend::FNodeRegistryKey& InRegistryKey, const Metasound::FNodeInitData& InInitData) = 0;
 
 	// Returns a list of possible nodes to use to convert from FromDataType to ToDataType.
 	// Returns an empty array if none are available.
-	TArray<FConverterNodeInfo> GetPossibleConverterNodes(const FName& FromDataType, const FName& ToDataType);
+	virtual TArray<FConverterNodeInfo> GetPossibleConverterNodes(const FName& FromDataType, const FName& ToDataType) = 0;
 
-	// Get the desired kind of literal for a given data type. Returns EConstructorArgType::Invalid if the data type couldn't be found.
-	Metasound::ELiteralType GetDesiredLiteralTypeForDataType(FName InDataType) const;
 
-	UClass* GetLiteralUClassForDataType(FName InDataType) const;
+
+	virtual bool RegisterDataType(const FDataTypeRegistryInfo& InDataInfo, const FDataTypeConstructorCallbacks& InCallbacks) = 0;
+
+	virtual bool RegisterEnumDataInterface(FName InDataType, TSharedPtr<IEnumDataTypeInterface>&& InInterface) = 0;
+
+	// Get the desired kind of literal for a given data type. Returns EConstructorArgType::Invalid if the data type couldn't be found.
+	virtual Metasound::ELiteralType GetDesiredLiteralTypeForDataType(FName InDataType) const = 0;
 
 	template<typename ArgType>
 	bool DoesDataTypeSupportLiteralType(FName InDataType) const
@@ -312,78 +370,24 @@ public:
 		return DoesDataTypeSupportLiteralType(InDataType, Metasound::TLiteralTypeInfo<ArgType>::GetLiteralArgTypeEnum());
 	}
 
-	// Get whether we can build a literal of this specific type for InDataType.
-	bool DoesDataTypeSupportLiteralType(FName InDataType, Metasound::ELiteralType InLiteralType) const;
-
-	bool RegisterDataType(const FDataTypeRegistryInfo& InDataInfo, const FDataTypeConstructorCallbacks& InCallbacks);
-
-	bool RegisterEnumDataInterface(FName InDataType, TSharedPtr<IEnumDataTypeInterface>&& InInterface);
-
-	/** Register external node with the frontend.
-	 *
-	 * @param InCreateNode - Function for creating node from FNodeInitData.
-	 * @param InCreateDescription - Function for creating a FMetasoundFrontendClass.
-	 *
-	 * @return True on success.
-	 */
-	bool RegisterExternalNode(Metasound::FCreateMetasoundNodeFunction&& InCreateNode, Metasound::FCreateMetasoundFrontendClassFunction&& InCreateDescription);
-
-	bool RegisterConversionNode(const FConverterNodeRegistryKey& InNodeKey, const FConverterNodeInfo& InNodeInfo);
+	virtual UClass* GetLiteralUClassForDataType(FName InDataType) const = 0;
+	virtual Metasound::FLiteral GenerateLiteralForUObject(const FName& InDataType, UObject* InObject) = 0;
+	virtual Metasound::FLiteral GenerateLiteralForUObjectArray(const FName& InDataType, TArray<UObject*> InObjectArray) = 0;
 
-	bool IsNodeRegistered(const FNodeRegistryKey& InKey) const
-	{
-		return ExternalNodeRegistry.Contains(InKey);
-	}
-
-	static FNodeRegistryKey GetRegistryKey(const FNodeClassMetadata& InNodeMetadata);
-	static FNodeRegistryKey GetRegistryKey(const FMetasoundFrontendClassMetadata& InNodeMetadata);
-	static bool GetRegistryKey(const FNodeRegistryElement& InElement, FNodeRegistryKey& OutKey);
-	static bool GetFrontendClassFromRegistered(const FMetasoundFrontendClassMetadata& InMetadata, FMetasoundFrontendClass& OutClass);
-	static bool GetInputNodeClassMetadataForDataType(const FName& InDataTypeName, FMetasoundFrontendClassMetadata& OutMetadata);
-	static bool GetOutputNodeClassMetadataForDataType(const FName& InDataTypeName, FMetasoundFrontendClassMetadata& OutMetadata);
+	// Get whether we can build a literal of this specific type for InDataType.
+	virtual bool DoesDataTypeSupportLiteralType(FName InDataType, Metasound::ELiteralType InLiteralType) const = 0;
 
-	// Return any data types that can be used as a metasound input type or output type.
-	TArray<FName> GetAllValidDataTypes();
 
 	// Get info about a specific data type (what kind of literals we can use, etc.)
 	// @returns false if InDataType wasn't found in the registry. 
-	bool GetInfoForDataType(FName InDataType, FDataTypeRegistryInfo& OutInfo);
+	virtual bool GetInfoForDataType(FName InDataType, FDataTypeRegistryInfo& OutInfo) = 0;
 
-	TSharedPtr<const Metasound::Frontend::IEnumDataTypeInterface> GetEnumInterfaceForDataType(FName InDataType) const;
+	virtual TSharedPtr<const Metasound::Frontend::IEnumDataTypeInterface> GetEnumInterfaceForDataType(FName InDataType) const = 0;
 
-	TSharedPtr<Metasound::IDataChannel, ESPMode::ThreadSafe> CreateDataChannelForDataType(const FName& InDataType, const Metasound::FOperatorSettings& InOperatorSettings) const;
+	virtual TSharedPtr<Metasound::IDataChannel, ESPMode::ThreadSafe> CreateDataChannelForDataType(const FName& InDataType, const Metasound::FOperatorSettings& InOperatorSettings) const = 0;
 
 private:
-	FMetasoundFrontendRegistryContainer();
-
 	static FMetasoundFrontendRegistryContainer* LazySingleton;
-
-	// This buffer is used to enqueue nodes and datatypes to register when the module has been initialized,
-	// in order to avoid bad behavior with ensures, logs, etc. on static initialization.
-	// The bad news is that TInlineAllocator is the safest allocator to use on static init.
-	// The good news is that none of these lambdas typically have captures, so this should have low memory overhead.
-	static constexpr int32 MaxNumNodesAndDatatypesToInitialize = 8192;
-	TArray<TUniqueFunction<void()>, TInlineAllocator<MaxNumNodesAndDatatypesToInitialize>> LazyInitCommands;
-	
-	FCriticalSection LazyInitCommandCritSection;
-
-	// Registry in which we keep all information about nodes implemented in C++.
-	TMap<FNodeRegistryKey, FNodeRegistryElement> ExternalNodeRegistry;
-
-	// Registry in which we keep lists of possible nodes to use to convert between two datatypes
-	TMap<FConverterNodeRegistryKey, FConverterNodeRegistryValue> ConverterNodeRegistry;
-
-	struct FDataTypeRegistryElement
-	{
-		Metasound::FDataTypeConstructorCallbacks Callbacks;
-
-		Metasound::FDataTypeRegistryInfo Info;
-
-		TSharedPtr<const Metasound::Frontend::IEnumDataTypeInterface> EnumInterface;
-	};
-
-	TMap<FName, FDataTypeRegistryElement> DataTypeRegistry;
-	TMap<FNodeRegistryKey, FDataTypeRegistryElement> DataTypeNodeRegistry;
 };
 
 
diff --git a/Engine/Plugins/Runtime/Metasound/Source/MetasoundFrontend/Public/MetasoundFrontendSearchEngine.h b/Engine/Plugins/Runtime/Metasound/Source/MetasoundFrontend/Public/MetasoundFrontendSearchEngine.h
new file mode 100644
--- /dev/null
+++ b/Engine/Plugins/Runtime/Metasound/Source/MetasoundFrontend/Public/MetasoundFrontendSearchEngine.h
@@ -0,0 +1,32 @@
+// Copyright Epic Games, Inc. All Rights Reserved.
+
+#pragma once
+
+#include "MetasoundFrontendDocument.h"
+#include "MetasoundNodeInterface.h"
+
+namespace Metasound
+{
+	namespace Frontend
+	{
+		/** Interface for frontend search engine. A frontend search engine provides
+		 * a simple interface for common frontend queries. It also serves as an
+		 * opportunity to cache queries to reduce CPU load. 
+		 */
+		class METASOUNDFRONTEND_API ISearchEngine
+		{
+			public:
+				/** Return an instance of a search engine. */
+				static ISearchEngine& Get();
+
+				/** Find all FMetasoundFrontendClasses with a matching FNodeClassName. */
+				virtual TArray<FMetasoundFrontendClass> FindClassesWithClassName(const FNodeClassName& InName) = 0;
+
+				virtual ~ISearchEngine() = default;
+
+			protected:
+
+				ISearchEngine() = default;
+		};
+	}
+}
diff --git a/Engine/Plugins/Runtime/Metasound/Source/MetasoundGenerator/Private/MetasoundGenerator.cpp b/Engine/Plugins/Runtime/Metasound/Source/MetasoundGenerator/Private/MetasoundGenerator.cpp
--- a/Engine/Plugins/Runtime/Metasound/Source/MetasoundGenerator/Private/MetasoundGenerator.cpp
+++ b/Engine/Plugins/Runtime/Metasound/Source/MetasoundGenerator/Private/MetasoundGenerator.cpp
@@ -12,11 +12,11 @@
 
 namespace Metasound
 {
-	FAsyncMetaSoundBuilder::FAsyncMetaSoundBuilder(FMetasoundGeneratorInitParams&& InInitParams)
-		: InitParams(MoveTemp(InInitParams))
-		, PlayTrigger(FTriggerWriteRef::CreateNew(InInitParams.OperatorSettings))
-		, FinishTrigger(FTriggerReadRef::CreateNew(InInitParams.OperatorSettings))
-		, bSuccess(false)
+
+	FAsyncMetaSoundBuilder::FAsyncMetaSoundBuilder(FMetasoundGenerator* InGenerator, FMetasoundGeneratorInitParams&& InInitParams, bool bInTriggerGenerator)
+		: Generator(InGenerator)
+		, InitParams(MoveTemp(InInitParams))
+		, bTriggerGenerator(bInTriggerGenerator)
 	{
 	}
 
@@ -28,7 +28,7 @@ namespace Metasound
 		ensureAlways(RootGraph->IsValid());
 
 		TArray<IOperatorBuilder::FBuildErrorPtr> BuildErrors;
-		GraphOperator = RootGraph->BuildOperator(InitParams.OperatorSettings, InitParams.Environment, BuildErrors);
+		TUniquePtr<IOperator> GraphOperator = RootGraph->BuildOperator(InitParams.OperatorSettings, InitParams.Environment, BuildErrors);
 
 		// Log build errors
 		for (const IOperatorBuilder::FBuildErrorPtr& Error : BuildErrors)
@@ -42,10 +42,10 @@ namespace Metasound
 		if (!GraphOperator.IsValid())
 		{
 			UE_LOG(LogMetasound, Error, TEXT("Failed to build Metasound operator from graph in MetasoundSource [%s]"), *InitParams.MetaSoundName);
-			bSuccess = false;
 		}
 		else
 		{
+			TArray<FAudioBufferReadRef> OutputBuffers;
 			FDataReferenceCollection Outputs = GraphOperator->GetOutputs();
 
 			// Get output audio buffers.
@@ -68,33 +68,24 @@ namespace Metasound
 
 			// References must be cached before moving the operator to the InitParams
 			FDataReferenceCollection Inputs = GraphOperator->GetInputs();
-			PlayTrigger = Inputs.GetDataWriteReferenceOrConstruct<FTrigger>(InitParams.InputName, InitParams.OperatorSettings, false);
-			FinishTrigger = Outputs.GetDataReadReferenceOrConstruct<FTrigger>(InitParams.IsFinishedOutputName, InitParams.OperatorSettings, false);
-			bSuccess = true;
-		}
-	}
+			FTriggerWriteRef PlayTrigger = Inputs.GetDataWriteReferenceOrConstruct<FTrigger>(InitParams.InputName, InitParams.OperatorSettings, false);
+			FTriggerReadRef FinishTrigger = Outputs.GetDataReadReferenceOrConstruct<FTrigger>(InitParams.IsFinishedOutputName, InitParams.OperatorSettings, false);
 
-	bool FAsyncMetaSoundBuilder::SetDataOnGenerator(FMetasoundGenerator& InGenerator)
-	{
-		if (bSuccess)
-		{
-			FMetasoundGeneratorData GeneratorData =
+			FMetasoundGeneratorData GeneratorData
 			{
+				InitParams.OperatorSettings,
 				MoveTemp(GraphOperator),
-				OutputBuffers,
+				MoveTemp(OutputBuffers),
 				MoveTemp(PlayTrigger),
 				MoveTemp(FinishTrigger)
 			};
 
-			InGenerator.SetGraph(MoveTemp(GeneratorData));
+			Generator->SetPendingGraph(MoveTemp(GeneratorData), bTriggerGenerator);
 		}
-		return bSuccess;
 	}
 
 	FMetasoundGenerator::FMetasoundGenerator(FMetasoundGeneratorInitParams&& InParams)
-		: bIsGraphBuilding(true)
-		, bIsPlaying(false)
-		, bIsFinished(false)
+		: bIsFinished(false)
 		, NumChannels(0)
 		, NumFramesPerExecute(0)
 		, NumSamplesPerExecute(0)
@@ -104,8 +95,14 @@ namespace Metasound
 		NumChannels = InParams.NumOutputChannels;
 		NumFramesPerExecute = InParams.OperatorSettings.GetNumFramesPerBlock();
 		NumSamplesPerExecute = NumChannels * NumFramesPerExecute;
-		BuilderTask = MakeUnique<FBuilderTask>(MoveTemp(InParams));
-		BuilderTask->StartBackgroundTask(GBackgroundPriorityThreadPool);
+
+		BuilderTask = MakeUnique<FBuilderTask>(this, MoveTemp(InParams), true /* bTriggerGenerator */);
+
+		BuilderTask->StartSynchronousTask();
+		BuilderTask = nullptr;
+		UpdateGraphIfPending();
+
+		//BuilderTask->StartBackgroundTask(GBackgroundPriorityThreadPool);
 	}
 
 	FMetasoundGenerator::~FMetasoundGenerator()
@@ -117,47 +114,68 @@ namespace Metasound
 		}
 	}
 
-	bool FMetasoundGenerator::IsTickable() const
+	void FMetasoundGenerator::SetPendingGraph(FMetasoundGeneratorData&& InData, bool bTriggerGraph)
 	{
-		return bIsGraphBuilding;
+		FScopeLock SetPendingGraphLock(&PendingGraphMutex);
+		{
+			PendingGraphData = MakeUnique<FMetasoundGeneratorData>(MoveTemp(InData));
+			bPendingGraphTrigger = bTriggerGraph;
+			bIsNewGraphPending = true;
+		}
 	}
 
-	void FMetasoundGenerator::Tick(float DeltaTime)
+	void FMetasoundGenerator::UpdateGraphIfPending()
 	{
-		if (BuilderTask->IsDone())
-		{	
-			FAsyncMetaSoundBuilder& Builder = BuilderTask->GetTask();
-			if (Builder.SetDataOnGenerator(*this))
-			{
-				// We're done building our graph now
-				bIsGraphBuilding = false;
-			}
-			else
-			{
-				// Failed to load/generate the graph, kill the metasound
-				bIsFinished = true;
-			}
-
-			BuilderTask = nullptr;
+		FScopeLock GraphLock(&PendingGraphMutex);
+		if (bIsNewGraphPending)
+		{
+			SetGraph(MoveTemp(PendingGraphData), bPendingGraphTrigger);
+			bIsNewGraphPending = false;
 		}
 	}
 
-	void FMetasoundGenerator::SetGraph(FMetasoundGeneratorData&& InData)
+	void FMetasoundGenerator::SetGraph(TUniquePtr<FMetasoundGeneratorData>&& InData, bool bTriggerGraph)
 	{
+		if (!InData.IsValid())
+		{
+			return;
+		}
+
 		InterleavedAudioBuffer.Reset();
 
 		GraphOutputAudio.Reset();
-		if (InData.OutputBuffers.Num() > 0)
+		if (InData->OutputBuffers.Num() == NumChannels)
 		{
-			GraphOutputAudio.Append(InData.OutputBuffers.GetData(), InData.OutputBuffers.Num());
+			if (InData->OutputBuffers.Num() > 0)
+			{
+				GraphOutputAudio.Append(InData->OutputBuffers.GetData(), InData->OutputBuffers.Num());
+			}
+		}
+		else
+		{
+			int32 FoundNumChannels = InData->OutputBuffers.Num();
+
+			UE_LOG(LogMetasound, Warning, TEXT("Metasound generator expected %d number of channels, found %d"), NumChannels, FoundNumChannels);
+
+			int32 NumChannelsToCopy = FMath::Min(FoundNumChannels, NumChannels);
+			int32 NumChannelsToCreate = NumChannels - NumChannelsToCopy;
+
+			if (NumChannelsToCopy > 0)
+			{
+				GraphOutputAudio.Append(InData->OutputBuffers.GetData(), NumChannelsToCopy);
+			}
+			for (int32 i = 0; i < NumChannelsToCreate; i++)
+			{
+				GraphOutputAudio.Add(TDataReadReference<FAudioBuffer>::CreateNew(InData->OperatorSettings));
+			}
 		}
 
-		OnPlayTriggerRef = InData.TriggerOnPlayRef;
-		OnFinishedTriggerRef = InData.TriggerOnFinishRef;
+		OnPlayTriggerRef = InData->TriggerOnPlayRef;
+		OnFinishedTriggerRef = InData->TriggerOnFinishRef;
 
 		// The graph operator and graph audio output contain all the values needed
 		// by the sound generator.
-		RootExecuter.SetOperator(MoveTemp(InData.GraphOperator));
+		RootExecuter.SetOperator(MoveTemp(InData->GraphOperator));
 
 
 		// Query the graph output to get the number of output audio channels.
@@ -169,6 +187,11 @@ namespace Metasound
 			// Preallocate interleaved buffer as it is necessary for any audio generation calls.
 			InterleavedAudioBuffer.AddUninitialized(NumSamplesPerExecute);
 		}
+
+		if (bTriggerGraph)
+		{
+			OnPlayTriggerRef->TriggerFrame(0);
+		}
 	}
 
 	int32 FMetasoundGenerator::GetNumChannels() const
@@ -183,19 +206,15 @@ namespace Metasound
 			return 0;
 		}
 
+		UpdateGraphIfPending();
+
 		// Output silent audio if we're still building a graph
-		if (bIsGraphBuilding || NumSamplesPerExecute < 1)
+		if (RootExecuter.IsNoOp() || NumSamplesPerExecute < 1)
 		{
 			FMemory::Memset(OutAudio, 0, sizeof(float) * NumSamplesRemaining);
 			return NumSamplesRemaining;
 		}
 
-		if (!bIsPlaying)
-		{
-			OnPlayTriggerRef->TriggerFrame(0);
-			bIsPlaying = true;
-		}
-
 		// If we have any audio left in the internal overflow buffer from 
 		// previous calls, write that to the output before generating more audio.
 		int32 NumSamplesWritten = FillWithBuffer(OverflowBuffer, OutAudio, NumSamplesRemaining);
@@ -204,7 +223,6 @@ namespace Metasound
 		{
 			NumSamplesRemaining -= NumSamplesWritten;
 			OverflowBuffer.RemoveAtSwap(0 /* Index */, NumSamplesWritten /* Count */, false /* bAllowShrinking */);
-
 		}
 
 		while (NumSamplesRemaining > 0)
diff --git a/Engine/Plugins/Runtime/Metasound/Source/MetasoundGenerator/Public/MetasoundGenerator.h b/Engine/Plugins/Runtime/Metasound/Source/MetasoundGenerator/Public/MetasoundGenerator.h
--- a/Engine/Plugins/Runtime/Metasound/Source/MetasoundGenerator/Public/MetasoundGenerator.h
+++ b/Engine/Plugins/Runtime/Metasound/Source/MetasoundGenerator/Public/MetasoundGenerator.h
@@ -30,8 +30,9 @@ namespace Metasound
 
 	struct FMetasoundGeneratorData
 	{
+		FOperatorSettings OperatorSettings;
 		TUniquePtr<Metasound::IOperator> GraphOperator;
-		TArrayView<TDataReadReference<FAudioBuffer>> OutputBuffers;
+		TArray<TDataReadReference<FAudioBuffer>> OutputBuffers;
 		FTriggerWriteRef TriggerOnPlayRef;
 		FTriggerReadRef TriggerOnFinishRef;
 	};
@@ -41,11 +42,11 @@ namespace Metasound
 	class FAsyncMetaSoundBuilder : public FNonAbandonableTask
 	{
 	public:
-		FAsyncMetaSoundBuilder(FMetasoundGeneratorInitParams&& InInitParams);
+		FAsyncMetaSoundBuilder(FMetasoundGenerator* InGenerator, FMetasoundGeneratorInitParams&& InInitParams, bool bInTriggerGenerator);
+
 		~FAsyncMetaSoundBuilder() = default;
 
 		void DoWork();
-		bool SetDataOnGenerator(FMetasoundGenerator& InGenerator);
 
 		FORCEINLINE TStatId GetStatId() const
 		{
@@ -53,20 +54,16 @@ namespace Metasound
 		}
 
 	private:
+
+		FMetasoundGenerator* Generator;
 		FMetasoundGeneratorInitParams InitParams;
-		TUniquePtr<IOperator> GraphOperator;
-		TArrayView<FAudioBufferReadRef> OutputBuffers;
-		FTriggerWriteRef PlayTrigger;
-		FTriggerReadRef FinishTrigger;
-		bool bSuccess;
+		bool bTriggerGenerator;
 	};
 
-
 	/** FMetasoundGenerator generates audio from a given metasound IOperator
 	 * which produces a multichannel audio output.
 	 */
-	class METASOUNDGENERATOR_API FMetasoundGenerator : public ISoundGenerator, 
-														public FTickableGameObject
+	class METASOUNDGENERATOR_API FMetasoundGenerator : public ISoundGenerator
 	{
 	public:
 		using FOperatorUniquePtr = TUniquePtr<Metasound::IOperator>;
@@ -81,15 +78,6 @@ namespace Metasound
 
 		virtual ~FMetasoundGenerator();
 
-		//~ Begin FTickableGameObject
-		virtual ETickableTickType GetTickableTickType() const override { return ETickableTickType::Conditional; }
-		virtual bool IsTickableWhenPaused() const override { return true; }
-		virtual bool IsTickableInEditor() const override { return true; }
-		virtual bool IsTickable() const override;
-		virtual void Tick(float DeltaTime) override;
-		virtual TStatId GetStatId() const override { RETURN_QUICK_DECLARE_CYCLE_STAT(MetasoundGenerator, STATGROUP_Tickables); }
-		//~ End FTickableGameObject
-
 		/** Set the value of a graph's input data using the assignment operator.
 		 *
 		 * @param InName - The name of the graph's input data reference.
@@ -137,17 +125,6 @@ namespace Metasound
 			}
 		}
 
-		/** Update the current graph operator with a new graph operator. The number of channels
-		 * of InGraphOutputAudioRef must match the existing number of channels reported by
-		 * GetNumChannels() in order for this function to successfully replace the graph operator.
-		 *
-		 * @param InGraphOperator - Unique pointer to the IOperator which executes the entire graph.
-		 * @param InGraphOutputAudioRef - Read reference to the audio buffer filled by the InGraphOperator.
-		 *
-		 * @return Returns true if the graph operator was successfully swapped. False otherwise.
-		 */
-//		bool UpdateGraphOperator(FMetasoundGeneratorInitParams&& InParams);
-
 		/** Return the number of audio channels. */
 		int32 GetNumChannels() const;
 
@@ -156,10 +133,21 @@ namespace Metasound
 		int32 GetDesiredNumSamplesToRenderPerCallback() const override;
 		bool IsFinished() const override;
 		//~ End FSoundGenerator
+		
+		/** Update the current graph operator with a new graph operator. The number of channels
+		 * of InGraphOutputAudioRef must match the existing number of channels reported by
+		 * GetNumChannels() in order for this function to successfully replace the graph operator.
+		 *
+		 * @param InData - Metasound data of built graph.
+		 * @param bTriggerGraph - If true, "OnPlay" will be triggered on the new graph.
+		 */
+		void SetPendingGraph(FMetasoundGeneratorData&& InData, bool bTriggerGraph);
+
 
 	private:
+		void UpdateGraphIfPending();
 		// Internal set graph after checking compatibility.
-		void SetGraph(FMetasoundGeneratorData&& InData);
+		void SetGraph(TUniquePtr<FMetasoundGeneratorData>&& InData, bool bTriggerGraph);
 
 		// Fill OutAudio with data in InBuffer, up to maximum number of samples.
 		// Returns the number of samples used.
@@ -171,7 +159,6 @@ namespace Metasound
 		FExecuter RootExecuter;
 
 		bool bIsGraphBuilding;
-		bool bIsPlaying;
 		bool bIsFinished;
 
 		int32 NumChannels;
@@ -193,7 +180,10 @@ namespace Metasound
 		typedef FAsyncTask<FAsyncMetaSoundBuilder> FBuilderTask;
 		TUniquePtr<FBuilderTask> BuilderTask;
 
-		friend class FAsyncMetaSoundBuilder;
+		FCriticalSection PendingGraphMutex;
+		TUniquePtr<FMetasoundGeneratorData> PendingGraphData;
+		bool bPendingGraphTrigger;
+		bool bIsNewGraphPending;
 	};
 }
 
