Metasound Wavewriter fixes:
Move to FString for filename caching to avoid underscore issues.
Fix up error message if failure to open file.
Make a "default" operator in case where file can't be written, so we don't crash.

diff --git a/Engine/Plugins/Runtime/Metasound/Source/MetasoundStandardNodes/Private/MetasoundStandardNodesCategories.cpp b/Engine/Plugins/Runtime/Metasound/Source/MetasoundStandardNodes/Private/MetasoundStandardNodesCategories.cpp
--- a/Engine/Plugins/Runtime/Metasound/Source/MetasoundStandardNodes/Private/MetasoundStandardNodesCategories.cpp
+++ b/Engine/Plugins/Runtime/Metasound/Source/MetasoundStandardNodes/Private/MetasoundStandardNodesCategories.cpp
@@ -14,7 +14,7 @@ namespace Metasound
 		const FText RandomUtils = { LOCTEXT("Metasound_RandomUtilsCategory", "Random Utilities") };
 		const FText Conversions  = { LOCTEXT("Metasound_ConversionsCategory", "Conversions") };
 		const FText DebugUtils = { LOCTEXT("Metasound_DebugUtilsCategory", "Debug Utilities") };
-		const FText Io = { LOCTEXT("Metasound_IoCategory", "File IO") };
+		const FText Io = { LOCTEXT("Metasound_IoCategory", "External IO") };
 	}
 }
 
diff --git a/Engine/Plugins/Runtime/Metasound/Source/MetasoundStandardNodes/Private/MetasoundWaveWriterNode.cpp b/Engine/Plugins/Runtime/Metasound/Source/MetasoundStandardNodes/Private/MetasoundWaveWriterNode.cpp
--- a/Engine/Plugins/Runtime/Metasound/Source/MetasoundStandardNodes/Private/MetasoundWaveWriterNode.cpp
+++ b/Engine/Plugins/Runtime/Metasound/Source/MetasoundStandardNodes/Private/MetasoundWaveWriterNode.cpp
@@ -120,7 +120,7 @@ namespace Metasound
 		virtual ~FFileWriteError() = default;
 
 		FFileWriteError(const FNode& InNode, const FString& InFilename)
-			: FBuildErrorBase(ErrorType, FText::FromString(FString::Printf(TEXT("File Writer Error while trying to write '%s'"), *InFilename)))
+			: FBuildErrorBase(ErrorType, FText::Format(FTextFormat(LOCTEXT("MetasoundFileWriterErrorDescription", "File Writer Error while trying to write '{0}'")), FText::FromString(InFilename)))
 		{
 			AddNode(InNode);
 		}
@@ -138,11 +138,11 @@ namespace Metasound
 			CacheFilenames();
 		}
 
-		FString GenerateNextNumberedFilename(const FName& InPrefix)
+		FString GenerateNextNumberedFilename(const FString& InPrefix)
 		{
 			FScopeLock Lock{ &Cs };
-			uint32& CurrentMax = FileIndexMap.FindOrAdd(InPrefix);
-			FString Filename{ *InPrefix.GetPlainNameString() };
+			uint32& CurrentMax = FileIndexMap.FindOrAdd(InPrefix.ToUpper());
+			FString Filename{ InPrefix };
 			Filename.Append(*Seperator);
 			Filename.AppendInt(++CurrentMax);
 			Filename.Append(*FileExtention);
@@ -169,7 +169,7 @@ namespace Metasound
 						int32 Number = FCString::Atoi(*NumberString);
 						if (Number >= 0)
 						{
-							uint32& CurrentMax = FileIndexMap.FindOrAdd(*Prefix);
+							uint32& CurrentMax = FileIndexMap.FindOrAdd(*Prefix.ToUpper());
 							if (static_cast<uint32>(Number) > CurrentMax)
 							{
 								CurrentMax = static_cast<uint32>(Number);
@@ -182,7 +182,7 @@ namespace Metasound
 		FCriticalSection Cs;
 		FString RootPath;
 		FString FileExtention;
-		TMap<FName, uint32> FileIndexMap;
+		TMap<FString, uint32> FileIndexMap;
 	};
 	const FString FNumberedFileCache::Seperator{ TEXT("_") };
 
@@ -253,8 +253,7 @@ namespace Metasound
 			static TUniquePtr<FNumberedFileCache> NumberedFileCache = MakeUnique<FNumberedFileCache>(*FPaths::AudioCaptureDir(), WaveExt);
 			
 			FStringReadRef FilenamePrefix = InputCol.GetDataReadReferenceOrConstructWithVertexDefault<FString>(InputInterface, FilenamePrefixPinName, Settings);
-			FName Prefix = FName{ **FilenamePrefix };
-			FString Filename = NumberedFileCache->GenerateNextNumberedFilename(Prefix);
+			FString Filename = NumberedFileCache->GenerateNextNumberedFilename(*FilenamePrefix);
 
 			if (TUniquePtr<FArchive> FileWriter = TUniquePtr<FArchive>(IFileManager::Get().CreateFileWriter(*Filename, IO_WRITE)))
 			{
@@ -268,7 +267,14 @@ namespace Metasound
 			
 			// Failed to open a writer object. Log an error.
 			OutErrors.Emplace(MakeUnique<FFileWriteError>(Node, Filename));
-			return {};
+
+			// Create a default operator with no-writer which will do nothing.
+			return MakeUnique<FWaveWriterOperator>(
+				Settings,
+				InputCol.GetDataReadReferenceOrConstruct<FAudioBuffer>(AudioInputPinName, Settings),
+				InputCol.GetDataReadReferenceOrConstructWithVertexDefault<bool>(InputInterface, EnabledPinName, Settings),
+				nullptr
+			);
 		}
 
 		void Execute()
