Metasound Editor Fixes
- Undo/Redo fixed when linking/unlinking nodes
- Node synchronization fix (adding ghost ed nodes when modifying default value in inspector)
- General synchronization logging for easier debugging
- Fix nodes jumping back to origin when swapping Metasound archetype

diff --git a/Engine/Plugins/Runtime/Metasound/Source/MetasoundEditor/Private/MetasoundEditor.cpp b/Engine/Plugins/Runtime/Metasound/Source/MetasoundEditor/Private/MetasoundEditor.cpp
--- a/Engine/Plugins/Runtime/Metasound/Source/MetasoundEditor/Private/MetasoundEditor.cpp
+++ b/Engine/Plugins/Runtime/Metasound/Source/MetasoundEditor/Private/MetasoundEditor.cpp
@@ -401,12 +401,16 @@ namespace Metasound
 			}
 
 			FSlateApplication::Get().DismissAllMenus();
+		}
 
-			if (Metasound)
-			{
-				UEdGraphPin::ResolveAllPinReferences();
-				Metasound::Editor::FGraphBuilder::SynchronizeGraph(*Metasound);
-			}
+		void FEditor::NotifyUserModifiedBySync()
+		{
+			FNotificationInfo Info(LOCTEXT("SynchronizationRequired", "Operation modified pin(s), connection(s), and/or node(s).  Please refer to graph."));
+			Info.bFireAndForget = true;
+			Info.bUseSuccessFailIcons = false;
+			Info.ExpireDuration = 5.0f;
+
+			MetasoundGraphEditor->AddNotification(Info, false /* bSuccess */);
 		}
 
 		void FEditor::NotifyPostChange(const FPropertyChangedEvent& PropertyChangedEvent, FProperty* PropertyThatChanged)
@@ -423,6 +427,7 @@ namespace Metasound
 					{
 						if (FGraphBuilder::SynchronizeGraph(*MetasoundObj))
 						{
+							NotifyUserModifiedBySync();
 							MetasoundObj->MarkPackageDirty();
 						}
 						MetasoundGraphEditor->NotifyGraphChanged();
@@ -1494,7 +1499,10 @@ namespace Metasound
 				MetasoundGraphEditor->SetNodeSelection(GraphNode, true);
 			}
 
-			FGraphBuilder::SynchronizeGraph(*Metasound);
+			if (FGraphBuilder::SynchronizeGraph(*Metasound))
+			{
+				NotifyUserModifiedBySync();
+			}
 
 			MetasoundGraphEditor->NotifyGraphChanged();
 			Metasound->MarkPackageDirty();
diff --git a/Engine/Plugins/Runtime/Metasound/Source/MetasoundEditor/Private/MetasoundEditor.h b/Engine/Plugins/Runtime/Metasound/Source/MetasoundEditor/Private/MetasoundEditor.h
--- a/Engine/Plugins/Runtime/Metasound/Source/MetasoundEditor/Private/MetasoundEditor.h
+++ b/Engine/Plugins/Runtime/Metasound/Source/MetasoundEditor/Private/MetasoundEditor.h
@@ -240,6 +240,8 @@ namespace Metasound
 		private:
 			void EditObjectSettings();
 
+			void NotifyUserModifiedBySync();
+
 			/** List of open tool panels; used to ensure only one exists at any one time */
 			TMap<FName, TWeakPtr<SDockableTab>> SpawnedToolPanels;
 
diff --git a/Engine/Plugins/Runtime/Metasound/Source/MetasoundEditor/Private/MetasoundEditorGraphBuilder.cpp b/Engine/Plugins/Runtime/Metasound/Source/MetasoundEditor/Private/MetasoundEditorGraphBuilder.cpp
--- a/Engine/Plugins/Runtime/Metasound/Source/MetasoundEditor/Private/MetasoundEditorGraphBuilder.cpp
+++ b/Engine/Plugins/Runtime/Metasound/Source/MetasoundEditor/Private/MetasoundEditorGraphBuilder.cpp
@@ -471,13 +471,10 @@ namespace Metasound
 							FAssetRegistryModule& AssetRegistryModule = FModuleManager::LoadModuleChecked<FAssetRegistryModule>("AssetRegistry");
 
 							const FString ClassName = Class->GetName();
-							FString ObjectPath = InInputPin.DefaultValue;
 
 							// Remove class prefix if included in default value path
-							if (ObjectPath.StartsWith(ClassName + TEXT(" ")))
-							{
-								ObjectPath.RemoveFromStart(ClassName + TEXT(" "));
-							}
+							FString ObjectPath = InInputPin.DefaultValue;
+							ObjectPath.RemoveFromStart(ClassName + TEXT(" "));
 
 							FARFilter Filter;
 							Filter.bRecursiveClasses = false;
@@ -1042,51 +1039,23 @@ namespace Metasound
 
 		bool FGraphBuilder::IsMatchingInputHandleAndPin(const Frontend::FInputHandle& InInputHandle, const UEdGraphPin& InEditorPin)
 		{
-			if (EEdGraphPinDirection::EGPD_Input == InEditorPin.Direction)
+			Frontend::FInputHandle PinInputHandle = GetInputHandleFromPin(&InEditorPin);
+			if (PinInputHandle->GetID() == InInputHandle->GetID())
 			{
-				if (InEditorPin.GetName() == InInputHandle->GetName())
-				{
-					IMetasoundEditorModule& EditorModule = FModuleManager::GetModuleChecked<IMetasoundEditorModule>("MetasoundEditor");
-
-					if (InEditorPin.PinType == EditorModule.FindDataType(InInputHandle->GetDataType()).PinType)
-					{
-						UEdGraphNode* EditorNode = InEditorPin.GetOwningNode();
-						if (UMetasoundEditorGraphNode* MetasoundEditorNode = Cast<UMetasoundEditorGraphNode>(EditorNode))
-						{
-							if  (MetasoundEditorNode->GetNodeHandle()->GetID() == InInputHandle->GetOwningNodeID())
-							{
-								return true;
-							}
-						}
-					}
-				}
+				return true;
 			}
+
 			return false;
 		}
 
 		bool FGraphBuilder::IsMatchingOutputHandleAndPin(const Frontend::FOutputHandle& InOutputHandle, const UEdGraphPin& InEditorPin)
 		{
-			if (EEdGraphPinDirection::EGPD_Output == InEditorPin.Direction)
+			Frontend::FOutputHandle PinOutputHandle = GetOutputHandleFromPin(&InEditorPin);
+			if (PinOutputHandle->GetID() == InOutputHandle->GetID())
 			{
-				if (InEditorPin.GetName() == InOutputHandle->GetName())
-				{
-					IMetasoundEditorModule& EditorModule = FModuleManager::GetModuleChecked<IMetasoundEditorModule>("MetasoundEditor");
-
-					if (InEditorPin.PinType == EditorModule.FindDataType(InOutputHandle->GetDataType()).PinType)
-					{
-						UEdGraphNode* EditorNode = InEditorPin.GetOwningNode();
-						UMetasoundEditorGraphNode* MetasoundEditorNode = Cast<UMetasoundEditorGraphNode>(EditorNode);
-
-						if (nullptr != MetasoundEditorNode)
-						{
-							if  (MetasoundEditorNode->GetNodeHandle()->GetID() == InOutputHandle->GetOwningNodeID())
-							{
-								return true;
-							}
-						}
-					}
-				}
+				return true;
 			}
+
 			return false;
 		}
 
@@ -1227,6 +1196,8 @@ namespace Metasound
 								if (ensure(!OutputEditorNode->IsEmpty()))
 								{
 									UEdGraphPin* OutputPin = (*OutputEditorNode)[0]->FindPinChecked(OutputHandle->GetName(), EEdGraphPinDirection::EGPD_Output);
+									const FText& OwningNodeName = NodeInput->GetOwningNode()->GetDisplayName();
+									UE_LOG(LogMetasoundEditor, Display, TEXT("Synchronizing Node '%s' Connection: Linking Pin '%s' to '%s'"), *OwningNodeName.ToString(), *MatchingPin->GetName(), *OutputPin->GetName());
 									MatchingPin->MakeLinkTo(OutputPin);
 									bIsNodeDirty = true;
 								}
@@ -1239,6 +1210,8 @@ namespace Metasound
 						if (!MatchingPin->LinkedTo.IsEmpty())
 						{
 							MatchingPin->BreakAllPinLinks();
+							const FText& OwningNodeName = NodeInput->GetOwningNode()->GetDisplayName();
+							UE_LOG(LogMetasoundEditor, Display, TEXT("Synchronizing Node '%s' Connection: Breaking all pin links to '%s'"), *OwningNodeName.ToString(), *NodeInput->GetDisplayName().ToString());
 							bIsNodeDirty = true;
 						}
 
@@ -1252,6 +1225,8 @@ namespace Metasound
 							const FString DefaultLiteralString = DefaultLiteral.ToString();
 							if (MatchingPin->DefaultValue != DefaultLiteralString)
 							{
+								const FText& OwningNodeName = NodeInput->GetOwningNode()->GetDisplayName();
+								UE_LOG(LogMetasoundEditor, Display, TEXT("Synchronizing Node '%s' Connection: Setting pin '%s' default value to '%s'"), *OwningNodeName.ToString(), *MatchingPin->GetName(), *DefaultLiteralString);
 								MatchingPin->DefaultValue = DefaultLiteralString;
 								bIsNodeDirty = true;
 							}
@@ -1356,32 +1331,9 @@ namespace Metasound
 					bIsEditorGraphDirty = true;
 
 					TArray<UMetasoundEditorGraphNode*> AddedNodes;
-					if (Style.Display.Locations.Num() > 0)
+					for (const TPair<FGuid, FVector2D>& Location : Style.Display.Locations)
 					{
-						for (const TPair<FGuid, FVector2D>& Location : Style.Display.Locations)
-						{
-							UMetasoundEditorGraphNode* NewNode = AddNode(InMetasound, Node, Location.Value, false /* bInSelectNewNode */);
-							if (ensure(NewNode))
-							{
-								FAssociatedNodes& AssociatedNodeData = AssociatedNodes.FindOrAdd(Node->GetID());
-								if (AssociatedNodeData.Node->IsValid())
-								{
-									ensure(AssociatedNodeData.Node == Node);
-								}
-								else
-								{
-									AssociatedNodeData.Node = Node;
-								}
-
-								AddedNodes.Add(NewNode);
-								AssociatedNodeData.EditorNodes.Add(NewNode);
-							}
-						}
-					}
-					else
-					{
-						// TODO: pass in a non-origin position, ideally the position of the output node that this is replacing (e.g. the old mono node, that may have been moved by a user)
-						UMetasoundEditorGraphNode* NewNode = AddNode(InMetasound, Node, { 0.0f, 0.0f }, false /* bInSelectNewNode */);
+						UMetasoundEditorGraphNode* NewNode = AddNode(InMetasound, Node, Location.Value, false /* bInSelectNewNode */);
 						if (ensure(NewNode))
 						{
 							FAssociatedNodes& AssociatedNodeData = AssociatedNodes.FindOrAdd(Node->GetID());
@@ -1422,7 +1374,7 @@ namespace Metasound
 			return bIsEditorGraphDirty;
 		}
 
-		bool FGraphBuilder::SynchronizeNodePins(UMetasoundEditorGraphNode& InEditorNode, Frontend::FNodeHandle InNode, bool bRemoveUnusedPins)
+		bool FGraphBuilder::SynchronizeNodePins(UMetasoundEditorGraphNode& InEditorNode, Frontend::FNodeHandle InNode, bool bRemoveUnusedPins, bool bLogChanges)
 		{
 			bool bIsNodeDirty = false;
 
@@ -1471,15 +1423,16 @@ namespace Metasound
 				}
 			}
 
-			bIsNodeDirty |= !InputHandles.IsEmpty();
-			bIsNodeDirty |= !OutputHandles.IsEmpty();
-			bIsNodeDirty |= !EditorPins.IsEmpty();
-
 			// Remove any unused editor pins.
 			if (bRemoveUnusedPins)
 			{
+				bIsNodeDirty |= !EditorPins.IsEmpty();
 				for (UEdGraphPin* Pin : EditorPins)
 				{
+					if (bLogChanges)
+					{
+						UE_LOG(LogMetasoundEditor, Display, TEXT("Synchronizing Node '%s' Pins: Removing Excess Editor Pin '%s'"), *InNode->GetDisplayName().ToString(), *Pin->GetName());
+					}
 					InEditorNode.RemovePin(Pin);
 				}
 			}
@@ -1489,9 +1442,14 @@ namespace Metasound
 			// to the graph.
 			if (EMetasoundFrontendClassType::Input != InNode->GetClassType())
 			{
+				bIsNodeDirty |= !InputHandles.IsEmpty();
 				InputHandles = InNode->GetOutputStyle().SortDefaults(InputHandles);
 				for (Frontend::FInputHandle& InputHandle : InputHandles)
 				{
+					if (bLogChanges)
+					{
+						UE_LOG(LogMetasoundEditor, Display, TEXT("Synchronizing Node '%s' Pins: Adding missing Editor Input Pin '%s'"), *InNode->GetDisplayName().ToString(), *InputHandle->GetDisplayName().ToString());
+					}
 					AddPinToNode(InEditorNode, InputHandle);
 				}
 			}
@@ -1501,9 +1459,14 @@ namespace Metasound
 			// to the graph.
 			if (EMetasoundFrontendClassType::Output != InNode->GetClassType())
 			{
+				bIsNodeDirty |= !OutputHandles.IsEmpty();
 				OutputHandles = InNode->GetOutputStyle().SortDefaults(OutputHandles);
 				for (Frontend::FOutputHandle& OutputHandle : OutputHandles)
 				{
+					if (bLogChanges)
+					{
+						UE_LOG(LogMetasoundEditor, Display, TEXT("Synchronizing Node '%s' Pins: Adding missing Editor Output Pin '%s'"), *InNode->GetDisplayName().ToString(), *OutputHandle->GetDisplayName().ToString());
+					}
 					AddPinToNode(InEditorNode, OutputHandle);
 				}
 			}
@@ -1539,12 +1502,14 @@ namespace Metasound
 					Inputs.Add(Input);
 					return;
 				}
+
 				if (!ensure(NodeHandle->GetNumInputs() == 1))
 				{
 					return;
 				}
 
 				Inputs.Add(Graph->FindOrAddInput(NodeHandle));
+				UE_LOG(LogMetasoundEditor, Display, TEXT("Synchronizing Inputs: Added missing input '%s'."), *NodeHandle->GetDisplayName().ToString());
 				bIsEditorGraphDirty = true;
 			}, EMetasoundFrontendClassType::Input);
 
@@ -1561,6 +1526,7 @@ namespace Metasound
 				}
 
 				Outputs.Add(Graph->FindOrAddOutput(NodeHandle));
+				UE_LOG(LogMetasoundEditor, Display, TEXT("Synchronizing Outputs: Added missing output '%s'."), *NodeHandle->GetDisplayName().ToString());
 				bIsEditorGraphDirty = true;
 			}, EMetasoundFrontendClassType::Output);
 
@@ -1569,6 +1535,7 @@ namespace Metasound
 			{
 				if (!Inputs.Contains(&Input))
 				{
+					UE_LOG(LogMetasoundEditor, Display, TEXT("Synchronizing Inputs: Removing stale input '%s'."), *Input.GetName());
 					ToRemove.Add(&Input);
 				}
 			});
@@ -1576,6 +1543,7 @@ namespace Metasound
 			{
 				if (!Outputs.Contains(&Output))
 				{
+					UE_LOG(LogMetasoundEditor, Display, TEXT("Synchronizing Outputs: Removing stale output '%s'."), *Output.GetName());
 					ToRemove.Add(&Output);
 				}
 			});
diff --git a/Engine/Plugins/Runtime/Metasound/Source/MetasoundEditor/Private/MetasoundEditorGraphBuilder.h b/Engine/Plugins/Runtime/Metasound/Source/MetasoundEditor/Private/MetasoundEditorGraphBuilder.h
--- a/Engine/Plugins/Runtime/Metasound/Source/MetasoundEditor/Private/MetasoundEditorGraphBuilder.h
+++ b/Engine/Plugins/Runtime/Metasound/Source/MetasoundEditor/Private/MetasoundEditorGraphBuilder.h
@@ -162,7 +162,7 @@ namespace Metasound
 			// Adds and removes pins so that the UMetasoundEditorGraphNode matches the InNode.
 			//
 			// @return True if the UMetasoundEditorGraphNode was altered. False otherwise.
-			static bool SynchronizeNodePins(UMetasoundEditorGraphNode& InEditorNode, Frontend::FNodeHandle InNode, bool bRemoveUnusedPins = true);
+			static bool SynchronizeNodePins(UMetasoundEditorGraphNode& InEditorNode, Frontend::FNodeHandle InNode, bool bRemoveUnusedPins = true, bool bLogChanges = true);
 
 			// Adds and removes connections so that the UEdGraph of the metasound has the same
 			// connections as the FMetasoundFrontendDocument graph.
diff --git a/Engine/Plugins/Runtime/Metasound/Source/MetasoundEditor/Private/MetasoundEditorGraphNode.cpp b/Engine/Plugins/Runtime/Metasound/Source/MetasoundEditor/Private/MetasoundEditorGraphNode.cpp
--- a/Engine/Plugins/Runtime/Metasound/Source/MetasoundEditor/Private/MetasoundEditorGraphNode.cpp
+++ b/Engine/Plugins/Runtime/Metasound/Source/MetasoundEditor/Private/MetasoundEditorGraphNode.cpp
@@ -152,7 +152,7 @@ void UMetasoundEditorGraphNode::ReconstructNode()
 	// are lost.
 	// TODO: User will want to see dead pins as well for node definition changes. Label and color-code dead
 	// pins (ex. red), and leave dead connections for reference like BP.
-	Editor::FGraphBuilder::SynchronizeNodePins(*this, GetNodeHandle(), false /* bRemoveUnusedPins */);
+	Editor::FGraphBuilder::SynchronizeNodePins(*this, GetNodeHandle(), false /* bRemoveUnusedPins */, false /* bLogChanges */);
 }
 
 void UMetasoundEditorGraphNode::AutowireNewNode(UEdGraphPin* FromPin)
diff --git a/Engine/Plugins/Runtime/Metasound/Source/MetasoundEditor/Private/MetasoundEditorGraphSchema.cpp b/Engine/Plugins/Runtime/Metasound/Source/MetasoundEditor/Private/MetasoundEditorGraphSchema.cpp
--- a/Engine/Plugins/Runtime/Metasound/Source/MetasoundEditor/Private/MetasoundEditorGraphSchema.cpp
+++ b/Engine/Plugins/Runtime/Metasound/Source/MetasoundEditor/Private/MetasoundEditorGraphSchema.cpp
@@ -166,7 +166,7 @@ UEdGraphNode* FMetasoundGraphSchemaAction_NewNode::PerformAction(UEdGraph* Paren
 	using namespace Metasound::Editor;
 	using namespace Metasound::Frontend;
 
-	const FScopedTransaction Transaction(LOCTEXT("MetasoundEditorNewNode", "Add New Metasound Node"));
+	const FScopedTransaction Transaction(LOCTEXT("AddNewNode", "Add New Metasound Node"));
 	ParentGraph->Modify();
 
 	UObject& ParentMetasound = CastChecked<UMetasoundEditorGraph>(ParentGraph)->GetMetasoundChecked();
@@ -190,7 +190,7 @@ UEdGraphNode* FMetasoundGraphSchemaAction_NewInput::PerformAction(UEdGraph* Pare
 	using namespace Metasound::Editor;
 	using namespace Metasound::Frontend;
 
-	const FScopedTransaction Transaction(LOCTEXT("MetasoundEditorNewInputNode", "Add New Metasound Input Node"));
+	const FScopedTransaction Transaction(LOCTEXT("AddNewInputNode", "Add New Metasound Input Node"));
 
 	UMetasoundEditorGraph* MetasoundGraph = CastChecked<UMetasoundEditorGraph>(ParentGraph);
 	UObject& ParentMetasound = MetasoundGraph->GetMetasoundChecked();
@@ -235,7 +235,6 @@ UEdGraphNode* FMetasoundGraphSchemaAction_PromoteToInput::PerformAction(UEdGraph
 	}
 
 	const FScopedTransaction Transaction(LOCTEXT("PromoteNodeInputToGraphInput", "Promote Metasound Node Input to Graph Input"));
-
 	ParentGraph->Modify();
 
 	UMetasoundEditorGraph* MetasoundGraph = CastChecked<UMetasoundEditorGraph>(ParentGraph);
@@ -287,8 +286,7 @@ UEdGraphNode* FMetasoundGraphSchemaAction_NewOutput::PerformAction(UEdGraph* Par
 	using namespace Metasound::Editor;
 	using namespace Metasound::Frontend;
 
-	const FScopedTransaction Transaction(LOCTEXT("MetasoundEditorNewInputNode", "Add New Metasound Output Node"));
-
+	const FScopedTransaction Transaction(LOCTEXT("AddNewOutputNode", "Add New Metasound Output Node"));
 	ParentGraph->Modify();
 
 	UMetasoundEditorGraph* MetasoundGraph = CastChecked<UMetasoundEditorGraph>(ParentGraph);
@@ -580,18 +578,27 @@ bool UMetasoundEditorGraphSchema::TryCreateConnection(UEdGraphPin* PinA, UEdGrap
 
 	// TODO: Implement YesWithConverterNode with selected conversion option
 
-	BreakPinLinks(*InputPin, false);
-	bool bModified = UEdGraphSchema::TryCreateConnection(PinA, PinB);
-	if (bModified)
+	// Must mark Metasound object as modified to avoid desync issues ***before*** attempting to create a connection
+	// so that transaction stack observes Frontend changes last if rolled back (i.e. undone).  UEdGraphSchema::TryCreateConnection
+	// intrinsically marks the respective pin EdGraphNodes as modified.
+	UEdGraphNode* PinANode = PinA->GetOwningNode();
+	if (UMetasoundEditorGraph* Graph = CastChecked<UMetasoundEditorGraph>(PinANode->GetGraph()))
 	{
-		bModified = Metasound::Editor::FGraphBuilder::ConnectNodes(*InputPin, *OutputPin, false);
+		Graph->GetMetasoundChecked().Modify();
 	}
 
-	if (bModified)
+	// This call to parent takes care of marking respective nodes for modification.
+	if (!UEdGraphSchema::TryCreateConnection(PinA, PinB))
 	{
-		InputPin->GetOwningNode()->MarkPackageDirty();
+		return false;
 	}
-	return bModified;
+
+	if (!Metasound::Editor::FGraphBuilder::ConnectNodes(*InputPin, *OutputPin, false /* bConnectEdPins */))
+	{
+		return false;
+	}
+
+	return true;
 }
 
 bool UMetasoundEditorGraphSchema::ShouldHidePinDefaultValue(UEdGraphPin* Pin) const
@@ -659,14 +666,14 @@ void UMetasoundEditorGraphSchema::BreakNodeLinks(UEdGraphNode& TargetNode) const
 {
 	using namespace Metasound::Frontend;
 
-	const FScopedTransaction Transaction(NSLOCTEXT("UnrealEd", "GraphEd_BreakNodeLinks", "Break Node Links"));
+	const FScopedTransaction Transaction(LOCTEXT("BreakNodeLinks", "Break Node Links"));
+	TargetNode.Modify();
 
 	TArray<UEdGraphPin*> Pins = TargetNode.GetAllPins();
 	for (UEdGraphPin* Pin : Pins)
 	{
 		BreakPinLinks(*Pin, false);
 	}
-
 	Super::BreakNodeLinks(TargetNode);
 }
 
@@ -675,7 +682,8 @@ void UMetasoundEditorGraphSchema::BreakPinLinks(UEdGraphPin& TargetPin, bool bSe
 	using namespace Metasound::Editor;
 	using namespace Metasound::Frontend;
 
-	const FScopedTransaction Transaction(NSLOCTEXT("UnrealEd", "GraphEd_BreakPinLinks", "Break Pin Links"));
+	const FScopedTransaction Transaction(LOCTEXT("BreakPinLinks", "Break Pin Links"));
+	TargetPin.GetOwningNode()->Modify();
 
 	FGraphBuilder::DisconnectPin(TargetPin);
 
diff --git a/Engine/Plugins/Runtime/Metasound/Source/MetasoundEditor/Private/MetasoundNodeDetailCustomization.cpp b/Engine/Plugins/Runtime/Metasound/Source/MetasoundEditor/Private/MetasoundNodeDetailCustomization.cpp
--- a/Engine/Plugins/Runtime/Metasound/Source/MetasoundEditor/Private/MetasoundNodeDetailCustomization.cpp
+++ b/Engine/Plugins/Runtime/Metasound/Source/MetasoundEditor/Private/MetasoundNodeDetailCustomization.cpp
@@ -210,13 +210,13 @@ namespace Metasound
 						if (Prop->GetValue(IntValue) != FPropertyAccess::Success)
 						{
 							IntValue = Interface->GetDefaultValue();
-							UE_LOG(LogTemp, Warning, TEXT("Failed to read int Property '%s', defaulting."), *GetNameSafe(Prop->GetProperty()));
+							UE_LOG(LogMetasoundEditor, Warning, TEXT("Failed to read int Property '%s', defaulting."), *GetNameSafe(Prop->GetProperty()));
 						}
 						if (TOptional<IEnumDataTypeInterface::FGenericInt32Entry> Result = Interface->FindByValue(IntValue))
 						{
 							return Result->DisplayName.ToString();
 						}
-						UE_LOG(LogTemp, Warning, TEXT("Failed to resolve int value '%d' to a valid enum value for enum '%s'"),
+						UE_LOG(LogMetasoundEditor, Warning, TEXT("Failed to resolve int value '%d' to a valid enum value for enum '%s'"),
 							IntValue, *Interface->GetNamespace().ToString());
 
 						// Return default (should always succeed as we can't have empty Enums and we must have a default).
@@ -242,7 +242,7 @@ namespace Metasound
 						}
 						else
 						{
-							UE_LOG(LogTemp, Warning, TEXT("Failed to Set Valid Value for Property '%s' with Value of '%s', writing default."),
+							UE_LOG(LogMetasoundEditor, Warning, TEXT("Failed to Set Valid Value for Property '%s' with Value of '%s', writing default."),
 								*GetNameSafe(Prop->GetProperty()), *InSelected);
 
 							ensure(Prop->SetValue(Interface->GetDefaultValue()) == FPropertyAccess::Success);
diff --git a/Engine/Plugins/Runtime/Metasound/Source/MetasoundFrontend/Private/MetasoundFrontendTransform.cpp b/Engine/Plugins/Runtime/Metasound/Source/MetasoundFrontend/Private/MetasoundFrontendTransform.cpp
--- a/Engine/Plugins/Runtime/Metasound/Source/MetasoundFrontend/Private/MetasoundFrontendTransform.cpp
+++ b/Engine/Plugins/Runtime/Metasound/Source/MetasoundFrontend/Private/MetasoundFrontendTransform.cpp
@@ -1,7 +1,9 @@
 // Copyright Epic Games, Inc. All Rights Reserved.
-
 #include "MetasoundFrontendTransform.h"
 
+#include "MetasoundFrontendDocument.h"
+
+
 namespace Metasound
 {
 	namespace Frontend
@@ -22,21 +24,27 @@ namespace Metasound
 			// Go through each input and add or swap if something is missing.
 			for (const FMetasoundFrontendClassVertex& RequiredInput : RequiredInputs)
 			{
-				const FMetasoundFrontendClassInput* ClassInput = Graph->FindClassInputWithName(RequiredInput.Name).Get();
-
-				bool bIsMatching = false;
-				if (nullptr != ClassInput)
-				{
-					bIsMatching = FMetasoundFrontendClassVertex::IsFunctionalEquivalent(RequiredInput, *ClassInput);
-				}
-
-				if (!bIsMatching)
+				if (const FMetasoundFrontendClassInput* ClassInput = Graph->FindClassInputWithName(RequiredInput.Name).Get())
 				{
-					bDidEdit = true;
-					if (nullptr != ClassInput)
+					if (!FMetasoundFrontendClassVertex::IsFunctionalEquivalent(RequiredInput, *ClassInput))
 					{
+						bDidEdit = true;
+
+						// Cache off node locations to push to new node
+						FNodeHandle InputNode = Graph->GetInputNodeWithName(ClassInput->Name);
+						const TMap<FGuid, FVector2D> Locations = InputNode->GetNodeStyle().Display.Locations;
+
 						Graph->RemoveInputVertex(RequiredInput.Name);
+						InputNode = Graph->AddInputVertex(RequiredInput);
+
+						FMetasoundFrontendNodeStyle Style = InputNode->GetNodeStyle();
+						Style.Display.Locations = Locations;
+						InputNode->SetNodeStyle(Style);
 					}
+				}
+				else
+				{
+					bDidEdit = true;
 					Graph->AddInputVertex(RequiredInput);
 				}
 			}
@@ -44,26 +52,32 @@ namespace Metasound
 			// Go through each output and add or swap if something is missing.
 			for (const FMetasoundFrontendClassVertex& RequiredOutput : RequiredOutputs)
 			{
-				const FMetasoundFrontendClassOutput* ClassOutput = Graph->FindClassOutputWithName(RequiredOutput.Name).Get();
-
-				bool bIsMatching = false;
-				if (nullptr != ClassOutput)
-				{
-					bIsMatching = FMetasoundFrontendClassVertex::IsFunctionalEquivalent(RequiredOutput, *ClassOutput);
-				}
-
-				if (!bIsMatching)
+				if (const FMetasoundFrontendClassOutput* ClassOutput = Graph->FindClassOutputWithName(RequiredOutput.Name).Get())
 				{
-					bDidEdit = true;
-					if (nullptr != ClassOutput)
+					if (!FMetasoundFrontendClassVertex::IsFunctionalEquivalent(RequiredOutput, *ClassOutput))
 					{
+						bDidEdit = true;
+
+						// Cache off node locations to push to new node
+						FNodeHandle OutputNode = Graph->GetOutputNodeWithName(ClassOutput->Name);
+						const TMap<FGuid, FVector2D> Locations = OutputNode->GetNodeStyle().Display.Locations;
+
 						Graph->RemoveOutputVertex(RequiredOutput.Name);
+						OutputNode = Graph->AddOutputVertex(RequiredOutput);
+
+						FMetasoundFrontendNodeStyle Style = OutputNode->GetNodeStyle();
+						Style.Display.Locations = Locations;
+						OutputNode->SetNodeStyle(Style);
 					}
+				}
+				else
+				{
+					bDidEdit = true;
 					Graph->AddOutputVertex(RequiredOutput);
 				}
 			}
 
 			return bDidEdit;
 		}
-	}
-}
+	} // namespace Frontend
+} // namespace Metasound
