Fix Metasound node pin alignment to be top aligned/justified.

diff --git a/Engine/Plugins/Runtime/Metasound/Source/MetasoundEditor/Private/SMetasoundGraphNode.cpp b/Engine/Plugins/Runtime/Metasound/Source/MetasoundEditor/Private/SMetasoundGraphNode.cpp
--- a/Engine/Plugins/Runtime/Metasound/Source/MetasoundEditor/Private/SMetasoundGraphNode.cpp
+++ b/Engine/Plugins/Runtime/Metasound/Source/MetasoundEditor/Private/SMetasoundGraphNode.cpp
@@ -301,6 +301,7 @@ TSharedRef<SWidget> SMetasoundGraphNode::CreateNodeContentArea()
 	{
 		ContentBox->AddSlot()
 			.HAlign(HAlign_Left)
+			.VAlign(VAlign_Top)
 			.FillWidth(1.0f)
 			[
 				SAssignNew(LeftNodeBox, SVerticalBox)
@@ -310,7 +311,7 @@ TSharedRef<SWidget> SMetasoundGraphNode::CreateNodeContentArea()
 	{
 		ContentBox->AddSlot()
 			.HAlign(HAlign_Left)
-			.VAlign(VAlign_Center)
+			.VAlign(VAlign_Top)
 			.AutoWidth()
 			[
 				SAssignNew(LeftNodeBox, SVerticalBox)
@@ -335,7 +336,7 @@ TSharedRef<SWidget> SMetasoundGraphNode::CreateNodeContentArea()
 	ContentBox->AddSlot()
 		.AutoWidth()
 		.HAlign(HAlign_Right)
-		.VAlign(VAlign_Center)
+		.VAlign(VAlign_Top)
 		[
 			SAssignNew(RightNodeBox, SVerticalBox)
 		];
