IWYU Fixes

diff --git a/Engine/Plugins/Runtime/Metasound/Source/MetasoundStandardNodes/Private/MetasoundWaveWriterNode.cpp b/Engine/Plugins/Runtime/Metasound/Source/MetasoundStandardNodes/Private/MetasoundWaveWriterNode.cpp
--- a/Engine/Plugins/Runtime/Metasound/Source/MetasoundStandardNodes/Private/MetasoundWaveWriterNode.cpp
+++ b/Engine/Plugins/Runtime/Metasound/Source/MetasoundStandardNodes/Private/MetasoundWaveWriterNode.cpp
@@ -2,12 +2,14 @@
 
 #include "MetasoundWaveWriterNode.h"
 
+#include "HAL/FileManager.h"
 #include "MetasoundNodeRegistrationMacro.h"
 #include "MetasoundPrimitives.h"
 #include "MetasoundStandardNodesNames.h"
 #include "MetasoundStandardNodesCategories.h"
 #include "MetasoundExecutableOperator.h"
 #include "MetasoundBuildError.h"
+#include "Misc/Paths.h"
 #include "Misc/ScopeLock.h"
 
 #define LOCTEXT_NAMESPACE "MetasoundStandardNodes"
