ADSR optimizations

diff --git a/Engine/Plugins/Runtime/Metasound/Source/MetasoundStandardNodes/Private/MetasoundADSREnvelopeNode.cpp b/Engine/Plugins/Runtime/Metasound/Source/MetasoundStandardNodes/Private/MetasoundADSREnvelopeNode.cpp
--- a/Engine/Plugins/Runtime/Metasound/Source/MetasoundStandardNodes/Private/MetasoundADSREnvelopeNode.cpp
+++ b/Engine/Plugins/Runtime/Metasound/Source/MetasoundStandardNodes/Private/MetasoundADSREnvelopeNode.cpp
@@ -34,6 +34,8 @@ namespace Metasound
 		METASOUND_PARAM(InputReleaseCurve, "Release Curve", "The exponential curve factor of the release. 1.0 = linear release, < 1.0 exponential release, > 1.0 logorithmic release.");
 
 		METASOUND_PARAM(OutputOnAttackTrigger, "On Attack Triggered", "Triggers when the envelope attack is triggered.");
+		METASOUND_PARAM(OutputOnDecayTrigger, "On Decay Triggered", "Triggers when the envelope decay begins and attack is finished.");
+		METASOUND_PARAM(OutputOnSustainTrigger, "On Sustain Triggered", "Triggers when the envelope sustain begins and attack is finished.");
 		METASOUND_PARAM(OutputOnReleaseTrigger, "On Release Triggered", "Triggers when the envelope attack is triggered.");
 		METASOUND_PARAM(OutputOnDone, "On Done", "Triggers when the envelope finishes.");
 		METASOUND_PARAM(OutputEnvelopeValue, "Out Envelope", "The output value of the envelope.");
@@ -41,48 +43,6 @@ namespace Metasound
 
 	namespace ADSREnvelopeNodePrivate
 	{
-// 		enum class EEnvState : uint8
-// 		{
-// 			Attack,
-// 			Decay,
-// 			Sustain,
-// 			Release
-// 		};
-// 
-// 		class FEnvSegment
-// 		{		
-// 		public:
-// 			FEnvSegment(EEnvState InState, int32 InSampleCount, float InCurveValue, float InTargetValue)
-// 				: State(InState)
-// 				, SampleCount(InSampleCount)
-// 				, CurveValue(InCurveValue)
-// 				, TargetValue(InTargetValue)
-// 			{}
-// 
-// 			void Update(float InSampleCount, float InTargetValue)
-// 			{
-// 				SampleCount = InSampleCount;
-// 				TargetValue = InTargetValue;
-// 			}
-// 
-// 			void Start()
-// 			{
-// 				CurrentSampleCount = 0;
-// 			}
-// 
-// 			// return true if done in this state
-// 			bool GetNextValue(float& OutValue)
-// 			{
-// 				if ()
-// 			}
-// 
-// 			EEnvState State;
-// 			int32 SampleCount = 0;
-// 			float TargetValue = 0.0f;
-// 			int32 CurrentSampleCount = 0;
-// 
-// 		};
-
 		struct FEnvState
 		{
 			// Where the envelope is. If INDEX_NONE, then the envelope is not triggered
@@ -116,15 +76,15 @@ namespace Metasound
 			bool bIsInRelease = false;
 		};
 
-		template<typename ValueType>
-		struct TADSREnvelope
-		{
-		};
 
-		template<>
-		struct TADSREnvelope<float>
+		struct FFloatADSREnvelope
 		{
-			static void GetNextEnvelopeOutput(FEnvState& InState, int32 StartFrame, int32 EndFrame, TArray<int32>& OutFinishedFrames, float& OutEnvelopeValue)
+			static float GetSampleRate(const FOperatorSettings& InOperatorSettings)
+			{		
+				return InOperatorSettings.GetActualBlockRate();
+			}
+
+			static void GetNextEnvelopeOutput(FEnvState& InState, int32 StartFrame, int32 EndFrame, TArray<int32>& OutOnDecayFrames, TArray<int32>& OutOnSustainFrames, TArray<int32>& OutOnDoneFrames, float& OutEnvelopeValue)
 			{
 				// Don't need to do anything if we're not generating the envelope at the top of the block since this is a block-rate envelope
 				if (StartFrame > 0 || InState.CurrentSampleIndex == INDEX_NONE)
@@ -133,9 +93,6 @@ namespace Metasound
 					return;
 				}
 
-				// Sample count to the end of the decay phase
-				int32 DecayEnvSampleCount = (InState.AttackSampleCount + InState.DecaySampleCount);
-
 				// If we are in the release state, jump forward in our sample count
 				if (InState.bIsInRelease)
 				{
@@ -156,56 +113,73 @@ namespace Metasound
 					InState.EnvEase.SetValue(TargeEnvelopeValue);
 					InState.CurrentEnvelopeValue = InState.EnvEase.GetNextValue();
 					OutEnvelopeValue = InState.CurrentEnvelopeValue;
+
+					if (InState.CurrentSampleIndex == InState.AttackSampleCount)
+					{
+						OutOnDecayFrames.Add(0);
+					}
 				}
 				// We are in decay
-				else if (InState.CurrentSampleIndex < DecayEnvSampleCount)
-				{
-					int32 SampleCountInDecayState = InState.CurrentSampleIndex++ - InState.AttackSampleCount;
-					float DecayFracton = (float)SampleCountInDecayState / InState.DecaySampleCount;
-					float TargetEnvelopeValue = 1.0f - (1.0f - InState.SustainLevel) * FMath::Pow(DecayFracton, InState.DecayCurveFactor);
-					InState.EnvEase.SetValue(TargetEnvelopeValue);
-					InState.CurrentEnvelopeValue = InState.EnvEase.GetNextValue();
-					OutEnvelopeValue = InState.CurrentEnvelopeValue;
-				}
-				// We are in sustain
-				else if (!InState.bIsInRelease)
-				{
-					InState.EnvEase.SetValue(InState.SustainLevel);
-					InState.CurrentEnvelopeValue = InState.EnvEase.GetNextValue();
-					OutEnvelopeValue = InState.CurrentEnvelopeValue;
-					InState.EnvelopeValueAtReleaseStart = OutEnvelopeValue;
-				}
-				// We are in release mode or finished
-				else 			
+				else
 				{
-					int32 ReleaseEnvSampleCount = (InState.AttackSampleCount + InState.DecaySampleCount + InState.ReleaseSampleCount);
-					// We are in release
-					if (InState.CurrentSampleIndex < ReleaseEnvSampleCount)
+					// Sample count to the end of the decay phase
+					int32 DecayEnvSampleCount = (InState.AttackSampleCount + InState.DecaySampleCount);
+					if (InState.CurrentSampleIndex < DecayEnvSampleCount)
 					{
-						int32 SampleCountInReleaseState = InState.CurrentSampleIndex++ - InState.DecaySampleCount - InState.AttackSampleCount;
-						float ReleaseFraction = (float)SampleCountInReleaseState / InState.ReleaseSampleCount;
-						float TargetEnvelopeValue = InState.EnvelopeValueAtReleaseStart * (1.0f - FMath::Pow(ReleaseFraction, InState.ReleaseCurveFactor));
+						int32 SampleCountInDecayState = InState.CurrentSampleIndex++ - InState.AttackSampleCount;
+						float DecayFracton = (float)SampleCountInDecayState / InState.DecaySampleCount;
+						float TargetEnvelopeValue = 1.0f - (1.0f - InState.SustainLevel) * FMath::Pow(DecayFracton, InState.DecayCurveFactor);
 						InState.EnvEase.SetValue(TargetEnvelopeValue);
 						InState.CurrentEnvelopeValue = InState.EnvEase.GetNextValue();
 						OutEnvelopeValue = InState.CurrentEnvelopeValue;
+
+						if (InState.CurrentSampleIndex == DecayEnvSampleCount)
+						{
+							OutOnSustainFrames.Add(0);
+						}
+					}
+					// We are in sustain
+					else if (!InState.bIsInRelease)
+					{
+						InState.EnvEase.SetValue(InState.SustainLevel);
+						InState.CurrentEnvelopeValue = InState.EnvEase.GetNextValue();
+						OutEnvelopeValue = InState.CurrentEnvelopeValue;
+						InState.EnvelopeValueAtReleaseStart = OutEnvelopeValue;
 					}
-					// We are done
+					// We are in release mode or finished
 					else
 					{
-						InState.CurrentSampleIndex = INDEX_NONE;
-						OutEnvelopeValue = 0.0f;
-						OutFinishedFrames.Add(0);
+						int32 ReleaseEnvSampleCount = (InState.AttackSampleCount + InState.DecaySampleCount + InState.ReleaseSampleCount);
+						// We are in release
+						if (InState.CurrentSampleIndex < ReleaseEnvSampleCount)
+						{
+							int32 SampleCountInReleaseState = InState.CurrentSampleIndex++ - InState.DecaySampleCount - InState.AttackSampleCount;
+							float ReleaseFraction = (float)SampleCountInReleaseState / InState.ReleaseSampleCount;
+							float TargetEnvelopeValue = InState.EnvelopeValueAtReleaseStart * (1.0f - FMath::Pow(ReleaseFraction, InState.ReleaseCurveFactor));
+							InState.EnvEase.SetValue(TargetEnvelopeValue);
+							InState.CurrentEnvelopeValue = InState.EnvEase.GetNextValue();
+							OutEnvelopeValue = InState.CurrentEnvelopeValue;
+						}
+						// We are done
+						else
+						{
+							InState.CurrentSampleIndex = INDEX_NONE;
+							OutEnvelopeValue = 0.0f;
+							OutOnDoneFrames.Add(0);
+						}
 					}
 				}
 			}
-
-			static bool IsAudio() { return false; }
 		};
 
-		template<>
-		struct TADSREnvelope<FAudioBuffer>
+		struct FAudioADSREnvelope
 		{
-			static void GetNextEnvelopeOutput(FEnvState& InState, int32 StartFrame, int32 EndFrame, TArray<int32>& OutFinishedFrames, FAudioBuffer& OutEnvelopeValue)
+			static float GetSampleRate(const FOperatorSettings& InOperatorSettings)
+			{
+					return InOperatorSettings.GetSampleRate();
+			}
+
+			static void GetNextEnvelopeOutput(FEnvState& InState, int32 StartFrame, int32 EndFrame, TArray<int32>& OutOnDecayFrames, TArray<int32>& OutOnSustainFrames, TArray<int32>& OutOnDoneFrames, FAudioBuffer& OutEnvelopeValue)
 			{
 				// If we are not active zero the buffer and early exit
 				if (InState.CurrentSampleIndex == INDEX_NONE)
@@ -225,11 +199,20 @@ namespace Metasound
 					}
 				}
 
+				// Init the end attack and decay frames to the start frame. If we're not in these states, it'll just start loops at the start frame
+				// But we want to keep track of which frames we have finished previous envelope states so the next state can start rendering at the
+				// correct sample vs assume it's at the start of the block.
+				int32 EndAttackFrame = StartFrame;
+				int32 EndDecayFrame = StartFrame;
+
 				float* OutEnvPtr = OutEnvelopeValue.GetData();
-				for (int32 i = StartFrame; i < EndFrame; ++i)
+				int32 AttackSamplesLeft = InState.AttackSampleCount - InState.CurrentSampleIndex;
+
+				// We are in attack
+				if (AttackSamplesLeft > 0)
 				{
-					// We are in attack
-					if (InState.CurrentSampleIndex <= InState.AttackSampleCount)
+					EndAttackFrame = FMath::Min(StartFrame + AttackSamplesLeft, EndFrame);
+					for (int32 i = StartFrame; i < EndAttackFrame; ++i)
 					{
 						float AttackFraction = (float)InState.CurrentSampleIndex++ / InState.AttackSampleCount;
 						float EnvValue = FMath::Pow(AttackFraction, InState.AttackCurveFactor);
@@ -238,11 +221,25 @@ namespace Metasound
 						InState.CurrentEnvelopeValue = InState.EnvEase.GetNextValue();
 						OutEnvPtr[i] = InState.CurrentEnvelopeValue;
 					}
-					else
+
+					// We have finished our attack phase in this block
+					if (InState.CurrentSampleIndex == InState.AttackSampleCount)
 					{
-						int32 DecayEnvSampleCount = (InState.AttackSampleCount + InState.DecaySampleCount);
-						// We are in decay
-						if (InState.CurrentSampleIndex < DecayEnvSampleCount)
+						OutOnDecayFrames.Add(EndAttackFrame);
+					}
+				}
+
+				// We are now done with our attack and may need to immediately start rendering our decay block
+				if (EndAttackFrame < EndFrame)
+				{
+					int32 DecaySampelsFromStart = InState.AttackSampleCount + InState.DecaySampleCount;
+					int32 DecaySamplesLeft = DecaySampelsFromStart - InState.CurrentSampleIndex;
+					// We are in decay
+					if (DecaySamplesLeft > 0)
+					{
+						EndDecayFrame = FMath::Min(StartFrame + EndAttackFrame + DecaySamplesLeft, EndFrame);
+
+						for (int32 i = EndAttackFrame; i < EndDecayFrame; ++i)
 						{
 							int32 SampleCountInDecayState = InState.CurrentSampleIndex++ - InState.AttackSampleCount;
 							float DecayFracton = (float)SampleCountInDecayState / InState.DecaySampleCount;
@@ -251,20 +248,37 @@ namespace Metasound
 							InState.CurrentEnvelopeValue = InState.EnvEase.GetNextValue();
 							OutEnvPtr[i] = InState.CurrentEnvelopeValue;
 						}
-						// We are in sustain
-						else if (!InState.bIsInRelease)
+
+						// We have finished the decay phase in this block
+						if (InState.CurrentSampleIndex == DecaySampelsFromStart)
 						{
-							InState.EnvEase.SetValue(InState.SustainLevel);
-							InState.CurrentEnvelopeValue = InState.EnvEase.GetNextValue();
-							InState.EnvelopeValueAtReleaseStart = InState.CurrentEnvelopeValue;
-							OutEnvPtr[i] = InState.CurrentEnvelopeValue;
+							OutOnSustainFrames.Add(EndDecayFrame);
+						}
+					}
+
+					// if the end decay frame is not the end of this current render frame, we are in a post-decay mode
+					// could be in sustain or release, or we're done.
+					if (EndDecayFrame < EndFrame)
+					{
+						// If we're not in release mode, we're in sustain mode
+						if (!InState.bIsInRelease)
+						{
+
+							for (int32 i = EndDecayFrame; i < EndFrame; ++i)
+							{
+								InState.EnvEase.SetValue(InState.SustainLevel);
+								InState.CurrentEnvelopeValue = InState.EnvEase.GetNextValue();
+								InState.EnvelopeValueAtReleaseStart = InState.CurrentEnvelopeValue;
+								OutEnvPtr[i] = InState.CurrentEnvelopeValue;
+							}
 						}
-						else 
+						// We're in release mode
+						else
 						{
-							int32 ReleaseEnvSampleCount = (InState.AttackSampleCount + InState.DecaySampleCount + InState.ReleaseSampleCount);
+							int32 ReleaseSamplesLeft = (InState.AttackSampleCount + InState.DecaySampleCount + InState.ReleaseSampleCount) - InState.CurrentSampleIndex;
 
-							// We are in release
-							if (InState.CurrentSampleIndex < ReleaseEnvSampleCount)
+							int32 EndReleaseFrame = FMath::Min(StartFrame + ReleaseSamplesLeft, EndFrame);
+							for (int32 i = EndDecayFrame; i < EndReleaseFrame; ++i)
 							{
 								int32 SampleCountInReleaseState = InState.CurrentSampleIndex++ - InState.DecaySampleCount - InState.AttackSampleCount;
 								float ReleaseFraction = (float)SampleCountInReleaseState / InState.ReleaseSampleCount;
@@ -273,38 +287,40 @@ namespace Metasound
 								InState.CurrentEnvelopeValue = InState.EnvEase.GetNextValue();
 								OutEnvPtr[i] = InState.CurrentEnvelopeValue;
 							}
-							// We're done
-							else
-							{
-								InState.CurrentEnvelopeValue = InState.EnvEase.GetNextValue();
-								OutEnvPtr[i] = InState.CurrentEnvelopeValue;
 
-								// Envelope is done
-								if (InState.EnvEase.IsDone())
+							// We're now done, lets taper it off and finish things
+							if (EndReleaseFrame < EndFrame)
+							{
+								for (int32 i = EndReleaseFrame; i < EndFrame; ++i)
 								{
-									// Zero out the rest of the envelope
-									int32 NumSamplesLeft = EndFrame - i - 1;
-									if (NumSamplesLeft > 0)
+									InState.CurrentEnvelopeValue = InState.EnvEase.GetNextValue();
+									OutEnvPtr[i] = InState.CurrentEnvelopeValue;
+
+									// Envelope is done
+									if (InState.EnvEase.IsDone())
 									{
-										FMemory::Memzero(&OutEnvPtr[i + 1], sizeof(float) * NumSamplesLeft);
+										// Zero out the rest of the envelope
+										int32 NumSamplesLeft = EndFrame - i - 1;
+										if (NumSamplesLeft > 0)
+										{
+											FMemory::Memzero(&OutEnvPtr[i + 1], sizeof(float) * NumSamplesLeft);
+										}
+										InState.CurrentSampleIndex = INDEX_NONE;
+										OutOnDoneFrames.Add(i);
+										break;
 									}
-									InState.CurrentSampleIndex = INDEX_NONE;
-									OutFinishedFrames.Add(i);
-									break;
 								}
 							}
 						}
 					}
 				}
 			}
-
-			static bool IsAudio() { return true; }
 		};
 
 	}
 
-	template<typename ValueType>
-	class TADSREnvelopeNodeOperator : public TExecutableOperator<TADSREnvelopeNodeOperator<ValueType>>
+	template<typename EnvelopeClass, typename ValueType>
+	class TADSREnvelopeNodeOperator : public TExecutableOperator<TADSREnvelopeNodeOperator<EnvelopeClass, ValueType>>
 	{
 	public:
 		static const FVertexInterface& GetDefaultInterface()
@@ -325,6 +341,8 @@ namespace Metasound
 				),
 				FOutputVertexInterface(
 					TOutputDataVertexModel<FTrigger>(METASOUND_GET_PARAM_NAME_AND_TT(OutputOnAttackTrigger)),
+					TOutputDataVertexModel<FTrigger>(METASOUND_GET_PARAM_NAME_AND_TT(OutputOnDecayTrigger)),
+					TOutputDataVertexModel<FTrigger>(METASOUND_GET_PARAM_NAME_AND_TT(OutputOnSustainTrigger)),
 					TOutputDataVertexModel<FTrigger>(METASOUND_GET_PARAM_NAME_AND_TT(OutputOnReleaseTrigger)),
 					TOutputDataVertexModel<FTrigger>(METASOUND_GET_PARAM_NAME_AND_TT(OutputOnDone)),
 					TOutputDataVertexModel<ValueType>(METASOUND_GET_PARAM_NAME_AND_TT(OutputEnvelopeValue))
@@ -366,63 +384,64 @@ namespace Metasound
 			return Metadata;
 		}
 
+		struct FOperatorArgs
+		{
+			FOperatorSettings OperatorSettings;
+			FTriggerReadRef TriggerAttackIn;
+			FTriggerReadRef TriggerReleaseIn;
+			FTimeReadRef AttackTime;
+			FTimeReadRef DecayTime;
+			FFloatReadRef SustainLevel;
+			FTimeReadRef ReleaseTime;
+			FFloatReadRef AttackCurveFactor;
+			FFloatReadRef DecayCurveFactor;
+			FFloatReadRef ReleaseCurveFactor;
+		};
+
 		static TUniquePtr<IOperator> CreateOperator(const FCreateOperatorParams& InParams, TArray<TUniquePtr<IOperatorBuildError>>& OutErrors)
 		{
 			using namespace ADSREnvelopeVertexNames;
 
 			const FInputVertexInterface& InputInterface = GetDefaultInterface().GetInputInterface();
 
-			FTriggerReadRef TriggerAttackIn = InParams.InputDataReferences.GetDataReadReferenceOrConstruct<FTrigger>(METASOUND_GET_PARAM_NAME(InputAttackTrigger), InParams.OperatorSettings);
-			FTriggerReadRef TriggerReleaseIn = InParams.InputDataReferences.GetDataReadReferenceOrConstruct<FTrigger>(METASOUND_GET_PARAM_NAME(InputReleaseTrigger), InParams.OperatorSettings);
-
-			FTimeReadRef AttackTime = InParams.InputDataReferences.GetDataReadReferenceOrConstructWithVertexDefault<FTime>(InputInterface, METASOUND_GET_PARAM_NAME(InputAttackTime), InParams.OperatorSettings);
-			FTimeReadRef DecayTime = InParams.InputDataReferences.GetDataReadReferenceOrConstructWithVertexDefault<FTime>(InputInterface, METASOUND_GET_PARAM_NAME(InputDecayTime), InParams.OperatorSettings);
-			FFloatReadRef SustainLevel = InParams.InputDataReferences.GetDataReadReferenceOrConstructWithVertexDefault<float>(InputInterface, METASOUND_GET_PARAM_NAME(InputSustainLevel), InParams.OperatorSettings);
-			FTimeReadRef ReleaseTime = InParams.InputDataReferences.GetDataReadReferenceOrConstructWithVertexDefault<FTime>(InputInterface, METASOUND_GET_PARAM_NAME(InputReleaseTime), InParams.OperatorSettings);
-
-			FFloatReadRef AttackCurveFactor = InParams.InputDataReferences.GetDataReadReferenceOrConstructWithVertexDefault<float>(InputInterface, METASOUND_GET_PARAM_NAME(InputAttackCurve), InParams.OperatorSettings);
-			FFloatReadRef DecayCurveFactor = InParams.InputDataReferences.GetDataReadReferenceOrConstructWithVertexDefault<float>(InputInterface, METASOUND_GET_PARAM_NAME(InputDecayCurve), InParams.OperatorSettings);
-			FFloatReadRef ReleaseCurveFactor = InParams.InputDataReferences.GetDataReadReferenceOrConstructWithVertexDefault<float>(InputInterface, METASOUND_GET_PARAM_NAME(InputReleaseCurve), InParams.OperatorSettings);
+			FOperatorArgs Args 
+			{
+				InParams.OperatorSettings,
+				InParams.InputDataReferences.GetDataReadReferenceOrConstruct<FTrigger>(METASOUND_GET_PARAM_NAME(InputAttackTrigger), InParams.OperatorSettings),
+				InParams.InputDataReferences.GetDataReadReferenceOrConstruct<FTrigger>(METASOUND_GET_PARAM_NAME(InputReleaseTrigger), InParams.OperatorSettings),
+				InParams.InputDataReferences.GetDataReadReferenceOrConstructWithVertexDefault<FTime>(InputInterface, METASOUND_GET_PARAM_NAME(InputAttackTime), InParams.OperatorSettings),
+				InParams.InputDataReferences.GetDataReadReferenceOrConstructWithVertexDefault<FTime>(InputInterface, METASOUND_GET_PARAM_NAME(InputDecayTime), InParams.OperatorSettings),
+				InParams.InputDataReferences.GetDataReadReferenceOrConstructWithVertexDefault<float>(InputInterface, METASOUND_GET_PARAM_NAME(InputSustainLevel), InParams.OperatorSettings),
+				InParams.InputDataReferences.GetDataReadReferenceOrConstructWithVertexDefault<FTime>(InputInterface, METASOUND_GET_PARAM_NAME(InputReleaseTime), InParams.OperatorSettings),
+				InParams.InputDataReferences.GetDataReadReferenceOrConstructWithVertexDefault<float>(InputInterface, METASOUND_GET_PARAM_NAME(InputAttackCurve), InParams.OperatorSettings),
+				InParams.InputDataReferences.GetDataReadReferenceOrConstructWithVertexDefault<float>(InputInterface, METASOUND_GET_PARAM_NAME(InputDecayCurve), InParams.OperatorSettings),
+				InParams.InputDataReferences.GetDataReadReferenceOrConstructWithVertexDefault<float>(InputInterface, METASOUND_GET_PARAM_NAME(InputReleaseCurve), InParams.OperatorSettings)
+			};
 
-			return MakeUnique<TADSREnvelopeNodeOperator<ValueType>>(InParams.OperatorSettings, TriggerAttackIn, TriggerReleaseIn, AttackTime, DecayTime, SustainLevel, ReleaseTime, AttackCurveFactor, DecayCurveFactor, ReleaseCurveFactor);
+			return MakeUnique<TADSREnvelopeNodeOperator<EnvelopeClass, ValueType>>(Args);
 		}
 
-		TADSREnvelopeNodeOperator(const FOperatorSettings& InSettings,
-			const FTriggerReadRef& InTriggerAttackIn,
-			const FTriggerReadRef& InTriggerReleaseIn,
-			const FTimeReadRef& InAttackTime,
-			const FTimeReadRef& InDecayTime,
-			const FFloatReadRef& InSustainLevel,
-			const FTimeReadRef& InReleaseTime,
-			const FFloatReadRef& InAttackCurveFactor,
-			const FFloatReadRef& InDecayCurveFactor,
-			const FFloatReadRef& InReleaseCurveFactor)
-			: TriggerAttackIn(InTriggerAttackIn)
-			, TriggerReleaseIn(InTriggerReleaseIn)
-			, AttackTime(InAttackTime)
-			, DecayTime(InDecayTime)
-			, SustainLevel(InSustainLevel)
-			, ReleaseTime(InReleaseTime)
-			, AttackCurveFactor(InAttackCurveFactor)
-			, DecayCurveFactor(InDecayCurveFactor)
-			, ReleaseCurveFactor(InReleaseCurveFactor)
-			, OnAttackTrigger(TDataWriteReferenceFactory<FTrigger>::CreateAny(InSettings))
-			, OnReleaseTrigger(TDataWriteReferenceFactory<FTrigger>::CreateAny(InSettings))
-			, OnDone(TDataWriteReferenceFactory<FTrigger>::CreateAny(InSettings))
-			, OutputEnvelope(TDataWriteReferenceFactory<ValueType>::CreateAny(InSettings))
+		TADSREnvelopeNodeOperator(const FOperatorArgs& InArgs)
+			: TriggerAttackIn(InArgs.TriggerAttackIn)
+			, TriggerReleaseIn(InArgs.TriggerReleaseIn)
+			, AttackTime(InArgs.AttackTime)
+			, DecayTime(InArgs.DecayTime)
+			, SustainLevel(InArgs.SustainLevel)
+			, ReleaseTime(InArgs.ReleaseTime)
+			, AttackCurveFactor(InArgs.AttackCurveFactor)
+			, DecayCurveFactor(InArgs.DecayCurveFactor)
+			, ReleaseCurveFactor(InArgs.ReleaseCurveFactor)
+			, OnAttackTrigger(TDataWriteReferenceFactory<FTrigger>::CreateAny(InArgs.OperatorSettings))
+			, OnDecayTrigger(TDataWriteReferenceFactory<FTrigger>::CreateAny(InArgs.OperatorSettings))
+			, OnSustainTrigger(TDataWriteReferenceFactory<FTrigger>::CreateAny(InArgs.OperatorSettings))
+			, OnReleaseTrigger(TDataWriteReferenceFactory<FTrigger>::CreateAny(InArgs.OperatorSettings))
+			, OnDone(TDataWriteReferenceFactory<FTrigger>::CreateAny(InArgs.OperatorSettings))
+			, OutputEnvelope(TDataWriteReferenceFactory<ValueType>::CreateAny(InArgs.OperatorSettings))
 		{
-			NumFramesPerBlock = InSettings.GetNumFramesPerBlock();
-
-			EnvState.EnvEase.SetEaseFactor(0.01f);
+			NumFramesPerBlock = InArgs.OperatorSettings.GetNumFramesPerBlock();
 
-			if (ADSREnvelopeNodePrivate::TADSREnvelope<ValueType>::IsAudio())
-			{
-				SampleRate = InSettings.GetSampleRate();
-			}
-			else
-			{
-				SampleRate = InSettings.GetActualBlockRate();
-			}
+			EnvState.EnvEase.SetEaseFactor(0.1f);
+			SampleRate = EnvelopeClass::GetSampleRate(InArgs.OperatorSettings);
 		}
 
 		virtual ~TADSREnvelopeNodeOperator() = default;
@@ -451,6 +470,8 @@ namespace Metasound
 
 			FDataReferenceCollection Outputs;
 			Outputs.AddDataReadReference(METASOUND_GET_PARAM_NAME(OutputOnAttackTrigger), OnAttackTrigger);
+			Outputs.AddDataReadReference(METASOUND_GET_PARAM_NAME(OutputOnDecayTrigger), OnDecayTrigger);
+			Outputs.AddDataReadReference(METASOUND_GET_PARAM_NAME(OutputOnSustainTrigger), OnSustainTrigger);
 			Outputs.AddDataReadReference(METASOUND_GET_PARAM_NAME(OutputOnReleaseTrigger), OnReleaseTrigger);
 			Outputs.AddDataReadReference(METASOUND_GET_PARAM_NAME(OutputOnDone), OnDone);
 			Outputs.AddDataReadReference(METASOUND_GET_PARAM_NAME(OutputEnvelopeValue), OutputEnvelope);
@@ -472,12 +493,36 @@ namespace Metasound
 			EnvState.ReleaseCurveFactor = FMath::Max(KINDA_SMALL_NUMBER, *ReleaseCurveFactor);
 		}
 
+		void ProcessEnvelopeOutput(int32 InStartFrame, int32 InEndFrame)
+		{
+			TArray<int32> OnDecayFrames;
+			TArray<int32> OnSustainFrames;
+			TArray<int32> OnDoneFrames;
+			EnvelopeClass::GetNextEnvelopeOutput(EnvState, InStartFrame, InEndFrame, OnDecayFrames, OnSustainFrames, OnDoneFrames, *OutputEnvelope);
+
+			for (int32 OnDecayFrame : OnDecayFrames)
+			{
+				OnDecayTrigger->TriggerFrame(OnDecayFrame);
+			}
+
+			for (int32 OnSustainFrame : OnSustainFrames)
+			{
+				OnSustainTrigger->TriggerFrame(OnSustainFrame);
+			}
+
+			for (int32 OnDoneFrame : OnDoneFrames)
+			{
+				OnDone->TriggerFrame(OnDoneFrame);
+			}
+		}
 
 		void Execute()
 		{
 			using namespace ADSREnvelopeNodePrivate;
 
 			OnAttackTrigger->AdvanceBlock();
+			OnDecayTrigger->AdvanceBlock();
+			OnSustainTrigger->AdvanceBlock();
 			OnReleaseTrigger->AdvanceBlock();
 			OnDone->AdvanceBlock();
  
@@ -501,13 +546,7 @@ namespace Metasound
 				// OnPreTrigger
 				[&](int32 StartFrame, int32 EndFrame)
 				{
-					TArray<int32> FinishedFrames;
-					TADSREnvelope<ValueType>::GetNextEnvelopeOutput(EnvState, StartFrame, EndFrame, FinishedFrames, *OutputEnvelope);
-
-					for (int32 FrameFinished : FinishedFrames)
-					{
-						OnDone->TriggerFrame(FrameFinished);
-					}
+					ProcessEnvelopeOutput(StartFrame, EndFrame);
 				},
 				// OnTrigger
 				[&](int32 StartFrame, int32 EndFrame)
@@ -519,13 +558,8 @@ namespace Metasound
 					EnvState.CurrentSampleIndex = 0;
 					EnvState.StartingEnvelopeValue = EnvState.CurrentEnvelopeValue;
 					EnvState.bIsInRelease = false;
-					// Generate the output (this will no-op if we're block rate)
-					TArray<int32> FinishedFrames;
-					TADSREnvelope<ValueType>::GetNextEnvelopeOutput(EnvState, StartFrame, EndFrame, FinishedFrames, *OutputEnvelope);
-					for (int32 FrameFinished : FinishedFrames)
-					{
-						OnDone->TriggerFrame(FrameFinished);
-					}
+
+					ProcessEnvelopeOutput(StartFrame, EndFrame);
 
 					// Forward the trigger
 					OnAttackTrigger->TriggerFrame(StartFrame);
@@ -545,6 +579,8 @@ namespace Metasound
 		FFloatReadRef ReleaseCurveFactor;
 
 		FTriggerWriteRef OnAttackTrigger;
+		FTriggerWriteRef OnDecayTrigger;
+		FTriggerWriteRef OnSustainTrigger;
 		FTriggerWriteRef OnReleaseTrigger;
 		FTriggerWriteRef OnDone;
 		TDataWriteReference<ValueType> OutputEnvelope;
@@ -560,7 +596,7 @@ namespace Metasound
 	 *
 	 *  Creates an Attack/Decay envelope node.
 	 */
-	template<typename ValueType>
+	template<typename EnvelopeClass, typename ValueType>
 	class METASOUNDSTANDARDNODES_API TADSREnvelopeNode : public FNodeFacade
 	{
 	public:
@@ -568,16 +604,16 @@ namespace Metasound
 		 * Constructor used by the Metasound Frontend.
 		 */
 		TADSREnvelopeNode(const FNodeInitData& InInitData)
-			: FNodeFacade(InInitData.InstanceName, InInitData.InstanceID, TFacadeOperatorClass<TADSREnvelopeNodeOperator<ValueType>>())
+			: FNodeFacade(InInitData.InstanceName, InInitData.InstanceID, TFacadeOperatorClass<TADSREnvelopeNodeOperator<EnvelopeClass, ValueType>>())
 		{}
 
 		virtual ~TADSREnvelopeNode() = default;
 	};
 
-	using FADSREnvelopeNodeFloat = TADSREnvelopeNode<float>;
+	using FADSREnvelopeNodeFloat = TADSREnvelopeNode<ADSREnvelopeNodePrivate::FFloatADSREnvelope, float>;
 	METASOUND_REGISTER_NODE(FADSREnvelopeNodeFloat)
 
-	using FADSREnvelopeAudioBuffer = TADSREnvelopeNode<FAudioBuffer>;
+	using FADSREnvelopeAudioBuffer = TADSREnvelopeNode<ADSREnvelopeNodePrivate::FAudioADSREnvelope, FAudioBuffer>;
 	METASOUND_REGISTER_NODE(FADSREnvelopeAudioBuffer)
 }
 
