Wave player node work

diff --git a/Engine/Plugins/Runtime/Metasound/Source/MetasoundEngine/Private/MetasoundWave.cpp b/Engine/Plugins/Runtime/Metasound/Source/MetasoundEngine/Private/MetasoundWave.cpp
--- a/Engine/Plugins/Runtime/Metasound/Source/MetasoundEngine/Private/MetasoundWave.cpp
+++ b/Engine/Plugins/Runtime/Metasound/Source/MetasoundEngine/Private/MetasoundWave.cpp
@@ -37,6 +37,7 @@ namespace Audio
 		StartTimeSeconds = InInitParams.StartTimeSeconds;
 
 		NumChannels = InWave->GetNumChannels();
+		NumFrames = InWave->GetNumFrames();
 		DecodeBlockSizeInFrames = 64;
 		DecodeBlockSizeInSamples = DecodeBlockSizeInFrames * NumChannels;
 
@@ -57,11 +58,12 @@ namespace Audio
 		return bSuccessful;
 	}
 
-
-	uint32 FSimpleDecoderWrapper::GenerateAudio(float* OutputDest, int32 NumOutputFrames, float PitchShiftInCents, bool bIsLooping)
+	uint32 FSimpleDecoderWrapper::GenerateAudio(float* OutputDest, int32 NumOutputFrames, int32& OutNumFramesConsumed, float PitchShiftInCents, bool bIsLooping)
 	{
 		const uint32 NumOutputSamples = NumOutputFrames * NumChannels;
 
+		OutNumFramesConsumed = 0;
+
 		if (OutputCircularBuffer.Num() < NumOutputSamples)
 		{
 
@@ -86,6 +88,7 @@ namespace Audio
 				Audio::IDecoderOutput::FPushedAudioDetails Details;
 				const Audio::IDecoder::EDecodeResult  DecodeResult = Decoder->Decode(bIsLooping);
 				const int32 NumFramesDecoded = Output->PopAudio(PreSrcBuffer, Details) / NumChannels;
+				OutNumFramesConsumed += NumFramesDecoded;
 				int32 NumResamplerOutputFrames = 0;
 				int32 Error = Resampler.ProcessAudio(PreSrcBuffer.GetData(), NumFramesDecoded, bDecoderIsDone, PostSrcBuffer.GetData(), MaxNumResamplerOutputFramesPerBlock, NumResamplerOutputFrames);
 				ensure(Error == 0);
@@ -93,6 +96,12 @@ namespace Audio
 				bDecoderIsDone = DecodeResult == Audio::IDecoder::EDecodeResult::Finished;
 				bDecoderHasLooped = DecodeResult == Audio::IDecoder::EDecodeResult::Looped;
 
+				if (!PostSrcBuffer.Num() || !NumResamplerOutputFrames)
+				{
+					continue;
+				}
+				int32 OutputFrames = FMath::Min((int32)PostSrcBuffer.Num(), (int32)(NumResamplerOutputFrames * NumChannels));
+
 				// perform linear pitch shift into OutputCircularBuffer
 				if (NumResamplerOutputFrames > 0)
 				{
diff --git a/Engine/Plugins/Runtime/Metasound/Source/MetasoundEngine/Private/MetasoundWavePlayerNode.cpp b/Engine/Plugins/Runtime/Metasound/Source/MetasoundEngine/Private/MetasoundWavePlayerNode.cpp
--- a/Engine/Plugins/Runtime/Metasound/Source/MetasoundEngine/Private/MetasoundWavePlayerNode.cpp
+++ b/Engine/Plugins/Runtime/Metasound/Source/MetasoundEngine/Private/MetasoundWavePlayerNode.cpp
@@ -14,9 +14,56 @@
 
 #define LOCTEXT_NAMESPACE "MetasoundWaveNode"
 
-
 namespace Metasound
 {
+	namespace WavePlayerVertexNames
+	{
+		static const TCHAR* InputTriggerPlayName = TEXT("Play");
+		static const TCHAR* InputTriggerStopName = TEXT("Stop");
+		static const TCHAR* InputWaveAssetName = TEXT("Wave Asset");
+		static const TCHAR* InputStartTimeName = TEXT("Start Time");
+		static const TCHAR* InputPitchShiftName = TEXT("Pitch Shift");
+		static const TCHAR* InputLoopName = TEXT("Loop");
+		static const TCHAR* InputLoopStartName = TEXT("Loop Start");
+		static const TCHAR* InputLoopDurationName = TEXT("Loop Duration");
+		//static const TCHAR* InputStartTimeAsPercentName = TEXT("Start/Loop Times As Percent");
+
+		static const TCHAR* OutputTriggerOnPlayName = TEXT("On Play");
+		static const TCHAR* OutputTriggerOnDoneName = TEXT("On Done");
+		static const TCHAR* OutputTriggerOnNearlyDoneName = TEXT("On Nearly Done");
+		static const TCHAR* OutputTriggerOnLoopedName = TEXT("On Looped");
+		static const TCHAR* OutputTriggerOnCuePointName = TEXT("On Cue Point");
+		static const TCHAR* OutputCuePointIDName = TEXT("Cue Point ID");
+		static const TCHAR* OutputCuePointLabelName = TEXT("Cue Point Label");
+		static const TCHAR* OutputLoopPercentName = TEXT("Loop Percent");
+		static const TCHAR* OutputPlaybackLocationName = TEXT("Playback Location");
+		static const TCHAR* OutputAudioLeftName = TEXT("Out Left");
+		static const TCHAR* OutputAudioRightName = TEXT("Out Right");
+
+		static FText InputTriggerPlayTT = LOCTEXT("PlayTT", "Play the wave player.");
+		static FText InputTriggerStopTT = LOCTEXT("StopTT", "Stop the wave player.");
+		static FText InputWaveAssetTT = LOCTEXT("WaveTT", "The wave asset to be real-time decoded.");
+		static FText InputStartTimeTT = LOCTEXT("StartTimeTT", "Time into the wave asset to start (seek) the wave asset. For real-time decoding, the wave asset must be set to seekable!)");
+		static FText InputPitchShiftTT = LOCTEXT("PitchShiftTT", "The pitch shift to use for the wave asset in semitones.");
+		static FText InputLoopTT = LOCTEXT("LoopTT", "Whether or not to loop between the start and specified end times.");
+		static FText InputLoopStartTT = LOCTEXT("LoopStartTT", "When to start the loop. Will be a percentage if \"Start Time As Percent\" is true.");
+		static FText InputLoopDurationTT = LOCTEXT("LoopDurationTT", "The duration of the loop when wave player is enabled for looping. A value of -1.0 will loop the whole wave asset.");
+		//static FText InputStartTimeAsPercentTT = LOCTEXT("StartTimeAsPercentTT", "Whether to treat the start time and loop start times as a percentage of total wave duration vs seconds (default).");
+
+		static FText OutputTriggerOnPlayTT = LOCTEXT("OnPlayTT", "Triggers when Play is triggered.");
+		static FText OutputTriggerOnDoneTT = LOCTEXT("OnDoneTT", "Triggers when the wave played has finished playing.");
+		static FText OutputTriggerOnNearlyDoneTT = LOCTEXT("OnNearlyDoneTT", "Triggers when the wave played has almost finished playing (the block before it finishes). Allows time for logic to trigger different variations to play seamlessly.");
+		static FText OutputTriggerOnLoopedTT = LOCTEXT("OnLoopedTT", "Triggers when the wave player has looped.");
+		static FText OutputTriggerOnCuePointTT = LOCTEXT("OnCuePointTT", "Triggers when a wave cue point was hit during playback.");
+		static FText OutputCuePointIDTT = LOCTEXT("CuePointIDTT", "The cue point ID that was triggered.");
+		static FText OutputCuePointLabelTT = LOCTEXT("CuePointLabelTT", "The cue point label that was triggered (if there was a label parsed in the imported .wav file).");
+		static FText OutputLoopPercentTT = LOCTEXT("LoopPercentTT", "Returns the current loop percent if looping is enabled.");
+		static FText OutputPlaybackLocationTT = LOCTEXT("PlaybackLocationTT", "Returns the absolute position of the wave playback as a precentage of wave duration.");
+		static FText OutputAudioLeftNameTT = LOCTEXT("AudioLeftTT", "The left channel audio output. Mono wave assets will be upmixed to dual stereo.");
+		static FText OutputAudioRightNameTT = LOCTEXT("AudioRightTT", "The right channel audio output. Mono wave assets will be upmixed to dual stereo.");
+
+	}
+
 	class FWavePlayerNode : public FNode
 	{
 		class FOperatorFactory : public IOperatorFactory
@@ -42,83 +89,123 @@ namespace Metasound
 		FVertexInterface Interface;
 	};
 
+	struct FWavePlayerOpArgs
+	{
+		FOperatorSettings Settings;
+		FTriggerReadRef PlayTrigger;
+		FTriggerReadRef StopTrigger;
+		FWaveAssetReadRef WaveAsset;
+		FTimeReadRef StartTime;
+		FFloatReadRef PitchShift;
+		FBoolReadRef bLoop;
+		FTimeReadRef LoopStartTime;
+		FTimeReadRef LoopDurationSeconds;
+		//FBoolReadRef bStartTimeAsPercentage;
+	};
 
 	class FWavePlayerOperator : public TExecutableOperator<FWavePlayerOperator>
-	{
+	{	
 	public:
-		FWavePlayerOperator(
-			const FOperatorSettings& InSettings,
-			const FWaveAssetReadRef& InWave,
-			const FTriggerReadRef& InPlayTrigger,
-			const FTriggerReadRef& InStopTrigger,
-			const FTriggerReadRef& InSeekTrigger,
-			const FFloatReadRef& InSeekTimeSeconds,
-			const FFloatReadRef& InPitchShiftSemiTones,
-			const FBoolReadRef& InLoop
-		)
-			: OperatorSettings(InSettings)
-			, PlayTrig(InPlayTrigger)
-			, StopTrig(InStopTrigger)
-			, SeekTrig(InSeekTrigger)
-			, Wave(InWave)
-			, SeekTimeSeconds(InSeekTimeSeconds)
-			, PitchShiftSemiTones(InPitchShiftSemiTones)
-			, IsLooping(InLoop)
-			, AudioBufferL(FAudioBufferWriteRef::CreateNew(InSettings))
-			, AudioBufferR(FAudioBufferWriteRef::CreateNew(InSettings))
-			, TrigggerOnLooped(FTriggerWriteRef::CreateNew(InSettings))
-			, TrigggerOnDone(FTriggerWriteRef::CreateNew(InSettings))
-			, OutputSampleRate(InSettings.GetSampleRate())
-			, OutputBlockSizeInFrames(InSettings.GetNumFramesPerBlock())
+		static constexpr float MaxPitchShiftOctaves = 6.0f;
+
+		FWavePlayerOperator(const FWavePlayerOpArgs& InArgs)
+			: OperatorSettings(InArgs.Settings)
+			, PlayTrigger(InArgs.PlayTrigger)
+			, StopTrigger(InArgs.StopTrigger)
+			, WaveAsset(InArgs.WaveAsset)
+			, StartTime(InArgs.StartTime)
+			//, bStartTimeAsPercentage(InArgs.bStartTimeAsPercentage)
+			, PitchShift(InArgs.PitchShift)
+			, bLoop(InArgs.bLoop)
+			, LoopStartTime(InArgs.LoopStartTime)
+			, LoopDurationSeconds(InArgs.LoopDurationSeconds)
+			, TriggerOnPlay(FTriggerWriteRef::CreateNew(InArgs.Settings))
+			, TriggerOnDone(FTriggerWriteRef::CreateNew(InArgs.Settings))
+			, TriggerOnNearlyDone(FTriggerWriteRef::CreateNew(InArgs.Settings))
+			, TriggerOnLooped(FTriggerWriteRef::CreateNew(InArgs.Settings))
+			, TriggerOnCuePoint(FTriggerWriteRef::CreateNew(InArgs.Settings))
+			, CuePointID(FInt32WriteRef::CreateNew(0))
+			, CuePointLabel(FStringWriteRef::CreateNew(TEXT("")))
+			, LoopPercent(FFloatWriteRef::CreateNew(0.0f))
+			, PlaybackLocation(FFloatWriteRef::CreateNew(0.0f))
+			, AudioBufferL(FAudioBufferWriteRef::CreateNew(InArgs.Settings))
+			, AudioBufferR(FAudioBufferWriteRef::CreateNew(InArgs.Settings))
+			, OutputSampleRate(InArgs.Settings.GetSampleRate())
+			, OutputBlockSizeInFrames(InArgs.Settings.GetNumFramesPerBlock())
 		{
 			check(OutputSampleRate);
 			check(AudioBufferL->Num() == OutputBlockSizeInFrames && AudioBufferR->Num() == OutputBlockSizeInFrames);
-
-			if (IsWaveValid())
-			{
-				CurrentSoundWaveName = (*Wave)->GetFName();
-			}
 		}
 
 		virtual FDataReferenceCollection GetInputs() const override
 		{
+			using namespace WavePlayerVertexNames;
+
 			FDataReferenceCollection InputDataReferences;
-			InputDataReferences.AddDataReadReference(TEXT("Audio"), FWaveAssetReadRef(Wave));
-			InputDataReferences.AddDataReadReference(TEXT("Play"), FTriggerReadRef(PlayTrig));
-			InputDataReferences.AddDataReadReference(TEXT("Stop"), FTriggerReadRef(StopTrig));
-			InputDataReferences.AddDataReadReference(TEXT("Seek"), FTriggerReadRef(SeekTrig));
-			InputDataReferences.AddDataReadReference(TEXT("SeekTime"), FFloatReadRef(SeekTimeSeconds));
-			InputDataReferences.AddDataReadReference(TEXT("PitchShift"), FFloatReadRef(PitchShiftSemiTones));
-			InputDataReferences.AddDataReadReference(TEXT("Loop"), FBoolReadRef(IsLooping));
+			InputDataReferences.AddDataReadReference(InputTriggerPlayName, PlayTrigger);
+			InputDataReferences.AddDataReadReference(InputTriggerStopName, StopTrigger);
+			InputDataReferences.AddDataReadReference(InputWaveAssetName, WaveAsset);
+			InputDataReferences.AddDataReadReference(InputStartTimeName, StartTime);
+			InputDataReferences.AddDataReadReference(InputPitchShiftName, PitchShift);
+			InputDataReferences.AddDataReadReference(InputLoopName, bLoop);
+			InputDataReferences.AddDataReadReference(InputLoopStartName, LoopStartTime);
+			InputDataReferences.AddDataReadReference(InputLoopDurationName, LoopDurationSeconds);
+			//InputDataReferences.AddDataReadReference(InputStartTimeAsPercentName, bStartTimeAsPercentage);
 			return InputDataReferences;
 		}
 
 		virtual FDataReferenceCollection GetOutputs() const override
 		{
+			using namespace WavePlayerVertexNames;
+
 			FDataReferenceCollection OutputDataReferences;
-			OutputDataReferences.AddDataReadReference(TEXT("AudioLeft"), FAudioBufferReadRef(AudioBufferL));
-			OutputDataReferences.AddDataReadReference(TEXT("AudioRight"), FAudioBufferReadRef(AudioBufferR));
-			OutputDataReferences.AddDataReadReference(TEXT("Looped"), FTriggerReadRef(TrigggerOnLooped));
-			OutputDataReferences.AddDataReadReference(TEXT("Done"), FTriggerReadRef(TrigggerOnDone));
+			OutputDataReferences.AddDataReadReference(OutputTriggerOnPlayName, TriggerOnPlay);
+			OutputDataReferences.AddDataReadReference(OutputTriggerOnDoneName, TriggerOnDone);
+			OutputDataReferences.AddDataReadReference(OutputTriggerOnNearlyDoneName, TriggerOnNearlyDone);
+			OutputDataReferences.AddDataReadReference(OutputTriggerOnLoopedName, TriggerOnLooped);
+			OutputDataReferences.AddDataReadReference(OutputTriggerOnCuePointName, TriggerOnCuePoint);
+			OutputDataReferences.AddDataReadReference(OutputCuePointIDName, CuePointID);
+			OutputDataReferences.AddDataReadReference(OutputCuePointLabelName, CuePointLabel);
+			OutputDataReferences.AddDataReadReference(OutputLoopPercentName, LoopPercent);
+			OutputDataReferences.AddDataReadReference(OutputPlaybackLocationName, PlaybackLocation);
+			OutputDataReferences.AddDataReadReference(OutputAudioLeftName, AudioBufferL);
+			OutputDataReferences.AddDataReadReference(OutputAudioRightName, AudioBufferR);
 			return OutputDataReferences;
 		}
 
-		void Execute()
+		float UpdateAndGetLoopStartTime()
 		{
-			TrigggerOnDone->AdvanceBlock();
-			TrigggerOnLooped->AdvanceBlock();
+			float LoopSeekTime = 0.0f;
+// 			if (*bStartTimeAsPercentage)
+// 			{
+// 				float LoopStartTimeClamped = FMath::Clamp((float)LoopStartTime->GetSeconds(), 0.0f, 1.0f);
+// 				LoopSeekTime = LoopStartTimeClamped * SoundAssetDurationSeconds;
+// 			}
+// 			else
+			{
+				LoopSeekTime = FMath::Clamp((float)LoopStartTime->GetSeconds(), 0.0f, SoundAssetDurationSeconds);
+			}
 
-			// Check if soundwave input is new
-			if (Wave->IsSoundWaveValid())
+			LoopStartFrame = SoundAssetSampleRate * LoopSeekTime;
+			FramesToConsumePlayBeforeLooping = LoopStartFrame;
+
+			if (!bLooped && PlaybackStartFrame > LoopStartFrame)
 			{
-				const FName SoundWaveName = (*Wave)->GetFName();
-				if (SoundWaveName != CurrentSoundWaveName)
-				{
-					ResetDecoder();
-					CurrentSoundWaveName = SoundWaveName;
-				}
+				FramesToConsumePlayBeforeLooping += (SoundAssetNumFrames - PlaybackStartFrame);
 			}
 
+			return LoopSeekTime;
+		}
+
+		void Execute()
+		{
+			TriggerOnPlay->AdvanceBlock();
+			TriggerOnDone->AdvanceBlock();
+			TriggerOnNearlyDone->AdvanceBlock();
+			TriggerOnCuePoint->AdvanceBlock();
+			TriggerOnLooped->AdvanceBlock();
+
+
 			// zero output buffers
 			FMemory::Memzero(AudioBufferL->GetData(), OutputBlockSizeInFrames * sizeof(float));
 			FMemory::Memzero(AudioBufferR->GetData(), OutputBlockSizeInFrames * sizeof(float));
@@ -126,15 +213,11 @@ namespace Metasound
 			// Parse triggers and render audio
 			int32 PlayTrigIndex = 0;
 			int32 NextPlayFrame = 0;
-			const int32 NumPlayTrigs = PlayTrig->NumTriggeredInBlock();
+			const int32 NumPlayTrigs = PlayTrigger->NumTriggeredInBlock();
 
 			int32 StopTrigIndex = 0;
 			int32 NextStopFrame = 0;
-			const int32 NumStopTrigs = StopTrig->NumTriggeredInBlock();
-
-			int32 SeekTrigIndex = 0;
-			int32 NextSeekFrame = 0;
-			const int32 NumSeekTrigs = SeekTrig->NumTriggeredInBlock();
+			const int32 NumStopTrigs = StopTrigger->NumTriggeredInBlock();
 
 			int32 CurrAudioFrame = 0;
 			int32 NextAudioFrame = 0;
@@ -143,11 +226,11 @@ namespace Metasound
 			{
 				const int32 NoTrigger = (OutputBlockSizeInFrames << 1);
 
-				// get the next Play and Seek indicies
+				// get the next Play and Stop indices
 				// (play)
 				if (PlayTrigIndex < NumPlayTrigs)
 				{
-					NextPlayFrame = (*PlayTrig)[PlayTrigIndex];
+					NextPlayFrame = (*PlayTrigger)[PlayTrigIndex];
 				}
 				else
 				{
@@ -157,26 +240,15 @@ namespace Metasound
 				// (stop)
 				if (StopTrigIndex < NumStopTrigs)
 				{
-					NextStopFrame = (*StopTrig)[StopTrigIndex];
+					NextStopFrame = (*StopTrigger)[StopTrigIndex];
 				}
 				else
 				{
 					NextStopFrame = NoTrigger;
 				}
 
-				// (seek)
-				if (SeekTrigIndex < NumSeekTrigs)
-				{
-					NextSeekFrame = (*SeekTrig)[SeekTrigIndex];
-				}
-				else
-				{
-					NextSeekFrame = NoTrigger;
-				}
-
 				// determine the next audio frame we are going to render up to
 				NextAudioFrame = FMath::Min(NextPlayFrame, NextStopFrame);
-				NextAudioFrame = FMath::Min(NextAudioFrame, NextSeekFrame);
 
 				// no more triggers, rendering to the end of the block
 				if (NextAudioFrame == NoTrigger)
@@ -198,14 +270,12 @@ namespace Metasound
 				}
 
 				// execute the next trigger
-				if (CurrAudioFrame == NextSeekFrame)
-				{
-					ExecuteSeekRequest();
-					++SeekTrigIndex;
-				}
-
 				if (CurrAudioFrame == NextPlayFrame)
 				{
+					if (*StartTime > 0.0f)
+					{
+						ExecuteSeekRequest();
+					}
 					StartPlaying();
 					++PlayTrigIndex;
 				}
@@ -213,8 +283,7 @@ namespace Metasound
 				if (CurrAudioFrame == NextStopFrame)
 				{
 					bIsPlaying = false;
-					ResetDecoder();
-					TrigggerOnDone->TriggerFrame(CurrAudioFrame);
+					TriggerOnDone->TriggerFrame(CurrAudioFrame);
 					++StopTrigIndex;
 				}
 			}
@@ -222,12 +291,42 @@ namespace Metasound
 
 		void StartPlaying()
 		{
-			ResetDecoder(true /* bLogFailures */);
-
-			if (!bIsPlaying)
+			if (IsWaveValid())
 			{
-				bIsPlaying = Decoder.CanGenerateAudio();
+				// Copy the wave asset off on init in case the user changes it while we're playing it.
+				// We'll only check for new wave assets when the current one finishes for sample accurate concantenation
+				CurrentWaveAsset = *WaveAsset;
+				SoundAssetSampleRate = CurrentWaveAsset->GetSampleRate();
+				SoundAssetDurationSeconds = CurrentWaveAsset->GetDuration();
+				SoundAssetNumFrames = CurrentWaveAsset->GetNumFrames();
+				CurrentSoundWaveName = CurrentWaveAsset->GetFName();
+
+				PlaybackStartFrame = FMath::Clamp((float)StartTime->GetSeconds(), 0.0f, SoundAssetDurationSeconds) / SoundAssetSampleRate;
+
+				float StartTimeSeconds = 0.0f;
+				const FSoundWaveProxyPtr& WaveProxy = CurrentWaveAsset.GetSoundWaveProxy();
+// 				if (*bStartTimeAsPercentage)
+// 				{
+// 
+// 					float Percentage = FMath::Clamp((float)StartTime->GetSeconds(), 0.0f, 1.0f);
+// 					StartTimeSeconds = Percentage * WaveProxy->GetDuration();
+// 				}
+// 				else
+				{
+					StartTimeSeconds = FMath::Clamp((float)StartTime->GetSeconds(), 0.0f, WaveProxy->GetDuration());
+				}
+
+				PlaybackStartFrame = (uint32)(StartTimeSeconds * SoundAssetSampleRate);
+				UpdateAndGetLoopStartTime();
+
+				StartDecoder(StartTimeSeconds, true /* bLogFailures */);
+
+				if (!bIsPlaying)
+				{
+					bIsPlaying = Decoder.CanGenerateAudio();
+				}
 			}
+
 		}
 
 		void ExecuteSeekRequest()
@@ -240,9 +339,9 @@ namespace Metasound
 			}
 		}
 
-		bool IsWaveValid(bool bReportToLog = false) const
+		bool IsWaveValidInternal(const FWaveAsset& InWaveAsset, bool bReportToLog = false) const
 		{
-			if (!Wave->IsSoundWaveValid())
+			if (!InWaveAsset.IsSoundWaveValid())
 			{
 				if (bReportToLog)
 				{
@@ -251,7 +350,7 @@ namespace Metasound
 				return false;
 			}
 
-			const FSoundWaveProxyPtr& WaveProxy = Wave->GetSoundWaveProxy();
+			const FSoundWaveProxyPtr& WaveProxy = InWaveAsset.GetSoundWaveProxy();
 
 			// WavePlayerNode currently only supports mono or stereo inputs
 			const int32 NumChannels = WaveProxy->GetNumChannels();
@@ -267,18 +366,38 @@ namespace Metasound
 			return true;
 		}
 
-		bool ResetDecoder(bool bLogFailures = false)
+		bool IsCurrentWaveValid(bool bReportToLog = false) const
 		{
-			Audio::FSimpleDecoderWrapper::InitParams Params;
-			Params.OutputBlockSizeInFrames = OutputBlockSizeInFrames;
-			Params.OutputSampleRate = OutputSampleRate;
-			Params.MaxPitchShiftMagnitudeAllowedInOctaves = 6.f;
-			Params.InitialPitchShiftSemitones = *PitchShiftSemiTones;
-			Params.StartTimeSeconds = FMath::Max(0.f, *SeekTimeSeconds);
-
-			if (IsWaveValid(bLogFailures))
+			return IsWaveValidInternal(CurrentWaveAsset, bReportToLog);
+		}
+
+		bool IsWaveValid(bool bReportToLog = false) const
+		{
+			return IsWaveValidInternal(*WaveAsset, bReportToLog);
+		}
+
+		float GetPlaybackRateFromPitchShift(float InPitchShift)
+		{
+			return FMath::Pow(2.0f, InPitchShift / 12.0f);
+		}
+
+		bool StartDecoder(float SeekTimeSeconds, bool bLogFailures = false)
+		{
+			if (IsCurrentWaveValid(bLogFailures))
 			{
-				const FSoundWaveProxyPtr& WaveProxy = Wave->GetSoundWaveProxy();
+				const FSoundWaveProxyPtr& WaveProxy = CurrentWaveAsset.GetSoundWaveProxy();
+
+				CuePoints = WaveProxy->GetCuePoints();
+
+				float PitchShiftClamped = FMath::Clamp(*PitchShift, -12.0f * MaxPitchShiftOctaves, 12.0f * MaxPitchShiftOctaves);
+
+				Audio::FSimpleDecoderWrapper::InitParams Params;
+				Params.OutputBlockSizeInFrames = OutputBlockSizeInFrames;
+				Params.OutputSampleRate = OutputSampleRate;
+				Params.MaxPitchShiftMagnitudeAllowedInOctaves = 6.f;
+				Params.InitialPitchShiftSemitones = PitchShiftClamped;
+				Params.StartTimeSeconds = SeekTimeSeconds;
+
 				return Decoder.Initialize(Params, WaveProxy);
 			}
 
@@ -291,48 +410,175 @@ namespace Metasound
 			// note: output is hard-coded to stereo (dual-mono)
 			float* FinalOutputLeft = AudioBufferL->GetData() + StartFrame;
 			float* FinalOutputRight = AudioBufferR->GetData() + StartFrame;
-			const int32 NumOutputFrames = (EndFrame - StartFrame);
-			int32 NumFramesDecoded = 0;
+			int32 NumOutputFramesToGenerate = (EndFrame - StartFrame);
+			int32 NumFramesGenerated = 0;
 
-			const bool bCanDecodeWave = IsWaveValid();
+			const bool bCanDecodeWave = IsCurrentWaveValid();
 			if (bCanDecodeWave)
 			{
 				ensure(Decoder.CanGenerateAudio());
 
-				const int32 NumInputChannels = (*Wave)->GetNumChannels();
+				const int32 NumInputChannels = CurrentWaveAsset->GetNumChannels();
 				const bool bNeedsUpmix = (NumInputChannels == 1);
 				const bool bNeedsDeinterleave = !bNeedsUpmix;
 
-				const int32 NumSamplesToGenerate = NumOutputFrames * NumInputChannels;
+				const int32 NumSamplesToGenerate = NumOutputFramesToGenerate * NumInputChannels;
+				
+				float PitchShiftClamped = FMath::Clamp(*PitchShift, -12.0f * MaxPitchShiftOctaves, 12.0f * MaxPitchShiftOctaves);
 
 				PostSrcBuffer.Reset(NumSamplesToGenerate);
 				PostSrcBuffer.AddZeroed(NumSamplesToGenerate);
 				float* PostSrcBufferPtr = PostSrcBuffer.GetData();
 
-				bool bIsLooping = *IsLooping;
-				NumFramesDecoded = Decoder.GenerateAudio(PostSrcBufferPtr, NumOutputFrames, (*PitchShiftSemiTones * 100.f), bIsLooping) / NumInputChannels;
+				// Update looping state which may have changed
+				bool bIsLooping = *bLoop;
 
-				// handle decoder having completed during it's decode
-				const bool bReachedEOF = !Decoder.CanGenerateAudio() || (NumFramesDecoded < NumOutputFrames);
-				if (bReachedEOF && !bIsLooping)
+				// If we're looping we got some extra logic we need to do before we generate audio
+				if (*bLoop)
 				{
-					bIsPlaying = false;
-					TrigggerOnDone->TriggerFrame(StartFrame + NumFramesDecoded);
+					float LoopDuration = SoundAssetDurationSeconds;
+					if (LoopDurationSeconds->GetSeconds() > 0.0f)
+					{
+						LoopDuration = FMath::Clamp((float)LoopDurationSeconds->GetSeconds(), 0.0f, SoundAssetDurationSeconds);
+					}
+					
+					// Give the loop frame a minimum width
+					LoopNumFrames = (uint32)FMath::Max(LoopDuration * SoundAssetSampleRate, 10.0f);
+
+					// We need to look for the case where we are about to loop and we don't want to "overshoot" the loop
+					// So we need to check the number output frames to generate against how much is left in our loop
+					if (NumTotalDecodedFramesInLoop < LoopNumFrames)
+					{
+						// The number of frames left in the loop we need to consume from the source file
+						int32 NumFramesLeftToConsumeInLoop = (int32)LoopNumFrames - NumTotalDecodedFramesInLoop;
+
+						// Translate that to a generated frames count taking into account the pitch scale and the playback rate.
+						// Note this is going to be slightly inaccurate due to pitch interpolation... but it should be super close
+						// and w/ loop cross fading implemented, won't matter too much.
+						float PlaybackRate = GetPlaybackRateFromPitchShift(PitchShiftClamped);
+
+						float SampleRateRatioWithPitchScale = PlaybackRate * SoundAssetSampleRate / OutputSampleRate;
+						int32 NumFramesLeftToGenerateForLoop = SampleRateRatioWithPitchScale * NumFramesLeftToConsumeInLoop;
+
+						// Generate the min of the requested frames to generate and the amount we need to finish the loop
+						NumOutputFramesToGenerate = FMath::Min(NumOutputFramesToGenerate, NumFramesLeftToGenerateForLoop);
+					}
+					else
+					{
+						NumOutputFramesToGenerate = 1;
+					}
+				}
+				else
+				{
+					LoopNumFrames = 0;
+					NumTotalDecodedFramesInLoop = 0;
 				}
-				else if (bReachedEOF && bIsLooping)
+
+				int32 NumFramesConsumed = 0;
+				if (NumOutputFramesToGenerate > 0)
 				{
-					TrigggerOnLooped->TriggerFrame(StartFrame + NumFramesDecoded);
+					NumFramesGenerated = Decoder.GenerateAudio(PostSrcBufferPtr, NumOutputFramesToGenerate, NumFramesConsumed, 100.0f * PitchShiftClamped, bIsLooping) / NumInputChannels;
+
+					// Update loop and playback state logic based on the number of frames consumed
+					FramesConsumedSinceStart += NumFramesConsumed;
+
+					int32 PrevConsumedFrameCount = CurrentConsumedFrameCount;
+					CurrentConsumedFrameCount = (CurrentConsumedFrameCount + NumFramesConsumed) % SoundAssetNumFrames;
+
+					*PlaybackLocation = SoundAssetDurationSeconds * (float)(CurrentConsumedFrameCount / SoundAssetNumFrames);
+
+					// Check for any cue trigger events 
+					FSoundWaveCuePoint OutCuePoint;
+					int32 OutFrameOffset = INDEX_NONE;
+					if (GetCuePointForBlock(PrevConsumedFrameCount, CurrentConsumedFrameCount, OutCuePoint, OutFrameOffset))
+					{
+						// Write the outputs
+						*CuePointID = OutCuePoint.CuePointID;
+						*CuePointLabel = OutCuePoint.Label;
+
+						// Do the trigger
+						TriggerOnCuePoint->TriggerFrame(StartFrame + OutFrameOffset);
+					}
+
+					// Check for nearly done trigger logic. If we're not looping and the next render block will likely finish the file.
+					if (!bIsLooping && (CurrentConsumedFrameCount + OutputBlockSizeInFrames) >= SoundAssetNumFrames)
+					{
+						TriggerOnNearlyDone->TriggerFrame(StartFrame + NumFramesGenerated);
+					}
+
+				}
+
+				// If we're looping and we've made it past the loop start point and we've not yet looped, check the looping condition
+				// Note: we need to handle the case where the "start time" of the player is past the loop start point. In that case
+				// we allow the audio file to wrap around and loop from the beginning before checking the loop logic
+				if (bIsLooping && (bLooped || FramesConsumedSinceStart >= FramesToConsumePlayBeforeLooping))
+				{					
+					NumTotalDecodedFramesInLoop += NumFramesConsumed;
+
+					if (NumTotalDecodedFramesInLoop >= LoopNumFrames || !NumOutputFramesToGenerate)
+					{
+						// We've looped now so we will continue to loop between the loop start and end points
+						bLooped = true;
+
+						// Trigger that we looped
+						TriggerOnLooped->TriggerFrame(StartFrame + NumFramesGenerated);
+					
+						// Update the loop start frame based on any recent inputs
+						float LoopSeekTime = UpdateAndGetLoopStartTime();
+
+						// Update our current loop frame
+						CurrentConsumedFrameCount = LoopStartFrame;
+
+						// Reset our decoded frame counting for sub loop
+						NumTotalDecodedFramesInLoop = 0;
+
+						*LoopPercent = 0.0f;
+
+						// Now seek to the loop start frame
+						StartDecoder(LoopSeekTime);
+					}
+					else
+					{
+						*LoopPercent = (float)NumTotalDecodedFramesInLoop / LoopNumFrames;
+					}
+				}
+				else
+				{
+					*LoopPercent = 0.0f;
+				}
+				
+				// handle decoder having completed during it's decode
+				const bool bReachedEOF = !Decoder.CanGenerateAudio() || (NumFramesGenerated < NumOutputFramesToGenerate);
+				if (bReachedEOF && !bIsLooping)
+				{
+					// Check the wave file input to see if it has changed... if it has, we want to restart everything up!
+					bool bConcatenatedPlayback = false;
+					if (WaveAsset->IsSoundWaveValid())
+					{
+						const FName SoundWaveName = (*WaveAsset)->GetFName();
+						if (SoundWaveName != CurrentSoundWaveName)
+						{
+							StartPlaying();
+							bConcatenatedPlayback = true;
+						}
+					}
+
+					if (!bConcatenatedPlayback)
+					{
+						bIsPlaying = false;
+						TriggerOnDone->TriggerFrame(StartFrame + NumFramesGenerated);
+					}
 				}
 
 				if (bNeedsUpmix)
 				{
-					static const float Sqrt2Over2 = 0.707106781f;
-					Audio::BufferMultiplyByConstant(PostSrcBufferPtr, Sqrt2Over2, FinalOutputLeft, NumFramesDecoded);
-					Audio::BufferMultiplyByConstant(PostSrcBufferPtr, Sqrt2Over2, FinalOutputRight, NumFramesDecoded);
+					// Dual mono output
+					FMemory::Memcpy(FinalOutputLeft, PostSrcBufferPtr, sizeof(float)* NumFramesGenerated);
+					FMemory::Memcpy(FinalOutputRight, PostSrcBufferPtr, sizeof(float) * NumFramesGenerated);
 				}
 				else if (bNeedsDeinterleave)
 				{
-					for (int32 i = 0; i < NumFramesDecoded; ++i)
+					for (int32 i = 0; i < NumFramesGenerated; ++i)
 					{
 						// deinterleave each stereo frame into output buffers
 						FinalOutputLeft[i] = PostSrcBufferPtr[(i << 1)];
@@ -342,40 +588,85 @@ namespace Metasound
 			}
 			else
 			{
-				FMemory::Memzero(FinalOutputLeft, sizeof(float) * NumOutputFrames);
-				FMemory::Memzero(FinalOutputRight, sizeof(float) * NumOutputFrames);
+				FMemory::Memzero(FinalOutputLeft, sizeof(float) * NumOutputFramesToGenerate);
+				FMemory::Memzero(FinalOutputRight, sizeof(float) * NumOutputFramesToGenerate);
 			}
 
-			return NumFramesDecoded;
+			return NumFramesGenerated;
 		}
 
 	private:
+
+		// Called to get the cue point in the given block. Will return the first cue point if multiple cues are in the same block.
+		// Note: we can't currently support sub-block cue triggers due to cue ID and cue label being block-rate params. 
+		bool GetCuePointForBlock(int32 InStartFrame, int32 InEndFrame, FSoundWaveCuePoint& OutCuePoint, int32& OutFrameOffset)
+		{
+			for (FSoundWaveCuePoint& CuePoint : CuePoints)
+			{
+				if (CuePoint.FramePosition >= InStartFrame && CuePoint.FramePosition < InEndFrame)
+				{
+					OutCuePoint = CuePoint;
+					OutFrameOffset = CuePoint.FramePosition - InStartFrame;
+					return true;
+				}
+			}
+			// No cue points in this block
+			return false;
+		}
+
 		const FOperatorSettings OperatorSettings;
 
 		// i/o
-		FTriggerReadRef PlayTrig;
-		FTriggerReadRef StopTrig;
-		FTriggerReadRef SeekTrig;
-		FWaveAssetReadRef Wave;
-		FFloatReadRef SeekTimeSeconds;
-		FFloatReadRef PitchShiftSemiTones;
-		FBoolReadRef IsLooping;
-
+		FTriggerReadRef PlayTrigger;
+		FTriggerReadRef StopTrigger;
+		FWaveAssetReadRef WaveAsset;
+		FTimeReadRef StartTime;
+		//FBoolReadRef bStartTimeAsPercentage;
+		FFloatReadRef PitchShift;
+		FBoolReadRef bLoop;
+		FTimeReadRef LoopStartTime;
+		FTimeReadRef LoopDurationSeconds;
+
+		FTriggerWriteRef TriggerOnPlay;
+		FTriggerWriteRef TriggerOnDone;
+		FTriggerWriteRef TriggerOnNearlyDone;
+		FTriggerWriteRef TriggerOnLooped;
+		FTriggerWriteRef TriggerOnCuePoint;
+		FInt32WriteRef CuePointID;
+		FStringWriteRef CuePointLabel;
+		FFloatWriteRef LoopPercent;
+		FFloatWriteRef PlaybackLocation;
 		FAudioBufferWriteRef AudioBufferL;
 		FAudioBufferWriteRef AudioBufferR;
-		FTriggerWriteRef TrigggerOnLooped;
-		FTriggerWriteRef TrigggerOnDone;
 
 		// source decode
 		Audio::FAlignedFloatBuffer PostSrcBuffer;
 		Audio::FSimpleDecoderWrapper Decoder;
 
-		FName CurrentSoundWaveName{ };
+		FName CurrentSoundWaveName;
 			
-		const float OutputSampleRate{ 0.f };
-		const int32 OutputBlockSizeInFrames{ 0 };
+		float SoundAssetSampleRate = 0.0f;
+		float SoundAssetDurationSeconds = 0.0f;
+		uint32 SoundAssetNumFrames = 0;
+
+		const float OutputSampleRate = 0.f;
+		const int32 OutputBlockSizeInFrames = 0;
 		
-		bool bIsPlaying{ false };
+		TArray<FSoundWaveCuePoint> CuePoints;
+		TArray<FSoundWaveCuePoint> CuePointsInCurrentBock;
+
+		// Data to track decoder
+		uint32 PlaybackStartFrame = 0;
+		uint32 CurrentConsumedFrameCount = 0;
+		uint32 FramesConsumedSinceStart = 0;
+		uint32 NumTotalDecodedFramesInLoop = 0;
+		uint32 LoopStartFrame = 0;
+		uint32 FramesToConsumePlayBeforeLooping = 0;
+		uint32 LoopNumFrames = 0;
+		bool bLooped = false;
+		bool bIsPlaying = false;
+
+		FWaveAsset CurrentWaveAsset;
 	};
 
 	TUniquePtr<IOperator> FWavePlayerNode::FOperatorFactory::CreateOperator(
@@ -383,48 +674,57 @@ namespace Metasound
 		FBuildErrorArray& OutErrors) 
 	{
 		using namespace Audio;
+		using namespace WavePlayerVertexNames;
 
 		const FWavePlayerNode& WaveNode = static_cast<const FWavePlayerNode&>(InParams.Node);
 
-		const FDataReferenceCollection& InputDataRefs = InParams.InputDataReferences;
-
-		FTriggerReadRef TriggerPlay = InputDataRefs.GetDataReadReferenceOrConstruct<FTrigger>(TEXT("Play"), InParams.OperatorSettings);
-		FTriggerReadRef TriggerStop = InputDataRefs.GetDataReadReferenceOrConstruct<FTrigger>(TEXT("Stop"), InParams.OperatorSettings);
-		FTriggerReadRef TriggerSeek = InputDataRefs.GetDataReadReferenceOrConstruct<FTrigger>(TEXT("Seek"), InParams.OperatorSettings);
-		FWaveAssetReadRef Wave = InputDataRefs.GetDataReadReferenceOrConstruct<FWaveAsset>(TEXT("Wave"));
-		FFloatReadRef SeekTimeSeconds = InputDataRefs.GetDataReadReferenceOrConstruct<float>(TEXT("SeekTime"), 0.f);
-		FFloatReadRef PitchShiftSemiTones = InputDataRefs.GetDataReadReferenceOrConstruct<float>(TEXT("PitchShift"));
-		FBoolReadRef IsLooping = InputDataRefs.GetDataReadReferenceOrConstruct<bool>(TEXT("Loop"), false);
-
-		return MakeUnique<FWavePlayerOperator>(
-			  InParams.OperatorSettings
-			, Wave
-			, TriggerPlay
-			, TriggerStop
-			, TriggerSeek
-			, SeekTimeSeconds
-			, PitchShiftSemiTones
-			, IsLooping
-			);
+		const FDataReferenceCollection& Inputs = InParams.InputDataReferences;
+
+		FWavePlayerOpArgs Args =
+		{
+			InParams.OperatorSettings,
+			Inputs.GetDataReadReferenceOrConstruct<FTrigger>(InputTriggerPlayName, InParams.OperatorSettings),
+			Inputs.GetDataReadReferenceOrConstruct<FTrigger>(InputTriggerStopName, InParams.OperatorSettings),
+			Inputs.GetDataReadReferenceOrConstruct<FWaveAsset>(InputWaveAssetName),
+			Inputs.GetDataReadReferenceOrConstruct<FTime>(InputStartTimeName),
+			Inputs.GetDataReadReferenceOrConstruct<float>(InputPitchShiftName),
+			Inputs.GetDataReadReferenceOrConstruct<bool>(InputLoopName),
+			Inputs.GetDataReadReferenceOrConstruct<FTime>(InputLoopStartName),
+			Inputs.GetDataReadReferenceOrConstruct<FTime>(InputLoopDurationName)
+			//Inputs.GetDataReadReferenceOrConstruct<bool>(InputStartTimeAsPercentName)
+		};
+
+		return MakeUnique<FWavePlayerOperator>(Args);
 	}
 
 	FVertexInterface FWavePlayerNode::DeclareVertexInterface()
 	{
+		using namespace WavePlayerVertexNames;
+
 		return FVertexInterface(
 			FInputVertexInterface(
-				TInputDataVertexModel<FWaveAsset>(TEXT("Wave"), LOCTEXT("WaveTooltip", "The Wave to be decoded")),
-				TInputDataVertexModel<FTrigger>(TEXT("Play"), LOCTEXT("PlayTooltip", "Play the input wave.")),
-				TInputDataVertexModel<FTrigger>(TEXT("Stop"), LOCTEXT("StopTooltip", "Stop the input wave.")),
-				TInputDataVertexModel<FTrigger>(TEXT("Seek"), LOCTEXT("PlayTooltip", "Seek to the given time. (Wave being played must be set to seekable!)")),
-				TInputDataVertexModel<float>(TEXT("SeekTime"), LOCTEXT("SeekTimeTooltip", "Seek time in seconds. (Wave being played must be set to seekable!)")),
-				TInputDataVertexModel<float>(TEXT("PitchShift"), LOCTEXT("PitchShiftTooltip", "Pitch Shift in semi-tones.")),
-				TInputDataVertexModel<bool>(TEXT("Loop"), LOCTEXT("LoopTooltip", "Wave will loop if true"))
-			),
+				TInputDataVertexModel<FTrigger>(InputTriggerPlayName, InputTriggerPlayTT),
+				TInputDataVertexModel<FTrigger>(InputTriggerStopName, InputTriggerStopTT),
+				TInputDataVertexModel<FWaveAsset>(InputWaveAssetName, InputWaveAssetTT),
+				TInputDataVertexModel<FTime>(InputStartTimeName, InputStartTimeTT, 0.0f),
+				TInputDataVertexModel<float>(InputPitchShiftName, InputPitchShiftTT, 0.0f),
+				TInputDataVertexModel<bool>(InputLoopName, InputLoopTT, false),
+				TInputDataVertexModel<FTime>(InputLoopStartName, InputLoopStartTT, 0.0f),
+				TInputDataVertexModel<FTime>(InputLoopDurationName, InputLoopDurationTT, -1.0f)
+				//TInputDataVertexModel<bool>(InputStartTimeAsPercentName, InputStartTimeAsPercentTT, false)
+				),
 			FOutputVertexInterface(
-				TOutputDataVertexModel<FAudioBuffer>(TEXT("AudioLeft"), LOCTEXT("AudioTooltip", "The output audio")),
-				TOutputDataVertexModel<FAudioBuffer>(TEXT("AudioRight"), LOCTEXT("AudioTooltip", "The output audio")),
-				TOutputDataVertexModel<FTrigger>(TEXT("Looped"), LOCTEXT("TriggerToolTip", "Trigger that notifies when the sound is has looped")),
-				TOutputDataVertexModel<FTrigger>(TEXT("Done"), LOCTEXT("TriggerToolTip", "Trigger that notifies when the sound is done playing"))
+				TOutputDataVertexModel<FTrigger>(OutputTriggerOnPlayName, OutputTriggerOnPlayTT),
+				TOutputDataVertexModel<FTrigger>(OutputTriggerOnDoneName, OutputTriggerOnDoneTT),
+				TOutputDataVertexModel<FTrigger>(OutputTriggerOnNearlyDoneName, OutputTriggerOnNearlyDoneTT),
+				TOutputDataVertexModel<FTrigger>(OutputTriggerOnLoopedName, OutputTriggerOnLoopedTT),
+				TOutputDataVertexModel<FTrigger>(OutputTriggerOnCuePointName, OutputTriggerOnCuePointTT),
+				TOutputDataVertexModel<int32>(OutputCuePointIDName, OutputCuePointIDTT),
+				TOutputDataVertexModel<FString>(OutputCuePointLabelName, OutputCuePointLabelTT),
+				TOutputDataVertexModel<float>(OutputLoopPercentName, OutputAudioRightNameTT),
+				TOutputDataVertexModel<float>(OutputPlaybackLocationName, OutputPlaybackLocationTT),
+				TOutputDataVertexModel<FAudioBuffer>(OutputAudioLeftName, OutputAudioLeftNameTT),
+				TOutputDataVertexModel<FAudioBuffer>(OutputAudioRightName, OutputAudioRightNameTT)
 			)
 		);
 	}
@@ -438,7 +738,7 @@ namespace Metasound
 			Info.MajorVersion = 1;
 			Info.MinorVersion = 0;
 			Info.DisplayName = LOCTEXT("Metasound_WavePlayerNodeDisplayName", "Wave Player");
-			Info.Description = LOCTEXT("Metasound_WavePlayerNodeDescription", "Plays a supplied Wave");
+			Info.Description = LOCTEXT("Metasound_WavePlayerNodeDescription", "Plays a wave asset.");
 			Info.Author = PluginAuthor;
 			Info.PromptIfMissing = PluginNodeMissingPrompt;
 			Info.DefaultInterface = DeclareVertexInterface();
diff --git a/Engine/Plugins/Runtime/Metasound/Source/MetasoundEngine/Public/MetasoundWave.h b/Engine/Plugins/Runtime/Metasound/Source/MetasoundEngine/Public/MetasoundWave.h
--- a/Engine/Plugins/Runtime/Metasound/Source/MetasoundEngine/Public/MetasoundWave.h
+++ b/Engine/Plugins/Runtime/Metasound/Source/MetasoundEngine/Public/MetasoundWave.h
@@ -65,7 +65,7 @@ namespace Metasound
 
 	DECLARE_METASOUND_DATA_REFERENCE_TYPES(FWaveAsset, METASOUNDENGINE_API, FWaveAssetTypeInfo, FWaveAssetReadRef, FWaveAssetWriteRef)
 }
-
+ 
 namespace Audio
 {
 	// Forward declares
@@ -94,7 +94,7 @@ namespace Audio
 		bool Initialize(const InitParams& InInitParams, const FSoundWaveProxyPtr& InWave);
 
 		// returns number of samples written.   
-		uint32 GenerateAudio(float* OutputDest, int32 NumOutputFrames, float PitchShiftInCents = 0.0f, bool bIsLooping = false);
+		uint32 GenerateAudio(float* OutputDest, int32 NumOutputFrames, int32& OutNumFramesConsumed, float PitchShiftInCents = 0.0f, bool bIsLooping = false);
 
 		void SeekToTime(const float InSeconds);
 
@@ -124,6 +124,7 @@ namespace Audio
 		float MaxPitchShiftRatio{ 2.0f };
 
 		uint32 NumChannels{ 1 };
+		uint32 NumFrames{ 0 };
 		uint32 DecodeBlockSizeInFrames{ 64 };
 		uint32 DecodeBlockSizeInSamples{ 64 };
 		
diff --git a/Engine/Source/Runtime/Engine/Classes/Sound/SoundWave.h b/Engine/Source/Runtime/Engine/Classes/Sound/SoundWave.h
--- a/Engine/Source/Runtime/Engine/Classes/Sound/SoundWave.h
+++ b/Engine/Source/Runtime/Engine/Classes/Sound/SoundWave.h
@@ -334,6 +334,7 @@ public:
 
 	float GetSampleRate() const { return SampleRate; }
 	uint32 GetNumChannels() const { return NumChannels; }
+	const TArray<FSoundWaveCuePoint>& GetCuePoints() const { return CuePoints; }
 	uint32 GetNumChunks() const;
 	uint32 GetSizeOfChunk(uint32 ChunkIndex) const;
 	int32 GetNumFrames() const { return NumFrames; }
@@ -377,6 +378,7 @@ private:
 
 	float SampleRate;
 	uint32 NumChannels;
+	TArray<FSoundWaveCuePoint> CuePoints;
 	uint32 NumChunks;
 	float Duration;
 	int32 NumFrames;
diff --git a/Engine/Source/Runtime/Engine/Private/SoundWave.cpp b/Engine/Source/Runtime/Engine/Private/SoundWave.cpp
--- a/Engine/Source/Runtime/Engine/Private/SoundWave.cpp
+++ b/Engine/Source/Runtime/Engine/Private/SoundWave.cpp
@@ -3071,7 +3071,10 @@ FSoundWaveProxy::FSoundWaveProxy(USoundWave* InWave)
 	, SoundWaveKeyCached(FObjectKey(InWave))
 	, SampleRate(InWave->GetSampleRateForCurrentPlatform())
 	, NumChannels(InWave->NumChannels)
+	, CuePoints(InWave->CuePoints)
 	, NumChunks(0)
+	, Duration(InWave->Duration)
+	, NumFrames((int32)(Duration * (float)InWave->SampleRate))
 	, FirstChunk(InWave->FirstChunk)
 	, LoadingBehavior(InWave->CachedSoundWaveLoadingBehavior)
 	, bLoadingBehaviorOverridden(InWave->bLoadingBehaviorOverridden)
