Pipe through Metasound Node Advanced Display markup
- Move descriptions to new FDataVertexMetadata struct along with new bIsAdvancedDisplay bool
- Add convenience getters for FMetasoundFrontendVertexMetadata (SO META!) that translates from FDataVertexMetadata on Controllers

diff --git a/Engine/Plugins/Runtime/Metasound/Source/MetasoundEditor/Private/MetasoundEditorGraphBuilder.cpp b/Engine/Plugins/Runtime/Metasound/Source/MetasoundEditor/Private/MetasoundEditorGraphBuilder.cpp
--- a/Engine/Plugins/Runtime/Metasound/Source/MetasoundEditor/Private/MetasoundEditorGraphBuilder.cpp
+++ b/Engine/Plugins/Runtime/Metasound/Source/MetasoundEditor/Private/MetasoundEditorGraphBuilder.cpp
@@ -1071,11 +1071,21 @@ namespace Metasound
 			{
 				NewPin->PinToolTip = InInputHandle->GetTooltip().ToString();
 				NewPin->PinType = EditorModule.FindDataType(InInputHandle->GetDataType()).PinType;
+				NewPin->bAdvancedView = InInputHandle->GetMetadata().bIsAdvancedDisplay;
+
+				if (NewPin->bAdvancedView)
+				{
+					if (InEditorNode.AdvancedPinDisplay == ENodeAdvancedPins::NoPins)
+					{
+						InEditorNode.AdvancedPinDisplay = ENodeAdvancedPins::Hidden;
+					}
+				}
 
 				FNodeHandle NodeHandle = InInputHandle->GetOwningNode();
 				FGraphBuilder::AddOrUpdateLiteralInput(InEditorNode.GetMetasoundChecked(), NodeHandle, *NewPin);
 			}
 
+
 			return NewPin;
 		}
 
@@ -1089,6 +1099,15 @@ namespace Metasound
 			{
 				NewPin->PinToolTip = InOutputHandle->GetTooltip().ToString();
 				NewPin->PinType = EditorModule.FindDataType(InOutputHandle->GetDataType()).PinType;
+				NewPin->bAdvancedView = InOutputHandle->GetMetadata().bIsAdvancedDisplay;
+
+				if (NewPin->bAdvancedView)
+				{
+					if (InEditorNode.AdvancedPinDisplay == ENodeAdvancedPins::NoPins)
+					{
+						InEditorNode.AdvancedPinDisplay = ENodeAdvancedPins::Hidden;
+					}
+				}
 			}
 
 			return NewPin;
diff --git a/Engine/Plugins/Runtime/Metasound/Source/MetasoundFrontend/Private/MetasoundFrontend.cpp b/Engine/Plugins/Runtime/Metasound/Source/MetasoundFrontend/Private/MetasoundFrontend.cpp
--- a/Engine/Plugins/Runtime/Metasound/Source/MetasoundFrontend/Private/MetasoundFrontend.cpp
+++ b/Engine/Plugins/Runtime/Metasound/Source/MetasoundFrontend/Private/MetasoundFrontend.cpp
@@ -15,6 +15,7 @@
 #include "MetasoundOperatorBuilder.h"
 #include "MetasoundPrimitives.h"
 #include "MetasoundRouter.h"
+#include "MetasoundVertex.h"
 #include "Modules/ModuleManager.h"
 #include "StructDeserializer.h"
 #include "StructSerializer.h"
@@ -93,14 +94,16 @@ namespace Metasound
 			TArray<int>& InputSortOrder = ClassInterface.InputStyle.DefaultSortOrder;
 			for (auto& InputTuple : InputInterface)
 			{
-				// TODO: lots to be added here. 
 				FMetasoundFrontendClassInput ClassInput;
 
 				ClassInput.Name = InputTuple.Value.GetVertexName();
 				ClassInput.TypeName = InputTuple.Value.GetDataTypeName();
 				ClassInput.VertexID = FGuid::NewGuid();
 				ClassInput.Metadata.DisplayName = FText::FromString(InputTuple.Value.GetVertexName());
-				ClassInput.Metadata.Description = InputTuple.Value.GetDescription();
+
+				const FDataVertexMetadata& VertexMetadata = InputTuple.Value.GetMetadata();
+				ClassInput.Metadata.Description = VertexMetadata.Description;
+				ClassInput.Metadata.bIsAdvancedDisplay = VertexMetadata.bIsAdvancedDisplay;
 
 				FLiteral DefaultLiteral = InputTuple.Value.GetDefaultLiteral();
 				if (DefaultLiteral.GetType() != ELiteralType::Invalid)
@@ -108,7 +111,12 @@ namespace Metasound
 					ClassInput.DefaultLiteral.SetFromLiteral(DefaultLiteral);
 				}
 
-				const int32 OrderIndex = InputInterface.GetOrderIndex(InputTuple.Key);
+				// Advanced display items are pushed to bottom of sort order
+				int32 OrderIndex = InputInterface.GetOrderIndex(InputTuple.Key);
+				if (ClassInput.Metadata.bIsAdvancedDisplay)
+				{
+					OrderIndex += InputInterface.Num();
+				}
 				InputSortOrder.Add(OrderIndex);
 
 				ClassInterface.Inputs.Add(ClassInput);
@@ -118,16 +126,23 @@ namespace Metasound
 			TArray<int32>& OutputSortOrder = ClassInterface.InputStyle.DefaultSortOrder;
 			for (auto& OutputTuple : OutputInterface)
 			{
-				// TODO: lots to be added here. 
 				FMetasoundFrontendClassOutput ClassOutput;
 
 				ClassOutput.Name = OutputTuple.Value.GetVertexName();
 				ClassOutput.TypeName = OutputTuple.Value.GetDataTypeName();
 				ClassOutput.VertexID = FGuid::NewGuid();
 				ClassOutput.Metadata.DisplayName = FText::FromString(OutputTuple.Value.GetVertexName());
-				ClassOutput.Metadata.Description = OutputTuple.Value.GetDescription();
 
-				const int32 OrderIndex = OutputInterface.GetOrderIndex(OutputTuple.Key);
+				const FDataVertexMetadata& VertexMetadata = OutputTuple.Value.GetMetadata();
+				ClassOutput.Metadata.Description = VertexMetadata.Description;
+				ClassOutput.Metadata.bIsAdvancedDisplay = VertexMetadata.bIsAdvancedDisplay;
+
+				// Advanced display items are pushed to bottom below non-advanced
+				int32 OrderIndex = OutputInterface.GetOrderIndex(OutputTuple.Key);
+				if (ClassOutput.Metadata.bIsAdvancedDisplay)
+				{
+					OrderIndex += OutputInterface.Num();
+				}
 				ClassInterface.OutputStyle.DefaultSortOrder.Add(OrderIndex);
 
 				ClassDescription.Interface.Outputs.Add(ClassOutput);
@@ -135,7 +150,6 @@ namespace Metasound
 
 			for (auto& EnvTuple : InNodeMetadata.DefaultInterface.GetEnvironmentInterface())
 			{
-				// TODO: lots to be added here. 
 				FMetasoundFrontendClassEnvironmentVariable EnvVar;
 
 				EnvVar.Name = EnvTuple.Value.GetVertexName();
diff --git a/Engine/Plugins/Runtime/Metasound/Source/MetasoundFrontend/Private/MetasoundFrontendInvalidController.h b/Engine/Plugins/Runtime/Metasound/Source/MetasoundFrontend/Private/MetasoundFrontendInvalidController.h
--- a/Engine/Plugins/Runtime/Metasound/Source/MetasoundFrontend/Private/MetasoundFrontendInvalidController.h
+++ b/Engine/Plugins/Runtime/Metasound/Source/MetasoundFrontend/Private/MetasoundFrontendInvalidController.h
@@ -7,6 +7,7 @@
 #include "MetasoundBuilderInterface.h"
 #include "MetasoundFrontendBaseClasses.h"
 #include "MetasoundFrontendController.h"
+#include "MetasoundFrontendDocument.h"
 #include "MetasoundGraph.h"
 
 namespace Metasound
@@ -44,11 +45,11 @@ namespace Metasound
 
 			bool IsValid() const override { return false; }
 			FGuid GetID() const override { return Metasound::FrontendInvalidID; }
-			const FName& GetDataType() const override { return MetasoundFrontendInvalidControllerPrivate::GetInvalid<FName>(); }  
+			const FName& GetDataType() const override { return MetasoundFrontendInvalidControllerPrivate::GetInvalid<FName>(); }
 			const FString& GetName() const override { return MetasoundFrontendInvalidControllerPrivate::GetInvalid<FString>(); }
 			const FText& GetDisplayName() const override { return FText::GetEmpty(); }
 			const FText& GetTooltip() const override { return FText::GetEmpty(); }
-
+			const FMetasoundFrontendVertexMetadata& GetMetadata() const override { return MetasoundFrontendInvalidControllerPrivate::GetInvalid<FMetasoundFrontendVertexMetadata>(); }
 			FGuid GetOwningNodeID() const override { return Metasound::FrontendInvalidID; }
 			TSharedRef<INodeController> GetOwningNode() override;
 			TSharedRef<const INodeController> GetOwningNode() const override;
@@ -91,6 +92,7 @@ namespace Metasound
 			const FString& GetName() const override { return MetasoundFrontendInvalidControllerPrivate::GetInvalid<FString>(); }
 			const FText& GetDisplayName() const override { return MetasoundFrontendInvalidControllerPrivate::GetInvalid<FText>(); }
 			const FText& GetTooltip() const override { return MetasoundFrontendInvalidControllerPrivate::GetInvalid<FText>(); }
+			const FMetasoundFrontendVertexMetadata& GetMetadata() const override { return MetasoundFrontendInvalidControllerPrivate::GetInvalid<FMetasoundFrontendVertexMetadata>(); }
 			const FMetasoundFrontendLiteral* GetDefaultLiteral() const override { return nullptr; }
 			FGuid GetOwningNodeID() const override { return Metasound::FrontendInvalidID; }
 			TSharedRef<INodeController> GetOwningNode() override;
diff --git a/Engine/Plugins/Runtime/Metasound/Source/MetasoundFrontend/Private/MetasoundFrontendStandardController.cpp b/Engine/Plugins/Runtime/Metasound/Source/MetasoundFrontend/Private/MetasoundFrontendStandardController.cpp
--- a/Engine/Plugins/Runtime/Metasound/Source/MetasoundFrontend/Private/MetasoundFrontendStandardController.cpp
+++ b/Engine/Plugins/Runtime/Metasound/Source/MetasoundFrontend/Private/MetasoundFrontendStandardController.cpp
@@ -513,8 +513,17 @@ namespace Metasound
 			return FrontendControllerIntrinsics::GetInvalidValueConstRef<FText>();
 		}
 
+		const FMetasoundFrontendVertexMetadata& FBaseOutputController::GetMetadata() const
+		{
+			if (ClassOutputPtr.IsValid())
+			{
+				return ClassOutputPtr->Metadata;
+			}
+
+			return FrontendControllerIntrinsics::GetInvalidValueConstRef<FMetasoundFrontendVertexMetadata>();
+		}
 
-		FDocumentAccess FBaseOutputController::ShareAccess() 
+		FDocumentAccess FBaseOutputController::ShareAccess()
 		{
 			FDocumentAccess Access;
 
@@ -581,6 +590,16 @@ namespace Metasound
 			return FrontendControllerIntrinsics::GetInvalidValueConstRef<FText>();
 		}
 
+		const FMetasoundFrontendVertexMetadata& FInputNodeOutputController::GetMetadata() const
+		{
+			if (OwningGraphClassInputPtr.IsValid())
+			{
+				return OwningGraphClassInputPtr->Metadata;
+			}
+
+			return FrontendControllerIntrinsics::GetInvalidValueConstRef<FMetasoundFrontendVertexMetadata>();
+		}
+
 		FDocumentAccess FInputNodeOutputController::ShareAccess() 
 		{
 			FDocumentAccess Access = FBaseOutputController::ShareAccess();
@@ -633,6 +652,16 @@ namespace Metasound
 			return FrontendControllerIntrinsics::GetInvalidValueConstRef<FText>();
 		}
 
+		const FMetasoundFrontendVertexMetadata& FOutputNodeOutputController::GetMetadata() const
+		{
+			if (OwningGraphClassOutputPtr.IsValid())
+			{
+				return OwningGraphClassOutputPtr->Metadata;
+			}
+
+			return FrontendControllerIntrinsics::GetInvalidValueConstRef<FMetasoundFrontendVertexMetadata>();
+		}
+
 		FConnectability FOutputNodeOutputController::CanConnectTo(const IInputController& InController) const 
 		{
 			// Cannot connect to a graph's output.
@@ -723,6 +752,15 @@ namespace Metasound
 			return FrontendControllerIntrinsics::GetInvalidValueConstRef<FText>();
 		}
 
+		const FMetasoundFrontendVertexMetadata& FBaseInputController::GetMetadata() const
+		{
+			if (const FMetasoundFrontendClassInput* ClassInput = ClassInputPtr.Get())
+			{
+				return ClassInput->Metadata;
+			}
+
+			return FrontendControllerIntrinsics::GetInvalidValueConstRef<FMetasoundFrontendVertexMetadata>();
+		}
 
 		bool FBaseInputController::IsConnected() const 
 		{
@@ -1009,7 +1047,17 @@ namespace Metasound
 			return FrontendControllerIntrinsics::GetInvalidValueConstRef<FText>();
 		}
 
-		const FMetasoundFrontendLiteral* FOutputNodeInputController::GetDefaultLiteral() const 
+		const FMetasoundFrontendVertexMetadata& FOutputNodeInputController::GetMetadata() const
+		{
+			if (OwningGraphClassOutputPtr.IsValid())
+			{
+				return OwningGraphClassOutputPtr->Metadata;
+			}
+
+			return FrontendControllerIntrinsics::GetInvalidValueConstRef<FMetasoundFrontendVertexMetadata>();
+		}
+
+		const FMetasoundFrontendLiteral* FOutputNodeInputController::GetDefaultLiteral() const
 		{
 			return nullptr;
 		}
@@ -1066,6 +1114,16 @@ namespace Metasound
 			return FrontendControllerIntrinsics::GetInvalidValueConstRef<FText>();
 		}
 
+		const FMetasoundFrontendVertexMetadata& FInputNodeInputController::GetMetadata() const
+		{
+			if (const FMetasoundFrontendClassInput* OwningInput = OwningGraphClassInputPtr.Get())
+			{
+				return OwningInput->Metadata;
+			}
+
+			return FrontendControllerIntrinsics::GetInvalidValueConstRef<FMetasoundFrontendVertexMetadata>();
+		}
+
 		FConnectability FInputNodeInputController::CanConnectTo(const IOutputController& InController) const 
 		{
 			static const FConnectability Connectability = {FConnectability::EConnectable::No};
diff --git a/Engine/Plugins/Runtime/Metasound/Source/MetasoundFrontend/Private/MetasoundFrontendStandardController.h b/Engine/Plugins/Runtime/Metasound/Source/MetasoundFrontend/Private/MetasoundFrontendStandardController.h
--- a/Engine/Plugins/Runtime/Metasound/Source/MetasoundFrontend/Private/MetasoundFrontendStandardController.h
+++ b/Engine/Plugins/Runtime/Metasound/Source/MetasoundFrontend/Private/MetasoundFrontendStandardController.h
@@ -49,6 +49,7 @@ namespace Metasound
 			// Output metadata
 			const FText& GetDisplayName() const override;
 			const FText& GetTooltip() const override;
+			const FMetasoundFrontendVertexMetadata& GetMetadata() const override;
 
 			// Return info on containing node. 
 			FGuid GetOwningNodeID() const override;
@@ -111,10 +112,10 @@ namespace Metasound
 
 			bool IsValid() const override;
 
-			// Output metadata
+			// Input metadata
 			const FText& GetDisplayName() const override;
 			const FText& GetTooltip() const override;
-
+			const FMetasoundFrontendVertexMetadata& GetMetadata() const override;
 
 		protected:
 			FDocumentAccess ShareAccess() override;
@@ -158,6 +159,7 @@ namespace Metasound
 			// Output metadata
 			const FText& GetDisplayName() const override;
 			const FText& GetTooltip() const override;
+			const FMetasoundFrontendVertexMetadata& GetMetadata() const override;
 
 			FConnectability CanConnectTo(const IInputController& InController) const override;
 			bool Connect(IInputController& InController) override;
@@ -200,6 +202,7 @@ namespace Metasound
 			// Input metadata
 			const FText& GetDisplayName() const override;
 			const FText& GetTooltip() const override;
+			const FMetasoundFrontendVertexMetadata& GetMetadata() const override;
 
 			// Owning node info
 			FGuid GetOwningNodeID() const override;
@@ -264,6 +267,8 @@ namespace Metasound
 			// Input metadata
 			const FText& GetDisplayName() const override;
 			const FText& GetTooltip() const override;
+			const FMetasoundFrontendVertexMetadata& GetMetadata() const override;
+
 			const FMetasoundFrontendLiteral* GetDefaultLiteral() const override;
 
 		protected:
@@ -305,6 +310,7 @@ namespace Metasound
 			// Input metadata
 			const FText& GetDisplayName() const override;
 			const FText& GetTooltip() const override;
+			const FMetasoundFrontendVertexMetadata& GetMetadata() const override;
 
 			FConnectability CanConnectTo(const IOutputController& InController) const override;
 			bool Connect(IOutputController& InController) override;
diff --git a/Engine/Plugins/Runtime/Metasound/Source/MetasoundFrontend/Public/MetasoundFrontendController.h b/Engine/Plugins/Runtime/Metasound/Source/MetasoundFrontend/Public/MetasoundFrontendController.h
--- a/Engine/Plugins/Runtime/Metasound/Source/MetasoundFrontend/Public/MetasoundFrontendController.h
+++ b/Engine/Plugins/Runtime/Metasound/Source/MetasoundFrontend/Public/MetasoundFrontendController.h
@@ -164,6 +164,9 @@ namespace Metasound
 			/** Returns the tooltip associated with this output. */
 			virtual const FText& GetTooltip() const = 0;
 
+			/** Returns all metadata associated with this output. */
+			virtual const FMetasoundFrontendVertexMetadata& GetMetadata() const = 0;
+
 			/** Returns the ID of the node which owns this output. */
 			virtual FGuid GetOwningNodeID() const = 0;
 			
@@ -235,6 +238,9 @@ namespace Metasound
 			/** Returns the data type name associated with this input. */
 			virtual const FText& GetTooltip() const = 0;
 
+			/** Returns all metadata associated with this input. */
+			virtual const FMetasoundFrontendVertexMetadata& GetMetadata() const = 0;
+
 			/** Returns the ID of the node which owns this output. */
 			virtual FGuid GetOwningNodeID() const = 0;
 			
diff --git a/Engine/Plugins/Runtime/Metasound/Source/MetasoundFrontend/Public/MetasoundFrontendDocument.h b/Engine/Plugins/Runtime/Metasound/Source/MetasoundFrontend/Public/MetasoundFrontendDocument.h
--- a/Engine/Plugins/Runtime/Metasound/Source/MetasoundFrontend/Public/MetasoundFrontendDocument.h
+++ b/Engine/Plugins/Runtime/Metasound/Source/MetasoundFrontend/Public/MetasoundFrontendDocument.h
@@ -473,15 +473,24 @@ struct FMetasoundFrontendInterfaceStyle
 	TArray<HandleType> SortDefaults(const TArray<HandleType>& InHandles) const
 	{
 		TArray<HandleType> SortedHandles = InHandles;
+
+		TMap<FGuid, int32> HandleIDToSortIndex;
 		for (int32 i = 0; i < DefaultSortOrder.Num(); ++i)
 		{
-			const int32 SortedIndex = DefaultSortOrder[i];
-			if (SortedHandles.IsValidIndex(SortedIndex) && InHandles.IsValidIndex(i))
+			if (ensure(InHandles.IsValidIndex(i)))
 			{
-				SortedHandles[SortedIndex] = InHandles[i];
+				const int32 SortIndex = DefaultSortOrder[i];
+				HandleIDToSortIndex.Add(InHandles[i]->GetID(), SortIndex);
 			}
 		}
 
+		SortedHandles.Sort([&](const HandleType& HandleA, const HandleType& HandleB)
+		{
+			const FGuid HandleAID = HandleA->GetID();
+			const FGuid HandleBID = HandleB->GetID();
+			return HandleIDToSortIndex[HandleAID] < HandleIDToSortIndex[HandleBID];
+		});
+
 		return SortedHandles;
 	}
 };
diff --git a/Engine/Plugins/Runtime/Metasound/Source/MetasoundGraphCore/Private/MetasoundVertex.cpp b/Engine/Plugins/Runtime/Metasound/Source/MetasoundGraphCore/Private/MetasoundVertex.cpp
--- a/Engine/Plugins/Runtime/Metasound/Source/MetasoundGraphCore/Private/MetasoundVertex.cpp
+++ b/Engine/Plugins/Runtime/Metasound/Source/MetasoundGraphCore/Private/MetasoundVertex.cpp
@@ -58,12 +58,12 @@ namespace Metasound
 		return VertexModel->CreateDefaultLiteral();
 	}
 
-	const FText& FInputDataVertex::GetDescription() const
+	const FDataVertexMetadata& FInputDataVertex::GetMetadata() const
 	{
-		return VertexModel->Description;
+		return VertexModel->VertexMetadata;
 	}
 
-	bool FInputDataVertex::IsReferenceOfSameType(const IDataReference& InReference) const 
+	bool FInputDataVertex::IsReferenceOfSameType(const IDataReference& InReference) const
 	{
 		return VertexModel->IsReferenceOfSameType(InReference);
 	}
@@ -147,17 +147,17 @@ namespace Metasound
 		return VertexModel->DataTypeName;
 	}
 
-	const FText& FOutputDataVertex::GetDescription() const
+	const FDataVertexMetadata& FOutputDataVertex::GetMetadata() const
 	{
-		return VertexModel->Description;
+		return VertexModel->VertexMetadata;
 	}
 
-	bool FOutputDataVertex::IsReferenceOfSameType(const IDataReference& InReference) const 
+	bool FOutputDataVertex::IsReferenceOfSameType(const IDataReference& InReference) const
 	{
 		return VertexModel->IsReferenceOfSameType(InReference);
 	}
 
-	const FName& FOutputDataVertex::GetVertexTypeName() const 
+	const FName& FOutputDataVertex::GetVertexTypeName() const
 	{
 		return VertexModel->GetVertexTypeName();
 	}
diff --git a/Engine/Plugins/Runtime/Metasound/Source/MetasoundGraphCore/Public/MetasoundVertex.h b/Engine/Plugins/Runtime/Metasound/Source/MetasoundGraphCore/Public/MetasoundVertex.h
--- a/Engine/Plugins/Runtime/Metasound/Source/MetasoundGraphCore/Public/MetasoundVertex.h
+++ b/Engine/Plugins/Runtime/Metasound/Source/MetasoundGraphCore/Public/MetasoundVertex.h
@@ -13,19 +13,22 @@
 namespace Metasound
 {
 	/** Key type for an FInputDataVertexColletion or 
-	 * FOutputDataVertexCollection. 
+	 * FOutputDataVertexCollection.
 	 */
-	typedef FString FVertexKey;
+	using FVertexKey = FString;
 
-	// Forward declare.
+	// Forward declarations
 	class FInputDataVertex;
-
-	// Forward declare.
 	class FOutputDataVertex;
-
-	// Forward declare
 	class FEnvironmentVertex;
 
+	// Vertex metadata
+	struct FDataVertexMetadata
+	{
+		FText Description;
+		bool bIsAdvancedDisplay = false;
+	};
+
 	/** FDataVertexModel
 	 *
 	 * An FDataVertexModel implements various vertex behaviors. 
@@ -45,12 +48,12 @@ namespace Metasound
 		 *
 		 * @InVertexName - Name of vertex.
 		 * @InDataTypeName - Name of data type.
-		 * @InDescription - Human readable vertex description.
+		 * @InMetadata - Metadata pertaining to given vertex.
 		 */
-		FDataVertexModel(const FString& InVertexName, const FName& InDataTypeName, const FText& InDescription)
+		FDataVertexModel(const FString& InVertexName, const FName& InDataTypeName, const FDataVertexMetadata& InMetadata)
 		:	VertexName(InVertexName)
 		,	DataTypeName(InDataTypeName)
-		,	Description(InDescription)
+		,	VertexMetadata(InMetadata)
 		{
 		}
 
@@ -62,8 +65,8 @@ namespace Metasound
 		/** Type name of data. */
 		const FName DataTypeName;
 
-		/** Description of the vertex. */
-		const FText Description;
+		/** Metadata associated with vertex. */
+		const FDataVertexMetadata VertexMetadata;
 
 		/** Test if a IDataReference matches the DataType of this vertex model.
 		 *
@@ -122,14 +125,27 @@ namespace Metasound
 	public:
 
 		FInputDataVertexModel(const FString& InVertexName, const FName& InDataTypeName, const FText& InDescription)
-			: FDataVertexModel(InVertexName, InDataTypeName, InDescription)
+			: FDataVertexModel(InVertexName, InDataTypeName, { InDescription })
+			, LiteralFactory(MakeUnique<TLiteralFactory<FLiteral::FNone>>(FLiteral::FNone{}))
+		{
+		}
+
+		FInputDataVertexModel(const FString& InVertexName, const FName& InDataTypeName, const FDataVertexMetadata& InMetadata)
+			: FDataVertexModel(InVertexName, InDataTypeName, InMetadata)
 			, LiteralFactory(MakeUnique<TLiteralFactory<FLiteral::FNone>>(FLiteral::FNone{}))
 		{
 		}
 
 		template<typename LiteralValueType>
 		FInputDataVertexModel(const FString& InVertexName, const FName& InDataTypeName, const FText& InDescription, const LiteralValueType& InLiteralValue)
-			: FDataVertexModel(InVertexName, InDataTypeName, InDescription)
+			: FDataVertexModel(InVertexName, InDataTypeName, { InDescription })
+			, LiteralFactory(MakeUnique<TLiteralFactory<LiteralValueType>>(InLiteralValue))
+		{
+		}
+
+		template<typename LiteralValueType>
+		FInputDataVertexModel(const FString& InVertexName, const FName& InDataTypeName, const FDataVertexMetadata& InMetadata, const LiteralValueType& InLiteralValue)
+			: FDataVertexModel(InVertexName, InDataTypeName, InMetadata)
 			, LiteralFactory(MakeUnique<TLiteralFactory<LiteralValueType>>(InLiteralValue))
 		{
 		}
@@ -208,7 +224,18 @@ namespace Metasound
 		 */
 		template<typename... ModelArgTypes>
 		TBaseVertexModel(const FString& InVertexName, const FText& InDescription, ModelArgTypes&&... ModelArgs)
-		:	VertexModelType(InVertexName, GetMetasoundDataTypeName<DataType>(), InDescription, Forward<ModelArgTypes>(ModelArgs)...)
+		:	VertexModelType(InVertexName, GetMetasoundDataTypeName<DataType>(), { InDescription }, Forward<ModelArgTypes>(ModelArgs)...)
+		{
+		}
+
+		/** TBaseVertexModel Constructor
+		 *
+		 * @InVertexName - Name of vertex.
+		 * @InVertexMetadata - Vertex metadata, used primarily for debugging or live edit.
+		 */
+		template<typename... ModelArgTypes>
+		TBaseVertexModel(const FString& InVertexName, const FDataVertexMetadata& InMetadata, ModelArgTypes&&... ModelArgs)
+		:	VertexModelType(InVertexName, GetMetasoundDataTypeName<DataType>(), InMetadata, Forward<ModelArgTypes>(ModelArgs)...)
 		{
 		}
 
@@ -270,6 +297,18 @@ namespace Metasound
 		{
 		}
 
+		/** TInputDataVertexModel Constructor
+		 *
+		 * @InVertexName - Name of vertex.
+		 * @InDescription - Human readable vertex description.
+		 * @InDefaultValue - Default Value of vertex
+		 */
+		template<typename LiteralValueType>
+		TInputDataVertexModel(const FString& InVertexName, const FDataVertexMetadata& InMetadata)
+		:	TBaseVertexModel<DataType, FInputDataVertexModel>(InVertexName, InMetadata)
+		{
+		}
+
 		TInputDataVertexModel(const TInputDataVertexModel& InOther) = default;
 
 		/** Return the vertex type name (not to be confused with the data type name).
@@ -333,7 +372,7 @@ namespace Metasound
 			const FName& GetDataTypeName() const;
 
 			/** Description of the vertex. */
-			const FText& GetDescription() const;
+			const FDataVertexMetadata& GetMetadata() const;
 
 			/** Default value of the vertex. */
 			FLiteral GetDefaultLiteral() const;
@@ -382,7 +421,7 @@ namespace Metasound
 		/** Create a clone of this TOutputDataVertexModel*/
 		virtual TUniquePtr<FOutputDataVertexModel> Clone() const override
 		{
-			return MakeUnique<TOutputDataVertexModel<DataType>>(this->VertexName, this->Description);
+			return MakeUnique<TOutputDataVertexModel<DataType>>(this->VertexName, this->VertexMetadata);
 		}
 	};
 
@@ -426,8 +465,8 @@ namespace Metasound
 			/** Type name of data reference. */
 			const FName& GetDataTypeName() const;
 
-			/** Description of the vertex. */
-			const FText& GetDescription() const;
+			/** Metadata of the vertex. */
+			const FDataVertexMetadata& GetMetadata() const;
 
 			/** Test if a IDataReference matches the DataType of this vertex.
 			 *
