Fixing trigger reroute node to make a permanent connection from input pin to output pin

Adding support for audio

diff --git a/Engine/Plugins/Runtime/Metasound/Source/MetasoundStandardNodes/Private/MetasoundTriggerRouteNode.cpp b/Engine/Plugins/Runtime/Metasound/Source/MetasoundStandardNodes/Private/MetasoundTriggerRouteNode.cpp
--- a/Engine/Plugins/Runtime/Metasound/Source/MetasoundStandardNodes/Private/MetasoundTriggerRouteNode.cpp
+++ b/Engine/Plugins/Runtime/Metasound/Source/MetasoundStandardNodes/Private/MetasoundTriggerRouteNode.cpp
@@ -14,6 +14,7 @@
 #include "MetasoundPrimitives.h"
 #include "MetasoundStandardNodesNames.h"
 #include "MetasoundStandardNodesCategories.h"
+#include "MetasoundAudioBuffer.h"
 
 #define LOCTEXT_NAMESPACE "MetasoundStandardNodes"
 
@@ -112,16 +113,33 @@ namespace Metasound
 	REGISTER_TRIGGER_ROUTE_NOTE(int32, 3)
 	REGISTER_TRIGGER_ROUTE_NOTE(int32, 4)
 	REGISTER_TRIGGER_ROUTE_NOTE(int32, 5)
+	REGISTER_TRIGGER_ROUTE_NOTE(int32, 6)
+	REGISTER_TRIGGER_ROUTE_NOTE(int32, 7)
+	REGISTER_TRIGGER_ROUTE_NOTE(int32, 8)
 
 	REGISTER_TRIGGER_ROUTE_NOTE(float, 2)
 	REGISTER_TRIGGER_ROUTE_NOTE(float, 3)
 	REGISTER_TRIGGER_ROUTE_NOTE(float, 4)
 	REGISTER_TRIGGER_ROUTE_NOTE(float, 5)
+	REGISTER_TRIGGER_ROUTE_NOTE(float, 6)
+	REGISTER_TRIGGER_ROUTE_NOTE(float, 7)
+	REGISTER_TRIGGER_ROUTE_NOTE(float, 8)
 
 	REGISTER_TRIGGER_ROUTE_NOTE(bool, 2)
 	REGISTER_TRIGGER_ROUTE_NOTE(bool, 3)
 	REGISTER_TRIGGER_ROUTE_NOTE(bool, 4)
 	REGISTER_TRIGGER_ROUTE_NOTE(bool, 5)
+	REGISTER_TRIGGER_ROUTE_NOTE(bool, 6)
+	REGISTER_TRIGGER_ROUTE_NOTE(bool, 7)
+	REGISTER_TRIGGER_ROUTE_NOTE(bool, 8)
+
+	REGISTER_TRIGGER_ROUTE_NOTE(FAudioBuffer, 2)
+	REGISTER_TRIGGER_ROUTE_NOTE(FAudioBuffer, 3)
+	REGISTER_TRIGGER_ROUTE_NOTE(FAudioBuffer, 4)
+	REGISTER_TRIGGER_ROUTE_NOTE(FAudioBuffer, 5)
+	REGISTER_TRIGGER_ROUTE_NOTE(FAudioBuffer, 6)
+	REGISTER_TRIGGER_ROUTE_NOTE(FAudioBuffer, 7)
+	REGISTER_TRIGGER_ROUTE_NOTE(FAudioBuffer, 8)
 }
 
 #undef LOCTEXT_NAMESPACE // MetasoundTriggerDelayNode
diff --git a/Engine/Plugins/Runtime/Metasound/Source/MetasoundStandardNodes/Public/MetasoundTriggerRouteNode.h b/Engine/Plugins/Runtime/Metasound/Source/MetasoundStandardNodes/Public/MetasoundTriggerRouteNode.h
--- a/Engine/Plugins/Runtime/Metasound/Source/MetasoundStandardNodes/Public/MetasoundTriggerRouteNode.h
+++ b/Engine/Plugins/Runtime/Metasound/Source/MetasoundStandardNodes/Public/MetasoundTriggerRouteNode.h
@@ -152,11 +152,13 @@ namespace Metasound
 					},
 					[this, i](int32 StartFrame, int32 EndFrame)
 					{
-						*OutputValue = *InputValues[i];
+						CurrentIndex = i;
 						OutputTrigger->TriggerFrame(StartFrame);
 					}
 				);
 			}
+
+			*OutputValue = *InputValues[CurrentIndex];
 		}
 
 	private:
@@ -167,6 +169,7 @@ namespace Metasound
 
 		FTriggerWriteRef OutputTrigger;
 		TDataWriteReference<ValueType> OutputValue;
+		int32 CurrentIndex = 0;
 	};
 
 	/** TTriggerRouteNode
