Remove spammy ensure

diff --git a/Engine/Plugins/Runtime/Metasound/Source/MetasoundFrontend/Public/MetasoundFrontendDocument.h b/Engine/Plugins/Runtime/Metasound/Source/MetasoundFrontend/Public/MetasoundFrontendDocument.h
--- a/Engine/Plugins/Runtime/Metasound/Source/MetasoundFrontend/Public/MetasoundFrontendDocument.h
+++ b/Engine/Plugins/Runtime/Metasound/Source/MetasoundFrontend/Public/MetasoundFrontendDocument.h
@@ -477,7 +477,7 @@ struct FMetasoundFrontendInterfaceStyle
 		TMap<FGuid, int32> HandleIDToSortIndex;
 		for (int32 i = 0; i < DefaultSortOrder.Num(); ++i)
 		{
-			if (ensure(InHandles.IsValidIndex(i)))
+			if (InHandles.IsValidIndex(i))
 			{
 				const int32 SortIndex = DefaultSortOrder[i];
 				HandleIDToSortIndex.Add(InHandles[i]->GetID(), SortIndex);
