Making LFOs hve built in map range

diff --git a/Engine/Plugins/Runtime/Metasound/Source/MetasoundStandardNodes/Private/MetasoundOscillatorNodes.cpp b/Engine/Plugins/Runtime/Metasound/Source/MetasoundStandardNodes/Private/MetasoundOscillatorNodes.cpp
--- a/Engine/Plugins/Runtime/Metasound/Source/MetasoundStandardNodes/Private/MetasoundOscillatorNodes.cpp
+++ b/Engine/Plugins/Runtime/Metasound/Source/MetasoundStandardNodes/Private/MetasoundOscillatorNodes.cpp
@@ -328,7 +328,7 @@ namespace Metasound
 		};
 		struct FSquarePolysmoothGenerator
 		{
-			FORCEINLINE float operator()(float InPhase, float InPhaseDelta, const FGeneratorArgs& InArgs)
+			float operator()(float InPhase, float InPhaseDelta, const FGeneratorArgs& InArgs)
 			{
 				// Taken piecemeal from Osc.cpp
 				// Lots of branches.
@@ -384,7 +384,7 @@ namespace Metasound
 			float PrevPhase = -1.0f;
 			float DPW_z1 = 0.0f;
 
-			FORCEINLINE float operator()(float InPhase, float InPhaseDelta, const FGeneratorArgs& InArgs)
+			float operator()(float InPhase, float InPhaseDelta, const FGeneratorArgs& InArgs)
 			{
 				// if the oscillator phase wrapped, our previous phase will be higher than current phase
 				if (PrevPhase > InPhase)
@@ -426,6 +426,9 @@ namespace Metasound
 		static constexpr const TCHAR* GlideFactorPinName = TEXT("Glide");
 		static constexpr const TCHAR* AudioOutPinName = TEXT("Audio");
 		static constexpr const TCHAR* BiPolarPinName = TEXT("Bi Polar");
+		static constexpr const TCHAR* MinOutputValuePinName = TEXT("Min Value");
+		static constexpr const TCHAR* MaxOutputValuePinName = TEXT("Max Value");
+
 
 		// Common to all Oscillators.
 		static FVertexInterface GetCommmonVertexInterface()
@@ -1206,12 +1209,13 @@ namespace Metasound
 			static const FVertexInterface Interface
 			{
 				FInputVertexInterface{
-					TInputDataVertexModel<bool>(FOscilatorFactoryBase::BiPolarPinName, LOCTEXT("LfoBiPolarDescription", "If the output is bipolar (-1..1) or unipolar (0..1)"), true),
 					TInputDataVertexModel<float>(FOscilatorFactoryBase::BaseFrequencyPinName, LOCTEXT("LfoFrequencyDescription", "Frequency of LFO (Hz), clamped at blockrate"), 5.f),
+					TInputDataVertexModel<FEnumLfoWaveshapeType>(WaveshapePinName, LOCTEXT("LfoShapeDescription", "Waveshape of the LFO")),
+					TInputDataVertexModel<float>(FOscilatorFactoryBase::MinOutputValuePinName, LOCTEXT("LfoMinValueDescription", "The minimum output value."), -1.0f),
+					TInputDataVertexModel<float>(FOscilatorFactoryBase::MaxOutputValuePinName, LOCTEXT("LfoMaxValueDescription", "The maximum output value."), 1.0f),
 					TInputDataVertexModel<FTrigger>(FOscilatorFactoryBase::PhaseResetPinName, LOCTEXT("LfoPhaseResetDescription", "Phase Reset (block rate only)")),
 					TInputDataVertexModel<float>(FOscilatorFactoryBase::PhaseOffsetPinName, LOCTEXT("LfoPhaseOffsetDescription", "Phase Offset In Degrees (0..360)"), 0.f),
-					TInputDataVertexModel<float>(FSquareOscilatorNode::FFactory::PulseWidthPinName, LOCTEXT("LfoPulseWidthDescription", "Pulse Width (0..1)"), 0.25f),
-					TInputDataVertexModel<FEnumLfoWaveshapeType>(WaveshapePinName, LOCTEXT("LfoShapeDescription", "Waveshape of the LFO")),
+					TInputDataVertexModel<float>(FSquareOscilatorNode::FFactory::PulseWidthPinName, LOCTEXT("LfoPulseWidthDescription", "Pulse Width (0..1)"), 0.5f)
 				},
 				FOutputVertexInterface{
 					TOutputDataVertexModel<float>(LfoOutPinName, LOCTEXT("LfoOutputDescription", "Output of the LFO (blockrate)"))
@@ -1251,24 +1255,26 @@ namespace Metasound
 				  Settings
 				, InputCol.GetDataReadReferenceOrConstructWithVertexDefault<float>(InputInterface, FOscilatorFactoryBase::BaseFrequencyPinName, Settings)
 				, InputCol.GetDataReadReferenceOrConstruct<FEnumLfoWaveshapeType>(WaveshapePinName)
+				, InputCol.GetDataReadReferenceOrConstructWithVertexDefault<float>(InputInterface, FOscilatorFactoryBase::MinOutputValuePinName, Settings)
+				, InputCol.GetDataReadReferenceOrConstructWithVertexDefault<float>(InputInterface, FOscilatorFactoryBase::MaxOutputValuePinName, Settings)
 				, InputCol.GetDataReadReferenceOrConstruct<FTrigger>(FOscilatorFactoryBase::PhaseResetPinName, Settings)
 				, InputCol.GetDataReadReferenceOrConstructWithVertexDefault<float>(InputInterface,FOscilatorFactoryBase::PhaseOffsetPinName, Settings)
-				, InputCol.GetDataReadReferenceOrConstructWithVertexDefault<bool>(InputInterface,FOscilatorFactoryBase::BiPolarPinName, Settings)
 				, InputCol.GetDataReadReferenceOrConstructWithVertexDefault<float>(InputInterface,FSquareOscilatorNode::FFactory::PulseWidthPinName, Settings)
 			);
 		}
 
-		FLfoOperator(const FOperatorSettings& InSettings, FFloatReadRef&& InFrequency, FEnumLfoWaveshapeTypeReadRef&& InType, 
-			FTriggerReadRef&& InPhaseReset, FFloatReadRef&& InPhaseOffset, FBoolReadRef&& InBiPolar, FFloatReadRef&& InPulseWidth)
-			: BlockRate{ InSettings.GetActualBlockRate() }
-			, Phase{ 0.0f }
-			, Frequency{ MoveTemp(InFrequency) }
-			, Waveshape{ MoveTemp(InType) }
-			, PhaseReset{ MoveTemp(InPhaseReset) }
-			, PhaseOffset{ MoveTemp(InPhaseOffset) }
-			, BiPolar{ MoveTemp(InBiPolar) }
-			, PulseWidth{ MoveTemp(InPulseWidth) }
-			, Output{ FFloatWriteRef::CreateNew(0.f) }
+		FLfoOperator(const FOperatorSettings& InSettings, FFloatReadRef&& InFrequency, FEnumLfoWaveshapeTypeReadRef&& InType, FFloatReadRef&& InMinValue, FFloatReadRef&& InMaxValue,
+			FTriggerReadRef&& InPhaseReset, FFloatReadRef&& InPhaseOffset, FFloatReadRef&& InPulseWidth)
+			: BlockRate{InSettings.GetActualBlockRate()}
+			, Phase{0.0f}
+			, Frequency{MoveTemp(InFrequency)}
+			, Waveshape{MoveTemp(InType)}
+			, MinValue{MoveTemp(InMinValue)}
+			, MaxValue{MoveTemp(InMaxValue)}
+			, PhaseReset{MoveTemp(InPhaseReset)}
+			, PhaseOffset{MoveTemp(InPhaseOffset)}
+			, PulseWidth{MoveTemp(InPulseWidth)}
+			, Output{FFloatWriteRef::CreateNew(0.f)}
 		{
 			ResetPhase();
 		}
@@ -1278,9 +1284,10 @@ namespace Metasound
 			FDataReferenceCollection InputDataReferences;
 			InputDataReferences.AddDataReadReference(FOscilatorFactoryBase::BaseFrequencyPinName, Frequency);
 			InputDataReferences.AddDataReadReference(WaveshapePinName, Waveshape);
+			InputDataReferences.AddDataReadReference(FOscilatorFactoryBase::MinOutputValuePinName, MinValue);
+			InputDataReferences.AddDataReadReference(FOscilatorFactoryBase::MaxOutputValuePinName, MaxValue);
 			InputDataReferences.AddDataReadReference(FOscilatorFactoryBase::PhaseOffsetPinName, PhaseOffset);
 			InputDataReferences.AddDataReadReference(FOscilatorFactoryBase::PhaseResetPinName, PhaseReset);
-			InputDataReferences.AddDataReadReference(FOscilatorFactoryBase::BiPolarPinName, BiPolar);
 			InputDataReferences.AddDataReadReference(FSquareOscilatorNode::FFactory::PulseWidthPinName, PulseWidth);
 			return InputDataReferences;
 		}
@@ -1308,7 +1315,7 @@ namespace Metasound
 			const float DeltaPhase = ClampedFreq * (1.f / BlockRate);
 			const float ClampPulseWidth = FMath::Clamp(*PulseWidth, 0.f, 1.f);
 
-			FGeneratorArgs Args{ BlockRate, ClampedFreq, ClampPulseWidth };
+			FGeneratorArgs Args{ BlockRate, ClampedFreq, 0.0f, ClampPulseWidth };
 
 			// We are not sample accurate.
 			if (PhaseReset->IsTriggeredInBlock()) 
@@ -1324,26 +1331,22 @@ namespace Metasound
 			{
 				case ELfoWaveshapeType::Sine:
 				{
-					FBhaskaraGenerator Generator;
-					Value = Generator(Phase, DeltaPhase, Args);
+					Value = SineGenerator(Phase, DeltaPhase, Args);
 					break;
 				}
 				case ELfoWaveshapeType::Saw:
 				{
-					FSawPolySmoothGenerator Generator;
-					Value = Generator(Phase, DeltaPhase, Args);
+					Value = SawGenerator(Phase, DeltaPhase, Args);
 					break;
 				}
 				case ELfoWaveshapeType::Triangle:
 				{
-					FTriangleGenerator Generator;
-					Value = Generator(Phase, DeltaPhase, Args);
+					Value = TriangleGenerator(Phase, DeltaPhase, Args);
 					break;
 				}
 				case ELfoWaveshapeType::Square:
 				{
-					FSquarePolysmoothGenerator Generator;
-					Value = Generator(Phase, DeltaPhase, Args);
+					Value = SquareGenerator(Phase, DeltaPhase, Args);
 					break;
 				}
 				default:
@@ -1353,10 +1356,7 @@ namespace Metasound
 				}
 			}
 
-			if (!*BiPolar)
-			{
-				Value = Audio::GetUnipolar(Value);
-			}
+			Value = FMath::GetMappedRangeValueClamped({ -1.0f, 1.0f }, { *MinValue, *MaxValue }, Value);
 
 			*Output = Value;
 
@@ -1369,11 +1369,18 @@ namespace Metasound
 		Generators::FWrapPhase Wrap;
 		FFloatReadRef Frequency;
 		FEnumLfoWaveshapeTypeReadRef Waveshape;
+		FFloatReadRef MinValue;
+		FFloatReadRef MaxValue;
 		FTriggerReadRef PhaseReset;
 		FFloatReadRef PhaseOffset;
-		FBoolReadRef BiPolar;
 		FFloatReadRef PulseWidth;
 		FFloatWriteRef Output;
+
+		Generators::FSawPolySmoothGenerator SawGenerator;
+		Generators::FSineWaveTableGenerator SineGenerator;
+		Generators::FTriangleGenerator TriangleGenerator;
+		Generators::FSquarePolysmoothGenerator SquareGenerator;
+
 	};
 
 	FLfoNode::FLfoNode(const FNodeInitData& InInitData)
