Refactor for Debug Log Node (making it headerfile-less)

diff --git a/Engine/Plugins/Runtime/Metasound/Source/MetasoundStandardNodes/Private/MetasoundDebugLogNode.cpp b/Engine/Plugins/Runtime/Metasound/Source/MetasoundStandardNodes/Private/MetasoundDebugLogNode.cpp
--- a/Engine/Plugins/Runtime/Metasound/Source/MetasoundStandardNodes/Private/MetasoundDebugLogNode.cpp
+++ b/Engine/Plugins/Runtime/Metasound/Source/MetasoundStandardNodes/Private/MetasoundDebugLogNode.cpp
@@ -1,21 +1,24 @@
 // Copyright Epic Games, Inc. All Rights Reserved.
 
-#include "MetasoundDebugLogNode.h"
-
 #include "MetasoundNodeRegistrationMacro.h"
 #include "MetasoundStandardNodesCategories.h"
+#include "MetasoundFacade.h"
+#include "MetasoundStandardNodesNames.h"
+#include "MetasoundTrigger.h"
+#include "MetasoundPrimitives.h"
 
 #define LOCTEXT_NAMESPACE "MetasoundStandardNodes"
 
 namespace Metasound
 {
-	namespace MetasoundDebugLogNodePrivate
+	namespace MetasoundPrintLogNodePrivate
 	{
+		//Creates Metadata for the Print Log Node
 		FNodeClassMetadata CreateNodeClassMetadata(const FName& InDataTypeName, const FName& InOperatorName, const FText& InDisplayName, const FText& InDescription, const FVertexInterface& InDefaultInterface)
 		{
 			FNodeClassMetadata Metadata
 			{
-				FNodeClassName{FName("DebugLog"), InOperatorName, InDataTypeName},
+				FNodeClassName{FName("Print Log"), InOperatorName, InDataTypeName},
 				1, // Major Version
 				0, // Minor Version
 				InDisplayName,
@@ -24,7 +27,7 @@ namespace Metasound
 				PluginNodeMissingPrompt,
 				InDefaultInterface,
 				{ StandardNodes::DebugUtils },
-				{TEXT("DebugLog")},
+				{TEXT("Print Log")},
 				FNodeDisplayStyle{}
 			};
 
@@ -32,7 +35,8 @@ namespace Metasound
 		}
 	}
 
-	namespace DebugLogVertexNames
+	//Getters for the name of each parameter used in Print Log
+	namespace PrintLogVertexNames
 	{
 		const FString& GetInputTriggerName()
 		{
@@ -40,36 +44,151 @@ namespace Metasound
 			return Name;
 		}
 
-		const FString& GetLabelDebugLogName()
+		const FString& GetLabelPrintLogName()
 		{
 			static const FString Name = TEXT("Label");
 			return Name;
 		}
 
-		const FString& GetToLogDebugLogName()
+		const FString& GetToLogPrintLogName()
 		{
 			static const FString Name = TEXT("Value To Log");
 			return Name;
 		}
 
-		const FString& GetOutputDebugLogName()
+		const FString& GetOutputPrintLogName()
 		{
 			static const FString Name = TEXT("Was Successful");
 			return Name;
 		}
 	}
 
-	using FDebugLogNodeInt32 = TDebugLogNode<int32>;
-	METASOUND_REGISTER_NODE(FDebugLogNodeInt32)
 
-	using FDebugLogNodeFloat = TDebugLogNode<float>;
-	METASOUND_REGISTER_NODE(FDebugLogNodeFloat)
+	template<typename PrintLogType>
+	class TPrintLogOperator : public TExecutableOperator<TPrintLogOperator<PrintLogType>>
+	{
+		public:
+			using FArrayDataReadReference = TDataReadReference<PrintLogType>;
+
+			static const FVertexInterface& GetDefaultInterface()
+			{
+				static const FVertexInterface DefaultInterface(
+					FInputVertexInterface(
+
+						TInputDataVertexModel<FTrigger>(PrintLogVertexNames::GetInputTriggerName(), LOCTEXT("PrintLogTrigger", "Trigger to write the set value to the log.")),
+						TInputDataVertexModel<FString>(PrintLogVertexNames::GetLabelPrintLogName(), LOCTEXT("PrintLogLabel", "The label to attach to the value that will be logged")),
+						TInputDataVertexModel<PrintLogType>(PrintLogVertexNames::GetToLogPrintLogName(), LOCTEXT("PrintLogValueToLog", "The value to record to the log when triggered"))
+					),
+					FOutputVertexInterface(
+					)
+				);
+
+				return DefaultInterface;
+			}
+
+			static const FNodeClassMetadata& GetNodeInfo()
+			{
+				auto CreateNodeClassMetadata = []() -> FNodeClassMetadata
+				{
+					FName DataTypeName = GetMetasoundDataTypeName<PrintLogType>();
+					FName OperatorName = TEXT("Print Log");
+					FText NodeDisplayName = FText::Format(LOCTEXT("PrintLogDisplayNamePattern", "Print Log ({0})"), FText::FromString(GetMetasoundDataTypeString<PrintLogType>()));
+					FText NodeDescription = LOCTEXT("PrintLogOpDescription", "Used to record values to the log, on trigger");
+					FVertexInterface NodeInterface = GetDefaultInterface();
+
+					return MetasoundPrintLogNodePrivate::CreateNodeClassMetadata(DataTypeName, OperatorName, NodeDisplayName, NodeDescription, NodeInterface);
+				};
+
+				static const FNodeClassMetadata Metadata = CreateNodeClassMetadata();
+				return Metadata;
+			}
+
+			static TUniquePtr<IOperator> CreateOperator(const FCreateOperatorParams& InParams, TArray<TUniquePtr<IOperatorBuildError>>& OutErrors)
+			{
+				const FInputVertexInterface& InputInterface = InParams.Node.GetVertexInterface().GetInputInterface();
+				const FDataReferenceCollection& InputCollection = InParams.InputDataReferences;
+
+				FTriggerReadRef Trigger = InputCollection.GetDataReadReferenceOrConstruct<FTrigger>(PrintLogVertexNames::GetInputTriggerName(), InParams.OperatorSettings);
+
+				TDataReadReference<FString> Label = InputCollection.GetDataReadReferenceOrConstructWithVertexDefault<FString>(InputInterface, PrintLogVertexNames::GetLabelPrintLogName(), InParams.OperatorSettings);
+				TDataReadReference<PrintLogType> ValueToLog = InputCollection.GetDataReadReferenceOrConstructWithVertexDefault<PrintLogType>(InputInterface, PrintLogVertexNames::GetToLogPrintLogName(), InParams.OperatorSettings);
+
+				return MakeUnique<TPrintLogOperator<PrintLogType>>(InParams.OperatorSettings, Trigger, Label, ValueToLog);
+			}
+
+
+			TPrintLogOperator(const FOperatorSettings& InSettings, TDataReadReference<FTrigger> InTrigger, TDataReadReference<FString> InLabelPrintLog, TDataReadReference<PrintLogType> InValueToLogPrintLog)
+				: Trigger(InTrigger)
+				, Label(InLabelPrintLog)
+				, ValueToLog(InValueToLogPrintLog)
+			{
+			}
+
+			virtual ~TPrintLogOperator() = default;
+
+
+			virtual FDataReferenceCollection GetInputs() const override
+			{
+				FDataReferenceCollection Inputs;
+
+				Inputs.AddDataReadReference(PrintLogVertexNames::GetInputTriggerName(), Trigger);
+				Inputs.AddDataReadReference(PrintLogVertexNames::GetLabelPrintLogName(), Label);
+				Inputs.AddDataReadReference(PrintLogVertexNames::GetToLogPrintLogName(), ValueToLog);
+
+				return Inputs;
+			}
+
+			virtual FDataReferenceCollection GetOutputs() const override
+			{
+				FDataReferenceCollection Outputs;
+
+				return Outputs;
+			}
+
+			void Execute()
+			{
+				if (*Trigger)
+				{
+					UE_LOG(LogMetasound, Display, TEXT("%s %s"), *(*Label), *LexToString(*ValueToLog));
+				}
+			}
+
+		private:
+
+			TDataReadReference<FTrigger> Trigger;
+			TDataReadReference<FString> Label;
+			TDataReadReference<PrintLogType> ValueToLog;
+	};
+
+	/** TPrintLogNode
+	 *
+	 *  Records a value to the log when triggered
+	 */
+	template<typename PrintLogType>
+	class METASOUNDSTANDARDNODES_API TPrintLogNode : public FNodeFacade
+	{
+		public:
+			/**
+			 * Constructor used by the Metasound Frontend.
+			 */
+			TPrintLogNode(const FNodeInitData& InInitData)
+				: FNodeFacade(InInitData.InstanceName, InInitData.InstanceID, TFacadeOperatorClass<TPrintLogOperator<PrintLogType>>())
+			{}
+
+			virtual ~TPrintLogNode() = default;
+	};
+
+	using FPrintLogNodeInt32 = TPrintLogNode<int32>;
+	METASOUND_REGISTER_NODE(FPrintLogNodeInt32)
+
+	using FPrintLogNodeFloat = TPrintLogNode<float>;
+	METASOUND_REGISTER_NODE(FPrintLogNodeFloat)
 
-	using FDebugLogNodeBool = TDebugLogNode<bool>;
-	METASOUND_REGISTER_NODE(FDebugLogNodeBool)
+	using FPrintLogNodeBool = TPrintLogNode<bool>;
+	METASOUND_REGISTER_NODE(FPrintLogNodeBool)
 
-	using FDebugLogNodeString = TDebugLogNode<FString>;
-	METASOUND_REGISTER_NODE(FDebugLogNodeString)
+	using FPrintLogNodeString = TPrintLogNode<FString>;
+	METASOUND_REGISTER_NODE(FPrintLogNodeString)
 }
 
 #undef LOCTEXT_NAMESPACE
\ No newline at end of file
diff --git a/Engine/Plugins/Runtime/Metasound/Source/MetasoundStandardNodes/Public/MetasoundDebugLogNode.h b/Engine/Plugins/Runtime/Metasound/Source/MetasoundStandardNodes/Public/MetasoundDebugLogNode.h
deleted file mode 100644
