Fixing up tooltips on trigger route node.

Removing the init node input and instead treating input value 0/trigger 0 as the init value. Is a cleaner UX.

diff --git a/Engine/Plugins/Runtime/Metasound/Source/MetasoundStandardNodes/Private/MetasoundTriggerRouteNode.cpp b/Engine/Plugins/Runtime/Metasound/Source/MetasoundStandardNodes/Private/MetasoundTriggerRouteNode.cpp
--- a/Engine/Plugins/Runtime/Metasound/Source/MetasoundStandardNodes/Private/MetasoundTriggerRouteNode.cpp
+++ b/Engine/Plugins/Runtime/Metasound/Source/MetasoundStandardNodes/Private/MetasoundTriggerRouteNode.cpp
@@ -1,6 +1,4 @@
 // Copyright Epic Games, Inc. All Rights Reserved.
-#include "MetasoundTriggerRouteNode.h"
-
 #include "CoreMinimal.h"
 
 #include "MetasoundBuilderInterface.h"
@@ -15,6 +13,7 @@
 #include "MetasoundStandardNodesNames.h"
 #include "MetasoundStandardNodesCategories.h"
 #include "MetasoundAudioBuffer.h"
+#include "MetasoundTrigger.h"
 
 #define LOCTEXT_NAMESPACE "MetasoundStandardNodes"
 
@@ -24,6 +23,59 @@
 
 namespace Metasound
 {
+	namespace TriggerRouteVertexNames
+	{
+		const FString GetInputTriggerName(uint32 InIndex)
+		{
+			return FText::Format(LOCTEXT("TriggerRouteTriggerInputName", "Set {0}"), InIndex).ToString();
+		}
+
+		const FText GetInputTriggerDescription(uint32 InIndex)
+		{
+			if (InIndex == 0)
+			{
+				return FText::Format(LOCTEXT("TriggerRouteInputTriggerDescInit", "The input trigger {0} to cause the corresponding input value to route to the output value. This trigger is the default output."), InIndex);
+			}
+			return FText::Format(LOCTEXT("TriggerRouteInputTriggerDesc", "The input trigger {0} to cause the corresponding input value to route to the output value."), InIndex);
+		}
+
+		const FString GetInputValueName(uint32 InIndex)
+		{
+			return FText::Format(LOCTEXT("TriggerRouteValueInputName", "Value {0}"), InIndex).ToString();
+
+		}
+
+		const FText GetInputValueDescription(uint32 InIndex)
+		{
+			return FText::Format(LOCTEXT("TriggerRouteValueDesc", "The input value ({0}) to route to the output when triggered by Set {0}."), InIndex);
+		}
+
+		const FString& GetOutputTriggerName()
+		{
+			static const FString Name = TEXT("On Set");
+			return Name;
+		}
+
+		const FText& GetOutputTriggerDescription()
+		{
+			static const FText Desc = LOCTEXT("TriggerRouteOnSetDesc", "Triggered when any of the input triggers are set.");
+			return Desc;
+		}
+
+		const FString& GetOutputValueName()
+		{
+			static const FString Name = TEXT("Value");
+			return Name;
+		}
+
+		const FText& GetOutputValueDescription()
+		{
+			static const FText Desc = LOCTEXT("TriggerRouteOutputValueDesc", "The output value set by the input triggers.");
+			return Desc;
+		}
+	}
+
+
 	namespace MetasoundTriggerRouteNodePrivate
 	{
 		FNodeClassMetadata CreateNodeClassMetadata(const FName& InDataTypeName, const FName& InOperatorName, const FText& InDisplayName, const FText& InDescription, const FVertexInterface& InDefaultInterface)
@@ -47,67 +99,156 @@
 		}
 	}
 
+	template<typename ValueType, uint32 NumInputs>
+	class TTriggerRouteOperator : public TExecutableOperator<TTriggerRouteOperator<ValueType, NumInputs>>
+	{
+	public:
+		static const FVertexInterface& GetDefaultInterface()
+		{
+			using namespace TriggerRouteVertexNames;
 
-	namespace TriggerRouteVertexNames
-	{
-		const FString& GetInputInitValueName()
-		{
-			static const FString Name = TEXT("Init");
-			return Name;
+			auto CreateDefaultInterface = []() -> FVertexInterface
+			{
+				FInputVertexInterface InputInterface;
+				for (uint32 i = 0; i < NumInputs; ++i)
+				{
+					InputInterface.Add(TInputDataVertexModel<FTrigger>(GetInputTriggerName(i), GetInputTriggerDescription(i)));
+					InputInterface.Add(TInputDataVertexModel<ValueType>(GetInputValueName(i), GetInputValueDescription(i)));
+				}
+
+				FOutputVertexInterface OutputInterface;
+				OutputInterface.Add(TOutputDataVertexModel<FTrigger>(GetOutputTriggerName(), GetOutputTriggerDescription()));
+				OutputInterface.Add(TOutputDataVertexModel<ValueType>(GetOutputValueName(), GetOutputValueDescription()));
+				return FVertexInterface(InputInterface, OutputInterface);
+			};
+
+			static const FVertexInterface DefaultInterface = CreateDefaultInterface();
+			return DefaultInterface;
 		}
 
-		const FText& GetInputInitValueDescription()
+		static const FNodeClassMetadata& GetNodeInfo()
 		{
-			static const FText Desc = LOCTEXT("TriggerRouteInputInitValueDesc", "The value to initialize the output to before any inputs are triggered.");
-			return Desc;
-		}
-		
-		const FString GetInputTriggerName(uint32 InIndex)
-		{
-			return FText::Format(LOCTEXT("TriggerRouteTriggerInputName", "Set {0}"), InIndex).ToString();
+			auto CreateNodeClassMetadata = []() -> FNodeClassMetadata
+			{
+				FName DataTypeName = GetMetasoundDataTypeName<ValueType>();
+				FName OperatorName = *FString::Printf(TEXT("Trigger Route (%s, %d)"), *DataTypeName.ToString(), NumInputs);
+				FText NodeDisplayName = FText::Format(LOCTEXT("TriggerRouteDisplayNamePattern", "Trigger Route ({0}, {1})"), FText::FromString(GetMetasoundDataTypeString<ValueType>()), NumInputs);
+				FText NodeDescription = LOCTEXT("TriggerRouteDescription", "Allows routing different values to the same output pin depending on trigger inputs.");
+				FVertexInterface NodeInterface = GetDefaultInterface();
+
+				return MetasoundTriggerRouteNodePrivate::CreateNodeClassMetadata(DataTypeName, OperatorName, NodeDisplayName, NodeDescription, NodeInterface);
+			};
+
+			static const FNodeClassMetadata Metadata = CreateNodeClassMetadata();
+			return Metadata;
 		}
 
-		const FText GetInputTriggerDescription(uint32 InIndex)
+		static TUniquePtr<IOperator> CreateOperator(const FCreateOperatorParams& InParams, TArray<TUniquePtr<IOperatorBuildError>>& OutErrors)
 		{
-			return FText::Format(LOCTEXT("TriggerRouteInputTriggerDesc", "The input trigger {0} to cause the corresponding input value to route to the output value."), InIndex);
+			using namespace TriggerRouteVertexNames;
+
+			const FInputVertexInterface& InputInterface = InParams.Node.GetVertexInterface().GetInputInterface();
+			const FDataReferenceCollection& InputCollection = InParams.InputDataReferences;
+
+			TArray<FTriggerReadRef> InputTriggers;
+			TArray<TDataReadReference<ValueType>> InputValues;
+
+			for (uint32 i = 0; i < NumInputs; ++i)
+			{
+				InputTriggers.Add(InputCollection.GetDataReadReferenceOrConstruct<FTrigger>(GetInputTriggerName(i), InParams.OperatorSettings));
+				InputValues.Add(InputCollection.GetDataReadReferenceOrConstructWithVertexDefault<ValueType>(InputInterface, GetInputValueName(i), InParams.OperatorSettings));
+			}
+
+			return MakeUnique<TTriggerRouteOperator<ValueType, NumInputs>>(InParams.OperatorSettings, MoveTemp(InputTriggers), MoveTemp(InputValues));
 		}
 
-		const FString GetInputValueName(uint32 InIndex)
+		TTriggerRouteOperator(const FOperatorSettings& InSettings, TArray<FTriggerReadRef>&& InInputTriggers, TArray<TDataReadReference<ValueType>>&& InInputValues)
+			: InputTriggers(MoveTemp(InInputTriggers))
+			, InputValues(MoveTemp(InInputValues))
+			, OutputTrigger(FTriggerWriteRef::CreateNew(InSettings))
+			, OutputValue(TDataWriteReferenceFactory<ValueType>::CreateAny(InSettings))
 		{
-			return FText::Format(LOCTEXT("TriggerRouteValueInputName", "Value {0}"), InIndex).ToString();
+			check(InputValues.Num() > 0)
+			CurrentIndex = 0;
+			*OutputValue = *InputValues[CurrentIndex];
 		}
 
-		const FText GetInputValueDescription(uint32 InIndex)
+		virtual ~TTriggerRouteOperator() = default;
+
+
+		virtual FDataReferenceCollection GetInputs() const override
 		{
-			static const FText Desc = FText::Format(LOCTEXT("TriggerRouteValueDesc", "The input value ({0}) to route to the output when triggered by Set {0}."), InIndex);
-			return Desc;
+			using namespace TriggerRouteVertexNames;
+
+			FDataReferenceCollection Inputs;
+			for (uint32 i = 0; i < NumInputs; ++i)
+			{
+				Inputs.AddDataReadReference(GetInputTriggerName(i), InputTriggers[i]);
+				Inputs.AddDataReadReference(GetInputValueName(i), InputValues[i]);
+			}
+
+			return Inputs;
 		}
 
-		const FString& GetOutputTriggerName()
+		virtual FDataReferenceCollection GetOutputs() const override
 		{
-			static const FString Name = TEXT("On Set");
-			return Name;
+			using namespace TriggerRouteVertexNames;
+
+			FDataReferenceCollection Outputs;
+			Outputs.AddDataReadReference(GetOutputTriggerName(), OutputTrigger);
+			Outputs.AddDataReadReference(GetOutputValueName(), OutputValue);
+
+			return Outputs;
 		}
 
-		const FText& GetOutputTriggerDescription()
+		void Execute()
 		{
-			static const FText Desc = LOCTEXT("TriggerRouteOnSetDesc", "Triggered when any of the input triggers are set.");
-			return Desc;
+			OutputTrigger->AdvanceBlock();
+
+			for (uint32 i = 0; i < NumInputs; ++i)
+			{
+				InputTriggers[i]->ExecuteBlock(
+					[&](int32 StartFrame, int32 EndFrame)
+					{
+					},
+					[this, i](int32 StartFrame, int32 EndFrame)
+					{
+						CurrentIndex = i;
+						OutputTrigger->TriggerFrame(StartFrame);
+					}
+					);
+			}
+
+			*OutputValue = *InputValues[CurrentIndex];
 		}
 
-		const FString& GetOutputValueName()
-		{
-			static const FString Name = TEXT("Value");
-			return Name;
-		}
+	private:
 
-		const FText& GetOutputValueDescription()
-		{
-			static const FText Desc = LOCTEXT("TriggerRouteOutputValueDesc", "The output value set by the input triggers.");
-			return Desc;
-		}
+		TArray<FTriggerReadRef> InputTriggers;
+		TArray<TDataReadReference<ValueType>> InputValues;
 
-	}
+		FTriggerWriteRef OutputTrigger;
+		TDataWriteReference<ValueType> OutputValue;
+		int32 CurrentIndex = 0;
+	};
 
+	/** TTriggerRouteNode
+	 *
+	 *  Routes values from multiple input pins to a single output pin based on trigger inputs.
+	 */
+	template<typename ValueType, uint32 NumInputs>
+	class METASOUNDSTANDARDNODES_API TTriggerRouteNode : public FNodeFacade
+	{
+	public:
+		/**
+		 * Constructor used by the Metasound Frontend.
+		 */
+		TTriggerRouteNode(const FNodeInitData& InInitData)
+			: FNodeFacade(InInitData.InstanceName, InInitData.InstanceID, TFacadeOperatorClass<TTriggerRouteOperator<ValueType, NumInputs>>())
+		{}
+
+		virtual ~TTriggerRouteNode() = default;
+	};
+
 	REGISTER_TRIGGER_ROUTE_NOTE(int32, 2)
 	REGISTER_TRIGGER_ROUTE_NOTE(int32, 3)
diff --git a/Engine/Plugins/Runtime/Metasound/Source/MetasoundStandardNodes/Public/MetasoundTriggerRouteNode.h b/Engine/Plugins/Runtime/Metasound/Source/MetasoundStandardNodes/Public/MetasoundTriggerRouteNode.h
deleted file mode 100644
