Name/Code convention pass on Filter Nodes

diff --git a/Engine/Plugins/Runtime/Metasound/Source/MetasoundStandardNodes/Private/MetasoundBasicFilters.cpp b/Engine/Plugins/Runtime/Metasound/Source/MetasoundStandardNodes/Private/MetasoundBasicFilters.cpp
--- a/Engine/Plugins/Runtime/Metasound/Source/MetasoundStandardNodes/Private/MetasoundBasicFilters.cpp
+++ b/Engine/Plugins/Runtime/Metasound/Source/MetasoundStandardNodes/Private/MetasoundBasicFilters.cpp
@@ -1,24 +1,78 @@
 // Copyright Epic Games, Inc. All Rights Reserved.
-#include "MetasoundBasicFilters.h"
 
-#include "MetasoundAudioBuffer.h"
-#include "MetasoundDataTypeRegistrationMacro.h"
-#include "MetasoundPrimitives.h"
-#include "MetasoundExecutableOperator.h"
+#include "CoreMinimal.h"
+#include "Internationalization/Text.h"
+
 #include "MetasoundFacade.h"
-#include "MetasoundNodeRegistrationMacro.h"
+#include "MetasoundPrimitives.h"
+#include "MetasoundAudioBuffer.h"
+#include "MetasoundParamHelper.h"
 #include "MetasoundOperatorSettings.h"
 #include "MetasoundStandardNodesNames.h"
+#include "MetasoundExecutableOperator.h"
+#include "MetasoundNodeRegistrationMacro.h"
 #include "MetasoundStandardNodesCategories.h"
+#include "MetasoundDataTypeRegistrationMacro.h"
 
 #include "DSP/Filter.h"
 #include "DSP/InterpolatedOnePole.h"
-#include "Math/NumericLimits.h"
 
 #define LOCTEXT_NAMESPACE "MetasoundBasicFilterNodes"
 
+
 namespace Metasound
 {
+	// forward declarations:
+	class FLadderFilterOperator;
+	class FStateVariableFilterOperator;
+	class FOnePoleLowPassFilterOperator;
+	class FOnePoleHighPassFilterOperator;
+	class FBiquadFilterOperator;
+
+#pragma region Parameter Names
+	namespace BasicFilterParameterNames
+	{
+		// inputs
+		METASOUND_PARAM(ParamAudioInput, "In", "Audio to be processed by the filter");	
+		METASOUND_PARAM(ParamCutoffFrequency, "Cutoff Frequency", "Controls cutoff frequency");
+		METASOUND_PARAM(ParamResonance, "Resonance", "Controls filter resonance");
+		METASOUND_PARAM(ParamBandwidth, "Bandwidth", "Controls bandwidth when applicable to the current filter type");
+		METASOUND_PARAM(ParamGainDb, "Gain", "Gain applied to the band when in Parametric mode (in decibels)");
+		METASOUND_PARAM(ParamFilterType, "Type", "Filter Type");
+		METASOUND_PARAM(ParamBandStopControl, "Band Stop Control", "Band Stop Control (applied to band stop output)");
+
+		// outputs
+		METASOUND_PARAM(ParamAudioOutput, "Out", "Audio processed by the filtter");
+		METASOUND_PARAM(ParamHighPassOutput, "High Pass Filter", "High Pass Filter Output");
+		METASOUND_PARAM(ParamLowPassOutput, "Low Pass Filter", "Low Pass Filter Output");
+		METASOUND_PARAM(ParamBandPassOutput, "Band Pass", "Band Pass Output");
+		METASOUND_PARAM(ParamBandStopOutput, "Band Stop", "Band Stop Output");
+
+	} // namespace BasicFilterParameterNames;
+
+	using namespace BasicFilterParameterNames;
+#pragma endregion
+
+
+#pragma region Biquad Filter Enum
+	DECLARE_METASOUND_ENUM(Audio::EBiquadFilter::Type, Audio::EBiquadFilter::Lowpass,
+	METASOUNDSTANDARDNODES_API, FEnumEBiquadFilterType, FEnumBiQuadFilterTypeInfo, FEnumBiQuadFilterReadRef, FEnumBiQuadFilterWriteRef);
+
+	DEFINE_METASOUND_ENUM_BEGIN(Audio::EBiquadFilter::Type, FEnumEBiquadFilterType, "BiquadFilterType")
+		DEFINE_METASOUND_ENUM_ENTRY(Audio::EBiquadFilter::Lowpass, LOCTEXT("LpDescription", "Low Pass"), LOCTEXT("LpDescriptionTT", "Low Pass Biquad Filter")),
+		DEFINE_METASOUND_ENUM_ENTRY(Audio::EBiquadFilter::Highpass, LOCTEXT("HpDescription", "High Pass"), LOCTEXT("HpDescriptionTT", "High Pass Biquad Filter")),
+		DEFINE_METASOUND_ENUM_ENTRY(Audio::EBiquadFilter::Bandpass, LOCTEXT("BpDescription", "Band Pass"), LOCTEXT("BpDescriptionTT", "Band Pass Biquad Filter")),
+		DEFINE_METASOUND_ENUM_ENTRY(Audio::EBiquadFilter::Notch, LOCTEXT("NotchDescription", "Notch "), LOCTEXT("NotchDescriptionTT", "Notch Biquad Filter")),
+		DEFINE_METASOUND_ENUM_ENTRY(Audio::EBiquadFilter::ParametricEQ, LOCTEXT("ParaEqDescription", "Parametric EQ"), LOCTEXT("ParaEqDescriptionTT", "Parametric EQ Biquad Filter")),
+		DEFINE_METASOUND_ENUM_ENTRY(Audio::EBiquadFilter::LowShelf, LOCTEXT("LowShelfDescription", "Low Shelf"), LOCTEXT("LowShelfDescriptionTT", "Low Shelf Biquad Filter")),
+		DEFINE_METASOUND_ENUM_ENTRY(Audio::EBiquadFilter::HighShelf, LOCTEXT("HighShelfDescription", "High Shelf"), LOCTEXT("HighShelfDescriptionTT", "High Shelf Biquad Filter")),
+		DEFINE_METASOUND_ENUM_ENTRY(Audio::EBiquadFilter::AllPass, LOCTEXT("AllPassDescription", "All Pass"), LOCTEXT("AllPassDescriptionTT", "All Pass Biquad Filter")),
+		DEFINE_METASOUND_ENUM_ENTRY(Audio::EBiquadFilter::ButterworthLowPass, LOCTEXT("LowPassButterDescription", "Butterworth Low Pass"), LOCTEXT("LowPassButterDescriptionTT", "Butterworth Low Pass Biquad Filter")),
+		DEFINE_METASOUND_ENUM_ENTRY(Audio::EBiquadFilter::ButterworthHighPass, LOCTEXT("HighPassButterDescription", "Butterworth High Pass"), LOCTEXT("HighPassButterDescriptionTT", "Butterworth High Pass Biquad Filter"))
+		DEFINE_METASOUND_ENUM_END()
+#pragma endregion
+
+
 #pragma region Ladder Filter
 
 	class FLadderFilterOperator : public TExecutableOperator<FLadderFilterOperator>
@@ -36,9 +90,9 @@ namespace Metasound
 		virtual FDataReferenceCollection GetInputs() const override
 		{
 			FDataReferenceCollection InputDataReferences;
-			InputDataReferences.AddDataReadReference(TEXT("In"), FAudioBufferReadRef(AudioInput));
-			InputDataReferences.AddDataReadReference(TEXT("Freq."), FFloatReadRef(Frequency));
-			InputDataReferences.AddDataReadReference(TEXT("Res."), FFloatReadRef(Resonance));
+			InputDataReferences.AddDataReadReference(METASOUND_GET_PARAM_NAME(ParamAudioInput), FAudioBufferReadRef(AudioInput));
+			InputDataReferences.AddDataReadReference(METASOUND_GET_PARAM_NAME(ParamCutoffFrequency), FFloatReadRef(Frequency));
+			InputDataReferences.AddDataReadReference(METASOUND_GET_PARAM_NAME(ParamResonance), FFloatReadRef(Resonance));
 
 			return InputDataReferences;
 		}
@@ -47,7 +101,7 @@ namespace Metasound
 		{
 			// expose read access to our output buffer for other processors in the graph
 			FDataReferenceCollection OutputDataReferences;
-			OutputDataReferences.AddDataReadReference(TEXT("Out"), FAudioBufferReadRef(AudioOutput));
+			OutputDataReferences.AddDataReadReference(METASOUND_GET_PARAM_NAME(ParamAudioOutput), FAudioBufferReadRef(AudioOutput));
 
 			return OutputDataReferences;
 		}
@@ -98,17 +152,7 @@ namespace Metasound
 
 			LadderFilter.Init(SampleRate, 1);
 		}
-	};
-
-	FLadderFilterNode::FLadderFilterNode(const FString& InInstanceName, const FGuid& InInstanceID)
-		: FNodeFacade(InInstanceName, InInstanceID, TFacadeOperatorClass<FLadderFilterOperator>())
-	{
-	}
-
-	FLadderFilterNode::FLadderFilterNode(const FNodeInitData& InInitData)
-		: FLadderFilterNode(InInitData.InstanceName, InInitData.InstanceID)
-	{
-	}
+	}; // class FLadderFilterOperator
 
 	const FNodeClassMetadata& FLadderFilterOperator::GetNodeInfo()
 	{
@@ -137,12 +181,12 @@ namespace Metasound
 	{
 		static const FVertexInterface Interface(
 			FInputVertexInterface(
-				TInputDataVertexModel<FAudioBuffer>(TEXT("In"), LOCTEXT("AudioInputTooltip", "Audio Input")),
-				TInputDataVertexModel<float>(TEXT("Freq."), LOCTEXT("FrequencyTooltip", "Controls cutoff frequency"), 20000.f),
-				TInputDataVertexModel<float>(TEXT("Res."), LOCTEXT("ResonanceTooltip", "Controls filter resonance"), 6.f)
+				TInputDataVertexModel<FAudioBuffer>(METASOUND_GET_PARAM_NAME(ParamAudioInput), METASOUND_GET_PARAM_TT(ParamAudioInput)),
+				TInputDataVertexModel<float>(METASOUND_GET_PARAM_NAME(ParamCutoffFrequency), METASOUND_GET_PARAM_TT(ParamCutoffFrequency), 20000.f),
+				TInputDataVertexModel<float>(METASOUND_GET_PARAM_NAME(ParamResonance), METASOUND_GET_PARAM_TT(ParamResonance), 6.f)
 			),
 			FOutputVertexInterface(
-				TOutputDataVertexModel<FAudioBuffer>(TEXT("Out"), LOCTEXT("FilterOutputToolTip", "Audio Out"))
+				TOutputDataVertexModel<FAudioBuffer>(METASOUND_GET_PARAM_NAME(ParamAudioOutput), METASOUND_GET_PARAM_TT(ParamAudioOutput))
 			)
 		);
 
@@ -151,14 +195,13 @@ namespace Metasound
 
 	TUniquePtr<IOperator> FLadderFilterOperator::CreateOperator(const FCreateOperatorParams& InParams, FBuildErrorArray& OutErrors)
 	{
-		const FLadderFilterNode& LadderFilterNode = static_cast<const FLadderFilterNode&>(InParams.Node);
 
 		const FDataReferenceCollection& InputDataRefs = InParams.InputDataReferences;
 
 		// inputs
-		FAudioBufferReadRef AudioIn = InputDataRefs.GetDataReadReferenceOrConstruct<FAudioBuffer>(TEXT("In"), InParams.OperatorSettings);
-		FFloatReadRef FrequencyIn = InputDataRefs.GetDataReadReferenceOrConstruct<float>(TEXT("Freq."));
-		FFloatReadRef ResonanceIn = InputDataRefs.GetDataReadReferenceOrConstruct<float>(TEXT("Res."));
+		FAudioBufferReadRef AudioIn = InputDataRefs.GetDataReadReferenceOrConstruct<FAudioBuffer>(METASOUND_GET_PARAM_NAME(ParamAudioInput), InParams.OperatorSettings);
+		FFloatReadRef FrequencyIn = InputDataRefs.GetDataReadReferenceOrConstruct<float>(METASOUND_GET_PARAM_NAME(ParamCutoffFrequency));
+		FFloatReadRef ResonanceIn = InputDataRefs.GetDataReadReferenceOrConstruct<float>(METASOUND_GET_PARAM_NAME(ParamResonance));
 
 		return MakeUnique<FLadderFilterOperator>(
 			InParams.OperatorSettings
@@ -174,7 +217,7 @@ namespace Metasound
 		const float CurrentResonance = FMath::Clamp(*Resonance, 1.0f, 10.0f);
 
 		const bool bNeedsUpdate =
-			   (!FMath::IsNearlyEqual(PreviousFrequency, CurrentFrequency))
+			(!FMath::IsNearlyEqual(PreviousFrequency, CurrentFrequency))
 			|| (!FMath::IsNearlyEqual(PreviousResonance, CurrentResonance));
 
 		if (bNeedsUpdate)
@@ -191,8 +234,6 @@ namespace Metasound
 		LadderFilter.ProcessAudio(AudioInput->GetData(), AudioInput->Num(), AudioOutput->GetData());
 	}
 
-	METASOUND_REGISTER_NODE(FLadderFilterNode);
-
 #pragma endregion
 
 
@@ -212,10 +253,10 @@ namespace Metasound
 		virtual FDataReferenceCollection GetInputs() const override
 		{
 			FDataReferenceCollection InputDataReferences;
-			InputDataReferences.AddDataReadReference(TEXT("In"), FAudioBufferReadRef(AudioInput));
-			InputDataReferences.AddDataReadReference(TEXT("Freq."), FFloatReadRef(Frequency));
-			InputDataReferences.AddDataReadReference(TEXT("Res."), FFloatReadRef(Resonance));
-			InputDataReferences.AddDataReadReference(TEXT("Band-Stop Control"), FFloatReadRef(BandStopControl));
+			InputDataReferences.AddDataReadReference(METASOUND_GET_PARAM_NAME(ParamAudioInput), FAudioBufferReadRef(AudioInput));
+			InputDataReferences.AddDataReadReference(METASOUND_GET_PARAM_NAME(ParamCutoffFrequency), FFloatReadRef(Frequency));
+			InputDataReferences.AddDataReadReference(METASOUND_GET_PARAM_NAME(ParamResonance), FFloatReadRef(Resonance));
+			InputDataReferences.AddDataReadReference(METASOUND_GET_PARAM_NAME(ParamBandStopControl), FFloatReadRef(BandStopControl));
 
 			return InputDataReferences;
 		}
@@ -224,10 +265,10 @@ namespace Metasound
 		{
 			// expose read access to our output buffer for other processors in the graph
 			FDataReferenceCollection OutputDataReferences;
-			OutputDataReferences.AddDataReadReference(TEXT("LP"), FAudioBufferReadRef(LowPassOutput));
-			OutputDataReferences.AddDataReadReference(TEXT("HP"), FAudioBufferReadRef(HighPassOutput));
-			OutputDataReferences.AddDataReadReference(TEXT("BP"), FAudioBufferReadRef(BandPassOutput));
-			OutputDataReferences.AddDataReadReference(TEXT("BS"), FAudioBufferReadRef(BandStopOutput));
+			OutputDataReferences.AddDataReadReference(METASOUND_GET_PARAM_NAME(ParamLowPassOutput), FAudioBufferReadRef(LowPassOutput));
+			OutputDataReferences.AddDataReadReference(METASOUND_GET_PARAM_NAME(ParamHighPassOutput), FAudioBufferReadRef(HighPassOutput));
+			OutputDataReferences.AddDataReadReference(METASOUND_GET_PARAM_NAME(ParamBandPassOutput), FAudioBufferReadRef(BandPassOutput));
+			OutputDataReferences.AddDataReadReference(METASOUND_GET_PARAM_NAME(ParamBandStopOutput), FAudioBufferReadRef(BandStopOutput));
 
 			return OutputDataReferences;
 		}
@@ -293,16 +334,6 @@ namespace Metasound
 		}
 	};
 
-	FStateVariableFilterNode::FStateVariableFilterNode(const FString& InInstanceName, const FGuid& InInstanceID)
-		: FNodeFacade(InInstanceName, InInstanceID, TFacadeOperatorClass<FStateVariableFilterOperator>())
-	{
-	}
-
-	FStateVariableFilterNode::FStateVariableFilterNode(const FNodeInitData& InInitData)
-		: FStateVariableFilterNode(InInitData.InstanceName, InInitData.InstanceID)
-	{
-	}
-
 	const FNodeClassMetadata& FStateVariableFilterOperator::GetNodeInfo()
 	{
 		auto InitNodeInfo = []() -> FNodeClassMetadata
@@ -313,7 +344,7 @@ namespace Metasound
 			Info.MinorVersion = 0;
 			Info.DisplayName = LOCTEXT("Metasound_StateVariableFilterNodeDisplayName", "State Variable Filter");
 			Info.Description = LOCTEXT("State_Variable_Filter_NodeDescription", "State Variable filter"),
-			Info.Author = PluginAuthor;
+				Info.Author = PluginAuthor;
 			Info.PromptIfMissing = PluginNodeMissingPrompt;
 			Info.DefaultInterface = DeclareVertexInterface();
 			Info.CategoryHierarchy.Emplace(StandardNodes::Filters);
@@ -330,16 +361,16 @@ namespace Metasound
 	{
 		static const FVertexInterface Interface(
 			FInputVertexInterface(
-				TInputDataVertexModel<FAudioBuffer>(TEXT("In"), LOCTEXT("AudioInputTooltip", "Audio Input")),
-				TInputDataVertexModel<float>(TEXT("Freq."), LOCTEXT("FrequencyTooltip", "Controls cutoff frequency"), 20000.f),
-				TInputDataVertexModel<float>(TEXT("Res."), LOCTEXT("ResonanceTooltip", "Controls filter resonance"), 0.f),
-				TInputDataVertexModel<float>(TEXT("Band-Stop Control"), LOCTEXT("BandStopControlTooltip", "Band Stop Control"), 0.f)
+				TInputDataVertexModel<FAudioBuffer>(METASOUND_GET_PARAM_NAME(ParamAudioInput), METASOUND_GET_PARAM_TT(ParamAudioInput)),
+				TInputDataVertexModel<float>(METASOUND_GET_PARAM_NAME(ParamCutoffFrequency), METASOUND_GET_PARAM_TT(ParamCutoffFrequency), 20000.f),
+				TInputDataVertexModel<float>(METASOUND_GET_PARAM_NAME(ParamResonance), METASOUND_GET_PARAM_TT(ParamResonance), 0.f),
+				TInputDataVertexModel<float>(METASOUND_GET_PARAM_NAME(ParamBandStopControl), METASOUND_GET_PARAM_TT(ParamBandStopControl), 0.f)
 			),
 			FOutputVertexInterface(
-				TOutputDataVertexModel<FAudioBuffer>(TEXT("LP"), LOCTEXT("FilterOutputToolTip", "Low Pass Output")),
-				TOutputDataVertexModel<FAudioBuffer>(TEXT("HP"), LOCTEXT("FilterOutputToolTip", "High Pass Output")),
-				TOutputDataVertexModel<FAudioBuffer>(TEXT("BP"), LOCTEXT("FilterOutputToolTip", "Band Pass Output")),
-				TOutputDataVertexModel<FAudioBuffer>(TEXT("BS"), LOCTEXT("FilterOutputToolTip", "Band Stop Output"))
+				TOutputDataVertexModel<FAudioBuffer>(METASOUND_GET_PARAM_NAME(ParamLowPassOutput), METASOUND_GET_PARAM_TT(ParamLowPassOutput)),
+				TOutputDataVertexModel<FAudioBuffer>(METASOUND_GET_PARAM_NAME(ParamHighPassOutput), METASOUND_GET_PARAM_TT(ParamHighPassOutput)),
+				TOutputDataVertexModel<FAudioBuffer>(METASOUND_GET_PARAM_NAME(ParamBandPassOutput), METASOUND_GET_PARAM_TT(ParamBandPassOutput)),
+				TOutputDataVertexModel<FAudioBuffer>(METASOUND_GET_PARAM_NAME(ParamBandStopOutput), METASOUND_GET_PARAM_TT(ParamBandStopOutput))
 			)
 		);
 
@@ -348,15 +379,14 @@ namespace Metasound
 
 	TUniquePtr<IOperator> FStateVariableFilterOperator::CreateOperator(const FCreateOperatorParams& InParams, FBuildErrorArray& OutErrors)
 	{
-		const FStateVariableFilterNode& StateVariableFilterNode = static_cast<const FStateVariableFilterNode&>(InParams.Node);
 
 		const FDataReferenceCollection& InputDataRefs = InParams.InputDataReferences;
 
 		// inputs
-		FAudioBufferReadRef AudioIn = InputDataRefs.GetDataReadReferenceOrConstruct<FAudioBuffer>(TEXT("In"), InParams.OperatorSettings);
-		FFloatReadRef FrequencyIn = InputDataRefs.GetDataReadReferenceOrConstruct<float>(TEXT("Freq."));
-		FFloatReadRef ResonanceIn = InputDataRefs.GetDataReadReferenceOrConstruct<float>(TEXT("Res."));
-		FFloatReadRef PassBandGainCompensationIn = InputDataRefs.GetDataReadReferenceOrConstruct<float>(TEXT("Band-Stop Control"));
+		FAudioBufferReadRef AudioIn = InputDataRefs.GetDataReadReferenceOrConstruct<FAudioBuffer>(METASOUND_GET_PARAM_NAME(ParamAudioInput), InParams.OperatorSettings);
+		FFloatReadRef FrequencyIn = InputDataRefs.GetDataReadReferenceOrConstruct<float>(METASOUND_GET_PARAM_NAME(ParamCutoffFrequency));
+		FFloatReadRef ResonanceIn = InputDataRefs.GetDataReadReferenceOrConstruct<float>(METASOUND_GET_PARAM_NAME(ParamResonance));
+		FFloatReadRef PassBandGainCompensationIn = InputDataRefs.GetDataReadReferenceOrConstruct<float>(METASOUND_GET_PARAM_NAME(ParamBandStopControl));
 
 		return MakeUnique<FStateVariableFilterOperator>(
 			InParams.OperatorSettings
@@ -374,7 +404,7 @@ namespace Metasound
 		const float CurrentBandStopControl = FMath::Clamp(*BandStopControl, 0.f, 1.f);
 
 		bool bNeedsUpdate =
-			   (!FMath::IsNearlyEqual(PreviousFrequency, CurrentFrequency))
+			(!FMath::IsNearlyEqual(PreviousFrequency, CurrentFrequency))
 			|| (!FMath::IsNearlyEqual(PreviousResonance, CurrentResonance))
 			|| (!FMath::IsNearlyEqual(PreviousBandStopControl, CurrentBandStopControl));
 
@@ -401,8 +431,6 @@ namespace Metasound
 		);
 	}
 
-	METASOUND_REGISTER_NODE(FStateVariableFilterNode);
-
 #pragma endregion
 
 
@@ -422,8 +450,8 @@ namespace Metasound
 		virtual FDataReferenceCollection GetInputs() const override
 		{
 			FDataReferenceCollection InputDataReferences;
-			InputDataReferences.AddDataReadReference(TEXT("In"), FAudioBufferReadRef(AudioInput));
-			InputDataReferences.AddDataReadReference(TEXT("Freq."), FFloatReadRef(Frequency));
+			InputDataReferences.AddDataReadReference(METASOUND_GET_PARAM_NAME(ParamAudioInput), FAudioBufferReadRef(AudioInput));
+			InputDataReferences.AddDataReadReference(METASOUND_GET_PARAM_NAME(ParamCutoffFrequency), FFloatReadRef(Frequency));
 
 			return InputDataReferences;
 		}
@@ -432,7 +460,7 @@ namespace Metasound
 		{
 			// expose read access to our output buffer for other processors in the graph
 			FDataReferenceCollection OutputDataReferences;
-			OutputDataReferences.AddDataReadReference(TEXT("Out"), FAudioBufferReadRef(AudioOutput));
+			OutputDataReferences.AddDataReadReference(METASOUND_GET_PARAM_NAME(ParamAudioOutput), FAudioBufferReadRef(AudioOutput));
 
 			return OutputDataReferences;
 		}
@@ -476,16 +504,6 @@ namespace Metasound
 		}
 	};
 
-	FOnePoleLowPassFilterNode::FOnePoleLowPassFilterNode(const FString& InInstanceName, const FGuid& InInstanceID)
-		: FNodeFacade(InInstanceName, InInstanceID, TFacadeOperatorClass<FOnePoleLowPassFilterOperator>())
-	{
-	}
-
-	FOnePoleLowPassFilterNode::FOnePoleLowPassFilterNode(const FNodeInitData& InInitData)
-		: FOnePoleLowPassFilterNode(InInitData.InstanceName, InInitData.InstanceID)
-	{
-	}
-
 	const FNodeClassMetadata& FOnePoleLowPassFilterOperator::GetNodeInfo()
 	{
 		auto InitNodeInfo = []() -> FNodeClassMetadata
@@ -496,7 +514,7 @@ namespace Metasound
 			Info.MinorVersion = 0;
 			Info.DisplayName = LOCTEXT("Metasound_OnePoleLpfNodeDisplayName", "One-Pole Low Pass Filter");
 			Info.Description = LOCTEXT("One_Pole_Low_Pass_Filter_NodeDescription", "One-Pole Low Pass Filter"),
-			Info.Author = PluginAuthor;
+				Info.Author = PluginAuthor;
 			Info.PromptIfMissing = PluginNodeMissingPrompt;
 			Info.DefaultInterface = DeclareVertexInterface();
 			Info.CategoryHierarchy.Emplace(StandardNodes::Filters);
@@ -513,11 +531,11 @@ namespace Metasound
 	{
 		static const FVertexInterface Interface(
 			FInputVertexInterface(
-				TInputDataVertexModel<FAudioBuffer>(TEXT("In"), LOCTEXT("AudioInputTooltip", "Audio Input")),
-				TInputDataVertexModel<float>(TEXT("Freq."), LOCTEXT("FrequencyTooltip", "Controls filter cutoff"), 20000.f)
+				TInputDataVertexModel<FAudioBuffer>(METASOUND_GET_PARAM_NAME(ParamAudioInput), METASOUND_GET_PARAM_TT(ParamAudioInput)),
+				TInputDataVertexModel<float>(METASOUND_GET_PARAM_NAME(ParamCutoffFrequency), METASOUND_GET_PARAM_TT(ParamCutoffFrequency), 20000.f)
 			),
 			FOutputVertexInterface(
-				TOutputDataVertexModel<FAudioBuffer>(TEXT("Out"), LOCTEXT("FilterOutputToolTip", "Audio Output"))
+				TOutputDataVertexModel<FAudioBuffer>(METASOUND_GET_PARAM_NAME(ParamAudioOutput), METASOUND_GET_PARAM_TT(ParamAudioOutput))
 			)
 		);
 
@@ -526,13 +544,12 @@ namespace Metasound
 
 	TUniquePtr<IOperator> FOnePoleLowPassFilterOperator::CreateOperator(const FCreateOperatorParams& InParams, FBuildErrorArray& OutErrors)
 	{
-		const FOnePoleLowPassFilterNode& OnePoleLowPassFilterNode = static_cast<const FOnePoleLowPassFilterNode&>(InParams.Node);
 
 		const FDataReferenceCollection& InputDataRefs = InParams.InputDataReferences;
 
 		// inputs
-		FAudioBufferReadRef AudioIn = InputDataRefs.GetDataReadReferenceOrConstruct<FAudioBuffer>(TEXT("In"), InParams.OperatorSettings);
-		FFloatReadRef FrequencyIn = InputDataRefs.GetDataReadReferenceOrConstruct<float>(TEXT("Freq."));
+		FAudioBufferReadRef AudioIn = InputDataRefs.GetDataReadReferenceOrConstruct<FAudioBuffer>(METASOUND_GET_PARAM_NAME(ParamAudioInput), InParams.OperatorSettings);
+		FFloatReadRef FrequencyIn = InputDataRefs.GetDataReadReferenceOrConstruct<float>(METASOUND_GET_PARAM_NAME(ParamCutoffFrequency));
 
 		return MakeUnique<FOnePoleLowPassFilterOperator>(
 			InParams.OperatorSettings
@@ -548,8 +565,6 @@ namespace Metasound
 		OnePoleLowPassFilter.StopFrequencyInterpolation();
 	}
 
-	METASOUND_REGISTER_NODE(FOnePoleLowPassFilterNode);
-
 #pragma endregion
 
 
@@ -569,8 +584,8 @@ namespace Metasound
 		virtual FDataReferenceCollection GetInputs() const override
 		{
 			FDataReferenceCollection InputDataReferences;
-			InputDataReferences.AddDataReadReference(TEXT("In"), FAudioBufferReadRef(AudioInput));
-			InputDataReferences.AddDataReadReference(TEXT("Freq."), FFloatReadRef(Frequency));
+			InputDataReferences.AddDataReadReference(METASOUND_GET_PARAM_NAME(ParamAudioInput), FAudioBufferReadRef(AudioInput));
+			InputDataReferences.AddDataReadReference(METASOUND_GET_PARAM_NAME(ParamCutoffFrequency), FFloatReadRef(Frequency));
 
 			return InputDataReferences;
 		}
@@ -579,7 +594,7 @@ namespace Metasound
 		{
 			// expose read access to our output buffer for other processors in the graph
 			FDataReferenceCollection OutputDataReferences;
-			OutputDataReferences.AddDataReadReference(TEXT("Out"), FAudioBufferReadRef(AudioOutput));
+			OutputDataReferences.AddDataReadReference(METASOUND_GET_PARAM_NAME(ParamAudioOutput), FAudioBufferReadRef(AudioOutput));
 
 			return OutputDataReferences;
 		}
@@ -623,16 +638,6 @@ namespace Metasound
 		}
 	};
 
-	FOnePoleHighPassFilterNode::FOnePoleHighPassFilterNode(const FString& InInstanceName, const FGuid& InInstanceID)
-		: FNodeFacade(InInstanceName, InInstanceID, TFacadeOperatorClass<FOnePoleHighPassFilterOperator>())
-	{
-	}
-
-	FOnePoleHighPassFilterNode::FOnePoleHighPassFilterNode(const FNodeInitData& InInitData)
-		: FOnePoleHighPassFilterNode(InInitData.InstanceName, InInitData.InstanceID)
-	{
-	}
-
 	const FNodeClassMetadata& FOnePoleHighPassFilterOperator::GetNodeInfo()
 	{
 		auto InitNodeInfo = []() -> FNodeClassMetadata
@@ -643,7 +648,7 @@ namespace Metasound
 			Info.MinorVersion = 0;
 			Info.DisplayName = LOCTEXT("Metasound_OnePoleHpfNodeDisplayName", "One-Pole High Pass Filter");
 			Info.Description = LOCTEXT("One_Pole_High_Pass_Filter_NodeDescription", "One-Pole High Pass Filter"),
-			Info.Author = PluginAuthor;
+				Info.Author = PluginAuthor;
 			Info.PromptIfMissing = PluginNodeMissingPrompt;
 			Info.DefaultInterface = DeclareVertexInterface();
 			Info.CategoryHierarchy.Emplace(StandardNodes::Filters);
@@ -660,11 +665,11 @@ namespace Metasound
 	{
 		static const FVertexInterface Interface(
 			FInputVertexInterface(
-				TInputDataVertexModel<FAudioBuffer>(TEXT("In"), LOCTEXT("AudioInputTooltip", "Audio In")),
-				TInputDataVertexModel<float>(TEXT("Freq."), LOCTEXT("FrequencyTooltip", "Freq In"), 10.f)
+				TInputDataVertexModel<FAudioBuffer>(METASOUND_GET_PARAM_NAME(ParamAudioInput), METASOUND_GET_PARAM_TT(ParamAudioInput)),
+				TInputDataVertexModel<float>(METASOUND_GET_PARAM_NAME(ParamCutoffFrequency), METASOUND_GET_PARAM_TT(ParamCutoffFrequency), 10.f)
 			),
 			FOutputVertexInterface(
-				TOutputDataVertexModel<FAudioBuffer>(TEXT("Out"), LOCTEXT("FilterOutputToolTip", "Audio Out"))
+				TOutputDataVertexModel<FAudioBuffer>(METASOUND_GET_PARAM_NAME(ParamAudioOutput), METASOUND_GET_PARAM_TT(ParamAudioOutput))
 			)
 		);
 
@@ -673,13 +678,12 @@ namespace Metasound
 
 	TUniquePtr<IOperator> FOnePoleHighPassFilterOperator::CreateOperator(const FCreateOperatorParams& InParams, FBuildErrorArray& OutErrors)
 	{
-		const FOnePoleHighPassFilterNode& OnePoleHighPassFilterNode = static_cast<const FOnePoleHighPassFilterNode&>(InParams.Node);
 
 		const FDataReferenceCollection& InputDataRefs = InParams.InputDataReferences;
 
 		// inputs
-		FAudioBufferReadRef AudioIn = InputDataRefs.GetDataReadReferenceOrConstruct<FAudioBuffer>(TEXT("In"), InParams.OperatorSettings);
-		FFloatReadRef FrequencyIn = InputDataRefs.GetDataReadReferenceOrConstruct<float>(TEXT("Freq."));
+		FAudioBufferReadRef AudioIn = InputDataRefs.GetDataReadReferenceOrConstruct<FAudioBuffer>(METASOUND_GET_PARAM_NAME(ParamAudioInput), InParams.OperatorSettings);
+		FFloatReadRef FrequencyIn = InputDataRefs.GetDataReadReferenceOrConstruct<float>(METASOUND_GET_PARAM_NAME(ParamCutoffFrequency));
 
 		return MakeUnique<FOnePoleHighPassFilterOperator>(
 			InParams.OperatorSettings
@@ -695,28 +699,10 @@ namespace Metasound
 		OnePoleHighPassFilter.StopFrequencyInterpolation();
 	}
 
-	METASOUND_REGISTER_NODE(FOnePoleHighPassFilterNode);
-
 #pragma endregion
 
 
 #pragma region Biquad Filter
-	DECLARE_METASOUND_ENUM(Audio::EBiquadFilter::Type, Audio::EBiquadFilter::Lowpass,
-	METASOUNDSTANDARDNODES_API, FEnumEBiquadFilterType, FEnumBiQuadFilterTypeInfo, FEnumBiQuadFilterReadRef, FEnumBiQuadFilterWriteRef);
-
-	DEFINE_METASOUND_ENUM_BEGIN(Audio::EBiquadFilter::Type, FEnumEBiquadFilterType, "BiquadFilterType")
-		DEFINE_METASOUND_ENUM_ENTRY(Audio::EBiquadFilter::Lowpass, LOCTEXT("LpDescription", "Low Pass"), LOCTEXT("LpDescriptionTT", "Low Pass Biquad Filter")),
-		DEFINE_METASOUND_ENUM_ENTRY(Audio::EBiquadFilter::Highpass, LOCTEXT("HpDescription", "High Pass"), LOCTEXT("HpDescriptionTT", "High Pass Biquad Filter")),
-		DEFINE_METASOUND_ENUM_ENTRY(Audio::EBiquadFilter::Bandpass, LOCTEXT("BpDescription", "Band Pass"), LOCTEXT("BpDescriptionTT", "Band Pass Biquad Filter")),
-		DEFINE_METASOUND_ENUM_ENTRY(Audio::EBiquadFilter::Notch, LOCTEXT("NotchDescription", "Notch "), LOCTEXT("NotchDescriptionTT", "Notch Biquad Filter")),
-		DEFINE_METASOUND_ENUM_ENTRY(Audio::EBiquadFilter::ParametricEQ, LOCTEXT("ParaEqDescription", "Parametric EQ"), LOCTEXT("ParaEqDescriptionTT", "Parametric EQ Biquad Filter")),
-		DEFINE_METASOUND_ENUM_ENTRY(Audio::EBiquadFilter::LowShelf, LOCTEXT("LowShelfDescription", "Low Shelf"), LOCTEXT("LowShelfDescriptionTT", "Low Shelf Biquad Filter")),
-		DEFINE_METASOUND_ENUM_ENTRY(Audio::EBiquadFilter::HighShelf, LOCTEXT("HighShelfDescription", "High Shelf"), LOCTEXT("HighShelfDescriptionTT", "High Shelf Biquad Filter")),
-		DEFINE_METASOUND_ENUM_ENTRY(Audio::EBiquadFilter::AllPass, LOCTEXT("AllPassDescription", "All Pass"), LOCTEXT("AllPassDescriptionTT", "All Pass Biquad Filter")),
-		DEFINE_METASOUND_ENUM_ENTRY(Audio::EBiquadFilter::ButterworthLowPass, LOCTEXT("LowPassButterDescription", "Butterworth Low Pass"), LOCTEXT("LowPassButterDescriptionTT", "Butterworth Low Pass Biquad Filter")),
-		DEFINE_METASOUND_ENUM_ENTRY(Audio::EBiquadFilter::ButterworthHighPass, LOCTEXT("HighPassButterDescription", "Butterworth High Pass"), LOCTEXT("HighPassButterDescriptionTT", "Butterworth High Pass Biquad Filter"))
-		DEFINE_METASOUND_ENUM_END()
-
 		class FBiquadFilterOperator : public TExecutableOperator<FBiquadFilterOperator>
 	{
 	private:
@@ -731,11 +717,11 @@ namespace Metasound
 		virtual FDataReferenceCollection GetInputs() const override
 		{
 			FDataReferenceCollection InputDataReferences;
-			InputDataReferences.AddDataReadReference(TEXT("In"), FAudioBufferReadRef(AudioInput));
-			InputDataReferences.AddDataReadReference(TEXT("Freq."), FFloatReadRef(Frequency));
-			InputDataReferences.AddDataReadReference(TEXT("Bandwidth"), FFloatReadRef(Bandwidth));
-			InputDataReferences.AddDataReadReference(TEXT("FilterGainDb"), FFloatReadRef(FilterGainDb));
-			InputDataReferences.AddDataReadReference(TEXT("Type"), FEnumBiQuadFilterReadRef(FilterType));
+			InputDataReferences.AddDataReadReference(METASOUND_GET_PARAM_NAME(ParamAudioInput), FAudioBufferReadRef(AudioInput));
+			InputDataReferences.AddDataReadReference(METASOUND_GET_PARAM_NAME(ParamCutoffFrequency), FFloatReadRef(Frequency));
+			InputDataReferences.AddDataReadReference(METASOUND_GET_PARAM_NAME(ParamBandwidth), FFloatReadRef(Bandwidth));
+			InputDataReferences.AddDataReadReference(METASOUND_GET_PARAM_NAME(ParamGainDb), FFloatReadRef(FilterGainDb));
+			InputDataReferences.AddDataReadReference(METASOUND_GET_PARAM_NAME(ParamFilterType), FEnumBiQuadFilterReadRef(FilterType));
 
 			return InputDataReferences;
 		}
@@ -744,7 +730,7 @@ namespace Metasound
 		{
 			// expose read access to our output buffer for other processors in the graph
 			FDataReferenceCollection OutputDataReferences;
-			OutputDataReferences.AddDataReadReference(TEXT("Out"), FAudioBufferReadRef(AudioOutput));
+			OutputDataReferences.AddDataReadReference(METASOUND_GET_PARAM_NAME(ParamAudioOutput), FAudioBufferReadRef(AudioOutput));
 
 			return OutputDataReferences;
 		}
@@ -805,16 +791,6 @@ namespace Metasound
 		}
 	};
 
-	FBiquadFilterNode::FBiquadFilterNode(const FString& InInstanceName, const FGuid& InInstanceID)
-		: FNodeFacade(InInstanceName, InInstanceID, TFacadeOperatorClass<FBiquadFilterOperator>())
-	{
-	}
-
-	FBiquadFilterNode::FBiquadFilterNode(const FNodeInitData& InInitData)
-		: FBiquadFilterNode(InInitData.InstanceName, InInitData.InstanceID)
-	{
-	}
-
 	const FNodeClassMetadata& FBiquadFilterOperator::GetNodeInfo()
 	{
 		auto InitNodeInfo = []() -> FNodeClassMetadata
@@ -825,7 +801,7 @@ namespace Metasound
 			Info.MinorVersion = 0;
 			Info.DisplayName = LOCTEXT("Metasound_BiquadFilterNodeDisplayName", "Biquad Filter");
 			Info.Description = LOCTEXT("Biquad_Filter_NodeDescription", "Biquad filter"),
-			Info.Author = PluginAuthor;
+				Info.Author = PluginAuthor;
 			Info.PromptIfMissing = PluginNodeMissingPrompt;
 			Info.DefaultInterface = DeclareVertexInterface();
 			Info.CategoryHierarchy.Emplace(StandardNodes::Filters);
@@ -842,14 +818,14 @@ namespace Metasound
 	{
 		static const FVertexInterface Interface(
 			FInputVertexInterface(
-				TInputDataVertexModel<FAudioBuffer>(TEXT("In"), LOCTEXT("AudioInputTooltip", "Audio Input")),
-				TInputDataVertexModel<float>(TEXT("Freq."), LOCTEXT("FrequencyTooltip", "Controls Filter Cutoff"), 20000.f),
-				TInputDataVertexModel<float>(TEXT("Bandwidth"), LOCTEXT("BandwidthTooltip", "Bandwidth Control"), 1.f),
-				TInputDataVertexModel<float>(TEXT("FilterGainDb"), LOCTEXT("FilterGainDbTooltip", "Filter Gain Control (decibels)"), 0.f), 
-				TInputDataVertexModel<FEnumEBiquadFilterType>(TEXT("Type"), LOCTEXT("FilterTypeDescription", "Biquad Filter Type"))
+				TInputDataVertexModel<FAudioBuffer>(METASOUND_GET_PARAM_NAME(ParamAudioInput), METASOUND_GET_PARAM_TT(ParamAudioInput)),
+				TInputDataVertexModel<float>(METASOUND_GET_PARAM_NAME(ParamCutoffFrequency), METASOUND_GET_PARAM_TT(ParamCutoffFrequency), 20000.f),
+				TInputDataVertexModel<float>(METASOUND_GET_PARAM_NAME(ParamBandwidth), METASOUND_GET_PARAM_TT(ParamBandwidth), 1.f),
+				TInputDataVertexModel<float>(METASOUND_GET_PARAM_NAME(ParamGainDb), METASOUND_GET_PARAM_TT(ParamGainDb), 0.f),
+				TInputDataVertexModel<FEnumEBiquadFilterType>(METASOUND_GET_PARAM_NAME(ParamFilterType), METASOUND_GET_PARAM_TT(ParamFilterType))
 			),
 			FOutputVertexInterface(
-				TOutputDataVertexModel<FAudioBuffer>(TEXT("Out"), LOCTEXT("FilterOutputToolTip", "Audio Output"))
+				TOutputDataVertexModel<FAudioBuffer>(METASOUND_GET_PARAM_NAME(ParamAudioOutput), METASOUND_GET_PARAM_TT(ParamAudioOutput))
 			)
 		);
 
@@ -858,16 +834,15 @@ namespace Metasound
 
 	TUniquePtr<IOperator> FBiquadFilterOperator::CreateOperator(const FCreateOperatorParams& InParams, FBuildErrorArray& OutErrors)
 	{
-		const FBiquadFilterNode& BiquadFilterNode = static_cast<const FBiquadFilterNode&>(InParams.Node);
 
 		const FDataReferenceCollection& InputDataRefs = InParams.InputDataReferences;
 
 		// inputs
-		FAudioBufferReadRef AudioIn = InputDataRefs.GetDataReadReferenceOrConstruct<FAudioBuffer>(TEXT("In"), InParams.OperatorSettings);
-		FFloatReadRef FrequencyIn = InputDataRefs.GetDataReadReferenceOrConstruct<float>(TEXT("Freq."));
-		FFloatReadRef BandwidthIn = InputDataRefs.GetDataReadReferenceOrConstruct<float>(TEXT("Bandwidth"));
-		FFloatReadRef FilterGainDbIn = InputDataRefs.GetDataReadReferenceOrConstruct<float>(TEXT("FilterGainDb"));
-			FEnumBiQuadFilterReadRef FilterType = InputDataRefs.GetDataReadReferenceOrConstruct<FEnumEBiquadFilterType>(TEXT("Type"));
+		FAudioBufferReadRef AudioIn = InputDataRefs.GetDataReadReferenceOrConstruct<FAudioBuffer>(METASOUND_GET_PARAM_NAME(ParamAudioInput), InParams.OperatorSettings);
+		FFloatReadRef FrequencyIn = InputDataRefs.GetDataReadReferenceOrConstruct<float>(METASOUND_GET_PARAM_NAME(ParamCutoffFrequency));
+		FFloatReadRef BandwidthIn = InputDataRefs.GetDataReadReferenceOrConstruct<float>(METASOUND_GET_PARAM_NAME(ParamBandwidth));
+		FFloatReadRef FilterGainDbIn = InputDataRefs.GetDataReadReferenceOrConstruct<float>(METASOUND_GET_PARAM_NAME(ParamGainDb));
+		FEnumBiQuadFilterReadRef FilterType = InputDataRefs.GetDataReadReferenceOrConstruct<FEnumEBiquadFilterType>(METASOUND_GET_PARAM_NAME(ParamFilterType));
 
 		return MakeUnique<FBiquadFilterOperator>(
 			InParams.OperatorSettings
@@ -912,11 +887,96 @@ namespace Metasound
 		BiquadFilter.ProcessAudio(AudioInput->GetData(), AudioInput->Num(), AudioOutput->GetData());
 	}
 
-	METASOUND_REGISTER_NODE(FBiquadFilterNode);
 
 #pragma endregion
 
 
+#pragma region Node Declarations
+	class METASOUNDSTANDARDNODES_API FLadderFilterNode : public FNodeFacade
+	{
+	public:
+		// public node api needs to define two conversion constructors:
+		// 1.) from FString
+		FLadderFilterNode(const FString& InInstanceName, const FGuid& InInstanceID)
+			: FNodeFacade(InInstanceName, InInstanceID, TFacadeOperatorClass<FLadderFilterOperator>())
+		{ }
+
+		// 2.) From an NodeInitData struct
+		FLadderFilterNode(const FNodeInitData& InInitData)
+			: FLadderFilterNode(InInitData.InstanceName, InInitData.InstanceID)
+		{ }
+	};
+
+	class METASOUNDSTANDARDNODES_API FStateVariableFilterNode : public FNodeFacade
+	{
+	public:
+		// public node api needs to define two conversion constructors:
+		// 1.) from FString
+		FStateVariableFilterNode(const FString& InInstanceName, const FGuid& InInstanceID)
+			: FNodeFacade(InInstanceName, InInstanceID, TFacadeOperatorClass<FStateVariableFilterOperator>())
+		{ }
+
+		// 2.) From an NodeInitData struct
+		FStateVariableFilterNode(const FNodeInitData& InInitData)
+			: FStateVariableFilterNode(InInitData.InstanceName, InInitData.InstanceID)
+		{ }
+	};
+
+	class METASOUNDSTANDARDNODES_API FOnePoleLowPassFilterNode : public FNodeFacade
+	{
+	public:
+		// public node api needs to define two conversion constructors:
+		// 1.) from FString
+		FOnePoleLowPassFilterNode(const FString& InInstanceName, const FGuid& InInstanceID)
+			: FNodeFacade(InInstanceName, InInstanceID, TFacadeOperatorClass<FOnePoleLowPassFilterOperator>())
+		{ }
+
+		// 2.) From an NodeInitData struct
+		FOnePoleLowPassFilterNode(const FNodeInitData& InInitData)
+			: FOnePoleLowPassFilterNode(InInitData.InstanceName, InInitData.InstanceID)
+		{ }
+	};
+
+	class METASOUNDSTANDARDNODES_API FOnePoleHighPassFilterNode : public FNodeFacade
+	{
+	public:
+		// public node api needs to define two conversion constructors:
+		// 1.) from FString
+		FOnePoleHighPassFilterNode(const FString& InInstanceName, const FGuid& InInstanceID)
+			: FNodeFacade(InInstanceName, InInstanceID, TFacadeOperatorClass<FOnePoleHighPassFilterOperator>())
+		{ }
+
+		// 2.) From an NodeInitData struct
+		FOnePoleHighPassFilterNode(const FNodeInitData& InInitData)
+			: FOnePoleHighPassFilterNode(InInitData.InstanceName, InInitData.InstanceID)
+		{ }
+	};
+
+	class METASOUNDSTANDARDNODES_API FBiquadFilterNode : public FNodeFacade
+	{
+	public:
+		// public node api needs to define two conversion constructors:
+		// 1.) from FString
+		FBiquadFilterNode(const FString& InInstanceName, const FGuid& InInstanceID)
+			: FNodeFacade(InInstanceName, InInstanceID, TFacadeOperatorClass<FBiquadFilterOperator>())
+		{ }
+
+		// 2.) From an NodeInitData struct
+		FBiquadFilterNode(const FNodeInitData& InInitData)
+			: FBiquadFilterNode(InInitData.InstanceName, InInitData.InstanceID)
+		{ }
+	};
+#pragma endregion
+
+
+#pragma region Node Registration
+	METASOUND_REGISTER_NODE(FLadderFilterNode);
+	METASOUND_REGISTER_NODE(FStateVariableFilterNode);
+	METASOUND_REGISTER_NODE(FOnePoleLowPassFilterNode);
+	METASOUND_REGISTER_NODE(FOnePoleHighPassFilterNode);
+	METASOUND_REGISTER_NODE(FBiquadFilterNode);
+#pragma endregion
+
 } // namespace Metasound
 
 
diff --git a/Engine/Plugins/Runtime/Metasound/Source/MetasoundStandardNodes/Public/MetasoundBasicFilters.h b/Engine/Plugins/Runtime/Metasound/Source/MetasoundStandardNodes/Public/MetasoundBasicFilters.h
deleted file mode 100644
