Cleaning up trigger counter and trigger delay node

diff --git a/Engine/Plugins/Runtime/Metasound/Source/MetasoundStandardNodes/Private/MetasoundTriggerCounterNode.cpp b/Engine/Plugins/Runtime/Metasound/Source/MetasoundStandardNodes/Private/MetasoundTriggerCounterNode.cpp
--- a/Engine/Plugins/Runtime/Metasound/Source/MetasoundStandardNodes/Private/MetasoundTriggerCounterNode.cpp
+++ b/Engine/Plugins/Runtime/Metasound/Source/MetasoundStandardNodes/Private/MetasoundTriggerCounterNode.cpp
@@ -1,6 +1,4 @@
 // Copyright Epic Games, Inc. All Rights Reserved.
-#include "MetasoundTriggerCounterNode.h"
-
 #include "CoreMinimal.h"
 
 #include "MetasoundBuilderInterface.h"
@@ -14,22 +12,24 @@
 #include "MetasoundPrimitives.h"
 #include "MetasoundStandardNodesNames.h"
 #include "MetasoundStandardNodesCategories.h"
+#include "MetasoundTrigger.h"
 #include "MetasoundParamHelper.h"
 
-#define LOCTEXT_NAMESPACE "MetasoundStandardNodes"
+#define LOCTEXT_NAMESPACE "MetasoundStandardNodes_CounterNode"
 
 namespace Metasound
 {
 	namespace TriggerCounterVertexNames
 	{
 		METASOUND_PARAM(InputInTrigger, "In", "Trigger Input to count.");
-		METASOUND_PARAM(InputReset, "Reset", "Resets the counter back to the start count.");
-		METASOUND_PARAM(InputStartCount, "Start Value", "The value to start output counter by.");
-		METASOUND_PARAM(InputStepSize, "Step Size", "The value to add to the current count each input trigger. Can be negative.");
-		METASOUND_PARAM(InputAutoResetCount, "Auto Reset Count", "The number of input triggers to automatically set the count back to Start Count. If 0, won't auto-reset.");
+		METASOUND_PARAM(InputReset, "Reset", "Resets the counter to zero and the value back to the start value.");
+		METASOUND_PARAM(InputStartValue, "Start Value", "The value to set the ouput to on init and reset.");
+		METASOUND_PARAM(InputStepSize, "Step Size", "The value to add to the current value for each input trigger. Can be negative.");
+		METASOUND_PARAM(InputAutoResetCount, "Reset Count", "The number of input triggers to automatically set the count back to Start Count. If 0, won't auto-reset.");
 		METASOUND_PARAM(OutputOnTrigger, "On Trigger", "Triggered when the input is triggered.");
 		METASOUND_PARAM(OutputOnReset, "On Reset", "Triggered when the input reset trigger is triggered or if the counter automatically resets.");
-		METASOUND_PARAM(OutputCount, "Count", "The current count value of the node.");
+		METASOUND_PARAM(OutputCount, "Count", "The current raw trigger count (i.e. number of times the input trigger has been triggered).");
+		METASOUND_PARAM(OutputValue, "Value", "The current value of the node given the starting value and step size.");
 	}
 
 	class FTriggerCounterOperator : public TExecutableOperator<FTriggerCounterOperator>
@@ -42,8 +42,8 @@ namespace Metasound
 			FTriggerCounterOperator(const FOperatorSettings& InSettings,
 				const FTriggerReadRef& InTriggerIn,
 				const FTriggerReadRef& InTriggerReset,
-				const FInt32ReadRef& InStartCount,
-				const FInt32ReadRef& InCountBy,
+				const FFloatReadRef& InStartValue,
+				const FFloatReadRef& InStepSize,
 				const FInt32ReadRef& InAutoResetCount);
 
 			virtual FDataReferenceCollection GetInputs() const override;
@@ -54,33 +54,36 @@ namespace Metasound
 		private:
 			FTriggerReadRef TriggerIn;
 			FTriggerReadRef TriggerReset;
-			FInt32ReadRef StartCount;
-			FInt32ReadRef CountBy;
+			FFloatReadRef StartValue;
+			FFloatReadRef StepSize;
 			FInt32ReadRef AutoResetCount;
 
 			FTriggerWriteRef TriggerOut;
 			FTriggerWriteRef TriggerOnReset;
-			FInt32WriteRef OutValue;
+			FFloatWriteRef OutValue;
+			FInt32WriteRef OutCount;
 
 			// Internal state to do the counting math
 			int32 CurrentAutoResetCount = 0;
 			int32 CurrentTriggerCount = 0;
+			float CurrentValue = 0.0f;
 	};
 
 	FTriggerCounterOperator::FTriggerCounterOperator(const FOperatorSettings& InSettings, 
 		const FTriggerReadRef& InTriggerIn, 
 		const FTriggerReadRef& InTriggerReset,
-		const FInt32ReadRef& InStartCount,
-		const FInt32ReadRef& InCountBy,
+		const FFloatReadRef& InStartValue,
+		const FFloatReadRef& InStepSize,
 		const FInt32ReadRef& InAutoResetCount)
 		: TriggerIn(InTriggerIn)
 		, TriggerReset(InTriggerReset)
-		, StartCount(InStartCount)
-		, CountBy(InCountBy)
+		, StartValue(InStartValue)
+		, StepSize(InStepSize)
 		, AutoResetCount(InAutoResetCount)
 		, TriggerOut(FTriggerWriteRef::CreateNew(InSettings))
 		, TriggerOnReset(FTriggerWriteRef::CreateNew(InSettings))
-		, OutValue(FInt32WriteRef::CreateNew(*StartCount))
+		, OutValue(FFloatWriteRef::CreateNew(*StartValue))
+		, OutCount(FInt32WriteRef::CreateNew(0))
 	{
 		CurrentAutoResetCount = FMath::Max(0, *AutoResetCount);
 	}
@@ -92,8 +95,8 @@ namespace Metasound
 		FDataReferenceCollection InputDataReferences;
 		InputDataReferences.AddDataReadReference(METASOUND_GET_PARAM_NAME(InputInTrigger), TriggerIn);
 		InputDataReferences.AddDataReadReference(METASOUND_GET_PARAM_NAME(InputReset), TriggerReset);
-		InputDataReferences.AddDataReadReference(METASOUND_GET_PARAM_NAME(InputStartCount), StartCount);
-		InputDataReferences.AddDataReadReference(METASOUND_GET_PARAM_NAME(InputStepSize), CountBy);
+		InputDataReferences.AddDataReadReference(METASOUND_GET_PARAM_NAME(InputStartValue), StartValue);
+		InputDataReferences.AddDataReadReference(METASOUND_GET_PARAM_NAME(InputStepSize), StepSize);
 		InputDataReferences.AddDataReadReference(METASOUND_GET_PARAM_NAME(InputAutoResetCount), AutoResetCount);
 
 		return InputDataReferences;
@@ -106,7 +109,8 @@ namespace Metasound
 		FDataReferenceCollection OutputDataReferences;
 		OutputDataReferences.AddDataReadReference(METASOUND_GET_PARAM_NAME(OutputOnTrigger), TriggerOut);
 		OutputDataReferences.AddDataReadReference(METASOUND_GET_PARAM_NAME(OutputOnReset), TriggerOnReset);
-		OutputDataReferences.AddDataReadReference(METASOUND_GET_PARAM_NAME(OutputCount), OutValue);
+		OutputDataReferences.AddDataReadReference(METASOUND_GET_PARAM_NAME(OutputCount), OutCount);
+		OutputDataReferences.AddDataReadReference(METASOUND_GET_PARAM_NAME(OutputValue), OutValue);
 
 		return OutputDataReferences;
 	}
@@ -126,7 +130,7 @@ namespace Metasound
 			[this](int32 StartFrame, int32 EndFrame)
 			{
 				CurrentTriggerCount = 0;
-				*OutValue = *StartCount;
+				*OutValue = *StartValue;
 
 				TriggerOnReset->TriggerFrame(StartFrame);
 			}
@@ -143,14 +147,15 @@ namespace Metasound
 				if (CurrentAutoResetCount > 0 && CurrentTriggerCount >= CurrentAutoResetCount)
 				{
 					CurrentTriggerCount = 0;
-					*OutValue = *StartCount;
+					*OutValue = *StartValue;
 					TriggerOnReset->TriggerFrame(StartFrame);
 				}
 				else
 				{
-					*OutValue += *CountBy;
+					*OutValue += *StepSize;
 				}
 
+				*OutCount = CurrentTriggerCount;
 				TriggerOut->TriggerFrame(StartFrame);
 			}
 		);
@@ -164,11 +169,11 @@ namespace Metasound
 		const FInputVertexInterface& InputInterface = GetVertexInterface().GetInputInterface();
 		FTriggerReadRef TriggerIn = InParams.InputDataReferences.GetDataReadReferenceOrConstruct<FTrigger>(METASOUND_GET_PARAM_NAME(InputInTrigger), InParams.OperatorSettings);
 		FTriggerReadRef TriggerReset = InParams.InputDataReferences.GetDataReadReferenceOrConstruct<FTrigger>(METASOUND_GET_PARAM_NAME(InputReset), InParams.OperatorSettings);
-		FInt32ReadRef StartCount = InParams.InputDataReferences.GetDataReadReferenceOrConstructWithVertexDefault<int32>(InputInterface, METASOUND_GET_PARAM_NAME(InputStartCount), InParams.OperatorSettings);
-		FInt32ReadRef StepSize = InParams.InputDataReferences.GetDataReadReferenceOrConstructWithVertexDefault<int32>(InputInterface, METASOUND_GET_PARAM_NAME(InputStepSize), InParams.OperatorSettings);
+		FFloatReadRef StartValue = InParams.InputDataReferences.GetDataReadReferenceOrConstructWithVertexDefault<float>(InputInterface, METASOUND_GET_PARAM_NAME(InputStartValue), InParams.OperatorSettings);
+		FFloatReadRef StepSize = InParams.InputDataReferences.GetDataReadReferenceOrConstructWithVertexDefault<float>(InputInterface, METASOUND_GET_PARAM_NAME(InputStepSize), InParams.OperatorSettings);
 		FInt32ReadRef AutoResetCount = InParams.InputDataReferences.GetDataReadReferenceOrConstructWithVertexDefault<int32>(InputInterface, METASOUND_GET_PARAM_NAME(InputAutoResetCount), InParams.OperatorSettings);
 
-		return MakeUnique<FTriggerCounterOperator>(InParams.OperatorSettings, TriggerIn, TriggerReset, StartCount, StepSize, AutoResetCount);
+		return MakeUnique<FTriggerCounterOperator>(InParams.OperatorSettings, TriggerIn, TriggerReset, StartValue, StepSize, AutoResetCount);
 	}
 
 	const FVertexInterface& FTriggerCounterOperator::GetVertexInterface()
@@ -179,14 +184,15 @@ namespace Metasound
 			FInputVertexInterface(
 				TInputDataVertexModel<FTrigger>(METASOUND_GET_PARAM_NAME_AND_TT(InputInTrigger)),
 				TInputDataVertexModel<FTrigger>(METASOUND_GET_PARAM_NAME_AND_TT(InputReset)),
-				TInputDataVertexModel<int32>(METASOUND_GET_PARAM_NAME_AND_TT(InputStartCount), 0),
-				TInputDataVertexModel<int32>(METASOUND_GET_PARAM_NAME_AND_TT(InputStepSize), 1),
+				TInputDataVertexModel<float>(METASOUND_GET_PARAM_NAME_AND_TT(InputStartValue), 0.0f),
+				TInputDataVertexModel<float>(METASOUND_GET_PARAM_NAME_AND_TT(InputStepSize), 1),
 				TInputDataVertexModel<int32>(METASOUND_GET_PARAM_NAME_AND_TT(InputAutoResetCount), 0)
 			),
 			FOutputVertexInterface(
 				TOutputDataVertexModel<FTrigger>(METASOUND_GET_PARAM_NAME_AND_TT(OutputOnTrigger)),
 				TOutputDataVertexModel<FTrigger>(METASOUND_GET_PARAM_NAME_AND_TT(OutputOnReset)),
-				TOutputDataVertexModel<int32>(METASOUND_GET_PARAM_NAME_AND_TT(OutputCount))
+				TOutputDataVertexModel<int32>(METASOUND_GET_PARAM_NAME_AND_TT(OutputCount)),
+				TOutputDataVertexModel<float>(METASOUND_GET_PARAM_NAME_AND_TT(OutputValue))
 			)
 		);
 
@@ -215,10 +221,16 @@ namespace Metasound
 		return Info;
 	}
 
-	FTriggerCounterNode::FTriggerCounterNode(const FNodeInitData& InInitData)
-	:	FNodeFacade(InInitData.InstanceName, InInitData.InstanceID, TFacadeOperatorClass<FTriggerCounterOperator>())
+	class FTriggerCounterNode : public FNodeFacade
 	{
-	}
+	public:
+		FTriggerCounterNode(const FNodeInitData& InInitData)
+			: FNodeFacade(InInitData.InstanceName, InInitData.InstanceID, TFacadeOperatorClass<FTriggerCounterOperator>())
+		{
+		}
+
+		virtual ~FTriggerCounterNode() = default;
+	};
 
 	METASOUND_REGISTER_NODE(FTriggerCounterNode)
 }
diff --git a/Engine/Plugins/Runtime/Metasound/Source/MetasoundStandardNodes/Private/MetasoundTriggerDelayNode.cpp b/Engine/Plugins/Runtime/Metasound/Source/MetasoundStandardNodes/Private/MetasoundTriggerDelayNode.cpp
--- a/Engine/Plugins/Runtime/Metasound/Source/MetasoundStandardNodes/Private/MetasoundTriggerDelayNode.cpp
+++ b/Engine/Plugins/Runtime/Metasound/Source/MetasoundStandardNodes/Private/MetasoundTriggerDelayNode.cpp
@@ -1,6 +1,4 @@
 // Copyright Epic Games, Inc. All Rights Reserved.
-#include "MetasoundTriggerDelayNode.h"
-
 #include "CoreMinimal.h"
 
 #include "MetasoundBuilderInterface.h"
@@ -17,20 +15,18 @@
 #include "MetasoundStandardNodesCategories.h"
 #include "MetasoundParamHelper.h"
 
-#define LOCTEXT_NAMESPACE "MetasoundStandardNodes"
+#define LOCTEXT_NAMESPACE "MetasoundStandardNodes_DelayNode"
 
 namespace Metasound
 {
 	namespace TriggerDelayVertexNames
 	{
-		METASOUND_PARAM(InputInTrigger, "In", "Input trigger which results in a delayed trigger.");
-		METASOUND_PARAM(InputReset, "Reset", "Resets the trigger delay, clearing the execution task if pending.");
+		METASOUND_PARAM(InputInTriggerDelay, "In", "Input trigger which results in a delayed trigger.");
+		METASOUND_PARAM(InputResetDelay, "Reset", "Resets the trigger delay, clearing the execution task if pending.");
 		METASOUND_PARAM(InputDelayTime, "Delay Time", "Time to delay and execute deferred trigger in seconds.");
 		METASOUND_PARAM(OutputOnTrigger, "Out", "The delayed output trigger.");
 	}
 
-	METASOUND_REGISTER_NODE(FTriggerDelayNode)
-
 	class FTriggerDelayOperator : public TExecutableOperator<FTriggerDelayOperator>
 	{
 		public:
@@ -74,8 +70,8 @@ namespace Metasound
 		using namespace TriggerDelayVertexNames;
 
 		FDataReferenceCollection InputDataReferences;
-		InputDataReferences.AddDataReadReference(METASOUND_GET_PARAM_NAME(InputInTrigger), TriggerIn);
-		InputDataReferences.AddDataReadReference(METASOUND_GET_PARAM_NAME(InputReset), TriggerReset);
+		InputDataReferences.AddDataReadReference(METASOUND_GET_PARAM_NAME(InputInTriggerDelay), TriggerIn);
+		InputDataReferences.AddDataReadReference(METASOUND_GET_PARAM_NAME(InputResetDelay), TriggerReset);
 		InputDataReferences.AddDataReadReference(METASOUND_GET_PARAM_NAME(InputDelayTime), DelayTime);
 		return InputDataReferences;
 	}
@@ -144,13 +140,11 @@ namespace Metasound
 	{
 		using namespace TriggerDelayVertexNames;
 
-		const FTriggerDelayNode& DelayNode = static_cast<const FTriggerDelayNode&>(InParams.Node);
-
 		const FInputVertexInterface& InputInterface = GetVertexInterface().GetInputInterface();
 		FTimeReadRef Delay = InParams.InputDataReferences.GetDataReadReferenceOrConstructWithVertexDefault<FTime>(InputInterface, METASOUND_GET_PARAM_NAME(InputDelayTime), InParams.OperatorSettings);
 
-		FTriggerReadRef TriggerIn = InParams.InputDataReferences.GetDataReadReferenceOrConstruct<FTrigger>(METASOUND_GET_PARAM_NAME(InputInTrigger), InParams.OperatorSettings);
-		FTriggerReadRef TriggerReset = InParams.InputDataReferences.GetDataReadReferenceOrConstruct<FTrigger>(METASOUND_GET_PARAM_NAME(InputReset), InParams.OperatorSettings);
+		FTriggerReadRef TriggerIn = InParams.InputDataReferences.GetDataReadReferenceOrConstruct<FTrigger>(METASOUND_GET_PARAM_NAME(InputInTriggerDelay), InParams.OperatorSettings);
+		FTriggerReadRef TriggerReset = InParams.InputDataReferences.GetDataReadReferenceOrConstruct<FTrigger>(METASOUND_GET_PARAM_NAME(InputResetDelay), InParams.OperatorSettings);
 
 		return MakeUnique<FTriggerDelayOperator>(InParams.OperatorSettings, TriggerReset, TriggerIn, Delay);
 	}
@@ -161,8 +155,8 @@ namespace Metasound
 
 		static const FVertexInterface Interface(
 			FInputVertexInterface(
-				TInputDataVertexModel<FTrigger>(METASOUND_GET_PARAM_NAME_AND_TT(InputInTrigger)),
-				TInputDataVertexModel<FTrigger>(METASOUND_GET_PARAM_NAME_AND_TT(InputReset)),
+				TInputDataVertexModel<FTrigger>(METASOUND_GET_PARAM_NAME_AND_TT(InputInTriggerDelay)),
+				TInputDataVertexModel<FTrigger>(METASOUND_GET_PARAM_NAME_AND_TT(InputResetDelay)),
 				TInputDataVertexModel<FTime>(METASOUND_GET_PARAM_NAME_AND_TT(InputDelayTime), 1.0f)
 			),
 			FOutputVertexInterface(
@@ -195,10 +189,16 @@ namespace Metasound
 		return Info;
 	}
 
-	FTriggerDelayNode::FTriggerDelayNode(const FNodeInitData& InInitData)
-	:	FNodeFacade(InInitData.InstanceName, InInitData.InstanceID, TFacadeOperatorClass<FTriggerDelayOperator>())
+	class METASOUNDSTANDARDNODES_API FTriggerDelayNode : public FNodeFacade
 	{
-	}
+	public:
+		FTriggerDelayNode(const FNodeInitData& InInitData)
+			: FNodeFacade(InInitData.InstanceName, InInitData.InstanceID, TFacadeOperatorClass<FTriggerDelayOperator>())
+		{
+		}
+	};
+
+	METASOUND_REGISTER_NODE(FTriggerDelayNode)
 }
 
 #undef LOCTEXT_NAMESPACE // MetasoundTriggerDelayNode
diff --git a/Engine/Plugins/Runtime/Metasound/Source/MetasoundStandardNodes/Public/MetasoundTriggerCounterNode.h b/Engine/Plugins/Runtime/Metasound/Source/MetasoundStandardNodes/Public/MetasoundTriggerCounterNode.h
deleted file mode 100644
diff --git a/Engine/Plugins/Runtime/Metasound/Source/MetasoundStandardNodes/Public/MetasoundTriggerDelayNode.h b/Engine/Plugins/Runtime/Metasound/Source/MetasoundStandardNodes/Public/MetasoundTriggerDelayNode.h
deleted file mode 100644
