Adding some clamps to filter nodes

diff --git a/Engine/Plugins/Runtime/Metasound/Source/MetasoundStandardNodes/Private/MetasoundBasicFilters.cpp b/Engine/Plugins/Runtime/Metasound/Source/MetasoundStandardNodes/Private/MetasoundBasicFilters.cpp
--- a/Engine/Plugins/Runtime/Metasound/Source/MetasoundStandardNodes/Private/MetasoundBasicFilters.cpp
+++ b/Engine/Plugins/Runtime/Metasound/Source/MetasoundStandardNodes/Private/MetasoundBasicFilters.cpp
@@ -560,7 +560,8 @@ namespace Metasound
 
 	void FOnePoleLowPassFilterOperator::Execute()
 	{
-		OnePoleLowPassFilter.StartFrequencyInterpolation(*Frequency);
+		float ClampedFreq = FMath::Clamp(0.0f, *Frequency, SampleRate);
+		OnePoleLowPassFilter.StartFrequencyInterpolation(*Frequency, 1);
 		OnePoleLowPassFilter.ProcessAudioBuffer(AudioInput->GetData(), AudioOutput->GetData(), AudioInput->Num());
 		OnePoleLowPassFilter.StopFrequencyInterpolation();
 	}
@@ -694,7 +695,8 @@ namespace Metasound
 
 	void FOnePoleHighPassFilterOperator::Execute()
 	{
-		OnePoleHighPassFilter.StartFrequencyInterpolation(*Frequency);
+		float ClampedFreq = FMath::Clamp(0.0f, *Frequency, SampleRate);
+		OnePoleHighPassFilter.StartFrequencyInterpolation(ClampedFreq, BlockSize);
 		OnePoleHighPassFilter.ProcessAudioBuffer(AudioInput->GetData(), AudioOutput->GetData(), AudioInput->Num());
 		OnePoleHighPassFilter.StopFrequencyInterpolation();
 	}
