Metasound Copy/Paste Bug UX fixes

diff --git a/Engine/Plugins/Runtime/Metasound/Source/MetasoundEditor/Private/MetasoundEditorGraphBuilder.cpp b/Engine/Plugins/Runtime/Metasound/Source/MetasoundEditor/Private/MetasoundEditorGraphBuilder.cpp
--- a/Engine/Plugins/Runtime/Metasound/Source/MetasoundEditor/Private/MetasoundEditorGraphBuilder.cpp
+++ b/Engine/Plugins/Runtime/Metasound/Source/MetasoundEditor/Private/MetasoundEditorGraphBuilder.cpp
@@ -387,17 +387,28 @@ namespace Metasound
 				const FString NewInputName = GenerateUniqueInputName(InMetasound);
 				Frontend::FNodeHandle NewInputNode = AddInputNodeHandle(InMetasound, NewInputName, TypeName, FText::GetEmpty(), EMetasoundFrontendNodeStyleDisplayVisibility::Hidden, &PinDataTypeDefaultLiteral);
 
-				const FMetasoundFrontendLiteral* DefaultLiteral = InputHandle->GetDefaultLiteral();
 				Frontend::FGraphHandle GraphHandle = InNodeHandle->GetOwningGraph();
 				const FGuid VertexID = GraphHandle->GetVertexIDForInputVertex(NewInputNode->GetNodeName());
 				if (bForcePinValueAsDefault)
 				{
 					GraphHandle->SetDefaultInput(VertexID, PinDataTypeDefaultLiteral);
 				}
-				else if (DefaultLiteral && DefaultLiteral->IsValid())
+				// Hack fix for triggers having special aspect in that they hide the default
+				// literal picker, but require a default value for graph synchronization purposes.
+				else if (InInputPin.PinType.PinCategory == PinCategoryTrigger)
+				{
+					FMetasoundFrontendLiteral DefaultTriggerLiteral;
+					DefaultTriggerLiteral.Set(false);
+					GraphHandle->SetDefaultInput(VertexID, DefaultTriggerLiteral);
+					InInputPin.DefaultValue = DefaultTriggerLiteral.ToString();
+				}
+				else if (const FMetasoundFrontendLiteral* DefaultLiteral = InputHandle->GetDefaultLiteral())
 				{
-					GraphHandle->SetDefaultInput(VertexID, *DefaultLiteral);
-					InInputPin.DefaultValue = DefaultLiteral->ToString();
+					if (DefaultLiteral->IsValid())
+					{
+						GraphHandle->SetDefaultInput(VertexID, *DefaultLiteral);
+						InInputPin.DefaultValue = DefaultLiteral->ToString();
+					}
 				}
 
 				TArray<FOutputHandle> OutputHandles = NewInputNode->GetOutputs();
diff --git a/Engine/Plugins/Runtime/Metasound/Source/MetasoundFrontend/Private/MetasoundFrontendLiteral.cpp b/Engine/Plugins/Runtime/Metasound/Source/MetasoundFrontend/Private/MetasoundFrontendLiteral.cpp
--- a/Engine/Plugins/Runtime/Metasound/Source/MetasoundFrontend/Private/MetasoundFrontendLiteral.cpp
+++ b/Engine/Plugins/Runtime/Metasound/Source/MetasoundFrontend/Private/MetasoundFrontendLiteral.cpp
@@ -30,13 +30,16 @@ namespace MetasoundFrontendLiteralPrivate
 	template<>
 	struct TLiteralValueToStringHelper<UObject*>
 	{
-		static FString Convert(const UObject* InType)
+		static FString Convert(const UObject* InObject)
 		{
-			if (nullptr == InType)
+			// Use empty string for null object as this is recognized
+			// as empty value in editor context (the string "nullptr"
+			// by contrast is not).
+			if (!InObject)
 			{
-				return FString("nullptr");
+				return FString();
 			}
-			return InType->GetFullName();
+			return InObject->GetFullName();
 		}
 	};
 
