Rename Metasound to MetaSound (PascalCase) in all places in UX. (Code-side to be cleaned up post EA)

diff --git a/Engine/Plugins/Runtime/Metasound/Source/MetasoundEditor/Private/MetasoundAssetTypeActions.cpp b/Engine/Plugins/Runtime/Metasound/Source/MetasoundEditor/Private/MetasoundAssetTypeActions.cpp
--- a/Engine/Plugins/Runtime/Metasound/Source/MetasoundEditor/Private/MetasoundAssetTypeActions.cpp
+++ b/Engine/Plugins/Runtime/Metasound/Source/MetasoundEditor/Private/MetasoundAssetTypeActions.cpp
@@ -15,12 +15,12 @@ namespace Metasound
 	{
 		static const TArray<FText> AssetTypeActionSubMenu
 		{
-			LOCTEXT("AssetMetasoundSubMenu", "Metasound")
+			LOCTEXT("AssetMetaSoundSubMenu", "MetaSound")
 		};
 
 		UClass* FAssetTypeActions_Metasound::GetSupportedClass() const
 		{
-			return UMetasound::StaticClass();
+			return UMetaSound::StaticClass();
 		}
 
 		const TArray<FText>& FAssetTypeActions_Metasound::GetSubMenus() const
@@ -34,7 +34,7 @@ namespace Metasound
 
 			for (UObject* Object : InObjects)
 			{
-				if (UMetasound* Metasound = Cast<UMetasound>(Object))
+				if (UMetaSound* Metasound = Cast<UMetaSound>(Object))
 				{
 					TSharedRef<FEditor> NewEditor = MakeShared<FEditor>();
 					NewEditor->InitMetasoundEditor(Mode, ToolkitHost, Metasound);
@@ -44,7 +44,7 @@ namespace Metasound
 
 		UClass* FAssetTypeActions_MetasoundSource::GetSupportedClass() const
 		{
-			return UMetasoundSource::StaticClass();
+			return UMetaSoundSource::StaticClass();
 		}
 
 		void FAssetTypeActions_MetasoundSource::OpenAssetEditor(const TArray<UObject*>& InObjects, TSharedPtr<IToolkitHost> ToolkitHost)
@@ -53,7 +53,7 @@ namespace Metasound
 
 			for (UObject* Object : InObjects)
 			{
-				if (UMetasoundSource* Metasound = Cast<UMetasoundSource>(Object))
+				if (UMetaSoundSource* Metasound = Cast<UMetaSoundSource>(Object))
 				{
 					TSharedRef<FEditor> NewEditor = MakeShared<FEditor>();
 					NewEditor->InitMetasoundEditor(Mode, ToolkitHost, Metasound);
diff --git a/Engine/Plugins/Runtime/Metasound/Source/MetasoundEditor/Private/MetasoundAssetTypeActions.h b/Engine/Plugins/Runtime/Metasound/Source/MetasoundEditor/Private/MetasoundAssetTypeActions.h
--- a/Engine/Plugins/Runtime/Metasound/Source/MetasoundEditor/Private/MetasoundAssetTypeActions.h
+++ b/Engine/Plugins/Runtime/Metasound/Source/MetasoundEditor/Private/MetasoundAssetTypeActions.h
@@ -16,8 +16,8 @@ namespace Metasound
 		{
 		public:
 			// IAssetTypeActions Implementation
-			virtual FText GetName() const override { return NSLOCTEXT("AssetTypeActions", "AssetTypeActions_Metasound", "Metasound"); }
-			virtual FColor GetTypeColor() const override { return FColor(220, 160, 0); }
+			virtual FText GetName() const override { return NSLOCTEXT("AssetTypeActions", "AssetTypeActions_MetaSound", "MetaSound"); }
+			virtual FColor GetTypeColor() const override { return FColor(103, 214, 66); }
 			virtual UClass* GetSupportedClass() const override;
 			virtual uint32 GetCategories() override { return EAssetTypeCategories::Sounds; }
 			virtual void OpenAssetEditor(const TArray<UObject*>& InObjects, TSharedPtr<IToolkitHost> ToolkitHost) override;
@@ -29,8 +29,8 @@ namespace Metasound
 		{
 		public:
 			// IAssetTypeActions Implementation
-			virtual FText GetName() const override { return NSLOCTEXT("AssetTypeActions", "AssetTypeActions_MetasoundSource", "MetasoundSource"); }
-			virtual FColor GetTypeColor() const override { return FColor(220, 160, 0); }
+			virtual FText GetName() const override { return NSLOCTEXT("AssetTypeActions", "AssetTypeActions_MetaSoundSource", "MetaSoundSource"); }
+			virtual FColor GetTypeColor() const override { return FColor(103, 214, 66); }
 			virtual UClass* GetSupportedClass() const override;
 			virtual uint32 GetCategories() override { return EAssetTypeCategories::Sounds; }
 			virtual void OpenAssetEditor(const TArray<UObject*>& InObjects, TSharedPtr<IToolkitHost> ToolkitHost) override;
diff --git a/Engine/Plugins/Runtime/Metasound/Source/MetasoundEditor/Private/MetasoundDetailCustomization.cpp b/Engine/Plugins/Runtime/Metasound/Source/MetasoundEditor/Private/MetasoundDetailCustomization.cpp
--- a/Engine/Plugins/Runtime/Metasound/Source/MetasoundEditor/Private/MetasoundDetailCustomization.cpp
+++ b/Engine/Plugins/Runtime/Metasound/Source/MetasoundEditor/Private/MetasoundDetailCustomization.cpp
@@ -77,7 +77,7 @@ namespace Metasound
 			{
 				case EMetasoundActiveDetailView::Metasound:
 				{
-					IDetailCategoryBuilder& GeneralCategoryBuilder = DetailLayout.EditCategory("Metasound");
+					IDetailCategoryBuilder& GeneralCategoryBuilder = DetailLayout.EditCategory("MetaSound");
 					const FName AuthorPropertyPath = BuildChildPath(GetMetadataPropertyPath(), GET_MEMBER_NAME_CHECKED(FMetasoundFrontendClassMetadata, Author));
 					const FName DescPropertyPath = BuildChildPath(GetMetadataPropertyPath(), GET_MEMBER_NAME_CHECKED(FMetasoundFrontendClassMetadata, Description));
 					const FName VersionPropertyPath = BuildChildPath(GetMetadataPropertyPath(), GET_MEMBER_NAME_CHECKED(FMetasoundFrontendClassMetadata, Version));
@@ -106,7 +106,7 @@ namespace Metasound
 
 				case EMetasoundActiveDetailView::General:
 				default:
-					DetailLayout.HideCategory("Metasound");
+					DetailLayout.HideCategory("MetaSound");
 
 					const bool bShouldBeInitiallyCollapsed = true;
 					IDetailCategoryBuilder& SoundCategory = DetailLayout.EditCategory("Sound");
diff --git a/Engine/Plugins/Runtime/Metasound/Source/MetasoundEditor/Private/MetasoundEditor.cpp b/Engine/Plugins/Runtime/Metasound/Source/MetasoundEditor/Private/MetasoundEditor.cpp
--- a/Engine/Plugins/Runtime/Metasound/Source/MetasoundEditor/Private/MetasoundEditor.cpp
+++ b/Engine/Plugins/Runtime/Metasound/Source/MetasoundEditor/Private/MetasoundEditor.cpp
@@ -146,7 +146,7 @@ namespace Metasound
 
 		void FEditor::RegisterTabSpawners(const TSharedRef<FTabManager>& InTabManager)
 		{
-			WorkspaceMenuCategory = InTabManager->AddLocalWorkspaceMenuCategory(LOCTEXT("WorkspaceMenu_MetasoundEditor", "Metasound Editor"));
+			WorkspaceMenuCategory = InTabManager->AddLocalWorkspaceMenuCategory(LOCTEXT("WorkspaceMenu_MetasoundEditor", "MetaSound Editor"));
 			auto WorkspaceMenuCategoryRef = WorkspaceMenuCategory.ToSharedRef();
 
 			FAssetEditorToolkit::RegisterTabSpawners(InTabManager);
@@ -181,7 +181,7 @@ namespace Metasound
 			{
 				return TabFactory::CreateMetasoundTab(InMetasoundMenu, Args);
 			}))
-			.SetDisplayName(LOCTEXT("MetasoundTab", "Metasound"))
+			.SetDisplayName(LOCTEXT("MetasoundTab", "MetaSound"))
 			.SetGroup(WorkspaceMenuCategoryRef)
 			.SetIcon(FSlateIcon(FEditorStyle::GetStyleSetName(), "LevelEditor.Tabs.Settings"));
 
@@ -268,7 +268,7 @@ namespace Metasound
 
 			CreateAnalyzers();
 
-			const TSharedRef<FTabManager::FLayout> StandaloneDefaultLayout = FTabManager::NewLayout("Standalone_MetasoundEditor_Layout_v7")
+			const TSharedRef<FTabManager::FLayout> StandaloneDefaultLayout = FTabManager::NewLayout("Standalone_MetasoundEditor_Layout_v8")
 				->AddArea
 				(
 					FTabManager::NewPrimaryArea()
@@ -346,17 +346,17 @@ namespace Metasound
 
 		FName FEditor::GetToolkitFName() const
 		{
-			return "MetasoundEditor";
+			return "MetaSoundEditor";
 		}
 
 		FText FEditor::GetBaseToolkitName() const
 		{
-			return LOCTEXT("AppLabel", "Metasound Editor");
+			return LOCTEXT("AppLabel", "MetaSound Editor");
 		}
 
 		FString FEditor::GetWorldCentricTabPrefix() const
 		{
-			return LOCTEXT("WorldCentricTabPrefix", "Metasound ").ToString();
+			return LOCTEXT("WorldCentricTabPrefix", "MetaSound ").ToString();
 		}
 
 		FLinearColor FEditor::GetWorldCentricTabColorScale() const
@@ -993,7 +993,7 @@ namespace Metasound
 			}
 
 			FGraphAppearanceInfo AppearanceInfo;
-			AppearanceInfo.CornerText = LOCTEXT("AppearanceCornerText_Metasound", "Metasound");
+			AppearanceInfo.CornerText = LOCTEXT("AppearanceCornerText_MetaSound", "MetaSound");
 
 			SGraphEditor::FGraphEditorEvents InEvents;
 			InEvents.OnSelectionChanged = SGraphEditor::FOnSelectionChanged::CreateSP(this, &FEditor::OnSelectedNodesChanged);
@@ -1688,7 +1688,7 @@ namespace Metasound
 			}
 		}
 
-		// TODO: Add ability to filter inputs/outputs in "Metasound" Tab
+		// TODO: Add ability to filter inputs/outputs in "MetaSound" Tab
 		FText FEditor::GetFilterText() const
 		{
 			return FText::GetEmpty();
diff --git a/Engine/Plugins/Runtime/Metasound/Source/MetasoundEditor/Private/MetasoundEditor.h b/Engine/Plugins/Runtime/Metasound/Source/MetasoundEditor/Private/MetasoundEditor.h
--- a/Engine/Plugins/Runtime/Metasound/Source/MetasoundEditor/Private/MetasoundEditor.h
+++ b/Engine/Plugins/Runtime/Metasound/Source/MetasoundEditor/Private/MetasoundEditor.h
@@ -36,7 +36,7 @@ class IDetailsView;
 class IToolkitHost;
 class SVerticalBox;
 class UEdGraphNode;
-class UMetasound;
+class UMetaSound;
 class UMetasoundEditorGraph;
 
 struct FGraphActionNode;
diff --git a/Engine/Plugins/Runtime/Metasound/Source/MetasoundEditor/Private/MetasoundEditorCommands.cpp b/Engine/Plugins/Runtime/Metasound/Source/MetasoundEditor/Private/MetasoundEditorCommands.cpp
--- a/Engine/Plugins/Runtime/Metasound/Source/MetasoundEditor/Private/MetasoundEditorCommands.cpp
+++ b/Engine/Plugins/Runtime/Metasound/Source/MetasoundEditor/Private/MetasoundEditorCommands.cpp
@@ -9,18 +9,18 @@ namespace Metasound
 	{
 		void FEditorCommands::RegisterCommands()
 		{
-			UI_COMMAND(Play, "Play", "Plays (or restarts) the Metasound", EUserInterfaceActionType::Button, FInputChord());
-			UI_COMMAND(Stop, "Stop", "Stops Metasound (If currently playing)", EUserInterfaceActionType::Button, FInputChord());
+			UI_COMMAND(Play, "Play", "Plays (or restarts) the MetaSound", EUserInterfaceActionType::Button, FInputChord());
+			UI_COMMAND(Stop, "Stop", "Stops MetaSound (If currently playing)", EUserInterfaceActionType::Button, FInputChord());
 			UI_COMMAND(TogglePlayback, "Toggle Playback", "Plays the Metasound or stops the currently playing Metasound", EUserInterfaceActionType::Button, FInputChord(EKeys::SpaceBar));
 
-			UI_COMMAND(Import, "Import", "Imports Metasound from Json", EUserInterfaceActionType::Button, FInputChord());
-			UI_COMMAND(Export, "Export", "Exports Metasound to Json", EUserInterfaceActionType::Button, FInputChord());
+			UI_COMMAND(Import, "Import", "Imports MetaSound from Json", EUserInterfaceActionType::Button, FInputChord());
+			UI_COMMAND(Export, "Export", "Exports MetaSound to Json", EUserInterfaceActionType::Button, FInputChord());
 
-			UI_COMMAND(BrowserSync, "Browse", "Selects the Metasound in the content browser. If referencing Metasound nodes are selected, selects referenced assets instead.", EUserInterfaceActionType::Button, FInputChord());
+			UI_COMMAND(BrowserSync, "Browse", "Selects the MetaSound in the content browser. If referencing Metasound nodes are selected, selects referenced assets instead.", EUserInterfaceActionType::Button, FInputChord());
 			UI_COMMAND(AddInput, "Add Input", "Adds an input to the node", EUserInterfaceActionType::Button, FInputChord());
 			UI_COMMAND(DeleteInput, "Delete Input", "Removes an input from the node", EUserInterfaceActionType::Button, FInputChord());
 
-			UI_COMMAND(EditMetasoundSettings, "Metasound", "Edit Metasound settings", EUserInterfaceActionType::Button, FInputChord());
+			UI_COMMAND(EditMetasoundSettings, "MetaSound", "Edit MetaSound settings", EUserInterfaceActionType::Button, FInputChord());
 			UI_COMMAND(EditGeneralSettings, "General", "Edit general sound settings", EUserInterfaceActionType::Button, FInputChord());
 		}
 	} // namespace Editor
diff --git a/Engine/Plugins/Runtime/Metasound/Source/MetasoundEditor/Private/MetasoundEditorGraphBuilder.h b/Engine/Plugins/Runtime/Metasound/Source/MetasoundEditor/Private/MetasoundEditorGraphBuilder.h
--- a/Engine/Plugins/Runtime/Metasound/Source/MetasoundEditor/Private/MetasoundEditorGraphBuilder.h
+++ b/Engine/Plugins/Runtime/Metasound/Source/MetasoundEditor/Private/MetasoundEditorGraphBuilder.h
@@ -14,7 +14,7 @@
 // Forward Declarations
 class UEdGraphNode;
 class UEdGraphPin;
-class UMetasound;
+class UMetaSound;
 class UMetasoundEditorGraphNode;
 
 struct FEdGraphPinType;
diff --git a/Engine/Plugins/Runtime/Metasound/Source/MetasoundEditor/Private/MetasoundEditorGraphNode.cpp b/Engine/Plugins/Runtime/Metasound/Source/MetasoundEditor/Private/MetasoundEditorGraphNode.cpp
--- a/Engine/Plugins/Runtime/Metasound/Source/MetasoundEditor/Private/MetasoundEditorGraphNode.cpp
+++ b/Engine/Plugins/Runtime/Metasound/Source/MetasoundEditor/Private/MetasoundEditorGraphNode.cpp
@@ -57,7 +57,7 @@ void UMetasoundEditorGraphNode::CreateInputPin()
 		return;
 	}
 
-	FString PinName; // get from UMetasound
+	FString PinName; // get from UMetaSound
 	UEdGraphPin* NewPin = CreatePin(EGPD_Input, TEXT("MetasoundEditorGraphNode"), *PinName);
 	if (NewPin->PinName.IsNone())
 	{
@@ -357,7 +357,7 @@ FString UMetasoundEditorGraphNode::GetDocumentationExcerptName() const
 {
 	// Default the node to searching for an excerpt named for the C++ node class name, including the U prefix.
 	// This is done so that the excerpt name in the doc file can be found by find-in-files when searching for the full class name.
-	return FString::Printf(TEXT("%s%s"), UMetasound::StaticClass()->GetPrefixCPP(), *UMetasound::StaticClass()->GetName());
+	return FString::Printf(TEXT("%s%s"), UMetaSound::StaticClass()->GetPrefixCPP(), *UMetaSound::StaticClass()->GetName());
 }
 
 FMetasoundFrontendClassName UMetasoundEditorGraphOutputNode::GetClassName() const
diff --git a/Engine/Plugins/Runtime/Metasound/Source/MetasoundEditor/Private/MetasoundEditorModule.cpp b/Engine/Plugins/Runtime/Metasound/Source/MetasoundEditor/Private/MetasoundEditorModule.cpp
--- a/Engine/Plugins/Runtime/Metasound/Source/MetasoundEditor/Private/MetasoundEditorModule.cpp
+++ b/Engine/Plugins/Runtime/Metasound/Source/MetasoundEditor/Private/MetasoundEditorModule.cpp
@@ -293,12 +293,12 @@ namespace Metasound
 				FPropertyEditorModule& PropertyModule = FModuleManager::LoadModuleChecked<FPropertyEditorModule>("PropertyEditor");
 				
 				PropertyModule.RegisterCustomClassLayout(
-					UMetasound::StaticClass()->GetFName(),
-					FOnGetDetailCustomizationInstance::CreateLambda([]() { return MakeShared<FMetasoundDetailCustomization>(UMetasound::GetDocumentPropertyName()); }));
+					UMetaSound::StaticClass()->GetFName(),
+					FOnGetDetailCustomizationInstance::CreateLambda([]() { return MakeShared<FMetasoundDetailCustomization>(UMetaSound::GetDocumentPropertyName()); }));
 
 				PropertyModule.RegisterCustomClassLayout(
-					UMetasoundSource::StaticClass()->GetFName(),
-					FOnGetDetailCustomizationInstance::CreateLambda([]() { return MakeShared<FMetasoundDetailCustomization>(UMetasoundSource::GetDocumentPropertyName()); }));
+					UMetaSoundSource::StaticClass()->GetFName(),
+					FOnGetDetailCustomizationInstance::CreateLambda([]() { return MakeShared<FMetasoundDetailCustomization>(UMetaSoundSource::GetDocumentPropertyName()); }));
 
 				PropertyModule.RegisterCustomClassLayout(
 					UMetasoundEditorGraphInput::StaticClass()->GetFName(),
diff --git a/Engine/Plugins/Runtime/Metasound/Source/MetasoundEditor/Private/MetasoundEditorTabFactory.h b/Engine/Plugins/Runtime/Metasound/Source/MetasoundEditor/Private/MetasoundEditorTabFactory.h
--- a/Engine/Plugins/Runtime/Metasound/Source/MetasoundEditor/Private/MetasoundEditorTabFactory.h
+++ b/Engine/Plugins/Runtime/Metasound/Source/MetasoundEditor/Private/MetasoundEditorTabFactory.h
@@ -7,10 +7,13 @@
 #include "Internationalization/Internationalization.h"
 #include "Widgets/Docking/SDockTab.h"
 #include "SAudioMeter.h"
+#include "Styling/ISlateStyle.h"
+#include "Styling/SlateStyleRegistry.h"
 #include "Templates/SharedPointer.h"
 
 #define LOCTEXT_NAMESPACE "MetasoundEditor"
 
+
 namespace Metasound
 {
 	namespace Editor
@@ -59,7 +62,7 @@ namespace Metasound
 
 				return SNew(SDockTab)
 					.Icon(FEditorStyle::GetBrush("LevelEditor.Tabs.Details"))
-					.Label(LOCTEXT("MetasoundInspectorTitle", "Inspector"))
+					.Label(LOCTEXT("MetaSoundInspectorTitle", "Inspector"))
 					[
 						DetailsView.ToSharedRef()
 					];
@@ -69,12 +72,18 @@ namespace Metasound
 			{
 				check(Args.GetTabId() == Names::Metasound);
 
-				return SNew(SDockTab)
-					.Icon(FAppStyle::Get().GetBrush("Icons.Settings"))
-					.Label(LOCTEXT("MetasoundGeneralTitle", "Metasound"))
-					[
-						GraphActionMenu.ToSharedRef()
-					];
+				TSharedRef<SDockTab> NewTab = SNew(SDockTab)
+				.Label(LOCTEXT("MetaSoundGeneralTitle", "MetaSound"))
+				[
+					GraphActionMenu.ToSharedRef()
+				];
+
+				if (const ISlateStyle* MetasoundStyle = FSlateStyleRegistry::FindSlateStyle("MetasoundStyle"))
+				{
+					NewTab->SetTabIcon(MetasoundStyle->GetBrush("ClassThumbnail.Metasound"));
+				}
+
+				return NewTab;
 			}
 
 			TSharedRef<SDockTab> CreatePaletteTab(TSharedPtr<SMetasoundPalette> Palette, const FSpawnTabArgs& Args)
@@ -83,7 +92,7 @@ namespace Metasound
 
 				return SNew(SDockTab)
 					.Icon(FEditorStyle::GetBrush("Kismet.Tabs.Palette"))
-					.Label(LOCTEXT("MetasoundPaletteTitle", "Palette"))
+					.Label(LOCTEXT("MetaSoundPaletteTitle", "Palette"))
 					[
 						Palette.ToSharedRef()
 					];
diff --git a/Engine/Plugins/Runtime/Metasound/Source/MetasoundEditor/Private/MetasoundFactory.cpp b/Engine/Plugins/Runtime/Metasound/Source/MetasoundEditor/Private/MetasoundFactory.cpp
--- a/Engine/Plugins/Runtime/Metasound/Source/MetasoundEditor/Private/MetasoundFactory.cpp
+++ b/Engine/Plugins/Runtime/Metasound/Source/MetasoundEditor/Private/MetasoundFactory.cpp
@@ -16,7 +16,7 @@ namespace Metasound
 	{
 		namespace FactoryPrivate
 		{
-			UObject* InitMetasound(UMetasoundSource* MetasoundSource, EObjectFlags Flags)
+			UObject* InitMetasound(UMetaSoundSource* MetasoundSource, EObjectFlags Flags)
 			{
 				FMetasoundFrontendClassMetadata Metadata;
 
@@ -49,7 +49,7 @@ namespace Metasound
 // UMetasoundFactory::UMetasoundFactory(const FObjectInitializer& ObjectInitializer)
 // 	: Super(ObjectInitializer)
 // {
-// 	SupportedClass = UMetasound::StaticClass();
+// 	SupportedClass = UMetaSound::StaticClass();
 // 
 // 	bCreateNew = true;
 // 	bEditorImport = false;
@@ -58,7 +58,7 @@ namespace Metasound
 // 
 // UObject* UMetasoundFactory::FactoryCreateNew(UClass* InClass, UObject* InParent, FName Name, EObjectFlags Flags, UObject* InContext, FFeedbackContext* InFeedbackContext)
 // {
-// 	UMetasound* Metasound = NewObject<UMetasound>(InParent, Name, Flags);
+// 	UMetaSound* Metasound = NewObject<UMetaSound>(InParent, Name, Flags);
 //	Metasound::Editor::FactoryPrivate::InitMetasound(Metasound, Flags);
 // 
 //	CastChecked<UMetasoundEditorGraph>(&Metasound->GetGraphChecked())->ParentMetasound = MetasoundSource;
@@ -70,7 +70,7 @@ namespace Metasound
 UMetasoundSourceFactory::UMetasoundSourceFactory(const FObjectInitializer& ObjectInitializer)
 	: Super(ObjectInitializer)
 {
-	SupportedClass = UMetasoundSource::StaticClass();
+	SupportedClass = UMetaSoundSource::StaticClass();
 
 	bCreateNew = true;
 	bEditorImport = false;
@@ -83,7 +83,7 @@ UObject* UMetasoundSourceFactory::FactoryCreateNew(UClass* InClass, UObject* InP
 	using namespace Metasound::Editor;
 	using namespace Metasound::Frontend;
 
-	UMetasoundSource* MetasoundSource = NewObject<UMetasoundSource>(InParent, Name, Flags);
+	UMetaSoundSource* MetasoundSource = NewObject<UMetaSoundSource>(InParent, Name, Flags);
 	Metasound::Editor::FactoryPrivate::InitMetasound(MetasoundSource, Flags);
 
 	CastChecked<UMetasoundEditorGraph>(&MetasoundSource->GetGraphChecked())->ParentMetasound = MetasoundSource;
diff --git a/Engine/Plugins/Runtime/Metasound/Source/MetasoundEditor/Public/MetasoundEditorGraphInputNodes.h b/Engine/Plugins/Runtime/Metasound/Source/MetasoundEditor/Public/MetasoundEditorGraphInputNodes.h
--- a/Engine/Plugins/Runtime/Metasound/Source/MetasoundEditor/Public/MetasoundEditorGraphInputNodes.h
+++ b/Engine/Plugins/Runtime/Metasound/Source/MetasoundEditor/Public/MetasoundEditorGraphInputNodes.h
@@ -23,7 +23,7 @@
 
 // Forward Declarations
 class UEdGraphPin;
-class UMetasound;
+class UMetaSound;
 
 namespace Metasound
 {
diff --git a/Engine/Plugins/Runtime/Metasound/Source/MetasoundEditor/Public/MetasoundEditorGraphNode.h b/Engine/Plugins/Runtime/Metasound/Source/MetasoundEditor/Public/MetasoundEditorGraphNode.h
--- a/Engine/Plugins/Runtime/Metasound/Source/MetasoundEditor/Public/MetasoundEditorGraphNode.h
+++ b/Engine/Plugins/Runtime/Metasound/Source/MetasoundEditor/Public/MetasoundEditorGraphNode.h
@@ -17,7 +17,7 @@
 
 // Forward Declarations
 class UEdGraphPin;
-class UMetasound;
+class UMetaSound;
 class UMetasoundEditorGraphOutput;
 
 namespace Metasound
diff --git a/Engine/Plugins/Runtime/Metasound/Source/MetasoundEditor/Public/MetasoundEditorGraphSchema.h b/Engine/Plugins/Runtime/Metasound/Source/MetasoundEditor/Public/MetasoundEditorGraphSchema.h
--- a/Engine/Plugins/Runtime/Metasound/Source/MetasoundEditor/Public/MetasoundEditorGraphSchema.h
+++ b/Engine/Plugins/Runtime/Metasound/Source/MetasoundEditor/Public/MetasoundEditorGraphSchema.h
@@ -19,7 +19,7 @@
 // Forward Declarations
 class UEdGraph;
 class UEdGraphNode;
-class UMetasound;
+class UMetaSound;
 class UMetasoundEditorGraphNode;
 
 
diff --git a/Engine/Plugins/Runtime/Metasound/Source/MetasoundEngine/Private/Metasound.cpp b/Engine/Plugins/Runtime/Metasound/Source/MetasoundEngine/Private/Metasound.cpp
--- a/Engine/Plugins/Runtime/Metasound/Source/MetasoundEngine/Private/Metasound.cpp
+++ b/Engine/Plugins/Runtime/Metasound/Source/MetasoundEngine/Private/Metasound.cpp
@@ -6,67 +6,67 @@
 #include "StructSerializer.h"
 #include "UObject/UnrealType.h"
 
-UMetasound::UMetasound(const FObjectInitializer& ObjectInitializer)
+UMetaSound::UMetaSound(const FObjectInitializer& ObjectInitializer)
 	: Super(ObjectInitializer)
 	, FMetasoundAssetBase(FMetasoundFrontendArchetype())
 {
 }
 
 #if WITH_EDITORONLY_DATA
-UEdGraph* UMetasound::GetGraph()
+UEdGraph* UMetaSound::GetGraph()
 {
 	return Graph;
 }
 
-const UEdGraph* UMetasound::GetGraph() const
+const UEdGraph* UMetaSound::GetGraph() const
 {
 	return Graph;
 }
 
-UEdGraph& UMetasound::GetGraphChecked()
+UEdGraph& UMetaSound::GetGraphChecked()
 {
 	check(Graph);
 	return *Graph;
 }
 
-const UEdGraph& UMetasound::GetGraphChecked() const
+const UEdGraph& UMetaSound::GetGraphChecked() const
 {
 	check(Graph);
 	return *Graph;
 }
 
-void UMetasound::SetGraph(UEdGraph* InGraph)
+void UMetaSound::SetGraph(UEdGraph* InGraph)
 {
 	Graph = InGraph;
 }
 #endif // WITH_EDITORONLY_DATA
 
 // Returns document object responsible for serializing asset
-Metasound::Frontend::TAccessPtr<FMetasoundFrontendDocument> UMetasound::GetDocument()
+Metasound::Frontend::TAccessPtr<FMetasoundFrontendDocument> UMetaSound::GetDocument()
 {
 	return Metasound::Frontend::MakeAccessPtr(MetasoundDocument.AccessPoint, MetasoundDocument);
 }
 
 // Returns document object responsible for serializing asset
-Metasound::Frontend::TAccessPtr<const FMetasoundFrontendDocument> UMetasound::GetDocument() const
+Metasound::Frontend::TAccessPtr<const FMetasoundFrontendDocument> UMetaSound::GetDocument() const
 {
 	return Metasound::Frontend::MakeAccessPtr<const FMetasoundFrontendDocument>(MetasoundDocument.AccessPoint, MetasoundDocument);
 }
 
-const TArray<FMetasoundFrontendArchetype>& UMetasound::GetPreferredMetasoundArchetypes() const
+const TArray<FMetasoundFrontendArchetype>& UMetaSound::GetPreferredMetasoundArchetypes() const
 {
-	// Not preferred archetypes for a basic UMetasound.
+	// Not preferred archetypes for a basic UMetaSound.
 	static const TArray<FMetasoundFrontendArchetype> Preferred;
 	return Preferred;
 }
 
-bool UMetasound::IsMetasoundArchetypeSupported(const FMetasoundFrontendArchetype& InArchetype) const
+bool UMetaSound::IsMetasoundArchetypeSupported(const FMetasoundFrontendArchetype& InArchetype) const
 {
 	// All archetypes are supported.
 	return true;
 }
 
-const FMetasoundFrontendArchetype& UMetasound::GetPreferredMetasoundArchetype(const FMetasoundFrontendDocument& InDocument) const
+const FMetasoundFrontendArchetype& UMetaSound::GetPreferredMetasoundArchetype(const FMetasoundFrontendDocument& InDocument) const
 {
 	// Prefer to keep original archetype.
 	return InDocument.Archetype;
diff --git a/Engine/Plugins/Runtime/Metasound/Source/MetasoundEngine/Private/MetasoundEngineModule.cpp b/Engine/Plugins/Runtime/Metasound/Source/MetasoundEngine/Private/MetasoundEngineModule.cpp
--- a/Engine/Plugins/Runtime/Metasound/Source/MetasoundEngine/Private/MetasoundEngineModule.cpp
+++ b/Engine/Plugins/Runtime/Metasound/Source/MetasoundEngine/Private/MetasoundEngineModule.cpp
@@ -17,11 +17,11 @@ class FMetasoundEngineModule : public IMetasoundEngineModule
 {
 	virtual void StartupModule() override
 	{
-		// If there is no archetype name, use UMetasound
-		Metasound::IMetasoundUObjectRegistry::RegisterUClassArchetype<UMetasound>(TEXT(""));
+		// If there is no archetype name, use UMetaSound
+		Metasound::IMetasoundUObjectRegistry::RegisterUClassArchetype<UMetaSound>(TEXT(""));
 
 		// Register preferred archetypes
-		Metasound::IMetasoundUObjectRegistry::RegisterUClassPreferredArchetypes<UMetasoundSource>();
+		Metasound::IMetasoundUObjectRegistry::RegisterUClassPreferredArchetypes<UMetaSoundSource>();
 
 		// flush node regsitration queue
 		FMetasoundFrontendRegistryContainer::Get()->RegisterPendingNodes();
diff --git a/Engine/Plugins/Runtime/Metasound/Source/MetasoundEngine/Private/MetasoundEngineModule.h b/Engine/Plugins/Runtime/Metasound/Source/MetasoundEngine/Private/MetasoundEngineModule.h
--- a/Engine/Plugins/Runtime/Metasound/Source/MetasoundEngine/Private/MetasoundEngineModule.h
+++ b/Engine/Plugins/Runtime/Metasound/Source/MetasoundEngine/Private/MetasoundEngineModule.h
@@ -6,7 +6,7 @@
 #include "Modules/ModuleInterface.h"
 
 // Forward Declarations
-class UMetasound;
+class UMetaSound;
 
 DECLARE_LOG_CATEGORY_EXTERN(LogMetasoundEngine, Log, All);
 
diff --git a/Engine/Plugins/Runtime/Metasound/Source/MetasoundEngine/Private/MetasoundSource.cpp b/Engine/Plugins/Runtime/Metasound/Source/MetasoundEngine/Private/MetasoundSource.cpp
--- a/Engine/Plugins/Runtime/Metasound/Source/MetasoundEngine/Private/MetasoundSource.cpp
+++ b/Engine/Plugins/Runtime/Metasound/Source/MetasoundEngine/Private/MetasoundSource.cpp
@@ -24,6 +24,9 @@
 #include "EdGraph/EdGraph.h"
 #endif // WITH_EDITORONLY_DATA
 
+#define LOCTEXT_NAMESPACE "MetasoundSource"
+
+
 static float MetaSoundBlockRateCVar = 100.f;
 FAutoConsoleVariableRef CVarMetaSoundBlockRate(
 	TEXT("au.MetaSound.BlockRate"),
@@ -32,13 +35,9 @@ FAutoConsoleVariableRef CVarMetaSoundBlockRate(
 	TEXT("Default: 100.0f"),
 	ECVF_Default);
 
+static const FName MetasoundSourceArchetypeName = "Metasound Source";
 
-
-#define LOCTEXT_NAMESPACE "MetasoundSource"
-
-static FName MetasoundSourceArchetypeName = FName(TEXT("Metasound Source"));
-
-UMetasoundSource::UMetasoundSource(const FObjectInitializer& ObjectInitializer)
+UMetaSoundSource::UMetaSoundSource(const FObjectInitializer& ObjectInitializer)
 	: Super(ObjectInitializer)
 	, FMetasoundAssetBase(GetMonoSourceArchetype())
 {
@@ -47,14 +46,14 @@ UMetasoundSource::UMetasoundSource(const FObjectInitializer& ObjectInitializer)
 	Duration = INDEFINITELY_LOOPING_DURATION;
 	bLooping = true;
 	
-	// todo: ensure that we have a method so that the audio engine can be authoritative over the sample rate the UMetasoundSource runs at.
+	// todo: ensure that we have a method so that the audio engine can be authoritative over the sample rate the UMetaSoundSource runs at.
 	SampleRate = 48000.f;
 
 }
 
 #if WITH_EDITOR
 
-void UMetasoundSource::PostEditUndo()
+void UMetaSoundSource::PostEditUndo()
 {
 	Super::PostEditUndo();
 
@@ -64,11 +63,11 @@ void UMetasoundSource::PostEditUndo()
 	}
 }
 
-void UMetasoundSource::PostEditChangeProperty(FPropertyChangedEvent& InEvent)
+void UMetaSoundSource::PostEditChangeProperty(FPropertyChangedEvent& InEvent)
 {
 	Super::PostEditChangeProperty(InEvent);
 
-	if (InEvent.GetPropertyName() == GET_MEMBER_NAME_CHECKED(UMetasoundSource, OutputFormat)) 
+	if (InEvent.GetPropertyName() == GET_MEMBER_NAME_CHECKED(UMetaSoundSource, OutputFormat)) 
 	{
 		switch (OutputFormat)
 		{
@@ -92,24 +91,24 @@ void UMetasoundSource::PostEditChangeProperty(FPropertyChangedEvent& InEvent)
 
 #endif // WITHEDITOR
 
-bool UMetasoundSource::IsPlayable() const
+bool UMetaSoundSource::IsPlayable() const
 {
 	// todo: cache off whether this metasound is buildable to an operator.
 	return true;
 }
 
-bool UMetasoundSource::SupportsSubtitles() const
+bool UMetaSoundSource::SupportsSubtitles() const
 {
 	return Super::SupportsSubtitles();
 }
 
-float UMetasoundSource::GetDuration()
+float UMetaSoundSource::GetDuration()
 {
 	// eh? this is kind of a weird field anyways.
 	return Super::GetDuration();
 }
 
-ISoundGeneratorPtr UMetasoundSource::CreateSoundGenerator(const FSoundGeneratorInitParams& InParams)
+ISoundGeneratorPtr UMetaSoundSource::CreateSoundGenerator(const FSoundGeneratorInitParams& InParams)
 {
 	using namespace Metasound;
 
@@ -142,7 +141,7 @@ ISoundGeneratorPtr UMetasoundSource::CreateSoundGenerator(const FSoundGeneratorI
 	const FMetasoundFrontendDocument* OriginalDoc = GetDocument().Get();
 	if (nullptr == OriginalDoc)
 	{
-		UE_LOG(LogMetasound, Error, TEXT("Cannot create sound generator. Null Metasound document in UMetasoundSource [Name:%s]"), *GetName());
+		UE_LOG(LogMetasound, Error, TEXT("Cannot create sound generator. Null Metasound document in UMetaSoundSource [Name:%s]"), *GetName());
 		return ISoundGeneratorPtr(nullptr);
 	}
 
@@ -170,7 +169,7 @@ ISoundGeneratorPtr UMetasoundSource::CreateSoundGenerator(const FSoundGeneratorI
 	return ISoundGeneratorPtr(new FMetasoundGenerator(MoveTemp(InitParams)));
 }
 
-TUniquePtr<IAudioInstanceTransmitter> UMetasoundSource::CreateInstanceTransmitter(const FAudioInstanceTransmitterInitParams& InParams) const
+TUniquePtr<IAudioInstanceTransmitter> UMetaSoundSource::CreateInstanceTransmitter(const FAudioInstanceTransmitterInitParams& InParams) const
 {
 	Metasound::FMetasoundInstanceTransmitter::FInitParams InitParams(GetOperatorSettings(InParams.SampleRate), InParams.InstanceID);
 
@@ -182,7 +181,7 @@ TUniquePtr<IAudioInstanceTransmitter> UMetasoundSource::CreateInstanceTransmitte
 	return MakeUnique<Metasound::FMetasoundInstanceTransmitter>(InitParams);
 }
 
-bool UMetasoundSource::GetReceiveNodeMetadataForDataType(const FName& InTypeName, FMetasoundFrontendClassMetadata& OutMetadata) const
+bool UMetaSoundSource::GetReceiveNodeMetadataForDataType(const FName& InTypeName, FMetasoundFrontendClassMetadata& OutMetadata) const
 {
 	using namespace Metasound;
 
@@ -197,7 +196,7 @@ bool UMetasoundSource::GetReceiveNodeMetadataForDataType(const FName& InTypeName
 	return false;
 }
 
-Metasound::Frontend::FNodeHandle UMetasoundSource::AddInputPinForSendAddress(const Metasound::FMetasoundInstanceTransmitter::FSendInfo& InSendInfo, Metasound::Frontend::FGraphHandle InGraph) const
+Metasound::Frontend::FNodeHandle UMetaSoundSource::AddInputPinForSendAddress(const Metasound::FMetasoundInstanceTransmitter::FSendInfo& InSendInfo, Metasound::Frontend::FGraphHandle InGraph) const
 {
 	FMetasoundFrontendClassInput Description;
 	FGuid VertexID = InGraph->GetNewVertexID();
@@ -211,7 +210,7 @@ Metasound::Frontend::FNodeHandle UMetasoundSource::AddInputPinForSendAddress(con
 	return InGraph->AddInputVertex(Description);
 }
 
-bool UMetasoundSource::CopyDocumentAndInjectReceiveNodes(uint64 InInstanceID, const FMetasoundFrontendDocument& InSourceDoc, FMetasoundFrontendDocument& OutDestDoc) const
+bool UMetaSoundSource::CopyDocumentAndInjectReceiveNodes(uint64 InInstanceID, const FMetasoundFrontendDocument& InSourceDoc, FMetasoundFrontendDocument& OutDestDoc) const
 {
 	using namespace Metasound;
 
@@ -237,7 +236,7 @@ bool UMetasoundSource::CopyDocumentAndInjectReceiveNodes(uint64 InInstanceID, co
 		Frontend::FNodeHandle ReceiveNode = RootGraph->AddNode(ReceiveNodeMetadata);
 
 		// Add receive node address to graph
-		Frontend::FNodeHandle AddressNode = UMetasoundSource::AddInputPinForSendAddress(InfoAndVertexName.SendInfo, RootGraph);
+		Frontend::FNodeHandle AddressNode = UMetaSoundSource::AddInputPinForSendAddress(InfoAndVertexName.SendInfo, RootGraph);
 		TArray<Frontend::FOutputHandle> AddressNodeOutputs = AddressNode->GetOutputs();
 		if (AddressNodeOutputs.Num() != 1)
 		{
@@ -295,7 +294,7 @@ bool UMetasoundSource::CopyDocumentAndInjectReceiveNodes(uint64 InInstanceID, co
 	return true;
 }
 
-TArray<FString> UMetasoundSource::GetTransmittableInputVertexNames() const
+TArray<FString> UMetaSoundSource::GetTransmittableInputVertexNames() const
 {
 	using namespace Metasound;
 
@@ -347,14 +346,14 @@ TArray<FString> UMetasoundSource::GetTransmittableInputVertexNames() const
 	return GraphInputVertexNames;
 }
 
-Metasound::FOperatorSettings UMetasoundSource::GetOperatorSettings(Metasound::FSampleRate InSampleRate) const
+Metasound::FOperatorSettings UMetaSoundSource::GetOperatorSettings(Metasound::FSampleRate InSampleRate) const
 {
 	// Metasound graph gets evaluated 100 times per second by default.
 	float BlockRate = FMath::Clamp(MetaSoundBlockRateCVar, 1.0f, 1000.0f); 
 	return Metasound::FOperatorSettings(InSampleRate, BlockRate);
 }
 
-Metasound::FSendAddress UMetasoundSource::CreateSendAddress(uint64 InInstanceID, const FString& InVertexName, const FName& InDataTypeName) const
+Metasound::FSendAddress UMetaSoundSource::CreateSendAddress(uint64 InInstanceID, const FString& InVertexName, const FName& InDataTypeName) const
 {
 	using namespace Metasound;
 
@@ -366,7 +365,7 @@ Metasound::FSendAddress UMetasoundSource::CreateSendAddress(uint64 InInstanceID,
 	return Address;
 }
 
-TArray<UMetasoundSource::FSendInfoAndVertexName> UMetasoundSource::GetSendInfos(uint64 InInstanceID) const
+TArray<UMetaSoundSource::FSendInfoAndVertexName> UMetaSoundSource::GetSendInfos(uint64 InInstanceID) const
 {
 	using FSendInfo = Metasound::FMetasoundInstanceTransmitter::FSendInfo;
 	using namespace Metasound::Frontend;
@@ -398,52 +397,52 @@ TArray<UMetasoundSource::FSendInfoAndVertexName> UMetasoundSource::GetSendInfos(
 	return SendInfos;
 }
 
-const TArray<FMetasoundFrontendArchetype>& UMetasoundSource::GetPreferredMetasoundArchetypes() const 
+const TArray<FMetasoundFrontendArchetype>& UMetaSoundSource::GetPreferredMetasoundArchetypes() const 
 {
 	static const TArray<FMetasoundFrontendArchetype> Preferred({GetMonoSourceArchetype(), GetStereoSourceArchetype()});
 
 	return Preferred;
 }
 
-const FString& UMetasoundSource::GetOnPlayInputName()
+const FString& UMetaSoundSource::GetOnPlayInputName()
 {
 	static const FString TriggerInputName = TEXT("On Play");
 	return TriggerInputName;
 }
 
-const FString& UMetasoundSource::GetAudioOutputName()
+const FString& UMetaSoundSource::GetAudioOutputName()
 {
 	static const FString AudioOutputName = TEXT("Generated Audio");
 	return AudioOutputName;
 }
 
-const FString& UMetasoundSource::GetIsFinishedOutputName()
+const FString& UMetaSoundSource::GetIsFinishedOutputName()
 {
 	static const FString OnFinishedOutputName = TEXT("On Finished");
 	return OnFinishedOutputName;
 }
 
-const FString& UMetasoundSource::GetAudioDeviceHandleVariableName()
+const FString& UMetaSoundSource::GetAudioDeviceHandleVariableName()
 {
 	static const FString AudioDeviceHandleVarName = TEXT("AudioDeviceHandle");
 	return AudioDeviceHandleVarName;
 }
 
-const FString& UMetasoundSource::GetSoundUniqueIdName()
+const FString& UMetaSoundSource::GetSoundUniqueIdName()
 {
 	static const FString SoundUniqueIdVarName = TEXT("SoundUniqueId");
 	return SoundUniqueIdVarName;
 }
 
 
-const FMetasoundFrontendArchetype& UMetasoundSource::GetBaseArchetype()
+const FMetasoundFrontendArchetype& UMetaSoundSource::GetBaseArchetype()
 {
 	auto CreateBaseArchetype = []() -> FMetasoundFrontendArchetype
 	{
 		FMetasoundFrontendArchetype Archetype;
 		
 		FMetasoundFrontendClassVertex OnPlayTrigger;
-		OnPlayTrigger.Name = UMetasoundSource::GetOnPlayInputName();
+		OnPlayTrigger.Name = UMetaSoundSource::GetOnPlayInputName();
 		OnPlayTrigger.Metadata.DisplayName = FText::FromString(OnPlayTrigger.Name);
 		OnPlayTrigger.TypeName = Metasound::Frontend::GetDataTypeName<Metasound::FTrigger>();
 		OnPlayTrigger.Metadata.Description = LOCTEXT("OnPlayTriggerToolTip", "Trigger executed when this source is played.");
@@ -452,7 +451,7 @@ const FMetasoundFrontendArchetype& UMetasoundSource::GetBaseArchetype()
 		Archetype.Interface.Inputs.Add(OnPlayTrigger);
 
 		FMetasoundFrontendClassVertex OnFinished;
-		OnFinished.Name = UMetasoundSource::GetIsFinishedOutputName();
+		OnFinished.Name = UMetaSoundSource::GetIsFinishedOutputName();
 		OnFinished.Metadata.DisplayName = FText::FromString(OnFinished.Name);
 		OnFinished.TypeName = Metasound::Frontend::GetDataTypeName<Metasound::FTrigger>();
 		OnFinished.Metadata.Description = LOCTEXT("OnFinishedToolTip", "Trigger executed to initiate stopping the source.");
@@ -461,7 +460,7 @@ const FMetasoundFrontendArchetype& UMetasoundSource::GetBaseArchetype()
 		Archetype.Interface.Outputs.Add(OnFinished);
 
 		FMetasoundFrontendEnvironmentVariable AudioDeviceHandle;
-		AudioDeviceHandle.Name = UMetasoundSource::GetAudioDeviceHandleVariableName();
+		AudioDeviceHandle.Name = UMetaSoundSource::GetAudioDeviceHandleVariableName();
 		AudioDeviceHandle.Metadata.DisplayName = FText::FromString(AudioDeviceHandle.Name);
 		AudioDeviceHandle.Metadata.Description = LOCTEXT("AudioDeviceHandleToolTip", "Audio device handle");
 
@@ -475,7 +474,7 @@ const FMetasoundFrontendArchetype& UMetasoundSource::GetBaseArchetype()
 	return BaseArchetype;
 }
 
-const FMetasoundFrontendArchetype& UMetasoundSource::GetMonoSourceArchetype()
+const FMetasoundFrontendArchetype& UMetaSoundSource::GetMonoSourceArchetype()
 {
 	auto CreateMonoArchetype = []() -> FMetasoundFrontendArchetype
 	{
@@ -483,7 +482,7 @@ const FMetasoundFrontendArchetype& UMetasoundSource::GetMonoSourceArchetype()
 		Archetype.Name = "MonoSource";
 
 		FMetasoundFrontendClassVertex GeneratedAudio;
-		GeneratedAudio.Name = UMetasoundSource::GetAudioOutputName();
+		GeneratedAudio.Name = UMetaSoundSource::GetAudioOutputName();
 		GeneratedAudio.TypeName = Metasound::Frontend::GetDataTypeName<Metasound::FMonoAudioFormat>();
 		GeneratedAudio.Metadata.DisplayName = LOCTEXT("GeneratedMono", "Audio");
 		GeneratedAudio.Metadata.Description = LOCTEXT("GeneratedAudioToolTip", "The resulting output audio from this source.");
@@ -499,7 +498,7 @@ const FMetasoundFrontendArchetype& UMetasoundSource::GetMonoSourceArchetype()
 	return MonoArchetype;
 }
 
-const FMetasoundFrontendArchetype& UMetasoundSource::GetStereoSourceArchetype()
+const FMetasoundFrontendArchetype& UMetaSoundSource::GetStereoSourceArchetype()
 {
 	auto CreateStereoArchetype = []() -> FMetasoundFrontendArchetype
 	{
@@ -507,7 +506,7 @@ const FMetasoundFrontendArchetype& UMetasoundSource::GetStereoSourceArchetype()
 		Archetype.Name = FName(TEXT("StereoSource"));
 
 		FMetasoundFrontendClassVertex GeneratedAudio;
-		GeneratedAudio.Name = UMetasoundSource::GetAudioOutputName();
+		GeneratedAudio.Name = UMetaSoundSource::GetAudioOutputName();
 		GeneratedAudio.TypeName = Metasound::Frontend::GetDataTypeName<Metasound::FStereoAudioFormat>();
 		GeneratedAudio.Metadata.DisplayName = LOCTEXT("GeneratedStereo", "Audio");
 		GeneratedAudio.Metadata.Description = LOCTEXT("GeneratedAudioToolTip", "The resulting output audio from this source.");
diff --git a/Engine/Plugins/Runtime/Metasound/Source/MetasoundEngine/Public/Metasound.h b/Engine/Plugins/Runtime/Metasound/Source/MetasoundEngine/Public/Metasound.h
--- a/Engine/Plugins/Runtime/Metasound/Source/MetasoundEngine/Public/Metasound.h
+++ b/Engine/Plugins/Runtime/Metasound/Source/MetasoundEngine/Public/Metasound.h
@@ -24,7 +24,7 @@ struct FPropertyChangedEvent;
  * Because of this, they can have any inputs or outputs they need.
  */
 UCLASS(hidecategories = object, BlueprintType)
-class METASOUNDENGINE_API UMetasound : public UObject, public FMetasoundAssetBase
+class METASOUNDENGINE_API UMetaSound : public UObject, public FMetasoundAssetBase
 {
 	GENERATED_BODY()
 
@@ -44,19 +44,19 @@ protected:
 #endif // WITH_EDITORONLY_DATA
 
 public:
-	UMetasound(const FObjectInitializer& ObjectInitializer);
+	UMetaSound(const FObjectInitializer& ObjectInitializer);
 
 #if WITH_EDITORONLY_DATA
 	// Returns document name (for editor purposes, and avoids making document public for edit
 	// while allowing editor to reference directly)
 	static FName GetDocumentPropertyName()
 	{
-		return GET_MEMBER_NAME_CHECKED(UMetasound, MetasoundDocument);
+		return GET_MEMBER_NAME_CHECKED(UMetaSound, MetasoundDocument);
 	}
 
 	// Returns the graph associated with this Metasound. Graph is required to be referenced on
 	// Metasound UObject for editor serialization purposes.
-	// @return Editor graph associated with UMetasound.
+	// @return Editor graph associated with UMetaSound.
 	virtual UEdGraph* GetGraph() override;
 	virtual const UEdGraph* GetGraph() const override;
 	virtual UEdGraph& GetGraphChecked() override;
@@ -64,7 +64,7 @@ public:
 
 	// Sets the graph associated with this Metasound. Graph is required to be referenced on
 	// Metasound UObject for editor serialization purposes.
-	// @param Editor graph associated with UMetasound.
+	// @param Editor graph associated with UMetaSound.
 	virtual void SetGraph(UEdGraph* InGraph) override;
 #endif // WITH_EDITORONLY_DATA
 
diff --git a/Engine/Plugins/Runtime/Metasound/Source/MetasoundEngine/Public/MetasoundSource.h b/Engine/Plugins/Runtime/Metasound/Source/MetasoundEngine/Public/MetasoundSource.h
--- a/Engine/Plugins/Runtime/Metasound/Source/MetasoundEngine/Public/MetasoundSource.h
+++ b/Engine/Plugins/Runtime/Metasound/Source/MetasoundEngine/Public/MetasoundSource.h
@@ -13,7 +13,7 @@
 
 #include "MetasoundSource.generated.h"
 
-/** Declares the output audio format of the UMetasoundSource */
+/** Declares the output audio format of the UMetaSoundSource */
 UENUM()
 enum class EMetasoundSourceAudioFormat : uint8
 {
@@ -34,7 +34,7 @@ public:
  * This Metasound type can be played as an audio source.
  */
 UCLASS(hidecategories = object, BlueprintType)
-class METASOUNDENGINE_API UMetasoundSource : public USoundWaveProcedural, public FMetasoundAssetBase
+class METASOUNDENGINE_API UMetaSoundSource : public USoundWaveProcedural, public FMetasoundAssetBase
 {
 	GENERATED_BODY()
 
@@ -48,7 +48,7 @@ protected:
 #endif // WITH_EDITORONLY_DATA
 
 public:
-	UMetasoundSource(const FObjectInitializer& ObjectInitializer);
+	UMetaSoundSource(const FObjectInitializer& ObjectInitializer);
 
 	// The output audio format of the metasound source.
 	UPROPERTY(EditAnywhere, BlueprintReadOnly, Category = Metasound)
@@ -59,14 +59,14 @@ public:
 	// while allowing editor to reference directly)
 	static FName GetDocumentPropertyName()
 	{
-		return GET_MEMBER_NAME_CHECKED(UMetasoundSource, RootMetasoundDocument);
+		return GET_MEMBER_NAME_CHECKED(UMetaSoundSource, RootMetasoundDocument);
 	}
 
 	virtual void PostEditUndo() override;
 
 	// Returns the graph associated with this Metasound. Graph is required to be referenced on
 	// Metasound UObject for editor serialization purposes.
-	// @return Editor graph associated with UMetasoundSource.
+	// @return Editor graph associated with UMetaSoundSource.
 	virtual UEdGraph* GetGraph() override
 	{
 		return Graph;
@@ -91,12 +91,26 @@ public:
 
 	// Sets the graph associated with this Metasound. Graph is required to be referenced on
 	// Metasound UObject for editor serialization purposes.
-	// @param Editor graph associated with UMetasoundSource.
+	// @param Editor graph associated with UMetaSoundSource.
 	virtual void SetGraph(UEdGraph* InGraph) override
 	{
 		Graph = CastChecked<UMetasoundEditorGraphBase>(InGraph);
 	}
 
+	virtual bool GetRedrawThumbnail() const override
+	{
+		return false;
+	}
+
+	virtual void SetRedrawThumbnail(bool bInRedraw) override
+	{
+	}
+
+	virtual bool CanVisualizeAsset() const override
+	{
+		return false;
+	}
+
 	void PostEditChangeProperty(FPropertyChangedEvent& InEvent);
 
 #endif // WITH_EDITOR
diff --git a/Engine/Plugins/Runtime/Metasound/Source/MetasoundFrontend/Private/MetasoundFrontend.cpp b/Engine/Plugins/Runtime/Metasound/Source/MetasoundFrontend/Private/MetasoundFrontend.cpp
--- a/Engine/Plugins/Runtime/Metasound/Source/MetasoundFrontend/Private/MetasoundFrontend.cpp
+++ b/Engine/Plugins/Runtime/Metasound/Source/MetasoundFrontend/Private/MetasoundFrontend.cpp
@@ -21,13 +21,6 @@
 #include "StructSerializer.h"
 #include "Serialization/MemoryReader.h"
 
-static int32 MetasoundUndoRollLimitCvar = 128;
-FAutoConsoleVariableRef CVarMetasoundUndoRollLimit(
-	TEXT("au.Metasound.Frontend.UndoRollLimit"),
-	MetasoundUndoRollLimitCvar,
-	TEXT("Sets the maximum size of our undo buffer for graph editing in the Metasound Frontend.\n")
-	TEXT("n: Number of undoable actions we buffer."),
-	ECVF_Default);
 
 namespace Metasound
 {
diff --git a/Engine/Plugins/Runtime/Metasound/Source/MetasoundFrontend/Private/MetasoundRouter.cpp b/Engine/Plugins/Runtime/Metasound/Source/MetasoundFrontend/Private/MetasoundRouter.cpp
--- a/Engine/Plugins/Runtime/Metasound/Source/MetasoundFrontend/Private/MetasoundRouter.cpp
+++ b/Engine/Plugins/Runtime/Metasound/Source/MetasoundFrontend/Private/MetasoundRouter.cpp
@@ -10,14 +10,14 @@
 
 // Convenience exec commands to push values to global params.
 static FAutoConsoleCommand GPushFloatCommand(
-	TEXT("au.metasound.SetFloat"),
-	TEXT("Use this with au.metasound.SetFloat [type] [address] [value]. Pushes a parameter value directly to a global address, which can then be received by Metasounds using a Receive node."),
+	TEXT("au.MetaSound.SetFloat"),
+	TEXT("Use this with au.MetaSound.SetFloat [type] [address] [value]. Pushes a parameter value directly to a global address, which can then be received by Metasounds using a Receive node."),
 	FConsoleCommandWithArgsDelegate::CreateStatic(
 		[](const TArray< FString >& Args)
 		{
 			if (Args.Num() < 3)
 			{
-				UE_LOG(LogTemp, Warning, TEXT("au.metasound.Set* should be called with three args- the data type of the channel, the address to send to and the value to send."));
+				UE_LOG(LogMetasound, Warning, TEXT("au.MetaSound.Set* should be called with three args- the data type of the channel, the address to send to and the value to send."));
 				return;
 			}
 
@@ -32,15 +32,15 @@ static FAutoConsoleCommand GPushFloatCommand(
 );
 
 static FAutoConsoleCommand GPushBoolCommand(
-	TEXT("au.metasound.SetBool"),
-	TEXT("Use this with au.metasound.SetBool [address] [value]. Pushes a parameter value directly to a global address, which can then be received by Metasounds using a Receive node."),
+	TEXT("au.MetaSound.SetBool"),
+	TEXT("Use this with au.MetaSound.SetBool [address] [value]. Pushes a parameter value directly to a global address, which can then be received by Metasounds using a Receive node."),
 	FConsoleCommandWithArgsDelegate::CreateStatic(
 		[](const TArray< FString >& Args)
 		{
 			if (Args.Num() < 2)
 			{
-				UE_LOG(LogTemp, Warning, TEXT("au.metasound.SetBool should be called with two args- the address to send to and the value to send."));
-				UE_LOG(LogTemp, Warning, TEXT("au.metasound.Set* should be called with three args- the data type of the channel, the address to send to and the value to send."));
+				UE_LOG(LogMetasound, Warning, TEXT("au.MetaSound.SetBool should be called with two args- the address to send to and the value to send."));
+				UE_LOG(LogMetasound, Warning, TEXT("au.MetaSound.Set* should be called with three args- the data type of the channel, the address to send to and the value to send."));
 				return;
 			}
 
@@ -57,14 +57,14 @@ static FAutoConsoleCommand GPushBoolCommand(
 );
 
 static FAutoConsoleCommand GPushIntCommand(
-	TEXT("au.metasound.SetInt"),
-	TEXT("Use this with au.metasound.SetInt [address] [value]. Pushes a parameter value directly to a global address, which can then be received by Metasounds using a Receive node."),
+	TEXT("au.MetaSound.SetInt"),
+	TEXT("Use this with au.MetaSound.SetInt [address] [value]. Pushes a parameter value directly to a global address, which can then be received by Metasounds using a Receive node."),
 	FConsoleCommandWithArgsDelegate::CreateStatic(
 		[](const TArray< FString >& Args)
 		{
 			if (Args.Num() < 2)
 			{
-				UE_LOG(LogTemp, Warning, TEXT("au.metasound.Set* should be called with three args- the data type of the channel, the address to send to and the value to send."));
+				UE_LOG(LogMetasound, Warning, TEXT("au.MetaSound.Set* should be called with three args- the data type of the channel, the address to send to and the value to send."));
 				return;
 			}
 
@@ -79,15 +79,15 @@ static FAutoConsoleCommand GPushIntCommand(
 );
 
 static FAutoConsoleCommand GPushStringCommand(
-	TEXT("au.metasound.SetString"),
-	TEXT("Use this with au.metasound.SetString [address] [value]. Pushes a parameter value directly to a global address, which can then be received by Metasounds using a Receive node."),
+	TEXT("au.MetaSound.SetString"),
+	TEXT("Use this with au.MetaSound.SetString [address] [value]. Pushes a parameter value directly to a global address, which can then be received by Metasounds using a Receive node."),
 	FConsoleCommandWithArgsDelegate::CreateStatic(
 		[](const TArray< FString >& Args)
 		{
 			if (Args.Num() < 3)
 			{
-				UE_LOG(LogTemp, Warning, TEXT("au.metasound.SetBool should be called with three args- the data type of the channel, the address to send to and the value to send."));
-				UE_LOG(LogTemp, Warning, TEXT("au.metasound.Set* should be called with three args- the data type of the channel, the address to send to and the value to send."));
+				UE_LOG(LogMetasound, Warning, TEXT("au.MetaSound.SetBool should be called with three args- the data type of the channel, the address to send to and the value to send."));
+				UE_LOG(LogMetasound, Warning, TEXT("au.MetaSound.Set* should be called with three args- the data type of the channel, the address to send to and the value to send."));
 				return;
 			}
 
diff --git a/Engine/Plugins/Runtime/Metasound/Source/MetasoundFrontend/Public/MetasoundFrontendController.h b/Engine/Plugins/Runtime/Metasound/Source/MetasoundFrontend/Public/MetasoundFrontendController.h
--- a/Engine/Plugins/Runtime/Metasound/Source/MetasoundFrontend/Public/MetasoundFrontendController.h
+++ b/Engine/Plugins/Runtime/Metasound/Source/MetasoundFrontend/Public/MetasoundFrontendController.h
@@ -22,7 +22,7 @@
  *
  * In general, the workflow for editing a Metasound graph will be:
  * 1) Load or create a metasound asset.
- * 2) Call UMetasound::GetDocumentHandle() to get a handle to the document for that asset.
+ * 2) Call UMetaSound::GetDocumentHandle() to get a handle to the document for that asset.
  * 
  * Typically the workflow for creating a Metasound subgraph will be
  * 1) Get a Metasound::Frontend::FGraphHandle (typically through the two steps described in the first paragraph)
diff --git a/Engine/Plugins/Runtime/Metasound/Source/MetasoundFrontend/Public/MetasoundFrontendDocument.h b/Engine/Plugins/Runtime/Metasound/Source/MetasoundFrontend/Public/MetasoundFrontendDocument.h
--- a/Engine/Plugins/Runtime/Metasound/Source/MetasoundFrontend/Public/MetasoundFrontendDocument.h
+++ b/Engine/Plugins/Runtime/Metasound/Source/MetasoundFrontend/Public/MetasoundFrontendDocument.h
@@ -682,7 +682,7 @@ struct FMetasoundFrontendArchetypeInterface
 };
 
 // This is used to describe the required inputs and outputs for a metasound, and is used to make sure we can use a metasound graph for specific applications.
-// For example, a UMetasoundSource needs to generate audio, so its RequiredOutputs will contain "MainAudioOutput"
+// For example, a UMetaSoundSource needs to generate audio, so its RequiredOutputs will contain "MainAudioOutput"
 USTRUCT()
 struct FMetasoundFrontendArchetype
 {
