More MetaSound UX Clean-up
- LogTemp & LogMetasound to LogMetaSound
- Hide Import/Export UI options (to/from Json) for EA from editor as its halfbaked
- Reorder node right-click context menu entries to follow suit with BP
- Add new dedicated metasound icon

diff --git a/Engine/Plugins/Runtime/Metasound/Source/MetasoundEditor/Private/MetasoundDetailCustomization.cpp b/Engine/Plugins/Runtime/Metasound/Source/MetasoundEditor/Private/MetasoundDetailCustomization.cpp
--- a/Engine/Plugins/Runtime/Metasound/Source/MetasoundEditor/Private/MetasoundDetailCustomization.cpp
+++ b/Engine/Plugins/Runtime/Metasound/Source/MetasoundEditor/Private/MetasoundDetailCustomization.cpp
@@ -29,7 +29,7 @@
 #include "Widgets/Text/STextBlock.h"
 
 
-#define LOCTEXT_NAMESPACE "MetasoundEditor"
+#define LOCTEXT_NAMESPACE "MetaSoundEditor"
 
 
 namespace Metasound
diff --git a/Engine/Plugins/Runtime/Metasound/Source/MetasoundEditor/Private/MetasoundEditor.cpp b/Engine/Plugins/Runtime/Metasound/Source/MetasoundEditor/Private/MetasoundEditor.cpp
--- a/Engine/Plugins/Runtime/Metasound/Source/MetasoundEditor/Private/MetasoundEditor.cpp
+++ b/Engine/Plugins/Runtime/Metasound/Source/MetasoundEditor/Private/MetasoundEditor.cpp
@@ -29,6 +29,7 @@
 #include "MetasoundEditorGraphSchema.h"
 #include "MetasoundEditorSettings.h"
 #include "MetasoundEditorTabFactory.h"
+#include "MetasoundLog.h"
 #include "MetasoundUObjectRegistry.h"
 #include "Misc/Attribute.h"
 #include "Modules/ModuleManager.h"
@@ -48,15 +49,15 @@
 #include "../../GraphEditor/Private/GraphActionNode.h"
 
 
-#define LOCTEXT_NAMESPACE "MetasoundEditor"
+#define LOCTEXT_NAMESPACE "MetaSoundEditor"
 
 static int32 ShowLiteralMetasoundInputsInEditorCVar = 0;
-FAutoConsoleVariableRef CVarShowLiteralMetasoundInputsInEditor(
-	TEXT("au.Debug.Editor.Metasounds.ShowLiteralInputs"),
-	ShowLiteralMetasoundInputsInEditorCVar,
-	TEXT("Show literal inputs in the Metasound Editor.\n")
-	TEXT("0: Disabled (default), !0: Enabled"),
-	ECVF_Default);
+// FAutoConsoleVariableRef CVarShowLiteralMetasoundInputsInEditor(
+// 	TEXT("au.Debug.Editor.Metasounds.ShowLiteralInputs"),
+// 	ShowLiteralMetasoundInputsInEditorCVar,
+// 	TEXT("Show literal inputs in the Metasound Editor.\n")
+// 	TEXT("0: Disabled (default), !0: Enabled"),
+// 	ECVF_Default);
 
 
 namespace Metasound
@@ -523,29 +524,30 @@ namespace Metasound
 				GetToolkitCommands(),
 				FToolBarExtensionDelegate::CreateLambda([this](FToolBarBuilder& ToolbarBuilder)
 				{
-					ToolbarBuilder.BeginSection("Utilities");
-					{
-						ToolbarBuilder.AddToolBarButton
-						(
-							FEditorCommands::Get().Import,
-							NAME_None,
-							TAttribute<FText>(),
-							TAttribute<FText>(),
-							TAttribute<FSlateIcon>::Create([this]() { return GetImportStatusImage(); }),
-							"ImportMetasound"
-						);
-
-						ToolbarBuilder.AddToolBarButton
-						(
-							FEditorCommands::Get().Export,
-							NAME_None,
-							TAttribute<FText>(),
-							TAttribute<FText>(),
-							TAttribute<FSlateIcon>::Create([this]() { return GetExportStatusImage(); }),
-							"ExportMetasound"
-						);
-					}
-					ToolbarBuilder.EndSection();
+				// TODO: Add OS SVD and clean this up post UE5.0 - Early Access
+// 					ToolbarBuilder.BeginSection("Utilities");
+// 					{
+// 						ToolbarBuilder.AddToolBarButton
+// 						(
+// 							FEditorCommands::Get().Import,
+// 							NAME_None,
+// 							TAttribute<FText>(),
+// 							TAttribute<FText>(),
+// 							TAttribute<FSlateIcon>::Create([this]() { return GetImportStatusImage(); }),
+// 							"ImportMetasound"
+// 						);
+// 
+// 						ToolbarBuilder.AddToolBarButton
+// 						(
+// 							FEditorCommands::Get().Export,
+// 							NAME_None,
+// 							TAttribute<FText>(),
+// 							TAttribute<FText>(),
+// 							TAttribute<FSlateIcon>::Create([this]() { return GetExportStatusImage(); }),
+// 							"ExportMetasound"
+// 						);
+// 					}
+// 					ToolbarBuilder.EndSection();
 
 					ToolbarBuilder.BeginSection("Audition");
 					{
@@ -584,13 +586,13 @@ namespace Metasound
 		FSlateIcon FEditor::GetImportStatusImage() const
 		{
 			const FName IconName = "MetasoundEditor.Import";
-			return FSlateIcon("MetasoundStyle", IconName);
+			return FSlateIcon("MetaSoundStyle", IconName);
 		}
 
 		FSlateIcon FEditor::GetSettingsImage() const
 		{
 			const FName IconName = "MetasoundEditor.Settings";
-			return FSlateIcon("MetasoundStyle", IconName);
+			return FSlateIcon("MetaSoundStyle", IconName);
 		}
 
 		FSlateIcon FEditor::GetExportStatusImage() const
@@ -601,7 +603,7 @@ namespace Metasound
 				IconName = "MetasoundEditor.ExportError";
 			}
 
-			return FSlateIcon("MetasoundStyle", IconName);
+			return FSlateIcon("MetaSoundStyle", IconName);
 		}
 
 		void FEditor::BindGraphCommands()
@@ -655,10 +657,10 @@ namespace Metasound
 			if (MetasoundAsset)
 			{
 				// TODO: Prompt OFD and provide path from user
-				const FString InputPath = FPaths::ProjectIntermediateDir() / TEXT("Metasounds") + FPaths::ChangeExtension(Metasound->GetPathName(), FMetasoundAssetBase::FileExtension);
+				const FString InputPath = FPaths::ProjectIntermediateDir() / TEXT("MetaSounds") + FPaths::ChangeExtension(Metasound->GetPathName(), FMetasoundAssetBase::FileExtension);
 				
-				// TODO: use the same directory as the currently open metasound
-				const FString OutputPath = FString("/Game/ImportedMetasound/GeneratedMetasound");
+				// TODO: use the same directory as the currently open MetaSound
+				const FString OutputPath = FString("/Game/ImportedMetaSound/GeneratedMetaSound");
 
 				FMetasoundFrontendDocument MetasoundDoc;
 
@@ -668,7 +670,7 @@ namespace Metasound
 
 					if (ImportClasses.Num() < 1)
 					{
-						UE_LOG(LogTemp, Warning, TEXT("Cannot create UObject from Metasound document. No UClass supports archetype \"%s\""), *MetasoundDoc.Archetype.Name.ToString());
+						UE_LOG(LogMetaSound, Warning, TEXT("Cannot create UObject from MetaSound document. No UClass supports archetype \"%s\""), *MetasoundDoc.Archetype.Name.ToString());
 					}
 					else
 					{
@@ -677,7 +679,7 @@ namespace Metasound
 							for (UClass* Cls : ImportClasses)
 							{
 								// TODO: could do a modal dialog to give user choice of import type.
-								UE_LOG(LogTemp, Warning, TEXT("Duplicate UClass support archetype \"%s\" with UClass \"%s\""), *MetasoundDoc.Archetype.Name.ToString(), *Cls->GetName());
+								UE_LOG(LogMetaSound, Warning, TEXT("Duplicate UClass support archetype \"%s\" with UClass \"%s\""), *MetasoundDoc.Archetype.Name.ToString(), *Cls->GetName());
 							}
 						}
 
@@ -686,7 +688,7 @@ namespace Metasound
 				}
 				else
 				{
-					UE_LOG(LogTemp, Warning, TEXT("Could not import Metasound at path: %s"), *InputPath);
+					UE_LOG(LogMetaSound, Warning, TEXT("Could not import MetaSound at path: %s"), *InputPath);
 				}
 			}
 		}
@@ -704,7 +706,7 @@ namespace Metasound
 			static const FString MetasoundExtension(TEXT(".metasound"));
 
 			// TODO: We could just make this an object.
-			const FString Path = FPaths::ProjectSavedDir() / TEXT("Metasounds") + FPaths::ChangeExtension(Metasound->GetPathName(), MetasoundExtension);
+			const FString Path = FPaths::ProjectSavedDir() / TEXT("MetaSounds") + FPaths::ChangeExtension(Metasound->GetPathName(), MetasoundExtension);
 			InMetasoundAsset->GetDocumentHandle()->ExportToJSONAsset(Path);
 		}
 
@@ -1145,7 +1147,7 @@ namespace Metasound
 				if (!Actions.IsEmpty())
 				{
 					check(Metasound);
-					const FScopedTransaction Transaction(LOCTEXT("MetasoundEditorDeleteSelectedNode", "Delete Metasound Variable"));
+					const FScopedTransaction Transaction(LOCTEXT("MetaSoundEditorDeleteSelectedNode", "Delete MetaSound Variable"));
 					Metasound->Modify();
 
 					UMetasoundEditorGraph* Graph = CastChecked<UMetasoundEditorGraph>(MetasoundGraphEditor->GetCurrentGraph());
@@ -1199,7 +1201,7 @@ namespace Metasound
 			const FGraphPanelSelectionSet SelectedNodes = MetasoundGraphEditor->GetSelectedNodes();
 			MetasoundGraphEditor->ClearSelectionSet();
 
-			const FScopedTransaction Transaction(LOCTEXT("MetasoundEditorDeleteSelectedNode", "Delete Selected Metasound Node(s)"));
+			const FScopedTransaction Transaction(LOCTEXT("MetaSoundEditorDeleteSelectedNode", "Delete Selected MetaSound Node(s)"));
 			check(Metasound);
 			Metasound->Modify();
 			UEdGraph* Graph = MetasoundGraphEditor->GetCurrentGraph();
@@ -1335,7 +1337,7 @@ namespace Metasound
 				return;
 			}
 
-			const FScopedTransaction Transaction(LOCTEXT("MetasoundEditorPaste", "Paste Metasound Node(s)"));
+			const FScopedTransaction Transaction(LOCTEXT("MetaSoundEditorPaste", "Paste MetaSound Node(s)"));
 			Metasound->Modify();
 			Graph->Modify();
 
@@ -1779,7 +1781,7 @@ namespace Metasound
 			{
 				case ENodeSection::Inputs:
 				{
-					const FScopedTransaction Transaction(LOCTEXT("AddInputNode", "Add Metasound Input"));
+					const FScopedTransaction Transaction(LOCTEXT("AddInputNode", "Add MetaSound Input"));
 					Metasound->Modify();
 
 					const FName DataTypeName = "Bool";
@@ -1805,7 +1807,7 @@ namespace Metasound
 
 				case ENodeSection::Outputs:
 				{
-					const FScopedTransaction Transaction(TEXT(""), LOCTEXT("AddOutputNode", "Add Metasound Output"), Metasound);
+					const FScopedTransaction Transaction(TEXT(""), LOCTEXT("AddOutputNode", "Add MetaSound Output"), Metasound);
 					Metasound->Modify();
 
 					const FName DataTypeName = "Bool";
diff --git a/Engine/Plugins/Runtime/Metasound/Source/MetasoundEditor/Private/MetasoundEditorCommands.cpp b/Engine/Plugins/Runtime/Metasound/Source/MetasoundEditor/Private/MetasoundEditorCommands.cpp
--- a/Engine/Plugins/Runtime/Metasound/Source/MetasoundEditor/Private/MetasoundEditorCommands.cpp
+++ b/Engine/Plugins/Runtime/Metasound/Source/MetasoundEditor/Private/MetasoundEditorCommands.cpp
@@ -11,12 +11,12 @@ namespace Metasound
 		{
 			UI_COMMAND(Play, "Play", "Plays (or restarts) the MetaSound", EUserInterfaceActionType::Button, FInputChord());
 			UI_COMMAND(Stop, "Stop", "Stops MetaSound (If currently playing)", EUserInterfaceActionType::Button, FInputChord());
-			UI_COMMAND(TogglePlayback, "Toggle Playback", "Plays the Metasound or stops the currently playing Metasound", EUserInterfaceActionType::Button, FInputChord(EKeys::SpaceBar));
+			UI_COMMAND(TogglePlayback, "Toggle Playback", "Plays or stops the currently playing MetaSound", EUserInterfaceActionType::Button, FInputChord(EKeys::SpaceBar));
 
 			UI_COMMAND(Import, "Import", "Imports MetaSound from Json", EUserInterfaceActionType::Button, FInputChord());
 			UI_COMMAND(Export, "Export", "Exports MetaSound to Json", EUserInterfaceActionType::Button, FInputChord());
 
-			UI_COMMAND(BrowserSync, "Browse", "Selects the MetaSound in the content browser. If referencing Metasound nodes are selected, selects referenced assets instead.", EUserInterfaceActionType::Button, FInputChord());
+			UI_COMMAND(BrowserSync, "Browse", "Selects the MetaSound in the content browser. If referencing MetaSound nodes are selected, selects referenced assets instead.", EUserInterfaceActionType::Button, FInputChord());
 			UI_COMMAND(AddInput, "Add Input", "Adds an input to the node", EUserInterfaceActionType::Button, FInputChord());
 			UI_COMMAND(DeleteInput, "Delete Input", "Removes an input from the node", EUserInterfaceActionType::Button, FInputChord());
 
diff --git a/Engine/Plugins/Runtime/Metasound/Source/MetasoundEditor/Private/MetasoundEditorGraph.cpp b/Engine/Plugins/Runtime/Metasound/Source/MetasoundEditor/Private/MetasoundEditorGraph.cpp
--- a/Engine/Plugins/Runtime/Metasound/Source/MetasoundEditor/Private/MetasoundEditorGraph.cpp
+++ b/Engine/Plugins/Runtime/Metasound/Source/MetasoundEditor/Private/MetasoundEditorGraph.cpp
@@ -9,7 +9,7 @@
 #include "MetasoundUObjectRegistry.h"
 #include "ScopedTransaction.h"
 
-#define LOCTEXT_NAMESPACE "MetasoundEditor"
+#define LOCTEXT_NAMESPACE "MetaSoundEditor"
 
 
 TArray<UMetasoundEditorGraphNode*> UMetasoundEditorGraphVariable::GetNodes() const
@@ -44,7 +44,7 @@ void UMetasoundEditorGraphVariable::SetDataType(FName InNewType)
 		return;
 	}
 
-	const FScopedTransaction Transaction(LOCTEXT("SetVariableDataType", "Set Metasound Variable Type"));
+	const FScopedTransaction Transaction(LOCTEXT("SetVariableDataType", "Set MetaSound Variable Type"));
 	Graph->GetMetasoundChecked().Modify();
 	Graph->Modify();
 
@@ -145,7 +145,7 @@ void UMetasoundEditorGraphInput::UpdateDocumentInput(bool bPostTransaction)
 		return;
 	}
 
-	const FScopedTransaction Transaction(LOCTEXT("Set Input Default", "Set Metasound Input Default"), bPostTransaction);
+	const FScopedTransaction Transaction(LOCTEXT("Set Input Default", "Set MetaSound Input Default"), bPostTransaction);
 	Metasound->Modify();
 
 	FMetasoundAssetBase* MetasoundAsset = IMetasoundUObjectRegistry::Get().GetObjectAsAssetBase(Metasound);
@@ -225,7 +225,7 @@ void UMetasoundEditorGraphInput::OnDataTypeChanged()
 	using namespace Metasound::Editor;
 	using namespace Metasound::Frontend;
 
-	IMetasoundEditorModule& EditorModule = FModuleManager::GetModuleChecked<IMetasoundEditorModule>("MetasoundEditor");
+	IMetasoundEditorModule& EditorModule = FModuleManager::GetModuleChecked<IMetasoundEditorModule>("MetaSoundEditor");
 	const FEditorDataType& EditorDataType = EditorModule.FindDataType(TypeName);
 	const EMetasoundFrontendLiteralType LiteralType = static_cast<EMetasoundFrontendLiteralType>(EditorDataType.RegistryInfo.PreferredLiteralType);
 
@@ -360,7 +360,7 @@ UMetasoundEditorGraphInput* UMetasoundEditorGraph::FindOrAddInput(Metasound::Fro
 
 		FMetasoundFrontendLiteral DefaultLiteral = Graph->GetDefaultInput(VertexID);
 		EMetasoundFrontendLiteralType LiteralType = DefaultLiteral.GetType();
-		IMetasoundEditorModule& EditorModule = FModuleManager::GetModuleChecked<IMetasoundEditorModule>("MetasoundEditor");
+		IMetasoundEditorModule& EditorModule = FModuleManager::GetModuleChecked<IMetasoundEditorModule>("MetaSoundEditor");
 		TSubclassOf<UMetasoundEditorGraphInputLiteral> InputLiteralClass = EditorModule.FindInputLiteralClass(LiteralType);
 
 		NewInput->Literal = NewObject<UMetasoundEditorGraphInputLiteral>(NewInput, InputLiteralClass, FName(), RF_Transactional);
diff --git a/Engine/Plugins/Runtime/Metasound/Source/MetasoundEditor/Private/MetasoundEditorGraphBuilder.cpp b/Engine/Plugins/Runtime/Metasound/Source/MetasoundEditor/Private/MetasoundEditorGraphBuilder.cpp
--- a/Engine/Plugins/Runtime/Metasound/Source/MetasoundEditor/Private/MetasoundEditorGraphBuilder.cpp
+++ b/Engine/Plugins/Runtime/Metasound/Source/MetasoundEditor/Private/MetasoundEditorGraphBuilder.cpp
@@ -26,7 +26,7 @@
 #include "Widgets/Notifications/SNotificationList.h"
 
 
-#define LOCTEXT_NAMESPACE "MetasoundEditor"
+#define LOCTEXT_NAMESPACE "MetaSoundEditor"
 
 
 namespace Metasound
@@ -436,7 +436,7 @@ namespace Metasound
 		{
 			using namespace Frontend;
 
-			IMetasoundEditorModule& EditorModule = FModuleManager::GetModuleChecked<IMetasoundEditorModule>("MetasoundEditor");
+			IMetasoundEditorModule& EditorModule = FModuleManager::GetModuleChecked<IMetasoundEditorModule>("MetaSoundEditor");
 
 			FInputHandle InputHandle = GetInputHandleFromPin(&InInputPin);
 			if (!ensure(InputHandle->IsValid()))
@@ -776,7 +776,7 @@ namespace Metasound
 				}
 			}
 
-			IMetasoundEditorModule& EditorModule = FModuleManager::GetModuleChecked<IMetasoundEditorModule>("MetasoundEditor");
+			IMetasoundEditorModule& EditorModule = FModuleManager::GetModuleChecked<IMetasoundEditorModule>("MetaSoundEditor");
 			for (int32 i = 0; i < InputHandles.Num(); ++i)
 			{
 				FInputHandle InputHandle = InputHandles[i];
@@ -1078,7 +1078,7 @@ namespace Metasound
 		{
 			using namespace Frontend;
 
-			IMetasoundEditorModule& EditorModule = FModuleManager::GetModuleChecked<IMetasoundEditorModule>("MetasoundEditor");
+			IMetasoundEditorModule& EditorModule = FModuleManager::GetModuleChecked<IMetasoundEditorModule>("MetaSoundEditor");
 			const FEdGraphPinType PinType = EditorModule.FindDataType(InInputHandle->GetDataType()).PinType;
 
 			UEdGraphPin* NewPin = InEditorNode.CreatePin(EGPD_Input, PinType, FName(*InInputHandle->GetName()));
@@ -1104,7 +1104,7 @@ namespace Metasound
 
 		UEdGraphPin* FGraphBuilder::AddPinToNode(UMetasoundEditorGraphNode& InEditorNode, Frontend::FOutputHandle InOutputHandle)
 		{
-			IMetasoundEditorModule& EditorModule = FModuleManager::GetModuleChecked<IMetasoundEditorModule>("MetasoundEditor");
+			IMetasoundEditorModule& EditorModule = FModuleManager::GetModuleChecked<IMetasoundEditorModule>("MetaSoundEditor");
 			FEdGraphPinType	PinType = EditorModule.FindDataType(InOutputHandle->GetDataType()).PinType;
 
 			UEdGraphPin* NewPin = InEditorNode.CreatePin(EGPD_Output, PinType, FName(*InOutputHandle->GetName()));
@@ -1522,7 +1522,7 @@ namespace Metasound
 			check(MetasoundAsset);
 			UMetasoundEditorGraph* Graph = CastChecked<UMetasoundEditorGraph>(MetasoundAsset->GetGraph());
 
-			IMetasoundEditorModule& EditorModule = FModuleManager::GetModuleChecked<IMetasoundEditorModule>("MetasoundEditor");
+			IMetasoundEditorModule& EditorModule = FModuleManager::GetModuleChecked<IMetasoundEditorModule>("MetaSoundEditor");
 
 			TSet<UMetasoundEditorGraphInput*> Inputs;
 			TSet<UMetasoundEditorGraphOutput*> Outputs;
diff --git a/Engine/Plugins/Runtime/Metasound/Source/MetasoundEditor/Private/MetasoundEditorGraphNode.cpp b/Engine/Plugins/Runtime/Metasound/Source/MetasoundEditor/Private/MetasoundEditorGraphNode.cpp
--- a/Engine/Plugins/Runtime/Metasound/Source/MetasoundEditor/Private/MetasoundEditorGraphNode.cpp
+++ b/Engine/Plugins/Runtime/Metasound/Source/MetasoundEditor/Private/MetasoundEditorGraphNode.cpp
@@ -18,7 +18,7 @@
 #include "ScopedTransaction.h"
 #include "ToolMenus.h"
 
-#define LOCTEXT_NAMESPACE "MetasoundEditor"
+#define LOCTEXT_NAMESPACE "MetaSoundEditor"
 
 
 UMetasoundEditorGraphNode::UMetasoundEditorGraphNode(const FObjectInitializer& ObjectInitializer)
@@ -227,7 +227,7 @@ void UMetasoundEditorGraphNode::PinDefaultValueChanged(UEdGraphPin* Pin)
 		GetMetasoundChecked().Modify();
 
 		FNodeHandle NodeHandle = GetNodeHandle();
-		IMetasoundEditorModule& EditorModule = FModuleManager::GetModuleChecked<IMetasoundEditorModule>("MetasoundEditor");
+		IMetasoundEditorModule& EditorModule = FModuleManager::GetModuleChecked<IMetasoundEditorModule>("MetaSoundEditor");
 		FGraphBuilder::AddOrUpdateLiteralInput(GetMetasoundChecked(), NodeHandle, *Pin);
 	}
 }
@@ -337,14 +337,6 @@ void UMetasoundEditorGraphNode::GetNodeContextMenuActions(UToolMenu* Menu, UGrap
 				}
 			}));
 		}
-
-		{
-			FToolMenuSection& Section = Menu->AddSection("MetasoundEditorGraphNodeEdit");
-			Section.AddMenuEntry(FGenericCommands::Get().Delete);
-			Section.AddMenuEntry(FGenericCommands::Get().Cut);
-			Section.AddMenuEntry(FGenericCommands::Get().Copy);
-			Section.AddMenuEntry(FGenericCommands::Get().Duplicate);
-		}
 	}
 }
 
diff --git a/Engine/Plugins/Runtime/Metasound/Source/MetasoundEditor/Private/MetasoundEditorGraphSchema.cpp b/Engine/Plugins/Runtime/Metasound/Source/MetasoundEditor/Private/MetasoundEditorGraphSchema.cpp
--- a/Engine/Plugins/Runtime/Metasound/Source/MetasoundEditor/Private/MetasoundEditorGraphSchema.cpp
+++ b/Engine/Plugins/Runtime/Metasound/Source/MetasoundEditor/Private/MetasoundEditorGraphSchema.cpp
@@ -6,6 +6,7 @@
 #include "EdGraphSchema_K2.h"
 #include "Editor.h"
 #include "Engine/Selection.h"
+#include "Framework/Commands/GenericCommands.h"
 #include "Framework/MultiBox/MultiBoxBuilder.h"
 #include "GraphEditor.h"
 #include "GraphEditorActions.h"
@@ -28,7 +29,7 @@
 #include "UObject/UObjectHash.h"
 #include "UObject/UObjectIterator.h"
 
-#define LOCTEXT_NAMESPACE "MetasoundEditor"
+#define LOCTEXT_NAMESPACE "MetaSoundEditor"
 
 
 namespace Metasound
@@ -166,7 +167,7 @@ UEdGraphNode* FMetasoundGraphSchemaAction_NewNode::PerformAction(UEdGraph* Paren
 	using namespace Metasound::Editor;
 	using namespace Metasound::Frontend;
 
-	const FScopedTransaction Transaction(LOCTEXT("AddNewNode", "Add New Metasound Node"));
+	const FScopedTransaction Transaction(LOCTEXT("AddNewNode", "Add New MetaSound Node"));
 	UObject& ParentMetasound = CastChecked<UMetasoundEditorGraph>(ParentGraph)->GetMetasoundChecked();
 	ParentMetasound.Modify();
 	ParentGraph->Modify();
@@ -208,7 +209,7 @@ UEdGraphNode* FMetasoundGraphSchemaAction_NewInput::PerformAction(UEdGraph* Pare
 		return nullptr;
 	}
 
-	const FScopedTransaction Transaction(LOCTEXT("AddNewInputNode", "Add New Metasound Input Node"));
+	const FScopedTransaction Transaction(LOCTEXT("AddNewInputNode", "Add New MetaSound Input Node"));
 	ParentMetasound.Modify();
 	MetasoundGraph->Modify();
 	Input->Modify();
@@ -240,7 +241,7 @@ UEdGraphNode* FMetasoundGraphSchemaAction_PromoteToInput::PerformAction(UEdGraph
 		return nullptr;
 	}
 
-	const FScopedTransaction Transaction(LOCTEXT("PromoteNodeInputToGraphInput", "Promote Metasound Node Input to Graph Input"));
+	const FScopedTransaction Transaction(LOCTEXT("PromoteNodeInputToGraphInput", "Promote MetaSound Node Input to Graph Input"));
 	UMetasoundEditorGraph* MetasoundGraph = CastChecked<UMetasoundEditorGraph>(ParentGraph);
 	UObject& ParentMetasound = MetasoundGraph->GetMetasoundChecked();
 	ParentMetasound.Modify();
@@ -306,7 +307,7 @@ UEdGraphNode* FMetasoundGraphSchemaAction_NewOutput::PerformAction(UEdGraph* Par
 		return nullptr;
 	}
 
-	const FScopedTransaction Transaction(LOCTEXT("AddNewOutputNode", "Add New Metasound Output Node"));
+	const FScopedTransaction Transaction(LOCTEXT("AddNewOutputNode", "Add New MetaSound Output Node"));
 	ParentMetasound.Modify();
 	ParentGraph->Modify();
 
@@ -329,7 +330,7 @@ UEdGraphNode* FMetasoundGraphSchemaAction_NewComment::PerformAction(UEdGraph* Pa
 {
 	using namespace Metasound::Editor;
 
-	const FScopedTransaction Transaction(LOCTEXT("AddNewOutputNode", "Add Comment to Metasound Graph"));
+	const FScopedTransaction Transaction(LOCTEXT("AddNewOutputNode", "Add Comment to MetaSound Graph"));
 	ParentGraph->Modify();
 
 	UEdGraphNode_Comment* CommentTemplate = NewObject<UEdGraphNode_Comment>();
@@ -477,6 +478,10 @@ void UMetasoundEditorGraphSchema::GetContextMenuActions(class UToolMenu* Menu, c
 	else if (Context->Node && Context->Node->IsA<UMetasoundEditorGraphNode>())
 	{
 		FToolMenuSection& Section = Menu->AddSection("MetasoundGraphSchemaNodeActions", LOCTEXT("NodeActionsMenuHeader", "Node Actions"));
+		Section.AddMenuEntry(FGenericCommands::Get().Delete);
+		Section.AddMenuEntry(FGenericCommands::Get().Cut);
+		Section.AddMenuEntry(FGenericCommands::Get().Copy);
+		Section.AddMenuEntry(FGenericCommands::Get().Duplicate);
 		Section.AddMenuEntry(FGraphEditorCommands::Get().BreakNodeLinks);
 	}
 
diff --git a/Engine/Plugins/Runtime/Metasound/Source/MetasoundEditor/Private/MetasoundEditorModule.cpp b/Engine/Plugins/Runtime/Metasound/Source/MetasoundEditor/Private/MetasoundEditorModule.cpp
--- a/Engine/Plugins/Runtime/Metasound/Source/MetasoundEditor/Private/MetasoundEditorModule.cpp
+++ b/Engine/Plugins/Runtime/Metasound/Source/MetasoundEditor/Private/MetasoundEditorModule.cpp
@@ -56,7 +56,7 @@ namespace Metasound
 		{
 		public:
 			FSlateStyle()
-				: FSlateStyleSet("MetasoundStyle")
+				: FSlateStyleSet("MetaSoundStyle")
 			{
 				SetParentStyleName(FEditorStyle::GetStyleSetName());
 
@@ -66,6 +66,9 @@ namespace Metasound
 				static const FVector2D Icon20x20(20.0f, 20.0f);
 				static const FVector2D Icon40x40(40.0f, 40.0f);
 
+				static const FVector2D Icon16 = FVector2D(16.0f, 16.0f);
+				static const FVector2D Icon64 = FVector2D(64.0f, 64.0f);
+
 				// Metasound Editor
 				{
 					// Actions
@@ -93,18 +96,16 @@ namespace Metasound
 
 					// Misc
 					Set("MetasoundEditor.Speaker", new FSlateImageBrush(RootToContentDir(TEXT("/Icons/speaker_144x.png")), FVector2D(144.0f, 144.0f)));
+					Set("MetasoundEditor.Metasound.Icon", new IMAGE_BRUSH_SVG(TEXT("Icons/metasound_icon"), Icon16));
 
 					// Class Icons
-					auto SetClassIcon = [this](const FString& ClassName)
+					auto SetClassIcon = [this, InIcon16 = Icon16, InIcon64 = Icon64](const FString& ClassName)
 					{
-						static const FVector2D Icon16 = FVector2D(16.0f, 16.0f);
-						static const FVector2D Icon64 = FVector2D(64.0f, 64.0f);
-
 						const FString IconFileName = FString::Printf(TEXT("Icons/%s"), *ClassName.ToLower());
 						const FSlateColor DefaultForeground(FStyleColors::Foreground);
 
-						Set(*FString::Printf(TEXT("ClassIcon.%s"), *ClassName), new IMAGE_BRUSH_SVG(IconFileName, Icon16));
-						Set(*FString::Printf(TEXT("ClassThumbnail.%s"), *ClassName), new IMAGE_BRUSH_SVG(IconFileName, Icon64));
+						Set(*FString::Printf(TEXT("ClassIcon.%s"), *ClassName), new IMAGE_BRUSH_SVG(IconFileName, InIcon16));
+						Set(*FString::Printf(TEXT("ClassThumbnail.%s"), *ClassName), new IMAGE_BRUSH_SVG(IconFileName, InIcon64));
 					};
 
 					SetClassIcon(TEXT("Metasound"));
@@ -336,9 +337,9 @@ namespace Metasound
 
 				ISettingsModule& SettingsModule = FModuleManager::LoadModuleChecked<ISettingsModule>("Settings");
 
-				SettingsModule.RegisterSettings("Editor", "Audio", "Metasound Editor",
-					NSLOCTEXT("MetasoundsEditor", "MetasoundEditorSettingsName", "Metasound Editor"),
-					NSLOCTEXT("MetasoundsEditor", "MetasoundEditorSettingsDescription", "Customize Metasound Editor."),
+				SettingsModule.RegisterSettings("Editor", "Audio", "MetaSound Editor",
+					NSLOCTEXT("MetaSoundsEditor", "MetaSoundEditorSettingsName", "MetaSound Editor"),
+					NSLOCTEXT("MetaSoundsEditor", "MetaSoundEditorSettingsDescription", "Customize MetaSound Editor."),
 					GetMutableDefault<UMetasoundEditorSettings>()
 				);
 			}
@@ -347,7 +348,7 @@ namespace Metasound
 			{
 				if (ISettingsModule* SettingsModule = FModuleManager::GetModulePtr<ISettingsModule>("Settings"))
 				{
-					SettingsModule->UnregisterSettings("Editor", "Audio", "Metasound Editor");
+					SettingsModule->UnregisterSettings("Editor", "Audio", "MetaSound Editor");
 				}
 
 				if (FModuleManager::Get().IsModuleLoaded(AssetToolName))
diff --git a/Engine/Plugins/Runtime/Metasound/Source/MetasoundEditor/Private/MetasoundEditorTabFactory.h b/Engine/Plugins/Runtime/Metasound/Source/MetasoundEditor/Private/MetasoundEditorTabFactory.h
--- a/Engine/Plugins/Runtime/Metasound/Source/MetasoundEditor/Private/MetasoundEditorTabFactory.h
+++ b/Engine/Plugins/Runtime/Metasound/Source/MetasoundEditor/Private/MetasoundEditorTabFactory.h
@@ -11,7 +11,7 @@
 #include "Styling/SlateStyleRegistry.h"
 #include "Templates/SharedPointer.h"
 
-#define LOCTEXT_NAMESPACE "MetasoundEditor"
+#define LOCTEXT_NAMESPACE "MetaSoundEditor"
 
 
 namespace Metasound
@@ -78,9 +78,9 @@ namespace Metasound
 					GraphActionMenu.ToSharedRef()
 				];
 
-				if (const ISlateStyle* MetasoundStyle = FSlateStyleRegistry::FindSlateStyle("MetasoundStyle"))
+				if (const ISlateStyle* MetasoundStyle = FSlateStyleRegistry::FindSlateStyle("MetaSoundStyle"))
 				{
-					NewTab->SetTabIcon(MetasoundStyle->GetBrush("ClassThumbnail.Metasound"));
+					NewTab->SetTabIcon(MetasoundStyle->GetBrush("MetasoundEditor.Metasound.Icon"));
 				}
 
 				return NewTab;
diff --git a/Engine/Plugins/Runtime/Metasound/Source/MetasoundEditor/Private/MetasoundNodeDetailCustomization.cpp b/Engine/Plugins/Runtime/Metasound/Source/MetasoundEditor/Private/MetasoundNodeDetailCustomization.cpp
--- a/Engine/Plugins/Runtime/Metasound/Source/MetasoundEditor/Private/MetasoundNodeDetailCustomization.cpp
+++ b/Engine/Plugins/Runtime/Metasound/Source/MetasoundEditor/Private/MetasoundNodeDetailCustomization.cpp
@@ -40,7 +40,7 @@
 #include "Widgets/Text/STextBlock.h"
 
 
-#define LOCTEXT_NAMESPACE "MetasoundEditor"
+#define LOCTEXT_NAMESPACE "MetaSoundEditor"
 
 namespace Metasound
 {
@@ -508,7 +508,7 @@ namespace Metasound
 			}
 
 			DataTypeNames.Reset();
-			IMetasoundEditorModule& EditorModule = FModuleManager::GetModuleChecked<IMetasoundEditorModule>("MetasoundEditor");
+			IMetasoundEditorModule& EditorModule = FModuleManager::GetModuleChecked<IMetasoundEditorModule>("MetaSoundEditor");
 			EditorModule.IterateDataTypes([&](const FEditorDataType& EditorDataType)
 			{
 				const FString TypeName = EditorDataType.RegistryInfo.DataTypeName.ToString();
diff --git a/Engine/Plugins/Runtime/Metasound/Source/MetasoundEditor/Private/MetasoundNodeDetailCustomization.h b/Engine/Plugins/Runtime/Metasound/Source/MetasoundEditor/Private/MetasoundNodeDetailCustomization.h
--- a/Engine/Plugins/Runtime/Metasound/Source/MetasoundEditor/Private/MetasoundNodeDetailCustomization.h
+++ b/Engine/Plugins/Runtime/Metasound/Source/MetasoundEditor/Private/MetasoundNodeDetailCustomization.h
@@ -25,7 +25,7 @@ class IPropertyHandle;
 class SCheckBox;
 class STextComboBox;
 
-#define LOCTEXT_NAMESPACE "MetasoundEditor"
+#define LOCTEXT_NAMESPACE "MetaSoundEditor"
 
 namespace Metasound
 {
@@ -162,7 +162,7 @@ namespace Metasound
 
 				if (GraphVariable.IsValid())
 				{
-					const FText TransactionLabel = FText::Format(LOCTEXT("SetTooltip", "Set the Metasound {0}'s tooltip"), VariableLabel);
+					const FText TransactionLabel = FText::Format(LOCTEXT("SetTooltip", "Set the MetaSound {0}'s tooltip"), VariableLabel);
 					const FScopedTransaction Transaction(TransactionLabel);
 					FNodeHandle NodeHandle = GraphVariable->GetNodeHandle();
 					NodeHandle->SetDescription(InNewText);
diff --git a/Engine/Plugins/Runtime/Metasound/Source/MetasoundEditor/Private/SMetasoundGraphNode.cpp b/Engine/Plugins/Runtime/Metasound/Source/MetasoundEditor/Private/SMetasoundGraphNode.cpp
--- a/Engine/Plugins/Runtime/Metasound/Source/MetasoundEditor/Private/SMetasoundGraphNode.cpp
+++ b/Engine/Plugins/Runtime/Metasound/Source/MetasoundEditor/Private/SMetasoundGraphNode.cpp
@@ -257,7 +257,7 @@ void SMetasoundGraphNode::MoveTo(const FVector2D& NewPosition, FNodeSet& NodeFil
 
 const FSlateBrush* SMetasoundGraphNode::GetNodeBodyBrush() const
 {
-// 	if (const ISlateStyle* MetasoundStyle = FSlateStyleRegistry::FindSlateStyle("MetasoundStyle"))
+// 	if (const ISlateStyle* MetasoundStyle = FSlateStyleRegistry::FindSlateStyle("MetaSoundStyle"))
 // 	{
 // 		if (GetMetasoundNode().GetNodeHandle()->GetClassType() == EMetasoundFrontendClassType::Input)
 // 		{
@@ -327,7 +327,7 @@ TSharedRef<SWidget> SMetasoundGraphNode::CreateNodeContentArea()
 				SAssignNew(LeftNodeBox, SVerticalBox)
 			];
 
-		if (const ISlateStyle* MetasoundStyle = FSlateStyleRegistry::FindSlateStyle("MetasoundStyle"))
+		if (const ISlateStyle* MetasoundStyle = FSlateStyleRegistry::FindSlateStyle("MetaSoundStyle"))
 		{
 			const FSlateBrush* ImageBrush = MetasoundStyle->GetBrush(StyleDisplay.ImageName);
 			ContentBox->AddSlot()
diff --git a/Engine/Plugins/Runtime/Metasound/Source/MetasoundEditor/Public/MetasoundEditorCommands.h b/Engine/Plugins/Runtime/Metasound/Source/MetasoundEditor/Public/MetasoundEditorCommands.h
--- a/Engine/Plugins/Runtime/Metasound/Source/MetasoundEditor/Public/MetasoundEditorCommands.h
+++ b/Engine/Plugins/Runtime/Metasound/Source/MetasoundEditor/Public/MetasoundEditorCommands.h
@@ -15,7 +15,7 @@ namespace Metasound
 		public:
 			/** Constructor */
 			FEditorCommands()
-				: TCommands<FEditorCommands>("MetasoundEditor", NSLOCTEXT("Contexts", "MetasoundEditor", "Metasound Graph Editor"), NAME_None, "MetasoundStyle")
+				: TCommands<FEditorCommands>("MetaSoundEditor", NSLOCTEXT("Contexts", "MetaSoundEditor", "MetaSound Graph Editor"), NAME_None, "MetaSoundStyle")
 			{
 			}
 
diff --git a/Engine/Plugins/Runtime/Metasound/Source/MetasoundEngine/Private/MetasoundSource.cpp b/Engine/Plugins/Runtime/Metasound/Source/MetasoundEngine/Private/MetasoundSource.cpp
--- a/Engine/Plugins/Runtime/Metasound/Source/MetasoundEngine/Private/MetasoundSource.cpp
+++ b/Engine/Plugins/Runtime/Metasound/Source/MetasoundEngine/Private/MetasoundSource.cpp
@@ -141,7 +141,7 @@ ISoundGeneratorPtr UMetaSoundSource::CreateSoundGenerator(const FSoundGeneratorI
 	const FMetasoundFrontendDocument* OriginalDoc = GetDocument().Get();
 	if (nullptr == OriginalDoc)
 	{
-		UE_LOG(LogMetasound, Error, TEXT("Cannot create sound generator. Null Metasound document in UMetaSoundSource [Name:%s]"), *GetName());
+		UE_LOG(LogMetaSound, Error, TEXT("Cannot create sound generator. Null Metasound document in UMetaSoundSource [Name:%s]"), *GetName());
 		return ISoundGeneratorPtr(nullptr);
 	}
 
diff --git a/Engine/Plugins/Runtime/Metasound/Source/MetasoundEngine/Private/MetasoundWavePlayerNode.cpp b/Engine/Plugins/Runtime/Metasound/Source/MetasoundEngine/Private/MetasoundWavePlayerNode.cpp
--- a/Engine/Plugins/Runtime/Metasound/Source/MetasoundEngine/Private/MetasoundWavePlayerNode.cpp
+++ b/Engine/Plugins/Runtime/Metasound/Source/MetasoundEngine/Private/MetasoundWavePlayerNode.cpp
@@ -345,7 +345,7 @@ namespace Metasound
 			{
 				if (bReportToLog)
 				{
-					UE_LOG(LogMetasound, Error, TEXT("Failed to initialize SoundWave decoder. WavePlayerNode references invalid or missing Wave."));
+					UE_LOG(LogMetaSound, Error, TEXT("Failed to initialize SoundWave decoder. WavePlayerNode references invalid or missing Wave."));
 				}
 				return false;
 			}
@@ -358,7 +358,7 @@ namespace Metasound
 			{
 				if (bReportToLog)
 				{
-					UE_LOG(LogMetasound, Error, TEXT("Failed to initialize SoundWave decoder. WavePlayerNode only supports 2 channels max [%s: %d Channels])"), *WaveProxy->GetFullName(), NumChannels);
+					UE_LOG(LogMetaSound, Error, TEXT("Failed to initialize SoundWave decoder. WavePlayerNode only supports 2 channels max [%s: %d Channels])"), *WaveProxy->GetFullName(), NumChannels);
 				}
 				return false;
 			}
diff --git a/Engine/Plugins/Runtime/Metasound/Source/MetasoundFrontend/Private/MetasoundAssetBase.cpp b/Engine/Plugins/Runtime/Metasound/Source/MetasoundFrontend/Private/MetasoundAssetBase.cpp
--- a/Engine/Plugins/Runtime/Metasound/Source/MetasoundFrontend/Private/MetasoundAssetBase.cpp
+++ b/Engine/Plugins/Runtime/Metasound/Source/MetasoundFrontend/Private/MetasoundAssetBase.cpp
@@ -30,7 +30,7 @@ void FMetasoundAssetBase::SetMetadata(FMetasoundFrontendClassMetadata& InMetadat
 
 	if (Doc.RootGraph.Metadata.Type != EMetasoundFrontendClassType::Graph)
 	{
-		UE_LOG(LogMetasound, Display, TEXT("Forcing class type to EMetasoundFrontendClassType::Graph on root graph metadata"));
+		UE_LOG(LogMetaSound, Display, TEXT("Forcing class type to EMetasoundFrontendClassType::Graph on root graph metadata"));
 		Doc.RootGraph.Metadata.Type = EMetasoundFrontendClassType::Graph;
 	}
 
diff --git a/Engine/Plugins/Runtime/Metasound/Source/MetasoundFrontend/Private/MetasoundFrontendGraph.cpp b/Engine/Plugins/Runtime/Metasound/Source/MetasoundFrontend/Private/MetasoundFrontendGraph.cpp
--- a/Engine/Plugins/Runtime/Metasound/Source/MetasoundFrontend/Private/MetasoundFrontendGraph.cpp
+++ b/Engine/Plugins/Runtime/Metasound/Source/MetasoundFrontend/Private/MetasoundFrontendGraph.cpp
@@ -172,13 +172,13 @@ namespace Metasound
 				}
 				else
 				{
-					UE_LOG(LogMetasound, Error, TEXT("Cannot create input node [NodeID:%s]. [Vertex:%s] cannot be constructed with the provided literal type."), *InNode.ID.ToString(), *InputVertex.Name);
+					UE_LOG(LogMetaSound, Error, TEXT("Cannot create input node [NodeID:%s]. [Vertex:%s] cannot be constructed with the provided literal type."), *InNode.ID.ToString(), *InputVertex.Name);
 				}
 			}
 		}
 		else
 		{
-			UE_LOG(LogMetasound, Error, TEXT("Cannot create input node [NodeID:%s]. No default literal set for input node."), *InNode.ID.ToString());
+			UE_LOG(LogMetaSound, Error, TEXT("Cannot create input node [NodeID:%s]. No default literal set for input node."), *InNode.ID.ToString());
 		}
 
 		return TUniquePtr<INode>(nullptr);
@@ -364,7 +364,7 @@ namespace Metasound
 							}
 							else
 							{
-								UE_LOG(LogMetasound, Error, TEXT("Failed to match input node [NodeID:%s, NodeName:%s] to owning graph [ClassID:%s] output."), *Node.ID.ToString(), *Node.Name, *InGraphClass.ID.ToString());
+								UE_LOG(LogMetaSound, Error, TEXT("Failed to match input node [NodeID:%s, NodeName:%s] to owning graph [ClassID:%s] output."), *Node.ID.ToString(), *Node.Name, *InGraphClass.ID.ToString());
 							}
 						}
 
@@ -381,7 +381,7 @@ namespace Metasound
 							}
 							else
 							{
-								UE_LOG(LogMetasound, Error, TEXT("Failed to match output node [NodeID:%s, NodeName:%s] to owning graph [ClassID:%s] output."), *Node.ID.ToString(), *Node.Name, *InGraphClass.ID.ToString());
+								UE_LOG(LogMetaSound, Error, TEXT("Failed to match output node [NodeID:%s, NodeName:%s] to owning graph [ClassID:%s] output."), *Node.ID.ToString(), *Node.Name, *InGraphClass.ID.ToString());
 							}
 						}
 
@@ -416,7 +416,7 @@ namespace Metasound
 			const INode* CoreNode = OutGraph.FindNode(Node.ID);
 			if (nullptr == CoreNode)
 			{
-				UE_LOG(LogMetasound, Display, TEXT("Could not find referenced node [NodeID:%s]"), *Node.ID.ToString());
+				UE_LOG(LogMetaSound, Display, TEXT("Could not find referenced node [NodeID:%s]"), *Node.ID.ToString());
 				continue;
 			}
 
@@ -440,14 +440,14 @@ namespace Metasound
 			if (nullptr == DestinationNodeAndVertex)
 			{
 				// TODO: bubble up error
-				UE_LOG(LogMetasound, Error, TEXT("Failed to add edge. Could not find destination [NodeID:%s, VertexID:%s]"), *Edge.ToNodeID.ToString(), *Edge.ToVertexID.ToString());
+				UE_LOG(LogMetaSound, Error, TEXT("Failed to add edge. Could not find destination [NodeID:%s, VertexID:%s]"), *Edge.ToNodeID.ToString(), *Edge.ToVertexID.ToString());
 				continue;
 			}
 
 			if (nullptr == DestinationNodeAndVertex->Node)
 			{
 				// TODO: bubble up error
-				UE_LOG(LogMetasound, Warning, TEXT("Skipping edge. Null destination node [NodeID:%s]"), *Edge.ToNodeID.ToString());
+				UE_LOG(LogMetaSound, Warning, TEXT("Skipping edge. Null destination node [NodeID:%s]"), *Edge.ToNodeID.ToString());
 				continue;
 			}
 
@@ -456,14 +456,14 @@ namespace Metasound
 
 			if (nullptr == SourceNodeAndVertex)
 			{
-				UE_LOG(LogMetasound, Error, TEXT("Failed to add edge. Could not find source [NodeID:%s, VertexID:%s]"), *Edge.FromNodeID.ToString(), *Edge.FromVertexID.ToString());
+				UE_LOG(LogMetaSound, Error, TEXT("Failed to add edge. Could not find source [NodeID:%s, VertexID:%s]"), *Edge.FromNodeID.ToString(), *Edge.FromVertexID.ToString());
 				continue;
 			}
 
 			if (nullptr == SourceNodeAndVertex->Node)
 			{
 				// TODO: bubble up error
-				UE_LOG(LogMetasound, Warning, TEXT("Skipping edge. Null source node [NodeID:%s]"), *Edge.FromNodeID.ToString());
+				UE_LOG(LogMetaSound, Warning, TEXT("Skipping edge. Null source node [NodeID:%s]"), *Edge.FromNodeID.ToString());
 				continue;
 			}
 
@@ -477,7 +477,7 @@ namespace Metasound
 
 			if (!bSuccess)
 			{
-				UE_LOG(LogMetasound, Error, TEXT("Failed to connect edge from [NodeID:%s, VertexID:%s] to [NodeID:%s, VertexID:%s]"), *Edge.FromNodeID.ToString(), *Edge.FromVertexID.ToString(), *Edge.ToNodeID.ToString(), *Edge.ToVertexID.ToString());
+				UE_LOG(LogMetaSound, Error, TEXT("Failed to connect edge from [NodeID:%s, VertexID:%s] to [NodeID:%s, VertexID:%s]"), *Edge.FromNodeID.ToString(), *Edge.FromVertexID.ToString(), *Edge.ToNodeID.ToString(), *Edge.ToVertexID.ToString());
 			}
 		}
 	}
@@ -530,7 +530,7 @@ namespace Metasound
 		{
 			// Likely this will change in the future and the builder will be able
 			// to build graphs with subgraphs..
-			UE_LOG(LogMetasound, Error, TEXT("Provided graph not flat. FFrontendGraphBuilder can only build flat graphs"));
+			UE_LOG(LogMetaSound, Error, TEXT("Provided graph not flat. FFrontendGraphBuilder can only build flat graphs"));
 			return TUniquePtr<FFrontendGraph>(nullptr);
 		}
 
diff --git a/Engine/Plugins/Runtime/Metasound/Source/MetasoundFrontend/Private/MetasoundFrontendLiteral.cpp b/Engine/Plugins/Runtime/Metasound/Source/MetasoundFrontend/Private/MetasoundFrontendLiteral.cpp
--- a/Engine/Plugins/Runtime/Metasound/Source/MetasoundFrontend/Private/MetasoundFrontendLiteral.cpp
+++ b/Engine/Plugins/Runtime/Metasound/Source/MetasoundFrontend/Private/MetasoundFrontendLiteral.cpp
@@ -263,7 +263,7 @@ void FMetasoundFrontendLiteral::SetFromLiteral(const Metasound::FLiteral& InLite
 			// is valid as a null proxy can safely correspond to a null UObject ptr).
 			if (InLiteral.Value.Get<Audio::IProxyDataPtr>().IsValid())
 			{
-				UE_LOG(LogMetasound, Error, TEXT("Cannot set UObjectProxy from Metasound::FLiteral"));
+				UE_LOG(LogMetaSound, Error, TEXT("Cannot set UObjectProxy from Metasound::FLiteral"));
 			}
 			Set(static_cast<UObject*>(nullptr));
 		}
@@ -307,7 +307,7 @@ void FMetasoundFrontendLiteral::SetFromLiteral(const Metasound::FLiteral& InLite
 			// is valid as a null proxy can safely correspond to a null UObject ptr).
 			if (!InLiteral.Value.Get<TArray<Audio::IProxyDataPtr>>().IsEmpty())
 			{
-				UE_LOG(LogMetasound, Error, TEXT("Cannot set UObjectProxy from Metasound::FLiteral"));
+				UE_LOG(LogMetaSound, Error, TEXT("Cannot set UObjectProxy from Metasound::FLiteral"));
 			}
 			Set(TArray<UObject*>());
 		}
@@ -424,7 +424,7 @@ Metasound::FLiteral FMetasoundFrontendLiteral::ToLiteral(const FName& InMetasoun
 
 	if (!bIsTypeSupported)
 	{
-		UE_LOG(LogMetasound, Error, TEXT("Reverting to default literal type for data type. Failed to create supported Metasound::FLiteral for data type [Name:%s] with FMetasoundFrontendLiteral [Literal:%s]"), *InMetasoundDataTypeName.ToString(), *ToString());
+		UE_LOG(LogMetaSound, Error, TEXT("Reverting to default literal type for data type. Failed to create supported Metasound::FLiteral for data type [Name:%s] with FMetasoundFrontendLiteral [Literal:%s]"), *InMetasoundDataTypeName.ToString(), *ToString());
 
 		Literal = GetDefaultParamForDataType(InMetasoundDataTypeName);
 	}
@@ -470,7 +470,7 @@ Metasound::FLiteral FMetasoundFrontendLiteral::ToLiteralNoProxy() const
 
 	if ((EMetasoundFrontendLiteralType::UObject == Type) || (EMetasoundFrontendLiteralType::UObjectArray == Type))
 	{
-		UE_LOG(LogMetasound, Error, TEXT("Cannot convert FMetasoundFrontendLiteral to Metasound::Literal without data type name [Literal:%s]"), *ToString());
+		UE_LOG(LogMetaSound, Error, TEXT("Cannot convert FMetasoundFrontendLiteral to Metasound::Literal without data type name [Literal:%s]"), *ToString());
 		return FLiteral::CreateInvalid();
 	}
 
diff --git a/Engine/Plugins/Runtime/Metasound/Source/MetasoundFrontend/Private/MetasoundFrontendRegistries.cpp b/Engine/Plugins/Runtime/Metasound/Source/MetasoundFrontend/Private/MetasoundFrontendRegistries.cpp
--- a/Engine/Plugins/Runtime/Metasound/Source/MetasoundFrontend/Private/MetasoundFrontendRegistries.cpp
+++ b/Engine/Plugins/Runtime/Metasound/Source/MetasoundFrontend/Private/MetasoundFrontendRegistries.cpp
@@ -255,7 +255,7 @@ namespace Metasound
 			{
 				FScopeLock ScopeLock(&LazyInitCommandCritSection);
 
-				UE_LOG(LogTemp, Display, TEXT("Processing %i Metasounds Frontend Registration Requests."), LazyInitCommands.Num());
+				UE_LOG(LogMetaSound, Display, TEXT("Processing %i Metasounds Frontend Registration Requests."), LazyInitCommands.Num());
 				uint64 CurrentTime = FPlatformTime::Cycles64();
 
 				for (TUniqueFunction<void()>& Command : LazyInitCommands)
@@ -266,7 +266,7 @@ namespace Metasound
 				LazyInitCommands.Empty();
 
 				uint64 CyclesUsed = FPlatformTime::Cycles64() - CurrentTime;
-				UE_LOG(LogTemp, Display, TEXT("Initializing Metasounds Frontend took %f seconds."), FPlatformTime::ToSeconds64(CyclesUsed));
+				UE_LOG(LogMetaSound, Display, TEXT("Initializing Metasounds Frontend took %f seconds."), FPlatformTime::ToSeconds64(CyclesUsed));
 			}
 
 			bool FRegistryContainerImpl::EnqueueInitCommand(TUniqueFunction<void()>&& InFunc)
@@ -274,7 +274,7 @@ namespace Metasound
 				FScopeLock ScopeLock(&LazyInitCommandCritSection);
 				if (LazyInitCommands.Num() >= MaxNumNodesAndDatatypesToInitialize)
 				{
-					UE_LOG(LogTemp, Warning, TEXT("Registering more that %d nodes and datatypes for metasounds! Consider increasing MetasoundFrontendRegistryContainer::MaxNumNodesAndDatatypesToInitialize."));
+					UE_LOG(LogMetaSound, Warning, TEXT("Registering more that %d nodes and datatypes for metasounds! Consider increasing MetasoundFrontendRegistryContainer::MaxNumNodesAndDatatypesToInitialize."));
 				}
 
 				LazyInitCommands.Add(MoveTemp(InFunc));
@@ -540,7 +540,7 @@ namespace Metasound
 					Metasound::Frontend::FNodeRegistryKey OutputNodeRegistryKey = GetRegistryKey(InCallbacks.CreateFrontendOutputClass().Metadata);
 					DataTypeNodeRegistry.Add(OutputNodeRegistryKey, InElement);
 
-					UE_LOG(LogMetasound, Display, TEXT("Registered Metasound Datatype %s."), *InDataInfo.DataTypeName.ToString());
+					UE_LOG(LogMetaSound, Display, TEXT("Registered Metasound Datatype %s."), *InDataInfo.DataTypeName.ToString());
 					return true;
 				}
 			}
diff --git a/Engine/Plugins/Runtime/Metasound/Source/MetasoundFrontend/Private/MetasoundFrontendStandardController.cpp b/Engine/Plugins/Runtime/Metasound/Source/MetasoundFrontend/Private/MetasoundFrontendStandardController.cpp
--- a/Engine/Plugins/Runtime/Metasound/Source/MetasoundFrontend/Private/MetasoundFrontendStandardController.cpp
+++ b/Engine/Plugins/Runtime/Metasound/Source/MetasoundFrontend/Private/MetasoundFrontendStandardController.cpp
@@ -880,13 +880,13 @@ namespace Metasound
 
 			if (ConverterInputs.Num() < 1)
 			{
-				UE_LOG(LogMetasound, Warning, TEXT("Converter node [Name: %s] does not support preferred input vertex [Vertex: %s]"), *InConverterInfo.NodeKey.NodeClassFullName.ToString(), *InConverterInfo.PreferredConverterInputPin);
+				UE_LOG(LogMetaSound, Warning, TEXT("Converter node [Name: %s] does not support preferred input vertex [Vertex: %s]"), *InConverterInfo.NodeKey.NodeClassFullName.ToString(), *InConverterInfo.PreferredConverterInputPin);
 				return false;
 			}
 
 			if (ConverterOutputs.Num() < 1)
 			{
-				UE_LOG(LogMetasound, Warning, TEXT("Converter node [Name: %s] does not support preferred output vertex [Vertex: %s]"), *InConverterInfo.NodeKey.NodeClassFullName.ToString(), *InConverterInfo.PreferredConverterOutputPin);
+				UE_LOG(LogMetaSound, Warning, TEXT("Converter node [Name: %s] does not support preferred output vertex [Vertex: %s]"), *InConverterInfo.NodeKey.NodeClassFullName.ToString(), *InConverterInfo.PreferredConverterOutputPin);
 				return false;
 			}
 
@@ -1152,7 +1152,7 @@ namespace Metasound
 			{
 				if (NodePtr->ClassID != ClassPtr->ID)
 				{
-					UE_LOG(LogMetasound, Warning, TEXT("Changing node's class id from [ClassID:%s] to [ClassID:%s]"), *NodePtr->ClassID.ToString(), *ClassPtr->ID.ToString());
+					UE_LOG(LogMetaSound, Warning, TEXT("Changing node's class id from [ClassID:%s] to [ClassID:%s]"), *NodePtr->ClassID.ToString(), *ClassPtr->ID.ToString());
 					NodePtr->ClassID = ClassPtr->ID;
 				}
 			}
@@ -1865,7 +1865,7 @@ namespace Metasound
 				}
 				else
 				{
-					UE_LOG(LogMetasound, Warning, TEXT("Frontend Node [NodeID:%s, ClassID:%s] is not of expected class class [ClassID:%s]"), *InParams.NodePtr->ID.ToString(), *InParams.NodePtr->ClassID.ToString(), *InParams.ClassPtr->ID.ToString());
+					UE_LOG(LogMetaSound, Warning, TEXT("Frontend Node [NodeID:%s, ClassID:%s] is not of expected class class [ClassID:%s]"), *InParams.NodePtr->ID.ToString(), *InParams.NodePtr->ClassID.ToString(), *InParams.ClassPtr->ID.ToString());
 				}
 			}
 			return FInvalidNodeController::GetInvalid();
@@ -1882,7 +1882,7 @@ namespace Metasound
 				}
 				else
 				{
-					UE_LOG(LogMetasound, Warning, TEXT("Frontend Node [NodeID:%s, ClassID:%s] is not of expected class class [ClassID:%s]"), *InParams.NodePtr->ID.ToString(), *InParams.NodePtr->ClassID.ToString(), *InParams.ClassPtr->ID.ToString());
+					UE_LOG(LogMetaSound, Warning, TEXT("Frontend Node [NodeID:%s, ClassID:%s] is not of expected class class [ClassID:%s]"), *InParams.NodePtr->ID.ToString(), *InParams.NodePtr->ClassID.ToString(), *InParams.ClassPtr->ID.ToString());
 				}
 			}
 			return FInvalidNodeController::GetInvalid();
@@ -1944,12 +1944,12 @@ namespace Metasound
 					}
 					else
 					{
-						UE_LOG(LogMetasound, Warning, TEXT("Frontend Node [NodeID:%s, ClassID:%s] is not of expected class class [ClassID:%s]"), *InParams.NodePtr->ID.ToString(), *InParams.NodePtr->ClassID.ToString(), *InParams.ClassPtr->ID.ToString());
+						UE_LOG(LogMetaSound, Warning, TEXT("Frontend Node [NodeID:%s, ClassID:%s] is not of expected class class [ClassID:%s]"), *InParams.NodePtr->ID.ToString(), *InParams.NodePtr->ClassID.ToString(), *InParams.ClassPtr->ID.ToString());
 					}
 				}
 				else
 				{
-					UE_LOG(LogMetasound, Warning, TEXT("Frontend Class of incorrect class type [ClassID:%s] while creating output node.. Must be EMetasoundFrontendClassType::Output."), *InParams.ClassPtr->ID.ToString());
+					UE_LOG(LogMetaSound, Warning, TEXT("Frontend Class of incorrect class type [ClassID:%s] while creating output node.. Must be EMetasoundFrontendClassType::Output."), *InParams.ClassPtr->ID.ToString());
 				}
 			}
 
@@ -2005,12 +2005,12 @@ namespace Metasound
 					}
 					else
 					{
-						UE_LOG(LogMetasound, Warning, TEXT("Frontend Node [NodeID:%s, ClassID:%s] is not of expected class class [ClassID:%s]"), *InParams.NodePtr->ID.ToString(), *InParams.NodePtr->ClassID.ToString(), *InParams.ClassPtr->ID.ToString());
+						UE_LOG(LogMetaSound, Warning, TEXT("Frontend Node [NodeID:%s, ClassID:%s] is not of expected class class [ClassID:%s]"), *InParams.NodePtr->ID.ToString(), *InParams.NodePtr->ClassID.ToString(), *InParams.ClassPtr->ID.ToString());
 					}
 				}
 				else
 				{
-					UE_LOG(LogMetasound, Warning, TEXT("Frontend Class of incorrect class type [ClassID:%s] while creating output node. Must be EMetasoundFrontendClassType::Output."), *InParams.ClassPtr->ID.ToString());
+					UE_LOG(LogMetaSound, Warning, TEXT("Frontend Class of incorrect class type [ClassID:%s] while creating output node. Must be EMetasoundFrontendClassType::Output."), *InParams.ClassPtr->ID.ToString());
 				}
 			}
 
@@ -2099,12 +2099,12 @@ namespace Metasound
 					}
 					else
 					{
-						UE_LOG(LogMetasound, Warning, TEXT("Frontend Node [NodeID:%s, ClassID:%s] is not of expected class class [ClassID:%s]"), *InParams.NodePtr->ID.ToString(), *InParams.NodePtr->ClassID.ToString(), *InParams.ClassPtr->ID.ToString());
+						UE_LOG(LogMetaSound, Warning, TEXT("Frontend Node [NodeID:%s, ClassID:%s] is not of expected class class [ClassID:%s]"), *InParams.NodePtr->ID.ToString(), *InParams.NodePtr->ClassID.ToString(), *InParams.ClassPtr->ID.ToString());
 					}
 				}
 				else
 				{
-					UE_LOG(LogMetasound, Warning, TEXT("Frontend Class of incorrect class type [ClassID:%s] while creating input node. Must be EMetasoundFrontendClassType::Input."), *InParams.ClassPtr->ID.ToString());
+					UE_LOG(LogMetaSound, Warning, TEXT("Frontend Class of incorrect class type [ClassID:%s] while creating input node. Must be EMetasoundFrontendClassType::Input."), *InParams.ClassPtr->ID.ToString());
 				}
 			}
 
@@ -2123,12 +2123,12 @@ namespace Metasound
 					}
 					else
 					{
-						UE_LOG(LogMetasound, Warning, TEXT("Frontend Node [NodeID:%s, ClassID:%s] is not of expected class class [ClassID:%s]"), *InParams.NodePtr->ID.ToString(), *InParams.NodePtr->ClassID.ToString(), *InParams.ClassPtr->ID.ToString());
+						UE_LOG(LogMetaSound, Warning, TEXT("Frontend Node [NodeID:%s, ClassID:%s] is not of expected class class [ClassID:%s]"), *InParams.NodePtr->ID.ToString(), *InParams.NodePtr->ClassID.ToString(), *InParams.ClassPtr->ID.ToString());
 					}
 				}
 				else
 				{
-					UE_LOG(LogMetasound, Warning, TEXT("Frontend Class of incorrect class type [ClassID:%s] while creating input node. Must be EMetasoundFrontendClassType::Input."), *InParams.ClassPtr->ID.ToString());
+					UE_LOG(LogMetaSound, Warning, TEXT("Frontend Class of incorrect class type [ClassID:%s] while creating input node. Must be EMetasoundFrontendClassType::Input."), *InParams.ClassPtr->ID.ToString());
 				}
 			}
 
@@ -2253,7 +2253,7 @@ namespace Metasound
 				}
 				else
 				{
-					UE_LOG(LogMetasound, Warning, TEXT("Failed to make graph controller [ClassID:%s]. Class must be EMeatsoundFrontendClassType::Graph."), *InParams.GraphClassPtr->ID.ToString())
+					UE_LOG(LogMetaSound, Warning, TEXT("Failed to make graph controller [ClassID:%s]. Class must be EMeatsoundFrontendClassType::Graph."), *InParams.GraphClassPtr->ID.ToString())
 				}
 			}
 			return FInvalidGraphController::GetInvalid();
@@ -2269,7 +2269,7 @@ namespace Metasound
 				}
 				else
 				{
-					UE_LOG(LogMetasound, Warning, TEXT("Failed to make graph controller [ClassID:%s]. Class must be EMeatsoundFrontendClassType::Graph."), *InParams.GraphClassPtr->ID.ToString())
+					UE_LOG(LogMetaSound, Warning, TEXT("Failed to make graph controller [ClassID:%s]. Class must be EMeatsoundFrontendClassType::Graph."), *InParams.GraphClassPtr->ID.ToString())
 				}
 			}
 			return FInvalidGraphController::GetInvalid();
@@ -2421,7 +2421,7 @@ namespace Metasound
 					}
 					else
 					{
-						UE_LOG(LogMetasound, Warning, TEXT("Failed to find class for node [NodeID:%s, ClassID:%s]"), *Node.ID.ToString(), *Node.ClassID.ToString());
+						UE_LOG(LogMetaSound, Warning, TEXT("Failed to find class for node [NodeID:%s, ClassID:%s]"), *Node.ID.ToString(), *Node.ClassID.ToString());
 					}
 				}
 			}
@@ -2447,7 +2447,7 @@ namespace Metasound
 					}
 					else
 					{
-						UE_LOG(LogMetasound, Warning, TEXT("Failed to find class for node [NodeID:%s, ClassID:%s]"), *Node.ID.ToString(), *Node.ClassID.ToString());
+						UE_LOG(LogMetaSound, Warning, TEXT("Failed to find class for node [NodeID:%s, ClassID:%s]"), *Node.ID.ToString(), *Node.ClassID.ToString());
 					}
 				}
 			}
@@ -2559,7 +2559,7 @@ namespace Metasound
 							}
 							else
 							{
-								UE_LOG(LogMetasound, Error, TEXT("Input node [TypeName:%s] does not contain input vertex with type [TypeName:%s]"), *InClassInput.TypeName.ToString(), *InClassInput.TypeName.ToString());
+								UE_LOG(LogMetaSound, Error, TEXT("Input node [TypeName:%s] does not contain input vertex with type [TypeName:%s]"), *InClassInput.TypeName.ToString(), *InClassInput.TypeName.ToString());
 							}
 
 							if (Node.Interface.Outputs.Num() == 1)
@@ -2580,12 +2580,12 @@ namespace Metasound
 					}
 					else 
 					{
-						UE_LOG(LogMetasound, Display, TEXT("Failed to add input. No input node registered for data type [TypeName:%s]"), *InClassInput.TypeName.ToString());
+						UE_LOG(LogMetaSound, Display, TEXT("Failed to add input. No input node registered for data type [TypeName:%s]"), *InClassInput.TypeName.ToString());
 					}
 				}
 				else
 				{
-					UE_LOG(LogMetasound, Display, TEXT("Failed to add input. Input with same name \"%s\" exists in class [ClassID:%s]"), *InClassInput.Name, *GraphClassPtr->ID.ToString());
+					UE_LOG(LogMetaSound, Display, TEXT("Failed to add input. Input with same name \"%s\" exists in class [ClassID:%s]"), *InClassInput.Name, *GraphClassPtr->ID.ToString());
 				}
 			}
 			return FInvalidNodeController::GetInvalid();
@@ -2645,7 +2645,7 @@ namespace Metasound
 							}
 							else
 							{
-								UE_LOG(LogMetasound, Error, TEXT("Output node [TypeName:%s] does not contain output vertex with type [TypeName:%s]"), *InClassOutput.TypeName.ToString(), *InClassOutput.TypeName.ToString());
+								UE_LOG(LogMetaSound, Error, TEXT("Output node [TypeName:%s] does not contain output vertex with type [TypeName:%s]"), *InClassOutput.TypeName.ToString(), *InClassOutput.TypeName.ToString());
 							}
 
 							FMetasoundFrontendClassOutput& NewOutput = GraphClassPtr->Interface.Outputs.Add_GetRef(InClassOutput);
@@ -2657,12 +2657,12 @@ namespace Metasound
 					}
 					else 
 					{
-						UE_LOG(LogMetasound, Display, TEXT("Failed to add output. No output node registered for data type [TypeName:%s]"), *InClassOutput.TypeName.ToString());
+						UE_LOG(LogMetaSound, Display, TEXT("Failed to add output. No output node registered for data type [TypeName:%s]"), *InClassOutput.TypeName.ToString());
 					}
 				}
 				else
 				{
-					UE_LOG(LogMetasound, Display, TEXT("Failed to add output. Output with same name \"%s\" exists in class [ClassID:%s]"), *InClassOutput.Name, *GraphClassPtr->ID.ToString());
+					UE_LOG(LogMetaSound, Display, TEXT("Failed to add output. Output with same name \"%s\" exists in class [ClassID:%s]"), *InClassOutput.Name, *GraphClassPtr->ID.ToString());
 				}
 			}
 
@@ -2911,7 +2911,7 @@ namespace Metasound
 				{
 					if (const FMetasoundFrontendClass* ExistingDependency = OwningDocument->FindClass(InInfo).Get())
 					{
-						UE_LOG(LogMetasound, Error, TEXT("Cannot add new subgraph. Metasound class already exists with matching metadata Name: \"%s\", Version %d.%d"), *(ExistingDependency->Metadata.ClassName.GetFullName().ToString()), ExistingDependency->Metadata.Version.Major, ExistingDependency->Metadata.Version.Minor);
+						UE_LOG(LogMetaSound, Error, TEXT("Cannot add new subgraph. Metasound class already exists with matching metadata Name: \"%s\", Version %d.%d"), *(ExistingDependency->Metadata.ClassName.GetFullName().ToString()), ExistingDependency->Metadata.Version.Major, ExistingDependency->Metadata.Version.Minor);
 					}
 					//else if (const FMetasoundFrontendClass* DependencyDescription = OwningDocument->FindOrAddClass(InInfo))
 					else if (FConstClassAccessPtr DependencyDescription = OwningDocument->FindOrAddClass(InInfo))
@@ -2921,7 +2921,7 @@ namespace Metasound
 				}
 				else
 				{
-					UE_LOG(LogMetasound, Warning, TEXT("Incompatible Metasound NodeType encountered when attempting to create an empty subgraph.  NodeType must equal EMetasoundFrontendClassType::Graph"));
+					UE_LOG(LogMetaSound, Warning, TEXT("Incompatible Metasound NodeType encountered when attempting to create an empty subgraph.  NodeType must equal EMetasoundFrontendClassType::Graph"));
 				}
 			}
 			
@@ -3059,7 +3059,7 @@ namespace Metasound
 					}
 					else
 					{
-						UE_LOG(LogMetasound, Warning, TEXT("Failed to find class for node [NodeID:%s, ClassID:%s]"), *Node.ID.ToString(), *Node.ClassID.ToString());
+						UE_LOG(LogMetaSound, Warning, TEXT("Failed to find class for node [NodeID:%s, ClassID:%s]"), *Node.ID.ToString(), *Node.ClassID.ToString());
 					}
 				}
 			}
@@ -3085,7 +3085,7 @@ namespace Metasound
 					}
 					else
 					{
-						UE_LOG(LogMetasound, Warning, TEXT("Failed to find class for node [NodeID:%s, ClassID:%s]"), *Node.ID.ToString(), *Node.ClassID.ToString());
+						UE_LOG(LogMetaSound, Warning, TEXT("Failed to find class for node [NodeID:%s, ClassID:%s]"), *Node.ID.ToString(), *Node.ClassID.ToString());
 					}
 				}
 			}
@@ -3111,7 +3111,7 @@ namespace Metasound
 					}
 					else
 					{
-						UE_LOG(LogMetasound, Warning, TEXT("Failed to find class for node [NodeID:%s, ClassID:%s]"), *Node.ID.ToString(), *Node.ClassID.ToString());
+						UE_LOG(LogMetaSound, Warning, TEXT("Failed to find class for node [NodeID:%s, ClassID:%s]"), *Node.ID.ToString(), *Node.ClassID.ToString());
 					}
 				}
 			}
@@ -3140,7 +3140,7 @@ namespace Metasound
 					}
 					else
 					{
-						UE_LOG(LogMetasound, Warning, TEXT("Failed to find class for node [NodeID:%s, ClassID:%s]"), *Node.ID.ToString(), *Node.ClassID.ToString());
+						UE_LOG(LogMetaSound, Warning, TEXT("Failed to find class for node [NodeID:%s, ClassID:%s]"), *Node.ID.ToString(), *Node.ClassID.ToString());
 					}
 				}
 			}
@@ -3169,7 +3169,7 @@ namespace Metasound
 					}
 					else
 					{
-						UE_LOG(LogMetasound, Warning, TEXT("Failed to find class for node [NodeID:%s, ClassID:%s]"), *Node.ID.ToString(), *Node.ClassID.ToString());
+						UE_LOG(LogMetaSound, Warning, TEXT("Failed to find class for node [NodeID:%s, ClassID:%s]"), *Node.ID.ToString(), *Node.ClassID.ToString());
 					}
 				}
 			}
@@ -3639,7 +3639,7 @@ namespace Metasound
 					}
 					else
 					{
-						UE_LOG(LogMetasound, Error, TEXT("Cannot add external dependency. No Metasound class found with matching metadata Name: \"%s\", Version %d.%d. Suggested solution \"%s\" by %s."), *InMetadata.ClassName.GetFullName().ToString(), InMetadata.Version.Major, InMetadata.Version.Minor, *InMetadata.PromptIfMissing.ToString(), *InMetadata.Author.ToString());
+						UE_LOG(LogMetaSound, Error, TEXT("Cannot add external dependency. No Metasound class found with matching metadata Name: \"%s\", Version %d.%d. Suggested solution \"%s\" by %s."), *InMetadata.ClassName.GetFullName().ToString(), InMetadata.Version.Major, InMetadata.Version.Minor, *InMetadata.PromptIfMissing.ToString(), *InMetadata.Author.ToString());
 					}
 				} 
 				else if (EMetasoundFrontendClassType::Graph == InMetadata.Type)
@@ -3652,7 +3652,7 @@ namespace Metasound
 				}
 				else
 				{
-					UE_LOG(LogMetasound, Error, TEXT("Unsupported metasound class type for node: \"%s\", Version %d.%d."), *InMetadata.ClassName.GetFullName().ToString(), InMetadata.Version.Major, InMetadata.Version.Minor);
+					UE_LOG(LogMetaSound, Error, TEXT("Unsupported metasound class type for node: \"%s\", Version %d.%d."), *InMetadata.ClassName.GetFullName().ToString(), InMetadata.Version.Major, InMetadata.Version.Minor);
 					checkNoEntry();
 				}
 
@@ -3782,7 +3782,7 @@ namespace Metasound
 				}
 				else
 				{
-					UE_LOG(LogMetasound, Error, TEXT("Failed to export Metasound json asset. Could not write to path \"%s\"."), *InAbsolutePath);
+					UE_LOG(LogMetaSound, Error, TEXT("Failed to export Metasound json asset. Could not write to path \"%s\"."), *InAbsolutePath);
 				}
 			}
 
diff --git a/Engine/Plugins/Runtime/Metasound/Source/MetasoundFrontend/Private/MetasoundRouter.cpp b/Engine/Plugins/Runtime/Metasound/Source/MetasoundFrontend/Private/MetasoundRouter.cpp
--- a/Engine/Plugins/Runtime/Metasound/Source/MetasoundFrontend/Private/MetasoundRouter.cpp
+++ b/Engine/Plugins/Runtime/Metasound/Source/MetasoundFrontend/Private/MetasoundRouter.cpp
@@ -17,7 +17,7 @@ static FAutoConsoleCommand GPushFloatCommand(
 		{
 			if (Args.Num() < 3)
 			{
-				UE_LOG(LogMetasound, Warning, TEXT("au.MetaSound.Set* should be called with three args- the data type of the channel, the address to send to and the value to send."));
+				UE_LOG(LogMetaSound, Warning, TEXT("au.MetaSound.Set* should be called with three args- the data type of the channel, the address to send to and the value to send."));
 				return;
 			}
 
@@ -39,8 +39,8 @@ static FAutoConsoleCommand GPushBoolCommand(
 		{
 			if (Args.Num() < 2)
 			{
-				UE_LOG(LogMetasound, Warning, TEXT("au.MetaSound.SetBool should be called with two args- the address to send to and the value to send."));
-				UE_LOG(LogMetasound, Warning, TEXT("au.MetaSound.Set* should be called with three args- the data type of the channel, the address to send to and the value to send."));
+				UE_LOG(LogMetaSound, Warning, TEXT("au.MetaSound.SetBool should be called with two args- the address to send to and the value to send."));
+				UE_LOG(LogMetaSound, Warning, TEXT("au.MetaSound.Set* should be called with three args- the data type of the channel, the address to send to and the value to send."));
 				return;
 			}
 
@@ -64,7 +64,7 @@ static FAutoConsoleCommand GPushIntCommand(
 		{
 			if (Args.Num() < 2)
 			{
-				UE_LOG(LogMetasound, Warning, TEXT("au.MetaSound.Set* should be called with three args- the data type of the channel, the address to send to and the value to send."));
+				UE_LOG(LogMetaSound, Warning, TEXT("au.MetaSound.Set* should be called with three args- the data type of the channel, the address to send to and the value to send."));
 				return;
 			}
 
@@ -86,8 +86,8 @@ static FAutoConsoleCommand GPushStringCommand(
 		{
 			if (Args.Num() < 3)
 			{
-				UE_LOG(LogMetasound, Warning, TEXT("au.MetaSound.SetBool should be called with three args- the data type of the channel, the address to send to and the value to send."));
-				UE_LOG(LogMetasound, Warning, TEXT("au.MetaSound.Set* should be called with three args- the data type of the channel, the address to send to and the value to send."));
+				UE_LOG(LogMetaSound, Warning, TEXT("au.MetaSound.SetBool should be called with three args- the data type of the channel, the address to send to and the value to send."));
+				UE_LOG(LogMetaSound, Warning, TEXT("au.MetaSound.Set* should be called with three args- the data type of the channel, the address to send to and the value to send."));
 				return;
 			}
 
@@ -160,7 +160,7 @@ namespace Metasound
 		else
 		{
 			// Otherwise, the subsystem FName was invalid.
-			UE_LOG(LogMetasound, Error, TEXT("Cannot create Sender. Did not find transmission subsystem [Name:%s]"), *InAddress.Subsystem.ToString());
+			UE_LOG(LogMetaSound, Error, TEXT("Cannot create Sender. Did not find transmission subsystem [Name:%s]"), *InAddress.Subsystem.ToString());
 		}
 
 		return MoveTemp(Sender);
@@ -194,7 +194,7 @@ namespace Metasound
 		else
 		{
 			// Otherwise, the subsystem FName was invalid.
-			UE_LOG(LogMetasound, Error, TEXT("Cannot create Receiver. Did not find transmission subsystem [Name:%s]"), *InAddress.Subsystem.ToString());
+			UE_LOG(LogMetaSound, Error, TEXT("Cannot create Receiver. Did not find transmission subsystem [Name:%s]"), *InAddress.Subsystem.ToString());
 		}
 
 		return MoveTemp(Receiver);
diff --git a/Engine/Plugins/Runtime/Metasound/Source/MetasoundFrontend/Public/MetasoundArrayNodes.h b/Engine/Plugins/Runtime/Metasound/Source/MetasoundFrontend/Public/MetasoundArrayNodes.h
--- a/Engine/Plugins/Runtime/Metasound/Source/MetasoundFrontend/Public/MetasoundArrayNodes.h
+++ b/Engine/Plugins/Runtime/Metasound/Source/MetasoundFrontend/Public/MetasoundArrayNodes.h
@@ -287,7 +287,7 @@ namespace Metasound
 				}
 				else
 				{
-					UE_LOG(LogMetasound, Error, TEXT("Attempt to get value at invalid index [ArraySize:%d, Index:%d]"), ArrayRef.Num(), IndexValue);
+					UE_LOG(LogMetaSound, Error, TEXT("Attempt to get value at invalid index [ArraySize:%d, Index:%d]"), ArrayRef.Num(), IndexValue);
 				}
 			}
 		}
@@ -425,7 +425,7 @@ namespace Metasound
 				}
 				else
 				{
-					UE_LOG(LogMetasound, Error, TEXT("Attempt to set value at invalid index [ArraySize:%d, Index:%d]"), ArrayRef.Num(), IndexValue);
+					UE_LOG(LogMetaSound, Error, TEXT("Attempt to set value at invalid index [ArraySize:%d, Index:%d]"), ArrayRef.Num(), IndexValue);
 				}
 			}
 		}
diff --git a/Engine/Plugins/Runtime/Metasound/Source/MetasoundFrontend/Public/MetasoundArrayRandomNode.h b/Engine/Plugins/Runtime/Metasound/Source/MetasoundFrontend/Public/MetasoundArrayRandomNode.h
--- a/Engine/Plugins/Runtime/Metasound/Source/MetasoundFrontend/Public/MetasoundArrayRandomNode.h
+++ b/Engine/Plugins/Runtime/Metasound/Source/MetasoundFrontend/Public/MetasoundArrayRandomNode.h
@@ -206,7 +206,7 @@ namespace Metasound
 			}
 			else
 			{
-				UE_LOG(LogMetasound, Error, TEXT("Array Random Get: Can't retrieve random elements from an empty array"));
+				UE_LOG(LogMetaSound, Error, TEXT("Array Random Get: Can't retrieve random elements from an empty array"));
 			}
 		}
 
diff --git a/Engine/Plugins/Runtime/Metasound/Source/MetasoundFrontend/Public/MetasoundArrayShuffleNode.h b/Engine/Plugins/Runtime/Metasound/Source/MetasoundFrontend/Public/MetasoundArrayShuffleNode.h
--- a/Engine/Plugins/Runtime/Metasound/Source/MetasoundFrontend/Public/MetasoundArrayShuffleNode.h
+++ b/Engine/Plugins/Runtime/Metasound/Source/MetasoundFrontend/Public/MetasoundArrayShuffleNode.h
@@ -247,7 +247,7 @@ namespace Metasound
 			}
 			else 
 			{
-				UE_LOG(LogMetasound, Error, TEXT("Array Shuffle: Can't shuffle an empty array"));
+				UE_LOG(LogMetaSound, Error, TEXT("Array Shuffle: Can't shuffle an empty array"));
 			}
 		}
 
diff --git a/Engine/Plugins/Runtime/Metasound/Source/MetasoundFrontend/Public/MetasoundAudioBuffer.h b/Engine/Plugins/Runtime/Metasound/Source/MetasoundFrontend/Public/MetasoundAudioBuffer.h
--- a/Engine/Plugins/Runtime/Metasound/Source/MetasoundFrontend/Public/MetasoundAudioBuffer.h
+++ b/Engine/Plugins/Runtime/Metasound/Source/MetasoundFrontend/Public/MetasoundAudioBuffer.h
@@ -62,7 +62,7 @@ namespace Metasound
 			FORCEINLINE const float* GetData() const
 			{
 #if METASOUNDGRAPHCORE_CHECKAUDIONUM
-				UE_CLOG(InitialNum != Buffer.Num(), LogMetasound, Error, TEXT("Metasound audio buffer size change detected.  Audio Buffers should not be resized."));
+				UE_CLOG(InitialNum != Buffer.Num(), LogMetaSound, Error, TEXT("MetaSound audio buffer size change detected.  Audio Buffers should not be resized."));
 #endif				
 				return Buffer.GetData();
 			}
@@ -71,7 +71,7 @@ namespace Metasound
 			FORCEINLINE float* GetData()
 			{
 #if METASOUNDGRAPHCORE_CHECKAUDIONUM
-				UE_CLOG(InitialNum != Buffer.Num(), LogMetasound, Error, TEXT("Metasound audio buffer size change detected.  Audio Buffers should not be resized."));
+				UE_CLOG(InitialNum != Buffer.Num(), LogMetaSound, Error, TEXT("MetaSound audio buffer size change detected.  Audio Buffers should not be resized."));
 #endif
 				return Buffer.GetData();
 			}
@@ -80,7 +80,7 @@ namespace Metasound
 			FORCEINLINE int32 Num() const
 			{
 #if METASOUNDGRAPHCORE_CHECKAUDIONUM
-				UE_CLOG(InitialNum != Buffer.Num(), LogMetasound, Error, TEXT("Metasound audio buffer size change detected.  Audio Buffers should not be resized."));
+				UE_CLOG(InitialNum != Buffer.Num(), LogMetaSound, Error, TEXT("MetaSound audio buffer size change detected.  Audio Buffers should not be resized."));
 #endif
 				return Buffer.Num();
 			}
@@ -89,7 +89,7 @@ namespace Metasound
 			FORCEINLINE operator const Audio::AlignedFloatBuffer& () const
 			{
 #if METASOUNDGRAPHCORE_CHECKAUDIONUM
-				UE_CLOG(InitialNum != Buffer.Num(), LogMetasound, Error, TEXT("Metasound audio buffer size change detected.  Audio Buffers should not be resized."));
+				UE_CLOG(InitialNum != Buffer.Num(), LogMetaSound, Error, TEXT("MetaSound audio buffer size change detected.  Audio Buffers should not be resized."));
 #endif
 return Buffer;
 			}
@@ -101,7 +101,7 @@ return Buffer;
 			FORCEINLINE operator Audio::AlignedFloatBuffer& ()
 			{
 #if METASOUNDGRAPHCORE_CHECKAUDIONUM
-				UE_CLOG(InitialNum != Buffer.Num(), LogMetasound, Error, TEXT("Metasound audio buffer size change detected.  Audio Buffers should not be resized."));
+				UE_CLOG(InitialNum != Buffer.Num(), LogMetaSound, Error, TEXT("MetaSound audio buffer size change detected.  Audio Buffers should not be resized."));
 #endif				
 				return Buffer;
 			}
diff --git a/Engine/Plugins/Runtime/Metasound/Source/MetasoundFrontend/Public/MetasoundConverterNodeRegistrationMacro.h b/Engine/Plugins/Runtime/Metasound/Source/MetasoundFrontend/Public/MetasoundConverterNodeRegistrationMacro.h
--- a/Engine/Plugins/Runtime/Metasound/Source/MetasoundFrontend/Public/MetasoundConverterNodeRegistrationMacro.h
+++ b/Engine/Plugins/Runtime/Metasound/Source/MetasoundFrontend/Public/MetasoundConverterNodeRegistrationMacro.h
@@ -6,10 +6,11 @@
 #include "Templates/Casts.h"
 
 #include "MetasoundDataReference.h"
-#include "MetasoundOutputNode.h"
+#include "MetasoundFrontendRegistries.h"
 #include "MetasoundInputNode.h"
+#include "MetasoundLog.h"
 #include "MetasoundOperatorInterface.h"
-#include "MetasoundFrontendRegistries.h"
+#include "MetasoundOutputNode.h"
 
 template<typename TNodeType, typename TFromType, typename TToType>
 bool RegisterConversionNode(const FString& FromPin, const FString& ToPin, const Metasound::FNodeClassMetadata& InNodeMetadata)
diff --git a/Engine/Plugins/Runtime/Metasound/Source/MetasoundFrontend/Public/MetasoundDataTypeRegistrationMacro.h b/Engine/Plugins/Runtime/Metasound/Source/MetasoundFrontend/Public/MetasoundDataTypeRegistrationMacro.h
--- a/Engine/Plugins/Runtime/Metasound/Source/MetasoundFrontend/Public/MetasoundDataTypeRegistrationMacro.h
+++ b/Engine/Plugins/Runtime/Metasound/Source/MetasoundFrontend/Public/MetasoundDataTypeRegistrationMacro.h
@@ -14,6 +14,7 @@
 #include "MetasoundFrontendRegistries.h"
 #include "MetasoundInputNode.h"
 #include "MetasoundLiteral.h"
+#include "MetasoundLog.h"
 #include "MetasoundNodeRegistrationMacro.h"
 #include "MetasoundOperatorInterface.h"
 #include "MetasoundOutputNode.h"
@@ -255,7 +256,7 @@ namespace Metasound
 			static bool bAlreadyRegisteredThisDataType = false;
 			if (bAlreadyRegisteredThisDataType)
 			{
-				UE_LOG(LogTemp, Display, TEXT("Tried to call REGISTER_METASOUND_DATATYPE twice with the same class %s. ignoring the second call. Likely because REGISTER_METASOUND_DATATYPE is in a header that's used in multiple modules. Consider moving it to a private header or cpp file."), TDataReferenceTypeInfo<TDataType>::TypeName)
+				UE_LOG(LogMetaSound, Display, TEXT("Tried to call REGISTER_METASOUND_DATATYPE twice with the same class %s. ignoring the second call. Likely because REGISTER_METASOUND_DATATYPE is in a header that's used in multiple modules. Consider moving it to a private header or cpp file."), TDataReferenceTypeInfo<TDataType>::TypeName)
 				return false;
 			}
 
diff --git a/Engine/Plugins/Runtime/Metasound/Source/MetasoundFrontend/Public/MetasoundNodeRegistrationMacro.h b/Engine/Plugins/Runtime/Metasound/Source/MetasoundFrontend/Public/MetasoundNodeRegistrationMacro.h
--- a/Engine/Plugins/Runtime/Metasound/Source/MetasoundFrontend/Public/MetasoundNodeRegistrationMacro.h
+++ b/Engine/Plugins/Runtime/Metasound/Source/MetasoundFrontend/Public/MetasoundNodeRegistrationMacro.h
@@ -6,6 +6,7 @@
 #include "MetasoundFrontend.h"
 #include "MetasoundFrontendDocument.h"
 #include "MetasoundFrontendRegistries.h"
+#include "MetasoundLog.h"
 #include "MetasoundNodeInterface.h"
 #include "Templates/IntegralConstant.h"
 
@@ -34,7 +35,7 @@ namespace Metasound
 		static bool bAlreadyRegisteredThisDataType = false;
 		if (bAlreadyRegisteredThisDataType)
 		{
-			UE_LOG(LogTemp, Display, TEXT("Tried to call METASOUND_REGISTER_NODE twice with the same class. ignoring the second call. Likely because METASOUND_REGISTER_NODE is in a header that's used in multiple modules. Consider moving it to a private header or cpp file."));
+			UE_LOG(LogMetaSound, Display, TEXT("Tried to call METASOUND_REGISTER_NODE twice with the same class. ignoring the second call. Likely because METASOUND_REGISTER_NODE is in a header that's used in multiple modules. Consider moving it to a private header or cpp file."));
 			return false;
 		}
 
diff --git a/Engine/Plugins/Runtime/Metasound/Source/MetasoundGenerator/Private/MetasoundGenerator.cpp b/Engine/Plugins/Runtime/Metasound/Source/MetasoundGenerator/Private/MetasoundGenerator.cpp
--- a/Engine/Plugins/Runtime/Metasound/Source/MetasoundGenerator/Private/MetasoundGenerator.cpp
+++ b/Engine/Plugins/Runtime/Metasound/Source/MetasoundGenerator/Private/MetasoundGenerator.cpp
@@ -35,13 +35,13 @@ namespace Metasound
 		{
 			if (Error.IsValid())
 			{
-				UE_LOG(LogMetasound, Warning, TEXT("MetasoundSource [%s] build error [%s] \"%s\""), *InitParams.MetaSoundName, *(Error->GetErrorType().ToString()), *(Error->GetErrorDescription().ToString()));
+				UE_LOG(LogMetaSound, Warning, TEXT("MetasoundSource [%s] build error [%s] \"%s\""), *InitParams.MetaSoundName, *(Error->GetErrorType().ToString()), *(Error->GetErrorDescription().ToString()));
 			}
 		}
 
 		if (!GraphOperator.IsValid())
 		{
-			UE_LOG(LogMetasound, Error, TEXT("Failed to build Metasound operator from graph in MetasoundSource [%s]"), *InitParams.MetaSoundName);
+			UE_LOG(LogMetaSound, Error, TEXT("Failed to build Metasound operator from graph in MetasoundSource [%s]"), *InitParams.MetaSoundName);
 		}
 		else
 		{
@@ -53,7 +53,7 @@ namespace Metasound
 			{
 				if (!Outputs.ContainsDataReadReference<FStereoAudioFormat>(InitParams.OutputName))
 				{
-					UE_LOG(LogMetasound, Warning, TEXT("MetasoundSource [%s] does not contain stereo output [%s] in output"), *InitParams.MetaSoundName, *InitParams.OutputName);
+					UE_LOG(LogMetaSound, Warning, TEXT("MetasoundSource [%s] does not contain stereo output [%s] in output"), *InitParams.MetaSoundName, *InitParams.OutputName);
 				}
 				OutputBuffers = Outputs.GetDataReadReferenceOrConstruct<FStereoAudioFormat>(InitParams.OutputName, InitParams.OperatorSettings)->GetBuffers();
 			}
@@ -61,7 +61,7 @@ namespace Metasound
 			{
 				if (!Outputs.ContainsDataReadReference<FMonoAudioFormat>(InitParams.OutputName))
 				{
-					UE_LOG(LogMetasound, Warning, TEXT("MetasoundSource [%s] does not contain mono output [%s] in output"), *InitParams.MetaSoundName, *InitParams.OutputName);
+					UE_LOG(LogMetaSound, Warning, TEXT("MetasoundSource [%s] does not contain mono output [%s] in output"), *InitParams.MetaSoundName, *InitParams.OutputName);
 				}
 				OutputBuffers = Outputs.GetDataReadReferenceOrConstruct<FMonoAudioFormat>(InitParams.OutputName, InitParams.OperatorSettings)->GetBuffers();
 			}
@@ -155,7 +155,7 @@ namespace Metasound
 		{
 			int32 FoundNumChannels = InData->OutputBuffers.Num();
 
-			UE_LOG(LogMetasound, Warning, TEXT("Metasound generator expected %d number of channels, found %d"), NumChannels, FoundNumChannels);
+			UE_LOG(LogMetaSound, Warning, TEXT("Metasound generator expected %d number of channels, found %d"), NumChannels, FoundNumChannels);
 
 			int32 NumChannelsToCopy = FMath::Min(FoundNumChannels, NumChannels);
 			int32 NumChannelsToCreate = NumChannels - NumChannelsToCopy;
diff --git a/Engine/Plugins/Runtime/Metasound/Source/MetasoundGraphCore/Private/MetasoundGraphCoreModule.cpp b/Engine/Plugins/Runtime/Metasound/Source/MetasoundGraphCore/Private/MetasoundGraphCoreModule.cpp
--- a/Engine/Plugins/Runtime/Metasound/Source/MetasoundGraphCore/Private/MetasoundGraphCoreModule.cpp
+++ b/Engine/Plugins/Runtime/Metasound/Source/MetasoundGraphCore/Private/MetasoundGraphCoreModule.cpp
@@ -7,7 +7,7 @@
 #include "Modules/ModuleInterface.h"
 #include "Modules/ModuleManager.h"
 
-DEFINE_LOG_CATEGORY(LogMetasound);
+DEFINE_LOG_CATEGORY(LogMetaSound);
 
 namespace Metasound 
 {
diff --git a/Engine/Plugins/Runtime/Metasound/Source/MetasoundGraphCore/Public/MetasoundLog.h b/Engine/Plugins/Runtime/Metasound/Source/MetasoundGraphCore/Public/MetasoundLog.h
--- a/Engine/Plugins/Runtime/Metasound/Source/MetasoundGraphCore/Public/MetasoundLog.h
+++ b/Engine/Plugins/Runtime/Metasound/Source/MetasoundGraphCore/Public/MetasoundLog.h
@@ -4,5 +4,5 @@
 
 #include "CoreMinimal.h"
 
-METASOUNDGRAPHCORE_API DECLARE_LOG_CATEGORY_EXTERN(LogMetasound, Log, All);
+METASOUNDGRAPHCORE_API DECLARE_LOG_CATEGORY_EXTERN(LogMetaSound, Log, All);
 
diff --git a/Engine/Plugins/Runtime/Metasound/Source/MetasoundStandardNodes/Private/MetasoundDebugLogNode.cpp b/Engine/Plugins/Runtime/Metasound/Source/MetasoundStandardNodes/Private/MetasoundDebugLogNode.cpp
--- a/Engine/Plugins/Runtime/Metasound/Source/MetasoundStandardNodes/Private/MetasoundDebugLogNode.cpp
+++ b/Engine/Plugins/Runtime/Metasound/Source/MetasoundStandardNodes/Private/MetasoundDebugLogNode.cpp
@@ -149,7 +149,7 @@ namespace Metasound
 			{
 				if (*Trigger)
 				{
-					UE_LOG(LogMetasound, Display, TEXT("%s %s"), *(*Label), *LexToString(*ValueToLog));
+					UE_LOG(LogMetaSound, Display, TEXT("%s %s"), *(*Label), *LexToString(*ValueToLog));
 				}
 			}
 
