- Fix ensures caused by nullptr passed to WaveAsset proxy. Nodes should handle case where proxy is null accordingly
- Only allow exact type matches for now for object proxies (So WaveAsset proxies cannot be hooked up to child classes of SoundWaves for EA).  Will revisit architecture post EA (add ability to filter types, etc.)

diff --git a/Engine/Plugins/Runtime/Metasound/Source/MetasoundFrontend/Public/MetasoundDataTypeRegistrationMacro.h b/Engine/Plugins/Runtime/Metasound/Source/MetasoundFrontend/Public/MetasoundDataTypeRegistrationMacro.h
--- a/Engine/Plugins/Runtime/Metasound/Source/MetasoundFrontend/Public/MetasoundDataTypeRegistrationMacro.h
+++ b/Engine/Plugins/Runtime/Metasound/Source/MetasoundFrontend/Public/MetasoundDataTypeRegistrationMacro.h
@@ -85,15 +85,16 @@ namespace Metasound
 		template<typename UClassToUse, typename TEnableIf<TIsDerivedFrom<UClassToUse, IAudioProxyDataFactory>::Value, bool>::Type = true>
 		IAudioProxyDataFactory* CastToAudioProxyDataFactory(UObject* InObject)
 		{
-			UClassToUse* DowncastObject = Cast<UClassToUse>(InObject);
-			if (ensureAlways(DowncastObject))
+			if (InObject)
 			{
-				return static_cast<IAudioProxyDataFactory*>(DowncastObject);
-			}
-			else
-			{
-				return nullptr;
+				UClassToUse* DowncastObject = Cast<UClassToUse>(InObject);
+				if (ensureAlways(DowncastObject))
+				{
+					return static_cast<IAudioProxyDataFactory*>(DowncastObject);
+				}
 			}
+
+			return nullptr;
 		}
 
 		template<typename UClassToUse, typename TEnableIf<!TIsDerivedFrom<UClassToUse, IAudioProxyDataFactory>::Value, bool>::Type = true>
@@ -303,21 +304,21 @@ namespace Metasound
 				static_assert(!bSpecifiedUClassForProxy || std::is_base_of<IAudioProxyDataFactory, UClassToUse>::value, "If a Metasound Datatype uses a UObject as a literal, the UClass of that object needs to also derive from Audio::IProxyDataFactory. See USoundWave as an example.");
 				ProxyGenerator = [](UObject* InObject) -> Audio::IProxyDataPtr
 				{
-					IAudioProxyDataFactory* ObjectAsFactory = CastToAudioProxyDataFactory<UClassToUse>(InObject);
-
-					if (ensureAlways(ObjectAsFactory))
+					if (InObject)
 					{
-						static FName ProxySubsystemName = TEXT("Metasound");
-						
-						Audio::FProxyDataInitParams ProxyInitParams;
-						ProxyInitParams.NameOfFeatureRequestingProxy = ProxySubsystemName;
+						IAudioProxyDataFactory* ObjectAsFactory = CastToAudioProxyDataFactory<UClassToUse>(InObject);
+						if (ensureAlways(ObjectAsFactory))
+						{
+							static FName ProxySubsystemName = TEXT("Metasound");
 
-						return ObjectAsFactory->CreateNewProxyData(ProxyInitParams);
-					}
-					else
-					{
-						return Audio::IProxyDataPtr(nullptr);
+							Audio::FProxyDataInitParams ProxyInitParams;
+							ProxyInitParams.NameOfFeatureRequestingProxy = ProxySubsystemName;
+
+							return ObjectAsFactory->CreateNewProxyData(ProxyInitParams);
+						}
 					}
+
+					return Audio::IProxyDataPtr(nullptr);
 				};
 			}
 
