Fix MetaSound Editor CanDuplicate/Duplicate logic
Fix crash when copying comments in MetasoundEditor

diff --git a/Engine/Plugins/Runtime/Metasound/Source/MetasoundEditor/Private/MetasoundEditor.cpp b/Engine/Plugins/Runtime/Metasound/Source/MetasoundEditor/Private/MetasoundEditor.cpp
--- a/Engine/Plugins/Runtime/Metasound/Source/MetasoundEditor/Private/MetasoundEditor.cpp
+++ b/Engine/Plugins/Runtime/Metasound/Source/MetasoundEditor/Private/MetasoundEditor.cpp
@@ -961,8 +961,8 @@ namespace Metasound
 					FCanExecuteAction::CreateLambda([this]() { return CanDeleteNodes(); }));
 
 				GraphEditorCommands->MapAction(FGenericCommands::Get().Duplicate,
-					FExecuteAction::CreateLambda([this] { CopySelectedNodes(); PasteNodes(); }),
-					FCanExecuteAction::CreateLambda([this]() { return CanCopyNodes() && CanPasteNodes(); }));
+					FExecuteAction::CreateLambda([this] { DuplicateNodes(); }),
+					FCanExecuteAction::CreateLambda([this]() { return CanDuplicateNodes(); }));
 
 				// Alignment Commands
 				GraphEditorCommands->MapAction(FGraphEditorCommands::Get().AlignNodesTop,
@@ -1258,25 +1258,12 @@ namespace Metasound
 			}
 		}
 
-		void FEditor::CopySelectedNodes()
- 		{
-			using namespace Metasound::Frontend;
-
-			// Export the selected nodes and place the text on the clipboard
+		void FEditor::CopySelectedNodes() const
+		{
+			FString NodeString;
 			const FGraphPanelSelectionSet SelectedNodes = MetasoundGraphEditor->GetSelectedNodes();
-			TMap<int32, FNodeHandle> HandlesToCopy;
-			for (UObject* SelectedNode : SelectedNodes)
-			{
-				if (UMetasoundEditorGraphNode* Node = Cast<UMetasoundEditorGraphNode>(SelectedNode))
-				{
-					HandlesToCopy.Add(Node->GetUniqueID(), Node->GetNodeHandle());
-					Node->PrepareForCopying();
-				}
-			}
-
-			FString ExportedEdText;
-			FEdGraphUtilities::ExportNodesToText(SelectedNodes, ExportedEdText);
-			FPlatformApplicationMisc::ClipboardCopy(*ExportedEdText);
+			FEdGraphUtilities::ExportNodesToText(SelectedNodes, NodeString);
+			FPlatformApplicationMisc::ClipboardCopy(*NodeString);
 		}
 
 		bool FEditor::CanCopyNodes() const
@@ -1294,6 +1281,34 @@ namespace Metasound
 			return false;
 		}
 
+		bool FEditor::CanDuplicateNodes() const
+		{
+			// If any of the nodes can be duplicated then allow copying
+			const FGraphPanelSelectionSet& SelectedNodes = MetasoundGraphEditor->GetSelectedNodes();
+			for (FGraphPanelSelectionSet::TConstIterator SelectedIter(SelectedNodes); SelectedIter; ++SelectedIter)
+			{
+				UEdGraphNode* Node = Cast<UEdGraphNode>(*SelectedIter);
+				if (!Node || !Node->CanDuplicateNode())
+				{
+					return false;
+				}
+			}
+
+			FString NodeString;
+			FEdGraphUtilities::ExportNodesToText(SelectedNodes, NodeString);
+
+			FMetasoundAssetBase* MetasoundAsset = IMetasoundUObjectRegistry::Get().GetObjectAsAssetBase(Metasound);
+			check(MetasoundAsset);
+
+			UEdGraph* Graph = MetasoundAsset->GetGraph();
+			if (!Graph)
+			{
+				return false;
+			}
+
+			return FEdGraphUtilities::CanImportNodesFromText(Graph, NodeString);
+		}
+
 		bool FEditor::CanDeleteNodes() const
 		{
 			if (MetasoundGraphEditor->GetSelectedNodes().IsEmpty())
@@ -1313,7 +1328,19 @@ namespace Metasound
 			return false;
 		}
 
+		void FEditor::DuplicateNodes()
+		{
+			const FGraphPanelSelectionSet SelectedNodes = MetasoundGraphEditor->GetSelectedNodes();
+			FEdGraphUtilities::ExportNodesToText(SelectedNodes, NodeTextToPaste);
+			PasteNodes(nullptr, LOCTEXT("MetaSoundEditorDuplicate", "Duplicate MetaSound Node(s)"));
+		}
+
 		void FEditor::PasteNodes(const FVector2D* InLocation)
+		{
+			PasteNodes(InLocation, LOCTEXT("MetaSoundEditorPaste", "Paste MetaSound Node(s)"));
+		}
+
+		void FEditor::PasteNodes(const FVector2D* InLocation, const FText& InTransactionText)
 		{
 			using namespace Frontend;
 
@@ -1337,7 +1364,7 @@ namespace Metasound
 				return;
 			}
 
-			const FScopedTransaction Transaction(LOCTEXT("MetaSoundEditorPaste", "Paste MetaSound Node(s)"));
+			const FScopedTransaction Transaction(InTransactionText);
 			Metasound->Modify();
 			Graph->Modify();
 
@@ -1345,9 +1372,9 @@ namespace Metasound
 			MetasoundGraphEditor->ClearSelectionSet();
 
 			TSet<UEdGraphNode*> PastedGraphNodes;
-			FEdGraphUtilities::ImportNodesFromText(Graph, ClipboardContent, PastedGraphNodes);
+			FEdGraphUtilities::ImportNodesFromText(Graph, NodeTextToPaste, PastedGraphNodes);
 
-			ClipboardContent.Empty();
+			NodeTextToPaste.Empty();
 
 			UMetasoundEditorGraph* MetasoundGraph = CastChecked<UMetasoundEditorGraph>(MetasoundAsset->GetGraph());
 
@@ -1392,7 +1419,7 @@ namespace Metasound
 						NodesToRemove.Add(GraphNode);
 					}
 				}
-				else
+				else if (!GraphNode->IsA<UEdGraphNode_Comment>())
 				{
 					checkNoEntry();
 				}
@@ -1484,7 +1511,6 @@ namespace Metasound
 			}
 
 			MetasoundGraphEditor->NotifyGraphChanged();
-			Metasound->MarkPackageDirty();
 		}
 
 		bool FEditor::CanPasteNodes()
@@ -1494,13 +1520,13 @@ namespace Metasound
 
 			if (UEdGraph* Graph = MetasoundAsset->GetGraph())
 			{
-				FPlatformApplicationMisc::ClipboardPaste(ClipboardContent);
-				if (FEdGraphUtilities::CanImportNodesFromText(Graph, ClipboardContent))
+				FPlatformApplicationMisc::ClipboardPaste(NodeTextToPaste);
+				if (FEdGraphUtilities::CanImportNodesFromText(Graph, NodeTextToPaste))
 				{
 					return true;
 				}
 
-				ClipboardContent.Empty();
+				NodeTextToPaste.Empty();
 			}
 
 			return false;
diff --git a/Engine/Plugins/Runtime/Metasound/Source/MetasoundEditor/Private/MetasoundEditor.h b/Engine/Plugins/Runtime/Metasound/Source/MetasoundEditor/Private/MetasoundEditor.h
--- a/Engine/Plugins/Runtime/Metasound/Source/MetasoundEditor/Private/MetasoundEditor.h
+++ b/Engine/Plugins/Runtime/Metasound/Source/MetasoundEditor/Private/MetasoundEditor.h
@@ -103,7 +103,12 @@ namespace Metasound
 			/** Whether pasting the currently selected nodes is permissible */
 			bool CanPasteNodes();
 
-			void PasteNodes(const FVector2D* Location = nullptr);
+			/* Duplicates the selected node(s) in the graph */
+			void DuplicateNodes();
+
+			/* Pastes node(s) from the clipboard to the graph */
+			void PasteNodes(const FVector2D* InLocation = nullptr);
+			void PasteNodes(const FVector2D* InLocation, const FText& InTransactionText);
 
 			int32 GetNumNodesSelected() const
 			{
@@ -155,13 +160,19 @@ namespace Metasound
 			/** Cut the currently selected nodes */
 			void CutSelectedNodes();
 
+			/** Copy the currently selected nodes to the given string */
+			void CopySelectedNodes(FString& OutString) const;
+
 			/** Copy the currently selected nodes */
-			void CopySelectedNodes();
+			void CopySelectedNodes() const;
 
-			/** Whether copying the currently selected nodes is permissible */
+			/** Whether copying the currently selected node(s) is permissible */
 			bool CanCopyNodes() const;
 
-			/** Whether the currently selected nodes can be deleted */
+			/** Whether or not the currently selected node(s) can be duplicated */
+			bool CanDuplicateNodes() const;
+
+			/** Whether the currently selected node(s) can be deleted */
 			bool CanDeleteNodes() const;
 
 			/** Called to undo the last action */
@@ -287,8 +298,8 @@ namespace Metasound
 			/** Whether or not metasound being edited is valid */
 			bool bPassedValidation = true;
 
-			/** Clipboard content used when pasting from clipboard to avoid deserializing twice */
-			FString ClipboardContent;
+			/** Text content used when either duplicating or pasting from clipboard (avoids deserializing twice) */
+			FString NodeTextToPaste;
 
 			/** Boolean state for when selection change handle should not respond due to selection state
 			  * being manually applied in code */
diff --git a/Engine/Plugins/Runtime/Metasound/Source/MetasoundEditor/Private/MetasoundEditorGraphNode.cpp b/Engine/Plugins/Runtime/Metasound/Source/MetasoundEditor/Private/MetasoundEditorGraphNode.cpp
--- a/Engine/Plugins/Runtime/Metasound/Source/MetasoundEditor/Private/MetasoundEditorGraphNode.cpp
+++ b/Engine/Plugins/Runtime/Metasound/Source/MetasoundEditor/Private/MetasoundEditorGraphNode.cpp
@@ -232,10 +232,6 @@ void UMetasoundEditorGraphNode::PinDefaultValueChanged(UEdGraphPin* Pin)
 	}
 }
 
-void UMetasoundEditorGraphNode::PrepareForCopying()
-{
-}
-
 void UMetasoundEditorGraphNode::PostEditImport()
 {
 }
diff --git a/Engine/Plugins/Runtime/Metasound/Source/MetasoundEditor/Public/MetasoundEditorGraphNode.h b/Engine/Plugins/Runtime/Metasound/Source/MetasoundEditor/Public/MetasoundEditorGraphNode.h
--- a/Engine/Plugins/Runtime/Metasound/Source/MetasoundEditor/Public/MetasoundEditorGraphNode.h
+++ b/Engine/Plugins/Runtime/Metasound/Source/MetasoundEditor/Public/MetasoundEditorGraphNode.h
@@ -54,7 +54,6 @@ public:
 	virtual void GetNodeContextMenuActions(class UToolMenu* Menu, class UGraphNodeContextMenuContext* Context) const override;
 	virtual FText GetTooltipText() const override;
 	virtual void PinDefaultValueChanged(UEdGraphPin* Pin) override;
-	virtual void PrepareForCopying() override;
 	virtual void ReconstructNode() override;
 	// End of UEdGraphNode interface
 
