Miscellaneous bug fixes - added desciption to MetaSounds Plugin, and fixed warning involving Slate styles for Audio Meter

diff --git a/Engine/Plugins/Runtime/AudioWidgets/Source/AudioWidgets/Public/SAudioMeter.h b/Engine/Plugins/Runtime/AudioWidgets/Source/AudioWidgets/Public/SAudioMeter.h
--- a/Engine/Plugins/Runtime/AudioWidgets/Source/AudioWidgets/Public/SAudioMeter.h
+++ b/Engine/Plugins/Runtime/AudioWidgets/Source/AudioWidgets/Public/SAudioMeter.h
@@ -34,7 +34,7 @@ public:
 		, _MeterPeakColor(FLinearColor::Blue)
 		, _MeterScaleColor(FLinearColor::White)
 		, _MeterScaleLabelColor(FLinearColor::Gray)
-		, _Style(&FCoreStyle::Get().GetWidgetStyle<FAudioMeterStyle>("AudioMeter"))
+		, _Style(&FAudioMeterStyle::GetDefault())
 	{
 	}
 
diff --git a/Engine/Plugins/Runtime/Metasound/Metasound.uplugin b/Engine/Plugins/Runtime/Metasound/Metasound.uplugin
--- a/Engine/Plugins/Runtime/Metasound/Metasound.uplugin
+++ b/Engine/Plugins/Runtime/Metasound/Metasound.uplugin
@@ -2,8 +2,8 @@
 	"FileVersion": 3,
 	"Version": 1,
 	"VersionName": "1.0",
-	"FriendlyName": "Metasound",
-	"Description": "",
+	"FriendlyName": "MetaSound",
+	"Description": "A high-performance audio system that enables sound designers to have complete control over audio DSP graph generation of sound sources, via sample-accurate control and modulation of sound using audio parameters and audio events from game data and Blueprints",
 	"Category": "Audio",
 	"CreatedBy": "Epic Games, Inc.",
 	"CreatedByURL": "http://epicgames.com",
@@ -14,8 +14,7 @@
 	"CanContainContent": true,
 	"IsBetaVersion": true,
 	"Installed": false,
-	"Modules": 
-	[
+	"Modules": [
 		{
 			"Name": "MetasoundGraphCore",
 			"Type": "Runtime",
@@ -47,8 +46,7 @@
 			"LoadingPhase": "PostEngineInit"
 		}
 	],
-	"Plugins": 
-	[
+	"Plugins": [
 		{
 			"Name": "AudioWidgets",
 			"Enabled": true
