Logarithm and Power MetaSounds

diff --git a/Engine/Plugins/Runtime/Metasound/Source/MetasoundEditor/Private/MetasoundEditorModule.cpp b/Engine/Plugins/Runtime/Metasound/Source/MetasoundEditor/Private/MetasoundEditorModule.cpp
--- a/Engine/Plugins/Runtime/Metasound/Source/MetasoundEditor/Private/MetasoundEditorModule.cpp
+++ b/Engine/Plugins/Runtime/Metasound/Source/MetasoundEditor/Private/MetasoundEditorModule.cpp
@@ -93,6 +93,9 @@ namespace Metasound
 					Set("MetasoundEditor.Graph.Node.Math.Modulo", new FSlateImageBrush(RootToContentDir(TEXT("/Graph/node_math_modulo_40x.png")), Icon40x40));
 					Set("MetasoundEditor.Graph.Node.Math.Multiply", new FSlateImageBrush(RootToContentDir(TEXT("/Graph/node_math_multiply_40x.png")), Icon40x40));
 					Set("MetasoundEditor.Graph.Node.Math.Subtract", new FSlateImageBrush(RootToContentDir(TEXT("/Graph/node_math_subtract_40x.png")), Icon40x40));
+					Set("MetasoundEditor.Graph.Node.Math.Modulo", new FSlateImageBrush(RootToContentDir(TEXT("/Graph/node_math_modulo_40x.png")), Icon40x40));
+					Set("MetasoundEditor.Graph.Node.Math.Power", new FSlateImageBrush(RootToContentDir(TEXT("/Graph/node_math_power_40x.png")), Icon40x40));
+					Set("MetasoundEditor.Graph.Node.Math.Logarithm", new FSlateImageBrush(RootToContentDir(TEXT("/Graph/node_math_logarithm_40x.png")), Icon40x40));
 
 					// Misc
 					Set("MetasoundEditor.Speaker", new FSlateImageBrush(RootToContentDir(TEXT("/Icons/speaker_144x.png")), FVector2D(144.0f, 144.0f)));
diff --git a/Engine/Plugins/Runtime/Metasound/Source/MetasoundStandardNodes/Private/MetasoundMathNodes.cpp b/Engine/Plugins/Runtime/Metasound/Source/MetasoundStandardNodes/Private/MetasoundMathNodes.cpp
--- a/Engine/Plugins/Runtime/Metasound/Source/MetasoundStandardNodes/Private/MetasoundMathNodes.cpp
+++ b/Engine/Plugins/Runtime/Metasound/Source/MetasoundStandardNodes/Private/MetasoundMathNodes.cpp
@@ -1,10 +1,1467 @@
 // Copyright Epic Games, Inc. All Rights Reserved.
-#include "MetasoundMathNodes.h"
+
+#include "CoreMinimal.h"
+#include "DSP/BufferVectorOperations.h"
+#include "MetasoundAudioBuffer.h"
+#include "MetasoundBuilderInterface.h"
+#include "MetasoundDataReference.h"
+#include "MetasoundDataReferenceCollection.h"
+#include "MetasoundExecutableOperator.h"
+#include "MetasoundFacade.h"
+#include "MetasoundNode.h"
+#include "MetasoundNodeRegistrationMacro.h"
+#include "MetasoundNodeInterface.h"
+#include "MetasoundOperatorInterface.h"
+#include "MetasoundPrimitives.h"
+#include "MetasoundStandardNodesNames.h"
+#include "MetasoundTime.h"
+#include "MetasoundTrigger.h"
 
 #define LOCTEXT_NAMESPACE "MetasoundMathOpNode"
 
+#define DEFINE_METASOUND_MATHOP(OpName, DataTypeName, DataClass, Description) \
+	class F##OpName##DataTypeName##Node \
+		: public ::Metasound::TMathOpNode<F##OpName##DataTypeName##Node, ::Metasound::TMathOp##OpName<DataClass, DataClass>, DataClass, DataClass> \
+	{ \
+	public: \
+		static FNodeClassName GetClassName() { return {::Metasound::StandardNodes::Namespace, TEXT(#OpName), TEXT(#DataTypeName)}; } \
+		static FText GetDisplayName() { return MathOpNames::Get##OpName##DisplayName<DataClass>(); } \
+		static FText GetDescription() { return Description; } \
+		static FName GetImageName() { return "MetasoundEditor.Graph.Node.Math." #OpName; } \
+		F##OpName##DataTypeName##Node(const FNodeInitData& InInitData) : TMathOpNode(InInitData) { } \
+	};
+
+#define DEFINE_METASOUND_OPERAND_TYPED_MATHOP(OpName, DataTypeName, DataClass, OperandTypeName, OperandDataClass, Description) \
+	class F##OpName##DataTypeName##OperandTypeName##Node \
+		: public ::Metasound::TMathOpNode<F##OpName##DataTypeName##OperandTypeName##Node, ::Metasound::TMathOp##OpName<DataClass, OperandDataClass>, DataClass, OperandDataClass> \
+	{ \
+	public: \
+		static FNodeClassName GetClassName() { return {::Metasound::StandardNodes::Namespace, TEXT(#OpName), TEXT(#DataTypeName " by " #OperandTypeName)}; } \
+		static FText GetDisplayName() { return MathOpNames::Get##OpName##DisplayName<DataClass, OperandDataClass>(); } \
+		static FText GetDescription() { return Description; } \
+		static FName GetImageName() { return "MetasoundEditor.Graph.Node.Math." #OpName; } \
+		F##OpName##DataTypeName##OperandTypeName##Node(const FNodeInitData& InInitData) : TMathOpNode(InInitData) { } \
+	};
+
 namespace Metasound
 {
+	namespace MathOpNames
+	{
+		static const FString PrimaryOperandName = TEXT("PrimaryOperand");
+
+		static const FString AdditionalOperandsName = TEXT("AdditionalOperands");
+
+		template<typename DataType>
+		const FText GetAddDisplayName()
+		{
+			static const FText DisplayName = FText::Format(LOCTEXT("AddNodeDisplayNamePattern", "Add ({0})"), FText::FromString(GetMetasoundDataTypeString<DataType>()));
+			return DisplayName;
+		}
+
+		template<typename DataType, typename OperandDataType>
+		const FText GetAddDisplayName()
+		{
+			static const FText DisplayName = FText::Format(
+				LOCTEXT("AddNodeDisplayNamePattern", "Add ({0} to {1})"),
+				FText::FromString(GetMetasoundDataTypeString<OperandDataType>()),
+				FText::FromString(GetMetasoundDataTypeString<DataType>()));
+			return DisplayName;
+		}
+
+		template<typename DataType>
+		const FText GetSubtractDisplayName()
+		{
+			static const FText DisplayName = FText::Format(LOCTEXT("SubtractNodeDisplayNamePattern", "Subtract ({0})"), FText::FromString(GetMetasoundDataTypeString<DataType>()));
+			return DisplayName;
+		}
+
+		template<typename DataType, typename OperandDataType>
+		const FText GetSubtractDisplayName()
+		{
+			static const FText DisplayName = FText::Format(
+				LOCTEXT("SubtractNodeOperandTypedDisplayNamePattern", "Subtract ({0} from {1})"),
+				FText::FromString(GetMetasoundDataTypeString<OperandDataType>()),
+				FText::FromString(GetMetasoundDataTypeString<DataType>()));
+			return DisplayName;
+		}
+
+		template<typename DataType>
+		const FText GetMultiplyDisplayName()
+		{
+			static const FText DisplayName = FText::Format(LOCTEXT("MultiplyNodeDisplayNamePattern", "Multiply ({0})"), FText::FromString(GetMetasoundDataTypeString<DataType>()));
+			return DisplayName;
+		}
+
+		template<typename DataType, typename OperandDataType>
+		const FText GetMultiplyDisplayName()
+		{
+			static const FText DisplayName = FText::Format(
+				LOCTEXT("MultiplyNodeOperandTypedDisplayNamePattern", "Multiply ({0} by {1})"),
+				FText::FromString(GetMetasoundDataTypeString<DataType>()),
+				FText::FromString(GetMetasoundDataTypeString<OperandDataType>()));
+			return DisplayName;
+		}
+
+		template<typename DataType>
+		const FText GetDivideDisplayName()
+		{
+			static const FText DisplayName = FText::Format(LOCTEXT("DivideNodeDisplayNamePattern", "Divide ({0})"), FText::FromString(GetMetasoundDataTypeString<DataType>()));
+			return DisplayName;
+		}
+
+		template<typename DataType, typename OperandDataType>
+		const FText GetDivideDisplayName()
+		{
+			static const FText DisplayName = FText::Format(
+				LOCTEXT("DivideNodeOperandTypedDisplayNamePattern", "Divide ({0} by {1})"),
+				FText::FromString(GetMetasoundDataTypeString<DataType>()),
+				FText::FromString(GetMetasoundDataTypeString<OperandDataType>()));
+			return DisplayName;
+		}
+
+		template<typename DataType>
+		const FText GetModuloDisplayName()
+		{
+			static const FText DisplayName = FText::Format(LOCTEXT("ModuloNodeDisplayNamePattern", "Modulo ({0})"), FText::FromString(GetMetasoundDataTypeString<DataType>()));
+			return DisplayName;
+		}
+
+		template<typename DataType, typename OperandDataType>
+		const FText GetModuloDisplayName()
+		{
+			static const FText DisplayName = FText::Format(
+				LOCTEXT("ModuloNodeOperandTypedDisplayNamePattern", "Modulo ({0} by {1})"),
+				FText::FromString(GetMetasoundDataTypeString<DataType>()),
+				FText::FromString(GetMetasoundDataTypeString<OperandDataType>()));
+			return DisplayName;
+		}
+
+		template<typename DataType>
+		const FText GetPowerDisplayName()
+		{
+			static const FText DisplayName = FText::Format(LOCTEXT("PowerNodeDisplayNamePattern", "Power ({0})"), FText::FromString(GetMetasoundDataTypeString<DataType>()));
+			return DisplayName;
+		}
+
+		template<typename DataType, typename OperandDataType>
+		const FText GetPowerDisplayName()
+		{
+			static const FText DisplayName = FText::Format(
+				LOCTEXT("PowerNodeOperandTypedDisplayNamePattern", "Power ({0} to the power of {1})"),
+				FText::FromString(GetMetasoundDataTypeString<DataType>()),
+				FText::FromString(GetMetasoundDataTypeString<OperandDataType>()));
+			return DisplayName;
+		}
+
+		template<typename DataType>
+		const FText GetLogarithmDisplayName()
+		{
+			static const FText DisplayName = FText::Format(LOCTEXT("LogNodeDisplayNamePattern", "Log ({0})"), FText::FromString(GetMetasoundDataTypeString<DataType>()));
+			return DisplayName;
+		}
+
+		template<typename DataType, typename OperandDataType>
+		const FText GetLogarithmDisplayName()
+		{
+			static const FText DisplayName = FText::Format(
+				LOCTEXT("LogarithmNodeOperandTypedDisplayNamePattern", "Log ({0}-Base logarithm of {1})"),
+				FText::FromString(GetMetasoundDataTypeString<DataType>()),
+				FText::FromString(GetMetasoundDataTypeString<OperandDataType>()));
+			return DisplayName;
+		}
+	}
+
+	template <typename TDerivedClass, typename TMathOpClass, typename TDataClass, typename TOperandDataClass>
+	class TMathOpNode : public FNodeFacade
+	{
+	private:
+		class TMathOperator : public TExecutableOperator<TMathOperator>
+		{
+		private:
+			using TDataClassReadRef = TDataReadReference<TDataClass>;
+			using TOperandDataClassReadRef = TDataReadReference<TOperandDataClass>;
+			using TDataClassWriteRef = TDataWriteReference<TDataClass>;
+
+			TMathOpClass InstanceData;
+
+			TDataClassReadRef PrimaryOperandRef;
+			TArray<TOperandDataClassReadRef> AdditionalOperandRefs;
+			TDataClassWriteRef ValueRef;
+
+		public:
+			static const FNodeClassMetadata& GetNodeInfo()
+			{
+				auto InitNodeInfo = []() -> FNodeClassMetadata
+				{
+					FNodeDisplayStyle DisplayStyle;
+					DisplayStyle.ImageName = TDerivedClass::GetImageName();
+					DisplayStyle.bShowName = false;
+					DisplayStyle.bShowInputNames = false;
+					DisplayStyle.bShowOutputNames = false;
+
+					FNodeClassMetadata Info;
+					Info.ClassName = TDerivedClass::GetClassName();
+					Info.MajorVersion = 1;
+					Info.MinorVersion = 0;
+					Info.DisplayName = TDerivedClass::GetDisplayName();
+					Info.Description = TDerivedClass::GetDescription();
+					Info.Author = PluginAuthor;
+					Info.PromptIfMissing = PluginNodeMissingPrompt;
+					Info.DefaultInterface = TMathOpClass::GetVertexInterface();
+					Info.DisplayStyle = DisplayStyle;
+					Info.CategoryHierarchy = { LOCTEXT("Metasound_MathCategory", "Math") };
+
+					return Info;
+				};
+
+				static const FNodeClassMetadata Info = InitNodeInfo();
+				return Info;
+			}
+
+			void Execute()
+			{
+				TMathOpClass::Calculate(InstanceData, PrimaryOperandRef, AdditionalOperandRefs, ValueRef);
+			}
+
+			static TUniquePtr<IOperator> CreateOperator(const FCreateOperatorParams& InParams, FBuildErrorArray& OutErrors)
+			{
+				const TMathOpNode& MathOpNode = static_cast<const TMathOpNode&>(InParams.Node);
+
+				const FInputVertexInterface& InputInterface = MathOpNode.GetVertexInterface().GetInputInterface();
+
+				FLiteral DefaultValue = FLiteral::CreateInvalid();
+				const FString& PrimaryOperandName = MathOpNames::PrimaryOperandName;
+				if (InputInterface.Contains(PrimaryOperandName))
+				{
+					DefaultValue = InputInterface[PrimaryOperandName].GetDefaultLiteral();
+				}
+				TDataClassReadRef PrimaryOperand = InParams.InputDataReferences.GetDataReadReferenceOrConstruct<TDataClass>(MathOpNames::PrimaryOperandName, TMathOpClass::GetDefault(InParams.OperatorSettings, DefaultValue));
+
+				// TODO: Support dynamic number of inputs
+				const FString& OpName = MathOpNames::AdditionalOperandsName;
+				if (InputInterface.Contains(OpName))
+				{
+					DefaultValue = InputInterface[OpName].GetDefaultLiteral();
+				}
+				TOperandDataClassReadRef Op1 = InParams.InputDataReferences.GetDataReadReferenceOrConstruct<TOperandDataClass>(OpName, TMathOpClass::GetDefaultOp(InParams.OperatorSettings, DefaultValue));
+				TArray<TOperandDataClassReadRef> AdditionalOperandRefs = { Op1 };
+
+				return MakeUnique<TMathOperator>(InParams.OperatorSettings, PrimaryOperand, AdditionalOperandRefs);
+			}
+
+			TMathOperator(const FOperatorSettings& InSettings, const TDataClassReadRef& InPrimaryOperand, const TArray<TOperandDataClassReadRef>& InAdditionalOperands)
+				: PrimaryOperandRef(InPrimaryOperand)
+				, AdditionalOperandRefs(InAdditionalOperands)
+				, ValueRef(TDataWriteReferenceFactory<TDataClass>::CreateAny(InSettings))
+			{
+			}
+
+			virtual FDataReferenceCollection GetInputs() const override
+			{
+				FDataReferenceCollection InputDataReferences;
+				InputDataReferences.AddDataReadReference(MathOpNames::PrimaryOperandName, PrimaryOperandRef);
+				InputDataReferences.AddDataReadReference(MathOpNames::AdditionalOperandsName, AdditionalOperandRefs[0]);
+				return InputDataReferences;
+			}
+
+			virtual FDataReferenceCollection GetOutputs() const override
+			{
+				FDataReferenceCollection OutputDataReferences;
+				OutputDataReferences.AddDataReadReference(TEXT("Out"), TDataClassReadRef(ValueRef));
+
+				return OutputDataReferences;
+			}
+		};
+
+	public:
+		TMathOpNode(const FNodeInitData& InInitData)
+			: FNodeFacade(InInitData.InstanceName, InInitData.InstanceID, TFacadeOperatorClass<TMathOperator>())
+		{
+		}
+	};
+
+	// Standard Math Operations
+	template <typename TDataClass, typename TOperandDataClass = TDataClass>
+	class TMathOpAdd
+	{
+	public:
+		using TDataClassReadRef = TDataReadReference<TDataClass>;
+		using TDataClassWriteRef = TDataWriteReference<TDataClass>;
+
+		using TOperandDataClassReadRef = TDataReadReference<TOperandDataClass>;
+
+		static const FVertexInterface& GetVertexInterface()
+		{
+			static const FVertexInterface Interface(
+				FInputVertexInterface(
+					TInputDataVertexModel<TDataClass>(MathOpNames::PrimaryOperandName, LOCTEXT("MathOpAddendInitialTooltip", "Initial addend."), static_cast<TDataClass>(0)),
+					TInputDataVertexModel<TOperandDataClass>(MathOpNames::AdditionalOperandsName, LOCTEXT("MathOpAddendsTooltip", "Additional addend(s)."), static_cast<TOperandDataClass>(0))
+				),
+				FOutputVertexInterface(
+					TOutputDataVertexModel<TDataClass>(TEXT("Out"), LOCTEXT("MathOpOutTooltip", "Math operation result"))
+				)
+			);
+
+			return Interface;
+		}
+
+		static TDataClass GetDefault(const FOperatorSettings& InSettings, const FLiteral& InVertexDefault)
+		{
+			return static_cast<TDataClass>(0);
+		}
+
+		static TOperandDataClass GetDefaultOp(const FOperatorSettings& InSettings, const FLiteral& InVertexDefault)
+		{
+			return static_cast<TOperandDataClass>(0);
+		}
+
+		static void Calculate(TMathOpAdd<TDataClass, TOperandDataClass>& InInstanceData, const TDataClassReadRef& InPrimaryOperand, const TArray<TOperandDataClassReadRef>& InAdditionalOperands, TDataClassWriteRef& OutResult)
+		{
+			if (!ensure(!InAdditionalOperands.IsEmpty()))
+			{
+				return;
+			}
+
+			*OutResult = *InPrimaryOperand;
+			for (int32 i = 0; i < InAdditionalOperands.Num(); ++i)
+			{
+				*OutResult += *InAdditionalOperands[i];
+			}
+		}
+	};
+
+	template <typename TDataClass, typename TOperandDataClass = TDataClass>
+	class TMathOpSubtract
+	{
+	public:
+		using TDataClassReadRef = TDataReadReference<TDataClass>;
+		using TDataClassWriteRef = TDataWriteReference<TDataClass>;
+
+		using TOperandDataClassReadRef = TDataReadReference<TOperandDataClass>;
+
+		static const FVertexInterface& GetVertexInterface()
+		{
+			static const FVertexInterface Interface(
+				FInputVertexInterface(
+					TInputDataVertexModel<TDataClass>(MathOpNames::PrimaryOperandName, LOCTEXT("MathOpMinuendTooltip", "Minuend."), static_cast<TDataClass>(0)),
+					TInputDataVertexModel<TOperandDataClass>(MathOpNames::AdditionalOperandsName, LOCTEXT("MathOpSubtrahendsTooltip", "Subtrahend(s)."), static_cast<TOperandDataClass>(0))
+				),
+				FOutputVertexInterface(
+					TOutputDataVertexModel<TDataClass>(TEXT("Out"), LOCTEXT("MathOpOutTooltip", "Subtraction result"))
+				)
+			);
+
+			return Interface;
+		}
+
+		static TDataClass GetDefault(const FOperatorSettings& InSettings, const FLiteral& InVertexDefault)
+		{
+			return static_cast<TDataClass>(0);
+		}
+
+		static TOperandDataClass GetDefaultOp(const FOperatorSettings& InSettings, const FLiteral& InVertexDefault)
+		{
+			return static_cast<TOperandDataClass>(0);
+		}
+
+		static void Calculate(TMathOpSubtract<TDataClass, TOperandDataClass>& InInstanceData, const TDataClassReadRef& InPrimaryOperand, const TArray<TOperandDataClassReadRef>& InAdditionalOperands, TDataClassWriteRef& OutResult)
+		{
+			if (!ensure(!InAdditionalOperands.IsEmpty()))
+			{
+				return;
+			}
+
+			*OutResult = *InPrimaryOperand;
+			for (int32 i = 0; i < InAdditionalOperands.Num(); ++i)
+			{
+				*OutResult -= *InAdditionalOperands[i];
+			}
+		}
+	};
+
+	template <typename TDataClass, typename TOperandDataClass = TDataClass>
+	class TMathOpMultiply
+	{
+	public:
+		using TDataClassReadRef = TDataReadReference<TDataClass>;
+		using TDataClassWriteRef = TDataWriteReference<TDataClass>;
+
+		using TOperandDataClassReadRef = TDataReadReference<TOperandDataClass>;
+
+		static const FVertexInterface& GetVertexInterface()
+		{
+			static const FVertexInterface Interface(
+				FInputVertexInterface(
+					TInputDataVertexModel<TDataClass>(MathOpNames::PrimaryOperandName, LOCTEXT("MathOpInitMultiplicandTooltip", "Initial multiplicand."), static_cast<TDataClass>(1)),
+					TInputDataVertexModel<TOperandDataClass>(MathOpNames::AdditionalOperandsName, LOCTEXT("MathOpMultiplicandsTooltip", "Additional multiplicand(s)."), static_cast<TOperandDataClass>(1))
+				),
+				FOutputVertexInterface(
+					TOutputDataVertexModel<TDataClass>(TEXT("Out"), LOCTEXT("MultiplicationResultTooltip", "Multiplication result"))
+				)
+			);
+
+			return Interface;
+		}
+
+		static TDataClass GetDefault(const FOperatorSettings& InSettings, const FLiteral& InVertexDefault)
+		{
+			return static_cast<TDataClass>(1);
+		}
+
+		static TOperandDataClass GetDefaultOp(const FOperatorSettings& InSettings, const FLiteral& InVertexDefault)
+		{
+			return static_cast<TOperandDataClass>(1);
+		}
+
+		static void Calculate(TMathOpMultiply<TDataClass, TOperandDataClass>& InInstanceData, const TDataClassReadRef& InPrimaryOperand, const TArray<TOperandDataClassReadRef>& InAdditionalOperands, TDataClassWriteRef& OutResult)
+		{
+			if (!ensure(!InAdditionalOperands.IsEmpty()))
+			{
+				return;
+			}
+
+			*OutResult = *InPrimaryOperand;
+			for (int32 i = 0; i < InAdditionalOperands.Num(); ++i)
+			{
+				*OutResult *= *InAdditionalOperands[i];
+			}
+		}
+	};
+
+	template <typename TDataClass, typename TOperandDataClass = TDataClass>
+	class TMathOpDivide
+	{
+	public:
+		using TDataClassReadRef = TDataReadReference<TDataClass>;
+		using TDataClassWriteRef = TDataWriteReference<TDataClass>;
+
+		using TOperandDataClassReadRef = TDataReadReference<TOperandDataClass>;
+
+		static const FVertexInterface& GetVertexInterface()
+		{
+			static const FVertexInterface Interface(
+				FInputVertexInterface(
+					TInputDataVertexModel<TDataClass>(MathOpNames::PrimaryOperandName, LOCTEXT("MathOpDividendTooltip", "Dividend."), static_cast<TDataClass>(1)),
+					TInputDataVertexModel<TOperandDataClass>(MathOpNames::AdditionalOperandsName, LOCTEXT("MathOpDivisorsTooltip", "Divisor(s)."), static_cast<TOperandDataClass>(1))
+				),
+				FOutputVertexInterface(
+					TOutputDataVertexModel<TDataClass>(TEXT("Out"), LOCTEXT("MathOpOutTooltip", "Math operation result"))
+				)
+			);
+
+			return Interface;
+		}
+
+		static TDataClass GetDefault(const FOperatorSettings& InSettings, const FLiteral& InVertexDefault)
+		{
+			return static_cast<TDataClass>(1);
+		}
+
+		static TOperandDataClass GetDefaultOp(const FOperatorSettings& InSettings, const FLiteral& InVertexDefault)
+		{
+			return static_cast<TOperandDataClass>(1);
+		}
+
+		static void Calculate(TMathOpDivide& InInstanceData, const TDataClassReadRef& InPrimaryOperand, const TArray<TOperandDataClassReadRef>& InAdditionalOperands, TDataClassWriteRef& OutResult)
+		{
+			if (!ensure(!InAdditionalOperands.IsEmpty()))
+			{
+				return;
+			}
+
+			*OutResult = *InPrimaryOperand;
+			for (int32 i = 0; i < InAdditionalOperands.Num(); ++i)
+			{
+				const TDataClass& OperandValue = *InAdditionalOperands[i];
+				if (OperandValue == static_cast<TDataClass>(0))
+				{
+					// TODO: Error here
+					return;
+				}
+
+				*OutResult /= OperandValue;
+			}
+		}
+	};
+
+	template <typename TDataClass, typename TOperandDataClass = TDataClass>
+	class TMathOpModulo
+	{
+	public:
+		using TDataClassReadRef = TDataReadReference<TDataClass>;
+		using TDataClassWriteRef = TDataWriteReference<TDataClass>;
+
+		using TOperandDataClassReadRef = TDataReadReference<TOperandDataClass>;
+
+		static const FVertexInterface& GetVertexInterface()
+		{
+			static const FVertexInterface Interface(
+				FInputVertexInterface(
+					TInputDataVertexModel<TDataClass>(MathOpNames::PrimaryOperandName, LOCTEXT("MathOpModuloDividendTooltip", "Dividend."), static_cast<TDataClass>(1)),
+					TInputDataVertexModel<TOperandDataClass>(MathOpNames::AdditionalOperandsName, LOCTEXT("MathOpModuloDivisorsTooltip", "Divisor(s)."), static_cast<TOperandDataClass>(1))
+				),
+				FOutputVertexInterface(
+					TOutputDataVertexModel<TDataClass>(TEXT("Out"), LOCTEXT("MathOpOutTooltip", "Math operation result"))
+				)
+			);
+
+			return Interface;
+		}
+
+		static TDataClass GetDefault(const FOperatorSettings& InSettings, const FLiteral& InVertexDefault)
+		{
+			return static_cast<TDataClass>(1);
+		}
+
+		static TOperandDataClass GetDefaultOp(const FOperatorSettings& InSettings, const FLiteral& InVertexDefault)
+		{
+			return static_cast<TOperandDataClass>(1);
+		}
+
+		static void Calculate(TMathOpModulo& InInstanceData, const TDataClassReadRef& InPrimaryOperand, const TArray<TOperandDataClassReadRef>& InAdditionalOperands, TDataClassWriteRef& OutResult)
+		{
+			if (!ensure(!InAdditionalOperands.IsEmpty()))
+			{
+				return;
+			}
+
+			*OutResult = *InPrimaryOperand;
+
+			// TODO: chaining modulo operations doesn't make too much sense... how do we forbid additional operands in some math ops?
+			for (int32 i = 0; i < InAdditionalOperands.Num(); ++i)
+			{
+				const TDataClass& OperandValue = *InAdditionalOperands[i];
+				if (OperandValue == static_cast<TDataClass>(0))
+				{
+					// TODO: Error here
+					return;
+				}
+
+				*OutResult %= OperandValue;
+			}
+		}
+	};
+
+	template <typename TDataClass, typename TOperandDataClass = TDataClass>
+	class TMathOpPower
+	{
+	public:
+		using TDataClassReadRef = TDataReadReference<TDataClass>;
+		using TDataClassWriteRef = TDataWriteReference<TDataClass>;
+
+		using TOperandDataClassReadRef = TDataReadReference<TOperandDataClass>;
+
+		static const FVertexInterface& GetVertexInterface()
+		{
+			static const FVertexInterface DefaultInterface(
+				FInputVertexInterface(
+					TInputDataVertexModel<float>(MathOpNames::PrimaryOperandName, LOCTEXT("Base", "The base of the power")),
+					TInputDataVertexModel<float>(MathOpNames::AdditionalOperandsName, LOCTEXT("Exponent", "The exponent to take the base to the power of"))
+				),
+				FOutputVertexInterface(
+					TOutputDataVertexModel<float>(TEXT("Out"), LOCTEXT("Result", "Returns Base to the Exponent power"))
+				)
+			);
+
+			return DefaultInterface;
+		}
+
+		static TDataClass GetDefault(const FOperatorSettings& InSettings, const FLiteral& InVertexDefault)
+		{
+			return static_cast<TDataClass>(1);
+		}
+
+		static TOperandDataClass GetDefaultOp(const FOperatorSettings& InSettings, const FLiteral& InVertexDefault)
+		{
+			return static_cast<TOperandDataClass>(1);
+		}
+
+		static void Calculate(TMathOpPower& InInstanceData, const TDataClassReadRef& InPrimaryOperand, const TArray<TOperandDataClassReadRef>& InAdditionalOperands, TDataClassWriteRef& OutResult)
+		{
+			if (!ensure(!InAdditionalOperands.IsEmpty()))
+			{
+				return;
+			}
+
+			*OutResult = *InPrimaryOperand;
+
+			for (int32 i = 0; i < InAdditionalOperands.Num(); ++i)
+			{
+				const TDataClass& OperandValue = *InAdditionalOperands[i];
+
+				*OutResult = FMath::Pow(*OutResult, OperandValue);
+			}
+		}
+	};
+
+	template <typename TDataClass, typename TOperandDataClass = TDataClass>
+	class TMathOpLogarithm
+	{
+	public:
+		using TDataClassReadRef = TDataReadReference<TDataClass>;
+		using TDataClassWriteRef = TDataWriteReference<TDataClass>;
+
+		using TOperandDataClassReadRef = TDataReadReference<TOperandDataClass>;
+
+		static const FVertexInterface& GetVertexInterface()
+		{
+			static const FVertexInterface DefaultInterface(
+				FInputVertexInterface(
+					TInputDataVertexModel<float>(MathOpNames::PrimaryOperandName, LOCTEXT("Base", "The base of the logarithm")),
+					TInputDataVertexModel<float>(MathOpNames::AdditionalOperandsName, LOCTEXT("Value", "The value to find the logarithm of"))
+				),
+				FOutputVertexInterface(
+					TOutputDataVertexModel<float>(TEXT("Out"), LOCTEXT("Result", "The logarithm of the inputted value"))
+				)
+			);
+
+			return DefaultInterface;
+		}
+
+		static TDataClass GetDefault(const FOperatorSettings& InSettings, const FLiteral& InVertexDefault)
+		{
+			return static_cast<TDataClass>(1);
+		}
+
+		static TOperandDataClass GetDefaultOp(const FOperatorSettings& InSettings, const FLiteral& InVertexDefault)
+		{
+			return static_cast<TOperandDataClass>(1);
+		}
+
+		static void Calculate(TMathOpLogarithm& InInstanceData, const TDataClassReadRef& InPrimaryOperand, const TArray<TOperandDataClassReadRef>& InAdditionalOperands, TDataClassWriteRef& OutResult)
+		{
+			if (!ensure(!InAdditionalOperands.IsEmpty()))
+			{
+				return;
+			}
+
+			*OutResult = *InPrimaryOperand;
+
+			//TODO: Find out how to disallow it from additional inputs, it doesn't really make sense in the context of logarithms
+			for (int32 i = 0; i < InAdditionalOperands.Num(); ++i)
+			{
+				const TDataClass& OperandValue = *InAdditionalOperands[i];
+
+				*OutResult = FMath::LogX(FMath::Max(SMALL_NUMBER, *OutResult), FMath::Max(SMALL_NUMBER, OperandValue));
+			}
+		}
+	};
+
+	// Specialized Math Operations
+	template <>
+	class TMathOpAdd<FAudioBuffer>
+	{
+	public:
+		static const FVertexInterface& GetVertexInterface()
+		{
+			static const FVertexInterface Interface(
+				FInputVertexInterface(
+					TInputDataVertexModel<FAudioBuffer>(MathOpNames::PrimaryOperandName, LOCTEXT("MathOpAddendTooltip", "First attend.")),
+					TInputDataVertexModel<FAudioBuffer>(MathOpNames::AdditionalOperandsName, LOCTEXT("MathOpAddendAdditionalTooltip", "Additional attends."))
+				),
+				FOutputVertexInterface(
+					TOutputDataVertexModel<FAudioBuffer>(TEXT("Out"), LOCTEXT("MathOpOutTooltip", "Math operation result"))
+				)
+			);
+
+			return Interface;
+		}
+
+		static FAudioBuffer GetDefault(const FOperatorSettings& InSettings, const FLiteral& InVertexDefault)
+		{
+			return FAudioBuffer(InSettings.GetNumFramesPerBlock());
+		}
+
+		static FAudioBuffer GetDefaultOp(const FOperatorSettings& InSettings, const FLiteral& InVertexDefault)
+		{
+			return GetDefault(InSettings, InVertexDefault);
+		}
+
+		static void Calculate(TMathOpAdd<FAudioBuffer>& InInstanceData, FAudioBufferReadRef& InPrimaryOperand, const TArray<FAudioBufferReadRef>& InAdditionalOperands, FAudioBufferWriteRef& OutResult)
+		{
+			if (!ensure(!InAdditionalOperands.IsEmpty()))
+			{
+				OutResult->Zero();
+				return;
+			}
+
+			const int32 NumSamples = InAdditionalOperands[0]->Num();
+			if (!ensure(NumSamples == OutResult->Num()))
+			{
+				// TODO: Error
+				OutResult->Zero();
+				return;
+			}
+
+			FMemory::Memcpy(OutResult->GetData(), InPrimaryOperand->GetData(), sizeof(float) * NumSamples);
+
+			for (int32 i = 0; i < InAdditionalOperands.Num(); ++i)
+			{
+				const FAudioBuffer& Operand = *InAdditionalOperands[i];
+				if (!ensure(NumSamples == Operand.Num()))
+				{
+					// TODO: Error
+					OutResult->Zero();
+					return;
+				}
+
+				const float* OperandData = InAdditionalOperands[i]->GetData();
+				float* OutData = OutResult->GetData();
+
+				if (NumSamples % AUDIO_SIMD_FLOAT_ALIGNMENT)
+				{
+					for (int32 SampleIndex = 0; SampleIndex < NumSamples; ++SampleIndex)
+					{
+						OutData[SampleIndex] += OperandData[SampleIndex];
+					}
+				}
+				else
+				{
+					const float StartGain = 1.0f;
+					const float EndGain = 1.0f;
+					Audio::MixInBufferFast(OperandData, OutData, NumSamples, StartGain, EndGain);
+				}
+			}
+		}
+	};
+
+	template <>
+	class TMathOpAdd<FTime>
+	{
+	public:
+		static const FVertexInterface& GetVertexInterface()
+		{
+			static const FVertexInterface Interface(
+				FInputVertexInterface(
+					TInputDataVertexModel<FTime>(MathOpNames::PrimaryOperandName, LOCTEXT("MathOpAddendTooltip", "First attend."), 0.0f),
+					TInputDataVertexModel<FTime>(MathOpNames::AdditionalOperandsName, LOCTEXT("MathOpAddendAdditionalTooltip", "Additional attends."), 0.0f)
+				),
+				FOutputVertexInterface(
+					TOutputDataVertexModel<FTime>(TEXT("Out"), LOCTEXT("MathOpOutTooltip", "Math operation result"))
+				)
+			);
+
+			return Interface;
+		}
+
+		static FTime GetDefault(const FOperatorSettings& InSettings, const FLiteral& InVertexDefault)
+		{
+			return FTime(InVertexDefault.Value.Get<float>());
+		}
+
+		static FTime GetDefaultOp(const FOperatorSettings& InSettings, const FLiteral& InVertexDefault)
+		{
+			return GetDefault(InSettings, InVertexDefault);
+		}
+
+		static void Calculate(TMathOpAdd& InInstanceData, const FTimeReadRef& InPrimaryOperand, const TArray<FTimeReadRef>& InAdditionalOperands, FTimeWriteRef& OutResult)
+		{
+			if (!ensure(!InAdditionalOperands.IsEmpty()))
+			{
+				return;
+			}
+
+			*OutResult = *InPrimaryOperand;
+			for (int32 i = 0; i < InAdditionalOperands.Num(); ++i)
+			{
+				*OutResult += *InAdditionalOperands[i];
+			}
+		}
+	};
+
+	template <>
+	class TMathOpAdd<FAudioBuffer, float>
+	{
+	public:
+		static const FVertexInterface& GetVertexInterface()
+		{
+			static const FVertexInterface Interface(
+				FInputVertexInterface(
+					TInputDataVertexModel<FAudioBuffer>(MathOpNames::PrimaryOperandName, LOCTEXT("MathOpAddendAudioTooltip", "Audio Buffer to add offset(s) to.")),
+					TInputDataVertexModel<float>(MathOpNames::AdditionalOperandsName, LOCTEXT("MathOpAddendAdditionalTooltip", "Float attends of which to offset buffer samples."), 0.0f)
+				),
+				FOutputVertexInterface(
+					TOutputDataVertexModel<FAudioBuffer>(TEXT("Out"), LOCTEXT("MathOpAudioFloatAddOutTooltip", "Resulting buffer"))
+				)
+			);
+
+			return Interface;
+		}
+
+		static FAudioBuffer GetDefault(const FOperatorSettings& InSettings, const FLiteral& InVertexDefault)
+		{
+			return FAudioBuffer(InSettings.GetNumFramesPerBlock());
+		}
+
+		static float GetDefaultOp(const FOperatorSettings& InSettings, const FLiteral& InVertexDefault)
+		{
+			return InVertexDefault.Value.Get<float>();
+		}
+
+		static void Calculate(TMathOpAdd& InInstanceData, const FAudioBufferReadRef& InPrimaryOperand, const TArray<FFloatReadRef>& InAdditionalOperands, FAudioBufferWriteRef& OutResult)
+		{
+			FMemory::Memcpy(OutResult->GetData(), InPrimaryOperand->GetData(), sizeof(float) * OutResult->Num());
+
+			const int32 SIMDRemainder = OutResult->Num() % AUDIO_SIMD_FLOAT_ALIGNMENT;
+			const int32 SIMDCount = OutResult->Num() - SIMDRemainder;
+
+			for (int32 i = 0; i < InAdditionalOperands.Num(); ++i)
+			{
+				Audio::AddConstantToBufferInplace(OutResult->GetData(), SIMDCount, *InAdditionalOperands[i]);
+
+				for (int32 j = SIMDCount; j < OutResult->Num(); ++j)
+				{
+					OutResult->GetData()[j] += *InAdditionalOperands[i];
+				}
+			}
+		}
+	};
+
+	// int64 has to be explicitly implemented as they are PODs that are not implemented as a literal base type.
+	template <>
+	class TMathOpAdd<int64>
+	{
+	public:
+		static const FVertexInterface& GetVertexInterface()
+		{
+			static const FVertexInterface Interface(
+				FInputVertexInterface(
+					TInputDataVertexModel<int64>(MathOpNames::PrimaryOperandName, LOCTEXT("MathOpInt64InitialAddendMinuendTooltip", "Initial Addend."), 0),
+					TInputDataVertexModel<int64>(MathOpNames::AdditionalOperandsName, LOCTEXT("MathOpInt64AddendsTooltip", "Additional Addend(s)."), 0)
+				),
+				FOutputVertexInterface(
+					TOutputDataVertexModel<int64>(TEXT("Out"), LOCTEXT("MathOpOutTooltip", "Addition result"))
+				)
+			);
+
+			return Interface;
+		}
+
+		static int32 GetDefault(const FOperatorSettings& InSettings, const FLiteral& InVertexDefault)
+		{
+			return static_cast<int64>(InVertexDefault.Value.Get<int32>());
+		}
+
+		static int32 GetDefaultOp(const FOperatorSettings& InSettings, const FLiteral& InVertexDefault)
+		{
+			return static_cast<int64>(InVertexDefault.Value.Get<int32>());
+		}
+
+		static void Calculate(TMathOpAdd& InInstanceData, const FInt64ReadRef& InPrimaryOperand, const TArray<FInt64ReadRef>& InAdditionalOperands, FInt64WriteRef& OutResult)
+		{
+			*OutResult = *InPrimaryOperand;
+
+			if (!ensure(!InAdditionalOperands.IsEmpty()))
+			{
+				return;
+			}
+
+			for (int32 i = 0; i < InAdditionalOperands.Num(); ++i)
+			{
+				const int64 OperandValue = *InAdditionalOperands[i];
+				*OutResult += OperandValue;
+			}
+		}
+	};
+
+	template <>
+	class TMathOpSubtract<FAudioBuffer>
+	{
+	public:
+		static const FVertexInterface& GetVertexInterface()
+		{
+			static const FVertexInterface Interface(
+				FInputVertexInterface(
+					TInputDataVertexModel<FAudioBuffer>(MathOpNames::PrimaryOperandName, LOCTEXT("MathOpBuffersMinuendTooltip", "Initial buffer to act as minuend.")),
+					TInputDataVertexModel<FAudioBuffer>(MathOpNames::AdditionalOperandsName, LOCTEXT("MathOpSubtractBuffersSubtrahendsTooltip", "Additional buffers to act as subtrahend(s)."))
+				),
+				FOutputVertexInterface(
+					TOutputDataVertexModel<FAudioBuffer>(TEXT("Out"), LOCTEXT("MathOpSubtractBuffersOutTooltip", "Resulting buffer"))
+				)
+			);
+
+			return Interface;
+		}
+
+		static FAudioBuffer GetDefault(const FOperatorSettings& InSettings, const FLiteral& InVertexDefault)
+		{
+			return FAudioBuffer(InSettings.GetNumFramesPerBlock());
+		}
+
+		static FAudioBuffer GetDefaultOp(const FOperatorSettings& InSettings, const FLiteral& InVertexDefault)
+		{
+			return GetDefault(InSettings, InVertexDefault);
+		}
+
+		static void Calculate(TMathOpSubtract& InInstanceData, FAudioBufferReadRef& InPrimaryOperand, const TArray<FAudioBufferReadRef>& InAdditionalOperands, FAudioBufferWriteRef& OutResult)
+		{
+			if (!ensure(!InAdditionalOperands.IsEmpty()))
+			{
+				OutResult->Zero();
+				return;
+			}
+
+			const int32 NumSamples = InAdditionalOperands[0]->Num();
+			if (!ensure(NumSamples == OutResult->Num()))
+			{
+				// TODO: Error
+				OutResult->Zero();
+				return;
+			}
+
+			FMemory::Memcpy(OutResult->GetData(), InPrimaryOperand->GetData(), sizeof(float) * NumSamples);
+
+			for (int32 i = 0; i < InAdditionalOperands.Num(); ++i)
+			{
+				const FAudioBuffer& Operand = *InAdditionalOperands[i];
+				if (!ensure(NumSamples == Operand.Num()))
+				{
+					// TODO: Error
+					OutResult->Zero();
+					return;
+				}
+
+				const float* OperandData = InAdditionalOperands[i]->GetData();
+				float* OutData = OutResult->GetData();
+
+				if (NumSamples % AUDIO_SIMD_FLOAT_ALIGNMENT)
+				{
+					for (int32 SampleIndex = 0; SampleIndex < NumSamples; ++SampleIndex)
+					{
+						OutData[i] -= OperandData[i];
+					}
+				}
+				else
+				{
+					Audio::BufferSubtractInPlace2Fast(OutData, OperandData, NumSamples);
+				}
+			}
+		}
+	};
+
+	template <>
+	class TMathOpSubtract<FTime>
+	{
+	public:
+		static const FVertexInterface& GetVertexInterface()
+		{
+			static const FVertexInterface Interface(
+				FInputVertexInterface(
+					TInputDataVertexModel<FTime>(MathOpNames::PrimaryOperandName, LOCTEXT("MathOpTimeMinuendTooltip", "Time minuend."), 0.0f),
+					TInputDataVertexModel<FTime>(MathOpNames::AdditionalOperandsName, LOCTEXT("MathOpTimeSubtrahendsTooltip", "Time subtrahends."), 0.0f)
+				),
+				FOutputVertexInterface(
+					TOutputDataVertexModel<FTime>(TEXT("Out"), LOCTEXT("MathOpTimeSubtractOutTooltip", "Resulting time value"))
+				)
+			);
+
+			return Interface;
+		}
+
+		static FTime GetDefault(const FOperatorSettings& InSettings, const FLiteral& InVertexDefault)
+		{
+			return FTime(InVertexDefault.Value.Get<float>());
+		}
+
+		static FTime GetDefaultOp(const FOperatorSettings& InSettings, const FLiteral& InVertexDefault)
+		{
+			return GetDefault(InSettings, InVertexDefault);
+		}
+
+		static void Calculate(TMathOpSubtract& InInstanceData, const FTimeReadRef& InPrimaryOperand, const TArray<FTimeReadRef>& InAdditionalOperands, FTimeWriteRef& OutResult)
+		{
+			if (!ensure(!InAdditionalOperands.IsEmpty()))
+			{
+				return;
+			}
+
+			*OutResult = *InPrimaryOperand;
+			for (int32 i = 0; i < InAdditionalOperands.Num(); ++i)
+			{
+				*OutResult -= *InAdditionalOperands[i];
+			}
+		}
+	};
+
+	// int64 has to be explicitly implemented as they are PODs that are not implemented as a literal base type.
+	template <>
+	class TMathOpSubtract<int64>
+	{
+	public:
+		static const FVertexInterface& GetVertexInterface()
+		{
+			static const FVertexInterface Interface(
+				FInputVertexInterface(
+					TInputDataVertexModel<int64>(MathOpNames::PrimaryOperandName, LOCTEXT("MathOpMinuendTooltip", "Minuend."), 0),
+					TInputDataVertexModel<int64>(MathOpNames::AdditionalOperandsName, LOCTEXT("MathOpSubtrahendsTooltip", "Subtrahend(s)."), 0)
+				),
+				FOutputVertexInterface(
+					TOutputDataVertexModel<int64>(TEXT("Out"), LOCTEXT("MathOpOutTooltip", "Subtraction result"))
+				)
+			);
+
+			return Interface;
+		}
+
+		static int32 GetDefault(const FOperatorSettings& InSettings, const FLiteral& InVertexDefault)
+		{
+			return static_cast<int64>(InVertexDefault.Value.Get<int32>());
+		}
+
+		static int32 GetDefaultOp(const FOperatorSettings& InSettings, const FLiteral& InVertexDefault)
+		{
+			return static_cast<int64>(InVertexDefault.Value.Get<int32>());
+		}
+
+		static void Calculate(TMathOpSubtract& InInstanceData, const FInt64ReadRef& InPrimaryOperand, const TArray<FInt64ReadRef>& InAdditionalOperands, FInt64WriteRef& OutResult)
+		{
+			*OutResult = *InPrimaryOperand;
+
+			if (!ensure(!InAdditionalOperands.IsEmpty()))
+			{
+				return;
+			}
+
+			for (int32 i = 0; i < InAdditionalOperands.Num(); ++i)
+			{
+				const int64 OperandValue = *InAdditionalOperands[i];
+				*OutResult -= OperandValue;
+			}
+		}
+	};
+
+	template <>
+	class TMathOpMultiply<FAudioBuffer>
+	{
+	public:
+		static const FVertexInterface& GetVertexInterface()
+		{
+			static const FVertexInterface Interface(
+				FInputVertexInterface(
+					TInputDataVertexModel<FAudioBuffer>(MathOpNames::PrimaryOperandName, LOCTEXT("MathOpAudioInitMultiplicandTooltip", "Initial audio to multiply.")),
+					TInputDataVertexModel<FAudioBuffer>(MathOpNames::AdditionalOperandsName, LOCTEXT("MathOpMultiplyAudioSubtrahendsTooltip", "Additional audio to multiply sample-by-sample."))
+				),
+				FOutputVertexInterface(
+					TOutputDataVertexModel<FAudioBuffer>(TEXT("Out"), LOCTEXT("MathOpMultiplyAudioOutTooltip", "Resulting buffer"))
+				)
+			);
+
+			return Interface;
+		}
+
+		static FAudioBuffer GetDefault(const FOperatorSettings& InSettings, const FLiteral& InVertexDefault)
+		{
+			return FAudioBuffer(InSettings.GetNumFramesPerBlock());
+		}
+
+		static FAudioBuffer GetDefaultOp(const FOperatorSettings& InSettings, const FLiteral& InVertexDefault)
+		{
+			return GetDefault(InSettings, InVertexDefault);
+		}
+
+		static void Calculate(TMathOpMultiply& InInstanceData, const FAudioBufferReadRef& InPrimaryOperand, const TArray<FAudioBufferReadRef>& InAdditionalOperands, FAudioBufferWriteRef& OutResult)
+		{
+			if (!ensure(!InAdditionalOperands.IsEmpty()))
+			{
+				OutResult->Zero();
+				return;
+			}
+
+			const int32 NumSamples = InPrimaryOperand->Num();
+			if (!ensure(NumSamples == OutResult->Num()))
+			{
+				// TODO: Error
+				OutResult->Zero();
+				return;
+			}
+
+			FMemory::Memcpy(OutResult->GetData(), InPrimaryOperand->GetData(), sizeof(float) * NumSamples);
+
+			for (int32 i = 0; i < InAdditionalOperands.Num(); ++i)
+			{
+				const FAudioBuffer& Operand = *InAdditionalOperands[i];
+				if (!ensure(NumSamples == Operand.Num()))
+				{
+					// TODO: Error
+					OutResult->Zero();
+					return;
+				}
+
+				const float* OperandData = InAdditionalOperands[i]->GetData();
+				float* OutData = OutResult->GetData();
+
+				if (NumSamples % AUDIO_SIMD_FLOAT_ALIGNMENT)
+				{
+					for (int32 SampleIndex = 0; SampleIndex < NumSamples; ++SampleIndex)
+					{
+						OutData[SampleIndex] += OperandData[SampleIndex];
+					}
+				}
+				else
+				{
+					Audio::MultiplyBuffersInPlace(OperandData, OutData, NumSamples);
+				}
+			}
+		}
+	};
+
+	template <>
+	class TMathOpMultiply<FAudioBuffer, float>
+	{
+		bool bInit = false;
+		float LastGain = 0.0f;
+
+	public:
+		static const FVertexInterface& GetVertexInterface()
+		{
+			static const FVertexInterface Interface(
+				FInputVertexInterface(
+					TInputDataVertexModel<FAudioBuffer>(MathOpNames::PrimaryOperandName, LOCTEXT("MathOpAudioMultiplyFloatTooltip", "Audio multiplicand.")),
+					TInputDataVertexModel<float>(MathOpNames::AdditionalOperandsName, LOCTEXT("MathOpAddendAdditionalTooltip", "Float multiplicand to apply sample-by-sample to audio. Interpolates over buffer size on value change."), 1.0f)
+				),
+				FOutputVertexInterface(
+					TOutputDataVertexModel<FAudioBuffer>(TEXT("Out"), LOCTEXT("MathOpAudioFloatMultiplyOutTooltip", "Resulting buffer"))
+				)
+			);
+
+			return Interface;
+		}
+
+		static FAudioBuffer GetDefault(const FOperatorSettings& InSettings, const FLiteral& InVertexDefault)
+		{
+			return FAudioBuffer(InSettings.GetNumFramesPerBlock());
+		}
+
+		static float GetDefaultOp(const FOperatorSettings& InSettings, const FLiteral& InVertexDefault)
+		{
+			return InVertexDefault.Value.Get<float>();
+		}
+
+		static void Calculate(TMathOpMultiply& InInstanceData, const FAudioBufferReadRef& InPrimaryOperand, const TArray<FFloatReadRef>& InAdditionalOperands, FAudioBufferWriteRef& OutResult)
+		{
+			if (!InInstanceData.bInit)
+			{
+				InInstanceData.bInit = true;
+				InInstanceData.LastGain = 1.0f;
+
+				for (int32 i = 0; i < InAdditionalOperands.Num(); ++i)
+				{
+					InInstanceData.LastGain *= *InAdditionalOperands[i];
+				}
+			}
+
+			FMemory::Memcpy(OutResult->GetData(), InPrimaryOperand->GetData(), sizeof(float) * InPrimaryOperand->Num());
+
+			float NewGain = 1.0f;
+			for (int32 i = 0; i < InAdditionalOperands.Num(); ++i)
+			{
+				const int32 SIMDRemainder = OutResult->Num() % AUDIO_SIMD_FLOAT_ALIGNMENT;
+				const int32 SIMDCount = OutResult->Num() - SIMDRemainder;
+
+				Audio::FadeBufferFast(OutResult->GetData(), SIMDCount, InInstanceData.LastGain, *InAdditionalOperands[i]);
+
+				for (int32 j = SIMDCount; j < OutResult->Num(); ++j)
+				{
+					OutResult->GetData()[j] *= *InAdditionalOperands[j];
+				}
+
+				NewGain *= *InAdditionalOperands[i];
+			}
+
+			InInstanceData.LastGain = NewGain;
+		}
+	};
+
+	template <>
+	class TMathOpMultiply<FTime, float>
+	{
+	public:
+		static const FVertexInterface& GetVertexInterface()
+		{
+			static const FVertexInterface Interface(
+				FInputVertexInterface(
+					TInputDataVertexModel<FTime>(MathOpNames::PrimaryOperandName, LOCTEXT("MathOpTimeMultiplyFloatTooltip", "Time multiplicand."), 1.0f),
+					TInputDataVertexModel<float>(MathOpNames::AdditionalOperandsName, LOCTEXT("MathOpAddendAdditionalTooltip", "Float multiplicand(s)."), 1.0f)
+				),
+				FOutputVertexInterface(
+					TOutputDataVertexModel<FTime>(TEXT("Out"), LOCTEXT("MathOpTimeMultiplyOutTooltip", "Resulting buffer"))
+				)
+			);
+
+			return Interface;
+		}
+
+		static FTime GetDefault(const FOperatorSettings& InSettings, const FLiteral& InVertexDefault)
+		{
+			return FTime(InVertexDefault.Value.Get<float>());
+		}
+
+		static float GetDefaultOp(const FOperatorSettings& InSettings, const FLiteral& InVertexDefault)
+		{
+			return 1.0f;
+		}
+
+		static void Calculate(TMathOpMultiply& InInstanceData, const FTimeReadRef& InPrimaryOperand, const TArray<FFloatReadRef>& InAdditionalOperands, FTimeWriteRef& OutResult)
+		{
+			*OutResult = *InPrimaryOperand;
+
+			if (!ensure(!InAdditionalOperands.IsEmpty()))
+			{
+				return;
+			}
+
+			for (int32 i = 0; i < InAdditionalOperands.Num(); ++i)
+			{
+				const float OperandValue = *InAdditionalOperands[i];
+				*OutResult *= OperandValue;
+			}
+		}
+	};
+
+	// int64 has to be explicitly implemented as they are PODs that are not implemented as a literal base type.
+	template <>
+	class TMathOpMultiply<int64>
+	{
+	public:
+		static const FVertexInterface& GetVertexInterface()
+		{
+			static const FVertexInterface Interface(
+				FInputVertexInterface(
+					TInputDataVertexModel<int64>(MathOpNames::PrimaryOperandName, LOCTEXT("MathOpInitMultiplicandTooltip", "Initial multiplicand."), 1),
+					TInputDataVertexModel<int64>(MathOpNames::AdditionalOperandsName, LOCTEXT("MathOpMultiplicandsTooltip", "Additional multiplicand(s)."), 1)
+				),
+				FOutputVertexInterface(
+					TOutputDataVertexModel<int64>(TEXT("Out"), LOCTEXT("MultiplicationResultTooltip", "Multiplication result"))
+				)
+			);
+
+			return Interface;
+		}
+
+		static int32 GetDefault(const FOperatorSettings& InSettings, const FLiteral& InVertexDefault)
+		{
+			return static_cast<int64>(InVertexDefault.Value.Get<int32>());
+		}
+
+		static int32 GetDefaultOp(const FOperatorSettings& InSettings, const FLiteral& InVertexDefault)
+		{
+			return static_cast<int64>(InVertexDefault.Value.Get<int32>());
+		}
+
+		static void Calculate(TMathOpMultiply& InInstanceData, const FInt64ReadRef& InPrimaryOperand, const TArray<FInt64ReadRef>& InAdditionalOperands, FInt64WriteRef& OutResult)
+		{
+			*OutResult = *InPrimaryOperand;
+
+			if (!ensure(!InAdditionalOperands.IsEmpty()))
+			{
+				return;
+			}
+
+			for (int32 i = 0; i < InAdditionalOperands.Num(); ++i)
+			{
+				const int64 OperandValue = *InAdditionalOperands[i];
+				*OutResult *= OperandValue;
+			}
+		}
+	};
+
+	template <>
+	class TMathOpDivide<FTime, float>
+	{
+	public:
+		static const FVertexInterface& GetVertexInterface()
+		{
+			static const FVertexInterface Interface(
+				FInputVertexInterface(
+					TInputDataVertexModel<FTime>(MathOpNames::PrimaryOperandName, LOCTEXT("MathOpTimeMultiplyFloatTooltip", "Time multiplicand."), 1.0f),
+					TInputDataVertexModel<float>(MathOpNames::AdditionalOperandsName, LOCTEXT("MathOpAddendAdditionalTooltip", "Float multiplicand(s)."), 1.0f)
+				),
+				FOutputVertexInterface(
+					TOutputDataVertexModel<FTime>(TEXT("Out"), LOCTEXT("MathOpTimeMultiplyOutTooltip", "Resulting buffer"))
+				)
+			);
+
+			return Interface;
+		}
+
+		static FTime GetDefault(const FOperatorSettings& InSettings, const FLiteral& InVertexDefault)
+		{
+			return FTime(InVertexDefault.Value.Get<float>());
+		}
+
+		static float GetDefaultOp(const FOperatorSettings& InSettings, const FLiteral& InVertexDefault)
+		{
+			return 1.0f;
+		}
+
+		static void Calculate(TMathOpDivide& InInstanceData, const FTimeReadRef& InPrimaryOperand, const TArray<FFloatReadRef>& InAdditionalOperands, FTimeWriteRef& OutResult)
+		{
+			*OutResult = *InPrimaryOperand;
+
+			if (!ensure(!InAdditionalOperands.IsEmpty()))
+			{
+				return;
+			}
+
+			for (int32 i = 0; i < InAdditionalOperands.Num(); ++i)
+			{
+				const float OperandValue = *InAdditionalOperands[i];
+				if (OperandValue == 0.0f)
+				{
+					// TODO: Error here
+					continue;
+				}
+
+				*OutResult /= OperandValue;
+			}
+		}
+	};
+
+	// int64 has to be explicitly implemented as they are PODs that are not implemented as a literal base type.
+	template <>
+	class TMathOpDivide<int64>
+	{
+	public:
+		static const FVertexInterface& GetVertexInterface()
+		{
+			static const FVertexInterface Interface(
+				FInputVertexInterface(
+					TInputDataVertexModel<int64>(MathOpNames::PrimaryOperandName, LOCTEXT("MathOpDivideDividendTooltip", "Dividend of operation."), 1),
+					TInputDataVertexModel<int64>(MathOpNames::AdditionalOperandsName, LOCTEXT("MathOpDivideDivisorTooltip", "Divisor of operation."), 1)
+				),
+				FOutputVertexInterface(
+					TOutputDataVertexModel<int64>(TEXT("Out"), LOCTEXT("MathOpDivideOutTooltip", "Resulting value"))
+				)
+			);
+
+			return Interface;
+		}
+
+		static int64 GetDefault(const FOperatorSettings& InSettings, const FLiteral& InVertexDefault)
+		{
+			// TODO: we are truncating to int32, why do we need int64?
+			return static_cast<int64>(InVertexDefault.Value.Get<int32>());
+		}
+
+		static int64 GetDefaultOp(const FOperatorSettings& InSettings, const FLiteral& InVertexDefault)
+		{
+			// TODO: we are truncating to int32, why do we need in64?
+			return static_cast<int64>(InVertexDefault.Value.Get<int32>());
+		}
+
+		static void Calculate(TMathOpDivide& InInstanceData, const FInt64ReadRef& InPrimaryOperand, const TArray<FInt64ReadRef>& InAdditionalOperands, FInt64WriteRef& OutResult)
+		{
+			*OutResult = *InPrimaryOperand;
+
+			if (!ensure(!InAdditionalOperands.IsEmpty()))
+			{
+				return;
+			}
+
+			for (int32 i = 0; i < InAdditionalOperands.Num(); ++i)
+			{
+				const int64 OperandValue = *InAdditionalOperands[i];
+				if (OperandValue == 0)
+				{
+					// TODO: Error here
+					continue;
+				}
+
+				*OutResult /= OperandValue;
+			}
+		}
+	};
+
+
+	// int64 has to be explicitly implemented as they are PODs that are not implemented as a literal base type.
+	template <>
+	class TMathOpModulo<int64>
+	{
+	public:
+		static const FVertexInterface& GetVertexInterface()
+		{
+			static const FVertexInterface Interface(
+				FInputVertexInterface(
+					TInputDataVertexModel<int64>(MathOpNames::PrimaryOperandName, LOCTEXT("MathOpModuloDividendTooltip", "Dividend of modulo operation."), 1),
+					TInputDataVertexModel<int64>(MathOpNames::AdditionalOperandsName, LOCTEXT("MathOpModuluDivisorTooltip", "Divisor of modulo operation."), 1)
+				),
+				FOutputVertexInterface(
+					TOutputDataVertexModel<int64>(TEXT("Out"), LOCTEXT("MathOpModuloOutTooltip", "Resulting value"))
+				)
+			);
+
+			return Interface;
+		}
+
+		static int64 GetDefault(const FOperatorSettings& InSettings, const FLiteral& InVertexDefault)
+		{
+			return static_cast<int64>(InVertexDefault.Value.Get<int32>());
+		}
+
+		static int64 GetDefaultOp(const FOperatorSettings& InSettings, const FLiteral& InVertexDefault)
+		{
+			return static_cast<int64>(InVertexDefault.Value.Get<int32>());
+		}
+
+		static void Calculate(TMathOpModulo& InInstanceData, const FInt64ReadRef& InPrimaryOperand, const TArray<FInt64ReadRef>& InAdditionalOperands, FInt64WriteRef& OutResult)
+		{
+			*OutResult = *InPrimaryOperand;
+
+			if (!ensure(!InAdditionalOperands.IsEmpty()))
+			{
+				return;
+			}
+
+			for (int32 i = 0; i < InAdditionalOperands.Num(); ++i)
+			{
+				const int64 OperandValue = *InAdditionalOperands[i];
+				if (OperandValue == 0)
+				{
+					// TODO: Error here
+					continue;
+				}
+
+				*OutResult %= OperandValue;
+			}
+		}
+	};
+
+	// Definitions
+	DEFINE_METASOUND_MATHOP(Add, Float, float, LOCTEXT("Metasound_MathAddFloatNodeDescription", "Adds floats."))
+	DEFINE_METASOUND_MATHOP(Add, Int32, int32, LOCTEXT("Metasound_MathAddInt32NodeDescription", "Adds int32s."))
+	DEFINE_METASOUND_MATHOP(Add, Int64, int64, LOCTEXT("Metasound_MathAddInt64NodeDescription", "Adds int64s."))
+	DEFINE_METASOUND_MATHOP(Add, Audio, FAudioBuffer, LOCTEXT("Metasound_MathAddBufferNodeDescription", "Adds buffers together by sample."))
+	DEFINE_METASOUND_MATHOP(Add, Time, FTime, LOCTEXT("Metasound_MathAddTimeNodeDescription", "Adds time values."))
+
+	DEFINE_METASOUND_OPERAND_TYPED_MATHOP(Add, Audio, FAudioBuffer, Float, float, LOCTEXT("Metasound_MathAddAudioFloatNodeDescription", "Add floats to buffer sample-by-sample."))
+
+	DEFINE_METASOUND_MATHOP(Subtract, Float, float, LOCTEXT("Metasound_MathSubractFloatNodeDescription", "Subtracts floats."))
+	DEFINE_METASOUND_MATHOP(Subtract, Int32, int32, LOCTEXT("Metasound_MathSubractInt32NodeDescription", "Subtracts int32s."))
+	DEFINE_METASOUND_MATHOP(Subtract, Int64, int64, LOCTEXT("Metasound_MathSubractInt64NodeDescription", "Subtracts int64s."))
+	DEFINE_METASOUND_MATHOP(Subtract, Audio, FAudioBuffer, LOCTEXT("Metasound_MathSubtractBufferNodeDescription", "Subtracts buffers sample-by-sample."))
+	DEFINE_METASOUND_MATHOP(Subtract, Time, FTime, LOCTEXT("Metasound_MathSubractTimeNodeDescription", "Subtracts time values."))
+
+	DEFINE_METASOUND_MATHOP(Multiply, Float, float, LOCTEXT("Metasound_MathMultiplyFloatNodeDescription", "Multiplies floats."))
+	DEFINE_METASOUND_MATHOP(Multiply, Int32, int32, LOCTEXT("Metasound_MathMultiplyInt32NodeDescription", "Multiplies int32s."))
+	DEFINE_METASOUND_MATHOP(Multiply, Int64, int64, LOCTEXT("Metasound_MathMultiplyInt64NodeDescription", "Multiplies int64s."))
+	DEFINE_METASOUND_MATHOP(Multiply, Audio, FAudioBuffer, LOCTEXT("Metasound_MathMultiplyBufferNodeDescription", "Multiplies buffers together sample-by-sample."))
+
+	DEFINE_METASOUND_OPERAND_TYPED_MATHOP(Multiply, Audio, FAudioBuffer, Float, float, LOCTEXT("Metasound_MathMultiplyAudioByFloatDescription", "Multiplies buffer by float scalars."))
+	DEFINE_METASOUND_OPERAND_TYPED_MATHOP(Multiply, Time, FTime, Float, float, LOCTEXT("Metasound_MathMultiplyTimeNodeDescription", "Scales time by floats."))
+
+	DEFINE_METASOUND_MATHOP(Divide, Float, float, LOCTEXT("Metasound_MathDivideFloatNodeDescription", "Divide float by another float."))
+	DEFINE_METASOUND_MATHOP(Divide, Int32, int32, LOCTEXT("Metasound_MathDivideInt32NodeDescription", "Divide int32 by another int32."))
+	DEFINE_METASOUND_MATHOP(Divide, Int64, int64, LOCTEXT("Metasound_MathDivideInt64NodeDescription", "Divide int64 by another int64."))
+
+	DEFINE_METASOUND_OPERAND_TYPED_MATHOP(Divide, Time, FTime, Float, float, LOCTEXT("Metasound_MathDivideTimeNodeDescription", "Divides time by floats."))
+
+	DEFINE_METASOUND_MATHOP(Modulo, Int32, int32, LOCTEXT("Metasound_MathModulusInt32NodeDescription", "Modulo int32 by another int32."))
+	DEFINE_METASOUND_MATHOP(Modulo, Int64, int64, LOCTEXT("Metasound_MathModulusInt64NodeDescription", "Modulo int64 by another int64."))
+
+	DEFINE_METASOUND_MATHOP(Power, Float, float, LOCTEXT("Metasound_MathPowerFloatNodeDescription", "Raise float to the power of another float."))
+
+	DEFINE_METASOUND_MATHOP(Logarithm, Float, float, LOCTEXT("Metasound_MathLogarithmFloatNodeDescription", "Calculate float-base logarithm of another float."))
+
 	METASOUND_REGISTER_NODE(FAddFloatNode)
 	METASOUND_REGISTER_NODE(FAddInt32Node)
 	METASOUND_REGISTER_NODE(FAddInt64Node)
@@ -32,6 +1489,10 @@ namespace Metasound
 
 	METASOUND_REGISTER_NODE(FModuloInt32Node)
 	METASOUND_REGISTER_NODE(FModuloInt64Node)
+
+	METASOUND_REGISTER_NODE(FPowerFloatNode)
+
+	METASOUND_REGISTER_NODE(FLogarithmFloatNode)
 }
 
 #undef LOCTEXT_NAMESPACE // MetasoundMathOpNode
diff --git a/Engine/Plugins/Runtime/Metasound/Source/MetasoundStandardNodes/Public/MetasoundMathNodes.h b/Engine/Plugins/Runtime/Metasound/Source/MetasoundStandardNodes/Public/MetasoundMathNodes.h
deleted file mode 100644
