Making our modules optimized

diff --git a/Engine/Plugins/Runtime/Metasound/Source/MetasoundEditor/MetasoundEditor.Build.cs b/Engine/Plugins/Runtime/Metasound/Source/MetasoundEditor/MetasoundEditor.Build.cs
--- a/Engine/Plugins/Runtime/Metasound/Source/MetasoundEditor/MetasoundEditor.Build.cs
+++ b/Engine/Plugins/Runtime/Metasound/Source/MetasoundEditor/MetasoundEditor.Build.cs
@@ -9,8 +9,6 @@ namespace UnrealBuildTool.Rules
 	{
 		public MetasoundEditor(ReadOnlyTargetRules Target) : base(Target)
 		{
-			OptimizeCode = CodeOptimization.Never;
-
 			PrivateDependencyModuleNames.AddRange
 			(
 				new string[]
diff --git a/Engine/Plugins/Runtime/Metasound/Source/MetasoundEngine/MetasoundEngine.Build.cs b/Engine/Plugins/Runtime/Metasound/Source/MetasoundEngine/MetasoundEngine.Build.cs
--- a/Engine/Plugins/Runtime/Metasound/Source/MetasoundEngine/MetasoundEngine.Build.cs
+++ b/Engine/Plugins/Runtime/Metasound/Source/MetasoundEngine/MetasoundEngine.Build.cs
@@ -9,9 +9,6 @@ namespace UnrealBuildTool.Rules
 	{
 		public MetasoundEngine(ReadOnlyTargetRules Target) : base(Target)
 		{
-
-			OptimizeCode = CodeOptimization.Never;
-
 			PrivateDependencyModuleNames.AddRange
 			(
 				new string[]
diff --git a/Engine/Plugins/Runtime/Metasound/Source/MetasoundFrontend/MetasoundFrontend.Build.cs b/Engine/Plugins/Runtime/Metasound/Source/MetasoundFrontend/MetasoundFrontend.Build.cs
--- a/Engine/Plugins/Runtime/Metasound/Source/MetasoundFrontend/MetasoundFrontend.Build.cs
+++ b/Engine/Plugins/Runtime/Metasound/Source/MetasoundFrontend/MetasoundFrontend.Build.cs
@@ -9,8 +9,6 @@ namespace UnrealBuildTool.Rules
 	{
 		public MetasoundFrontend(ReadOnlyTargetRules Target) : base(Target)
 		{
-			OptimizeCode = CodeOptimization.Never;
-
 			PublicDependencyModuleNames.AddRange
 			(
 				new string[]
diff --git a/Engine/Plugins/Runtime/Metasound/Source/MetasoundGenerator/MetasoundGenerator.Build.cs b/Engine/Plugins/Runtime/Metasound/Source/MetasoundGenerator/MetasoundGenerator.Build.cs
--- a/Engine/Plugins/Runtime/Metasound/Source/MetasoundGenerator/MetasoundGenerator.Build.cs
+++ b/Engine/Plugins/Runtime/Metasound/Source/MetasoundGenerator/MetasoundGenerator.Build.cs
@@ -9,8 +9,6 @@ namespace UnrealBuildTool.Rules
 	{
         public MetasoundGenerator(ReadOnlyTargetRules Target) : base(Target)
 		{
-			//OptimizeCode = CodeOptimization.Never;
-
 			PrivateDependencyModuleNames.AddRange(
 				new string[] {
 					"SignalProcessing",
diff --git a/Engine/Plugins/Runtime/Metasound/Source/MetasoundGraphCore/MetasoundGraphCore.Build.cs b/Engine/Plugins/Runtime/Metasound/Source/MetasoundGraphCore/MetasoundGraphCore.Build.cs
--- a/Engine/Plugins/Runtime/Metasound/Source/MetasoundGraphCore/MetasoundGraphCore.Build.cs
+++ b/Engine/Plugins/Runtime/Metasound/Source/MetasoundGraphCore/MetasoundGraphCore.Build.cs
@@ -9,8 +9,6 @@ namespace UnrealBuildTool.Rules
 	{
         public MetasoundGraphCore(ReadOnlyTargetRules Target) : base(Target)
 		{
-            //OptimizeCode = CodeOptimization.Never;
-
             PublicDependencyModuleNames.AddRange(
 				new string[] {
                     "Core",
