- Fix for "Promote to Graph Input" failing to build array literals properly
- Implement 'C' for creating comment around currently selected nodes in Metasound Editor

diff --git a/Engine/Plugins/Runtime/Metasound/Source/MetasoundEditor/Private/MetasoundEditor.cpp b/Engine/Plugins/Runtime/Metasound/Source/MetasoundEditor/Private/MetasoundEditor.cpp
--- a/Engine/Plugins/Runtime/Metasound/Source/MetasoundEditor/Private/MetasoundEditor.cpp
+++ b/Engine/Plugins/Runtime/Metasound/Source/MetasoundEditor/Private/MetasoundEditor.cpp
@@ -911,6 +911,14 @@ namespace Metasound
 
 		void FEditor::OnCreateComment()
 		{
+			if (MetasoundGraphEditor.IsValid())
+			{
+				if (UEdGraph* Graph = MetasoundGraphEditor->GetCurrentGraph())
+				{
+					FMetasoundGraphSchemaAction_NewComment CommentAction;
+					CommentAction.PerformAction(Graph, nullptr, MetasoundGraphEditor->GetPasteLocation());
+				}
+			}
 		}
 
 		void FEditor::CreateGraphEditorWidget()
@@ -936,11 +944,7 @@ namespace Metasound
 					FExecuteAction::CreateSP(this, &FEditor::DeleteInput),
 					FCanExecuteAction::CreateSP(this, &FEditor::CanDeleteInput));
 
-				// Graph Editor Commands
-				GraphEditorCommands->MapAction(FGraphEditorCommands::Get().CreateComment,
-					FExecuteAction::CreateSP(this, &FEditor::OnCreateComment));
-
-				// Editing commands
+				// Editing Commands
 				GraphEditorCommands->MapAction(FGenericCommands::Get().SelectAll,
 					FExecuteAction::CreateLambda([this]() { MetasoundGraphEditor->SelectAllNodes(); }));
 
@@ -992,6 +996,10 @@ namespace Metasound
 
 				GraphEditorCommands->MapAction(FGraphEditorCommands::Get().DistributeNodesVertically,
 					FExecuteAction::CreateLambda([this]() { MetasoundGraphEditor->OnDistributeNodesV(); }));
+
+				// Node Commands
+				GraphEditorCommands->MapAction(FGraphEditorCommands::Get().CreateComment,
+					FExecuteAction::CreateSP(this, &FEditor::OnCreateComment));
 			}
 
 			FGraphAppearanceInfo AppearanceInfo;
diff --git a/Engine/Plugins/Runtime/Metasound/Source/MetasoundEditor/Private/MetasoundEditorGraphBuilder.cpp b/Engine/Plugins/Runtime/Metasound/Source/MetasoundEditor/Private/MetasoundEditorGraphBuilder.cpp
--- a/Engine/Plugins/Runtime/Metasound/Source/MetasoundEditor/Private/MetasoundEditorGraphBuilder.cpp
+++ b/Engine/Plugins/Runtime/Metasound/Source/MetasoundEditor/Private/MetasoundEditorGraphBuilder.cpp
@@ -514,14 +514,44 @@ namespace Metasound
 				break;
 
 				case ELiteralType::BooleanArray:
+				{
+					OutDefaultLiteral.Set(TArray<bool>());
+				}
+				break;
+
 				case ELiteralType::FloatArray:
+				{
+					OutDefaultLiteral.Set(TArray<float>());
+				}
+				break;
+
 				case ELiteralType::IntegerArray:
+				{
+					OutDefaultLiteral.Set(TArray<int32>());
+				}
+				break;
+
 				case ELiteralType::NoneArray:
+				{
+					OutDefaultLiteral.Set(FMetasoundFrontendLiteral::FDefaultArray());
+				}
+				break;
+
 				case ELiteralType::StringArray:
+				{
+					OutDefaultLiteral.Set(TArray<FString>());
+				}
+				break;
+
 				case ELiteralType::UObjectProxyArray:
+				{
+					OutDefaultLiteral.Set(TArray<UObject*>());
+				}
+				break;
+
 				case ELiteralType::None:
 				{
-					return false;
+					OutDefaultLiteral.Set(FMetasoundFrontendLiteral::FDefault());
 				}
 				break;
 
diff --git a/Engine/Plugins/Runtime/Metasound/Source/MetasoundEditor/Private/MetasoundNodeDetailCustomization.cpp b/Engine/Plugins/Runtime/Metasound/Source/MetasoundEditor/Private/MetasoundNodeDetailCustomization.cpp
--- a/Engine/Plugins/Runtime/Metasound/Source/MetasoundEditor/Private/MetasoundNodeDetailCustomization.cpp
+++ b/Engine/Plugins/Runtime/Metasound/Source/MetasoundEditor/Private/MetasoundNodeDetailCustomization.cpp
@@ -1063,9 +1063,7 @@ namespace Metasound
 
 			GraphHandle->IterateConstNodes([this, NodeID, InNewName](FConstNodeHandle NodeToCompare)
 			{
-				// Disregard display name collisions with hidden nodes
-				const bool bIsVisible = NodeToCompare->GetNodeStyle().Display.Visibility == EMetasoundFrontendNodeStyleDisplayVisibility::Visible;
-				if (NodeID != NodeToCompare->GetID() && bIsVisible)
+				if (NodeID != NodeToCompare->GetID())
 				{
 					if (InNewName.CompareToCaseIgnored(NodeToCompare->GetDisplayName()) == 0)
 					{
