Various fixups and cleanups on metasound nodes

UE-109509

diff --git a/Engine/Plugins/Runtime/Metasound/Source/MetasoundStandardNodes/Private/MetasoundBasicFilters.cpp b/Engine/Plugins/Runtime/Metasound/Source/MetasoundStandardNodes/Private/MetasoundBasicFilters.cpp
--- a/Engine/Plugins/Runtime/Metasound/Source/MetasoundStandardNodes/Private/MetasoundBasicFilters.cpp
+++ b/Engine/Plugins/Runtime/Metasound/Source/MetasoundStandardNodes/Private/MetasoundBasicFilters.cpp
@@ -33,20 +33,20 @@ namespace Metasound
 	namespace BasicFilterParameterNames
 	{
 		// inputs
-		METASOUND_PARAM(ParamAudioInput, "In", "Audio to be processed by the filter");	
-		METASOUND_PARAM(ParamCutoffFrequency, "Cutoff Frequency", "Controls cutoff frequency");
-		METASOUND_PARAM(ParamResonance, "Resonance", "Controls filter resonance");
-		METASOUND_PARAM(ParamBandwidth, "Bandwidth", "Controls bandwidth when applicable to the current filter type");
-		METASOUND_PARAM(ParamGainDb, "Gain", "Gain applied to the band when in Parametric mode (in decibels)");
-		METASOUND_PARAM(ParamFilterType, "Type", "Filter Type");
-		METASOUND_PARAM(ParamBandStopControl, "Band Stop Control", "Band Stop Control (applied to band stop output)");
+		METASOUND_PARAM(ParamAudioInput, "In", "Audio to be processed by the filter.");	
+		METASOUND_PARAM(ParamCutoffFrequency, "Cutoff Frequency", "Controls cutoff frequency.");
+		METASOUND_PARAM(ParamResonance, "Resonance", "Controls filter resonance.");
+		METASOUND_PARAM(ParamBandwidth, "Bandwidth", "Controls bandwidth when applicable to the current filter type.");
+		METASOUND_PARAM(ParamGainDb, "Gain", "Gain applied to the band when in Parametric mode (in decibels).");
+		METASOUND_PARAM(ParamFilterType, "Type", "Filter type.");
+		METASOUND_PARAM(ParamBandStopControl, "Band Stop Control", "Band stop Control (applied to band stop output).");
 
 		// outputs
-		METASOUND_PARAM(ParamAudioOutput, "Out", "Audio processed by the filtter");
-		METASOUND_PARAM(ParamHighPassOutput, "High Pass Filter", "High Pass Filter Output");
-		METASOUND_PARAM(ParamLowPassOutput, "Low Pass Filter", "Low Pass Filter Output");
-		METASOUND_PARAM(ParamBandPassOutput, "Band Pass", "Band Pass Output");
-		METASOUND_PARAM(ParamBandStopOutput, "Band Stop", "Band Stop Output");
+		METASOUND_PARAM(ParamAudioOutput, "Out", "Audio processed by the filter.");
+		METASOUND_PARAM(ParamHighPassOutput, "High Pass Filter", "High pass filter output.");
+		METASOUND_PARAM(ParamLowPassOutput, "Low Pass Filter", "Low pass filter output.");
+		METASOUND_PARAM(ParamBandPassOutput, "Band Pass", "Band pass filter output.");
+		METASOUND_PARAM(ParamBandStopOutput, "Band Stop", "Band stop filter output.");
 
 	} // namespace BasicFilterParameterNames;
 
@@ -59,16 +59,16 @@ namespace Metasound
 	METASOUNDSTANDARDNODES_API, FEnumEBiquadFilterType, FEnumBiQuadFilterTypeInfo, FEnumBiQuadFilterReadRef, FEnumBiQuadFilterWriteRef);
 
 	DEFINE_METASOUND_ENUM_BEGIN(Audio::EBiquadFilter::Type, FEnumEBiquadFilterType, "BiquadFilterType")
-		DEFINE_METASOUND_ENUM_ENTRY(Audio::EBiquadFilter::Lowpass, LOCTEXT("LpDescription", "Low Pass"), LOCTEXT("LpDescriptionTT", "Low Pass Biquad Filter")),
-		DEFINE_METASOUND_ENUM_ENTRY(Audio::EBiquadFilter::Highpass, LOCTEXT("HpDescription", "High Pass"), LOCTEXT("HpDescriptionTT", "High Pass Biquad Filter")),
-		DEFINE_METASOUND_ENUM_ENTRY(Audio::EBiquadFilter::Bandpass, LOCTEXT("BpDescription", "Band Pass"), LOCTEXT("BpDescriptionTT", "Band Pass Biquad Filter")),
-		DEFINE_METASOUND_ENUM_ENTRY(Audio::EBiquadFilter::Notch, LOCTEXT("NotchDescription", "Notch "), LOCTEXT("NotchDescriptionTT", "Notch Biquad Filter")),
-		DEFINE_METASOUND_ENUM_ENTRY(Audio::EBiquadFilter::ParametricEQ, LOCTEXT("ParaEqDescription", "Parametric EQ"), LOCTEXT("ParaEqDescriptionTT", "Parametric EQ Biquad Filter")),
-		DEFINE_METASOUND_ENUM_ENTRY(Audio::EBiquadFilter::LowShelf, LOCTEXT("LowShelfDescription", "Low Shelf"), LOCTEXT("LowShelfDescriptionTT", "Low Shelf Biquad Filter")),
-		DEFINE_METASOUND_ENUM_ENTRY(Audio::EBiquadFilter::HighShelf, LOCTEXT("HighShelfDescription", "High Shelf"), LOCTEXT("HighShelfDescriptionTT", "High Shelf Biquad Filter")),
-		DEFINE_METASOUND_ENUM_ENTRY(Audio::EBiquadFilter::AllPass, LOCTEXT("AllPassDescription", "All Pass"), LOCTEXT("AllPassDescriptionTT", "All Pass Biquad Filter")),
-		DEFINE_METASOUND_ENUM_ENTRY(Audio::EBiquadFilter::ButterworthLowPass, LOCTEXT("LowPassButterDescription", "Butterworth Low Pass"), LOCTEXT("LowPassButterDescriptionTT", "Butterworth Low Pass Biquad Filter")),
-		DEFINE_METASOUND_ENUM_ENTRY(Audio::EBiquadFilter::ButterworthHighPass, LOCTEXT("HighPassButterDescription", "Butterworth High Pass"), LOCTEXT("HighPassButterDescriptionTT", "Butterworth High Pass Biquad Filter"))
+		DEFINE_METASOUND_ENUM_ENTRY(Audio::EBiquadFilter::Lowpass, LOCTEXT("LpDescription", "Low Pass"), LOCTEXT("LpDescriptionTT", "Low pass Biquad filter.")),
+		DEFINE_METASOUND_ENUM_ENTRY(Audio::EBiquadFilter::Highpass, LOCTEXT("HpDescription", "High Pass"), LOCTEXT("HpDescriptionTT", "High pass Biquad filter.")),
+		DEFINE_METASOUND_ENUM_ENTRY(Audio::EBiquadFilter::Bandpass, LOCTEXT("BpDescription", "Band Pass"), LOCTEXT("BpDescriptionTT", "Band pass Biquad filter.")),
+		DEFINE_METASOUND_ENUM_ENTRY(Audio::EBiquadFilter::Notch, LOCTEXT("NotchDescription", "Notch "), LOCTEXT("NotchDescriptionTT", "Notch biquad filter.")),
+		DEFINE_METASOUND_ENUM_ENTRY(Audio::EBiquadFilter::ParametricEQ, LOCTEXT("ParaEqDescription", "Parametric EQ"), LOCTEXT("ParaEqDescriptionTT", "Parametric EQ biquad filter.")),
+		DEFINE_METASOUND_ENUM_ENTRY(Audio::EBiquadFilter::LowShelf, LOCTEXT("LowShelfDescription", "Low Shelf"), LOCTEXT("LowShelfDescriptionTT", "Low shelf biquad filter.")),
+		DEFINE_METASOUND_ENUM_ENTRY(Audio::EBiquadFilter::HighShelf, LOCTEXT("HighShelfDescription", "High Shelf"), LOCTEXT("HighShelfDescriptionTT", "High shelf biquad filter.")),
+		DEFINE_METASOUND_ENUM_ENTRY(Audio::EBiquadFilter::AllPass, LOCTEXT("AllPassDescription", "All Pass"), LOCTEXT("AllPassDescriptionTT", "All pass biquad Filter.")),
+		DEFINE_METASOUND_ENUM_ENTRY(Audio::EBiquadFilter::ButterworthLowPass, LOCTEXT("LowPassButterDescription", "Butterworth Low Pass"), LOCTEXT("LowPassButterDescriptionTT", "Butterworth Low Pass Biquad Filter.")),
+		DEFINE_METASOUND_ENUM_ENTRY(Audio::EBiquadFilter::ButterworthHighPass, LOCTEXT("HighPassButterDescription", "Butterworth High Pass"), LOCTEXT("HighPassButterDescriptionTT", "Butterworth High Pass Biquad Filter."))
 		DEFINE_METASOUND_ENUM_END()
 #pragma endregion
 
diff --git a/Engine/Plugins/Runtime/Metasound/Source/MetasoundStandardNodes/Private/MetasoundDelayNode.cpp b/Engine/Plugins/Runtime/Metasound/Source/MetasoundStandardNodes/Private/MetasoundDelayNode.cpp
--- a/Engine/Plugins/Runtime/Metasound/Source/MetasoundStandardNodes/Private/MetasoundDelayNode.cpp
+++ b/Engine/Plugins/Runtime/Metasound/Source/MetasoundStandardNodes/Private/MetasoundDelayNode.cpp
@@ -1,7 +1,5 @@
 // Copyright Epic Games, Inc. All Rights Reserved.
 
-#include "MetasoundDelayNode.h"
-
 #include "Internationalization/Text.h"
 #include "MetasoundExecutableOperator.h"
 #include "MetasoundNodeRegistrationMacro.h"
@@ -11,8 +9,10 @@
 #include "MetasoundTime.h"
 #include "MetasoundAudioBuffer.h"
 #include "DSP/Delay.h"
+#include "MetasoundStandardNodesCategories.h"
+#include "MetasoundFacade.h"
 
-#define LOCTEXT_NAMESPACE "MetasoundStandardNodes"
+#define LOCTEXT_NAMESPACE "MetasoundStandardNodes_Delay"
 
 namespace Metasound
 {
@@ -173,7 +173,7 @@ namespace Metasound
 				TInputDataVertexModel<FAudioBuffer>(Delay::InParamNameAudioInput, LOCTEXT("AudioInputTooltip", "Audio input.")),
 				TInputDataVertexModel<FTime>(Delay::InParamNameDelayTime, LOCTEXT("DelayTimeTooltip", "The amount of time to delay the audio."), 1.0f),
 				TInputDataVertexModel<float>(Delay::InParamNameDryLevel, LOCTEXT("DryLevelTooltip", "The dry level of the delay."), 0.0f),
-				TInputDataVertexModel<float>(Delay::InParamNameWetLevel, LOCTEXT("FeedbackTooltip", "The wet level of the delay."), 1.0f),
+				TInputDataVertexModel<float>(Delay::InParamNameWetLevel, LOCTEXT("WetlevelTooltip", "The wet level of the delay."), 1.0f),
 				TInputDataVertexModel<float>(Delay::InParamNameFeedbackAmount, LOCTEXT("FeedbackTooltip", "Feedback amount."), 0.0f)
 			),
 			FOutputVertexInterface(
@@ -197,7 +197,7 @@ namespace Metasound
 			Info.Author = PluginAuthor;
 			Info.PromptIfMissing = PluginNodeMissingPrompt;
 			Info.DefaultInterface = GetVertexInterface();
-
+			Info.CategoryHierarchy.Emplace(StandardNodes::Delays);
 			return Info;
 		};
 
@@ -208,7 +208,6 @@ namespace Metasound
 
 	TUniquePtr<IOperator> FDelayOperator::CreateOperator(const FCreateOperatorParams& InParams, FBuildErrorArray& OutErrors)
 	{
-		const FDelayNode& DelayNode = static_cast<const FDelayNode&>(InParams.Node);
 		const FDataReferenceCollection& InputCollection = InParams.InputDataReferences;
 		const FInputVertexInterface& InputInterface = GetVertexInterface().GetInputInterface();
 
@@ -221,10 +220,18 @@ namespace Metasound
 		return MakeUnique<FDelayOperator>(InParams.OperatorSettings, AudioIn, DelayTime, DryLevel, WetLevel, Feedback);
 	}
 
-	FDelayNode::FDelayNode(const FNodeInitData& InitData)
-		: FNodeFacade(InitData.InstanceName, InitData.InstanceID, TFacadeOperatorClass<FDelayOperator>())
+	class FDelayNode : public FNodeFacade
 	{
-	}
+	public:
+		/**
+		 * Constructor used by the Metasound Frontend.
+		 */
+		FDelayNode(const FNodeInitData& InitData)
+			: FNodeFacade(InitData.InstanceName, InitData.InstanceID, TFacadeOperatorClass<FDelayOperator>())
+		{
+		}
+	};
+
 
 	METASOUND_REGISTER_NODE(FDelayNode)
 }
diff --git a/Engine/Plugins/Runtime/Metasound/Source/MetasoundStandardNodes/Private/MetasoundITDPannerNode.cpp b/Engine/Plugins/Runtime/Metasound/Source/MetasoundStandardNodes/Private/MetasoundITDPannerNode.cpp
--- a/Engine/Plugins/Runtime/Metasound/Source/MetasoundStandardNodes/Private/MetasoundITDPannerNode.cpp
+++ b/Engine/Plugins/Runtime/Metasound/Source/MetasoundStandardNodes/Private/MetasoundITDPannerNode.cpp
@@ -1,7 +1,5 @@
 // Copyright Epic Games, Inc. All Rights Reserved.
 
-#include "MetasoundITDPannerNode.h"
-
 #include "Internationalization/Text.h"
 #include "MetasoundExecutableOperator.h"
 #include "MetasoundNodeRegistrationMacro.h"
@@ -14,8 +12,10 @@
 #include "DSP/BufferVectorOperations.h"
 #include "DSP/Delay.h"
 #include "DSP/Dsp.h"
+#include "MetasoundFacade.h"
+#include "MetasoundStandardNodesCategories.h"
 
-#define LOCTEXT_NAMESPACE "MetasoundStandardNodes"
+#define LOCTEXT_NAMESPACE "MetasoundStandardNodes_ITDPannerNode"
 
 namespace Metasound
 {
@@ -318,7 +318,7 @@ namespace Metasound
 			Info.Author = PluginAuthor;
 			Info.PromptIfMissing = PluginNodeMissingPrompt;
 			Info.DefaultInterface = GetVertexInterface();
-
+			Info.CategoryHierarchy.Emplace(StandardNodes::Spatalization);
 			return Info;
 		};
 
@@ -329,7 +329,6 @@ namespace Metasound
 
 	TUniquePtr<IOperator> FITDPannerOperator::CreateOperator(const FCreateOperatorParams& InParams, FBuildErrorArray& OutErrors)
 	{
-		const FITDPannerNode& ITDPannerNode = static_cast<const FITDPannerNode&>(InParams.Node);
 		const FDataReferenceCollection& InputCollection = InParams.InputDataReferences;
 		const FInputVertexInterface& InputInterface = GetVertexInterface().GetInputInterface();
 
@@ -343,10 +342,17 @@ namespace Metasound
 		return MakeUnique<FITDPannerOperator>(InParams.OperatorSettings, AudioIn, PanningAngle, DistanceFactor, HeadWidth);
 	}
 
-	FITDPannerNode::FITDPannerNode(const FNodeInitData& InitData)
-		: FNodeFacade(InitData.InstanceName, InitData.InstanceID, TFacadeOperatorClass<FITDPannerOperator>())
+	class FITDPannerNode : public FNodeFacade
 	{
-	}
+	public:
+		/**
+		 * Constructor used by the Metasound Frontend.
+		 */
+		FITDPannerNode(const FNodeInitData& InitData)
+			: FNodeFacade(InitData.InstanceName, InitData.InstanceID, TFacadeOperatorClass<FITDPannerOperator>())
+		{
+		}
+	};
 
 	METASOUND_REGISTER_NODE(FITDPannerNode)
 }
diff --git a/Engine/Plugins/Runtime/Metasound/Source/MetasoundStandardNodes/Private/MetasoundOscNode.cpp b/Engine/Plugins/Runtime/Metasound/Source/MetasoundStandardNodes/Private/MetasoundOscNode.cpp
deleted file mode 100644
diff --git a/Engine/Plugins/Runtime/Metasound/Source/MetasoundStandardNodes/Private/MetasoundStandardNodesCategories.cpp b/Engine/Plugins/Runtime/Metasound/Source/MetasoundStandardNodes/Private/MetasoundStandardNodesCategories.cpp
--- a/Engine/Plugins/Runtime/Metasound/Source/MetasoundStandardNodes/Private/MetasoundStandardNodesCategories.cpp
+++ b/Engine/Plugins/Runtime/Metasound/Source/MetasoundStandardNodes/Private/MetasoundStandardNodesCategories.cpp
@@ -10,12 +10,15 @@ namespace Metasound
 	{
 		const FText Generators = { LOCTEXT("Metasound_GeneratorsCategory", "Generators") };
 		const FText Filters = { LOCTEXT("Metasound_FiltersCategory", "Filters") };
-		const FText TriggerUtils = { LOCTEXT("Metasound_TriggerUtilsCategory", "Trigger Utilities") };
+		const FText TriggerUtils = { LOCTEXT("Metasound_TriggerUtilsCategory", "Triggers") };
 		const FText Music = { LOCTEXT("Metasound_MusicCategory", "Music") };
-		const FText RandomUtils = { LOCTEXT("Metasound_RandomUtilsCategory", "Random Utilities") };
+		const FText RandomUtils = { LOCTEXT("Metasound_RandomCategory", "Random") };
 		const FText Conversions  = { LOCTEXT("Metasound_ConversionsCategory", "Conversions") };
-		const FText DebugUtils = { LOCTEXT("Metasound_DebugUtilsCategory", "Debug Utilities") };
+		const FText DebugUtils = { LOCTEXT("Metasound_DebugUtilsCategory", "Debug") };
 		const FText Io = { LOCTEXT("Metasound_IoCategory", "External IO") };
+		const FText Delays = { LOCTEXT("Metasound_EffectsCategory", "Delays") };
+		const FText Spatalization = { LOCTEXT("Metasound_SpatializationCategory", "Spatialization") };
+		const FText Dynamics = { LOCTEXT("Metasound_DynamicsCategory", "Dynamics") };
 	}
 }
 
diff --git a/Engine/Plugins/Runtime/Metasound/Source/MetasoundStandardNodes/Private/MetasoundStereoDelayNode.cpp b/Engine/Plugins/Runtime/Metasound/Source/MetasoundStandardNodes/Private/MetasoundStereoDelayNode.cpp
--- a/Engine/Plugins/Runtime/Metasound/Source/MetasoundStandardNodes/Private/MetasoundStereoDelayNode.cpp
+++ b/Engine/Plugins/Runtime/Metasound/Source/MetasoundStandardNodes/Private/MetasoundStereoDelayNode.cpp
@@ -1,7 +1,5 @@
 // Copyright Epic Games, Inc. All Rights Reserved.
 
-#include "MetasoundStereoDelayNode.h"
-
 #include "Internationalization/Text.h"
 #include "MetasoundExecutableOperator.h"
 #include "MetasoundNodeRegistrationMacro.h"
@@ -12,8 +10,10 @@
 #include "MetasoundTime.h"
 #include "MetasoundAudioBuffer.h"
 #include "DSP/Delay.h"
+#include "MetasoundStandardNodesCategories.h"
+#include "MetasoundFacade.h"
 
-#define LOCTEXT_NAMESPACE "MetasoundStandardNodes"
+#define LOCTEXT_NAMESPACE "MetasoundStandardNodes_DelayNode"
 
 namespace Metasound
 {
@@ -314,7 +314,7 @@ namespace Metasound
 			Info.Author = PluginAuthor;
 			Info.PromptIfMissing = PluginNodeMissingPrompt;
 			Info.DefaultInterface = GetVertexInterface();
-
+			Info.CategoryHierarchy.Emplace(StandardNodes::Delays);
 			return Info;
 		};
 
@@ -325,7 +325,6 @@ namespace Metasound
 
 	TUniquePtr<IOperator> FStereoDelayOperator::CreateOperator(const FCreateOperatorParams& InParams, FBuildErrorArray& OutErrors)
 	{
-		const FStereoDelayNode& StereoDelayNode = static_cast<const FStereoDelayNode&>(InParams.Node);
 		const FDataReferenceCollection& InputCollection = InParams.InputDataReferences;
 		const FInputVertexInterface& InputInterface = GetVertexInterface().GetInputInterface();
 
@@ -341,10 +340,18 @@ namespace Metasound
 		return MakeUnique<FStereoDelayOperator>(InParams.OperatorSettings, LeftAudioIn, RightAudioIn, StereoDelayMode, DelayTime, DelayRatio, DryLevel, WetLevel, Feedback);
 	}
 
-	FStereoDelayNode::FStereoDelayNode(const FNodeInitData& InitData)
-		: FNodeFacade(InitData.InstanceName, InitData.InstanceID, TFacadeOperatorClass<FStereoDelayOperator>())
+
+	class FStereoDelayNode : public FNodeFacade
 	{
-	}
+	public:
+		/**
+		 * Constructor used by the Metasound Frontend.
+		 */
+		FStereoDelayNode(const FNodeInitData& InitData)
+			: FNodeFacade(InitData.InstanceName, InitData.InstanceID, TFacadeOperatorClass<FStereoDelayOperator>())
+		{
+		}
+	};
 
 	METASOUND_REGISTER_NODE(FStereoDelayNode)
 }
diff --git a/Engine/Plugins/Runtime/Metasound/Source/MetasoundStandardNodes/Private/MetasoundStereoPannerNode.cpp b/Engine/Plugins/Runtime/Metasound/Source/MetasoundStandardNodes/Private/MetasoundStereoPannerNode.cpp
--- a/Engine/Plugins/Runtime/Metasound/Source/MetasoundStandardNodes/Private/MetasoundStereoPannerNode.cpp
+++ b/Engine/Plugins/Runtime/Metasound/Source/MetasoundStandardNodes/Private/MetasoundStereoPannerNode.cpp
@@ -1,7 +1,5 @@
 // Copyright Epic Games, Inc. All Rights Reserved.
 
-#include "MetasoundStereoPannerNode.h"
-
 #include "Internationalization/Text.h"
 #include "MetasoundExecutableOperator.h"
 #include "MetasoundNodeRegistrationMacro.h"
@@ -12,8 +10,10 @@
 #include "MetasoundTime.h"
 #include "MetasoundAudioBuffer.h"
 #include "DSP/BufferVectorOperations.h"
+#include "MetasoundStandardNodesCategories.h"
+#include "MetasoundFacade.h"
 
-#define LOCTEXT_NAMESPACE "MetasoundStandardNodes"
+#define LOCTEXT_NAMESPACE "MetasoundStandardNodes_StereoPanner"
 
 namespace Metasound
 {
@@ -257,7 +257,7 @@ namespace Metasound
 			Info.Author = PluginAuthor;
 			Info.PromptIfMissing = PluginNodeMissingPrompt;
 			Info.DefaultInterface = GetVertexInterface();
-
+			Info.CategoryHierarchy.Emplace(StandardNodes::Spatalization);
 			return Info;
 		};
 
@@ -268,7 +268,6 @@ namespace Metasound
 
 	TUniquePtr<IOperator> FStereoPannerOperator::CreateOperator(const FCreateOperatorParams& InParams, FBuildErrorArray& OutErrors)
 	{
-		const FStereoPannerNode& StereoPannerNode = static_cast<const FStereoPannerNode&>(InParams.Node);
 		const FDataReferenceCollection& InputCollection = InParams.InputDataReferences;
 		const FInputVertexInterface& InputInterface = GetVertexInterface().GetInputInterface();
 
@@ -281,10 +280,18 @@ namespace Metasound
 		return MakeUnique<FStereoPannerOperator>(InParams.OperatorSettings, AudioIn, PanningAmount, PanningLaw);
 	}
 
-	FStereoPannerNode::FStereoPannerNode(const FNodeInitData& InitData)
-		: FNodeFacade(InitData.InstanceName, InitData.InstanceID, TFacadeOperatorClass<FStereoPannerOperator>())
+	class FStereoPannerNode : public FNodeFacade
 	{
-	}
+	public:
+		/**
+		 * Constructor used by the Metasound Frontend.
+		 */
+		FStereoPannerNode(const FNodeInitData& InitData)
+			: FNodeFacade(InitData.InstanceName, InitData.InstanceID, TFacadeOperatorClass<FStereoPannerOperator>())
+		{
+		}
+	};
+
 
 	METASOUND_REGISTER_NODE(FStereoPannerNode)
 }
diff --git a/Engine/Plugins/Runtime/Metasound/Source/MetasoundStandardNodes/Private/MetasoundWaveShaperNode.cpp b/Engine/Plugins/Runtime/Metasound/Source/MetasoundStandardNodes/Private/MetasoundWaveShaperNode.cpp
--- a/Engine/Plugins/Runtime/Metasound/Source/MetasoundStandardNodes/Private/MetasoundWaveShaperNode.cpp
+++ b/Engine/Plugins/Runtime/Metasound/Source/MetasoundStandardNodes/Private/MetasoundWaveShaperNode.cpp
@@ -1,7 +1,5 @@
 // Copyright Epic Games, Inc. All Rights Reserved.
 
-#include "MetasoundWaveShaperNode.h"
-
 #include "Internationalization/Text.h"
 #include "MetasoundExecutableOperator.h"
 #include "MetasoundNodeRegistrationMacro.h"
@@ -171,7 +169,6 @@ namespace Metasound
 
 	TUniquePtr<IOperator> FWaveShaperOperator::CreateOperator(const FCreateOperatorParams& InParams, FBuildErrorArray& OutErrors)
 	{
-		const FWaveShaperNode& WaveShaperNode = static_cast<const FWaveShaperNode&>(InParams.Node);
 		const FDataReferenceCollection& InputCollection = InParams.InputDataReferences;
 		const FInputVertexInterface& InputInterface = GetVertexInterface().GetInputInterface();
 
@@ -184,10 +181,17 @@ namespace Metasound
 		return MakeUnique<FWaveShaperOperator>(InParams.OperatorSettings, AudioIn, InAmount, InBias, OutGain, InType);
 	}
 
-	FWaveShaperNode::FWaveShaperNode(const FNodeInitData& InitData)
-		: FNodeFacade(InitData.InstanceName, InitData.InstanceID, TFacadeOperatorClass<FWaveShaperOperator>())
+	class FWaveShaperNode : public FNodeFacade
 	{
-	}
+	public:
+		/**
+		 * Constructor used by the Metasound Frontend.
+		 */
+		FWaveShaperNode(const FNodeInitData& InitData)
+			: FNodeFacade(InitData.InstanceName, InitData.InstanceID, TFacadeOperatorClass<FWaveShaperOperator>())
+		{
+		}
+	};
 
 	METASOUND_REGISTER_NODE(FWaveShaperNode)
 }
diff --git a/Engine/Plugins/Runtime/Metasound/Source/MetasoundStandardNodes/Public/MetasoundDelayNode.h b/Engine/Plugins/Runtime/Metasound/Source/MetasoundStandardNodes/Public/MetasoundDelayNode.h
deleted file mode 100644
diff --git a/Engine/Plugins/Runtime/Metasound/Source/MetasoundStandardNodes/Public/MetasoundITDPannerNode.h b/Engine/Plugins/Runtime/Metasound/Source/MetasoundStandardNodes/Public/MetasoundITDPannerNode.h
deleted file mode 100644
diff --git a/Engine/Plugins/Runtime/Metasound/Source/MetasoundStandardNodes/Public/MetasoundOscNode.h b/Engine/Plugins/Runtime/Metasound/Source/MetasoundStandardNodes/Public/MetasoundOscNode.h
deleted file mode 100644
diff --git a/Engine/Plugins/Runtime/Metasound/Source/MetasoundStandardNodes/Public/MetasoundStandardNodesCategories.h b/Engine/Plugins/Runtime/Metasound/Source/MetasoundStandardNodes/Public/MetasoundStandardNodesCategories.h
--- a/Engine/Plugins/Runtime/Metasound/Source/MetasoundStandardNodes/Public/MetasoundStandardNodesCategories.h
+++ b/Engine/Plugins/Runtime/Metasound/Source/MetasoundStandardNodes/Public/MetasoundStandardNodesCategories.h
@@ -16,5 +16,8 @@ namespace Metasound
 		extern const FText Conversions;
 		extern const FText DebugUtils;
 		extern const FText Io;
+		extern const FText Delays;
+		extern const FText Spatalization;
+		extern const FText Dynamics;
 	}
 }
diff --git a/Engine/Plugins/Runtime/Metasound/Source/MetasoundStandardNodes/Public/MetasoundStereoDelayNode.h b/Engine/Plugins/Runtime/Metasound/Source/MetasoundStandardNodes/Public/MetasoundStereoDelayNode.h
deleted file mode 100644
diff --git a/Engine/Plugins/Runtime/Metasound/Source/MetasoundStandardNodes/Public/MetasoundStereoPannerNode.h b/Engine/Plugins/Runtime/Metasound/Source/MetasoundStandardNodes/Public/MetasoundStereoPannerNode.h
deleted file mode 100644
diff --git a/Engine/Plugins/Runtime/Metasound/Source/MetasoundStandardNodes/Public/MetasoundWaveShaperNode.h b/Engine/Plugins/Runtime/Metasound/Source/MetasoundStandardNodes/Public/MetasoundWaveShaperNode.h
deleted file mode 100644
