Fixing array get node shared state playback issues

- Adding environment variable for if a sound is preview or not
- Clears the shared state on init if the sound is a preview sound

diff --git a/Engine/Plugins/Runtime/Metasound/Source/MetasoundEngine/Private/MetasoundSource.cpp b/Engine/Plugins/Runtime/Metasound/Source/MetasoundEngine/Private/MetasoundSource.cpp
--- a/Engine/Plugins/Runtime/Metasound/Source/MetasoundEngine/Private/MetasoundSource.cpp
+++ b/Engine/Plugins/Runtime/Metasound/Source/MetasoundEngine/Private/MetasoundSource.cpp
@@ -137,6 +137,7 @@ ISoundGeneratorPtr UMetaSoundSource::CreateSoundGenerator(const FSoundGeneratorI
 
 	// Set the unique object ID as an environment variable
 	Environment.SetValue<uint32>(GetSoundUniqueIdName(), GetUniqueID());
+	Environment.SetValue<bool>(GetIsPreviewSoundName(), InParams.bIsPreviewSound);
 
 	const FMetasoundFrontendDocument* OriginalDoc = GetDocument().Get();
 	if (nullptr == OriginalDoc)
@@ -434,6 +435,11 @@ const FString& UMetaSoundSource::GetSoundUniqueIdName()
 	return SoundUniqueIdVarName;
 }
 
+const FString& UMetaSoundSource::GetIsPreviewSoundName()
+{
+	static const FString SoundIsPreviewSoundName = TEXT("IsPreviewSound");
+	return SoundIsPreviewSoundName;
+}
 
 const FMetasoundFrontendArchetype& UMetaSoundSource::GetBaseArchetype()
 {
diff --git a/Engine/Plugins/Runtime/Metasound/Source/MetasoundEngine/Public/MetasoundSource.h b/Engine/Plugins/Runtime/Metasound/Source/MetasoundEngine/Public/MetasoundSource.h
--- a/Engine/Plugins/Runtime/Metasound/Source/MetasoundEngine/Public/MetasoundSource.h
+++ b/Engine/Plugins/Runtime/Metasound/Source/MetasoundEngine/Public/MetasoundSource.h
@@ -174,6 +174,7 @@ private:
 	static const FString& GetIsFinishedOutputName();
 	static const FString& GetAudioDeviceHandleVariableName();
 	static const FString& GetSoundUniqueIdName();
+	static const FString& GetIsPreviewSoundName();
 	static const FMetasoundFrontendArchetype& GetBaseArchetype();
 	static const FMetasoundFrontendArchetype& GetMonoSourceArchetype();
 	static const FMetasoundFrontendArchetype& GetStereoSourceArchetype();
diff --git a/Engine/Plugins/Runtime/Metasound/Source/MetasoundFrontend/Private/MetasoundArrayRandomNode.cpp b/Engine/Plugins/Runtime/Metasound/Source/MetasoundFrontend/Private/MetasoundArrayRandomNode.cpp
--- a/Engine/Plugins/Runtime/Metasound/Source/MetasoundFrontend/Private/MetasoundArrayRandomNode.cpp
+++ b/Engine/Plugins/Runtime/Metasound/Source/MetasoundFrontend/Private/MetasoundArrayRandomNode.cpp
@@ -184,14 +184,18 @@ namespace Metasound
 		return RGM;
 	}
 
-	void FSharedStateRandomGetManager::InitSharedState(uint32 InSharedStateId, int32 InSeed, int32 InNumElements, const TArray<float>& InWeights, int32 InNoRepeatOrder)
+	void FSharedStateRandomGetManager::InitSharedState(InitSharedStateArgs& InArgs)
 	{
 		FScopeLock Lock(&CritSect);
 
-		if (!RandomGets.Contains(InSharedStateId))
+		if (!InArgs.bIsPreviewSound && RandomGets.Contains(InArgs.SharedStateId))
 		{
-			RandomGets.Add(InSharedStateId, MakeUnique<FArrayRandomGet>(InSeed, InNumElements, InWeights, InNoRepeatOrder));
+			// TODO: check the actual arguments here against what has been registered, and if the new args are different warn they'll be ignored.
+			return;
 		}
+
+		// For preview sounds, we'll stomp the state
+		RandomGets.Add(InArgs.SharedStateId, MakeUnique<FArrayRandomGet>(InArgs.Seed, InArgs.NumElements, InArgs.Weights, InArgs.NoRepeatOrder));
 	}
 
 	int32 FSharedStateRandomGetManager::NextValue(uint32 InSharedStateId)
@@ -215,6 +219,13 @@ namespace Metasound
 		(*RG)->SetNoRepeatOrder(InNoRepeatOrder);
 	}
 
+	void FSharedStateRandomGetManager::SetRandomWeights(uint32 InSharedStateId, const TArray<float>& InRandomWeights)
+	{
+		FScopeLock Lock(&CritSect);
+		TUniquePtr<FArrayRandomGet>* RG = RandomGets.Find(InSharedStateId);
+		(*RG)->SetRandomWeights(InRandomWeights);
+	}
+
 	void FSharedStateRandomGetManager::ResetSeed(uint32 InSharedStateId)
 	{
 		FScopeLock Lock(&CritSect);
diff --git a/Engine/Plugins/Runtime/Metasound/Source/MetasoundFrontend/Public/MetasoundArrayRandomNode.h b/Engine/Plugins/Runtime/Metasound/Source/MetasoundFrontend/Public/MetasoundArrayRandomNode.h
--- a/Engine/Plugins/Runtime/Metasound/Source/MetasoundFrontend/Public/MetasoundArrayRandomNode.h
+++ b/Engine/Plugins/Runtime/Metasound/Source/MetasoundFrontend/Public/MetasoundArrayRandomNode.h
@@ -68,15 +68,26 @@ namespace Metasound
 		FRandomStream RandomStream;
 	};
 
+	struct InitSharedStateArgs
+	{
+		uint32 SharedStateId = 0;
+		int32 Seed = INDEX_NONE;
+		int32 NumElements = 0;
+		int32 NoRepeatOrder = 0;
+		bool bIsPreviewSound = false;
+		TArray<float> Weights;
+	};
+
 	class METASOUNDFRONTEND_API FSharedStateRandomGetManager
 	{
 	public:
 		static FSharedStateRandomGetManager& Get();
 
-		void InitSharedState(uint32 InSharedStateId, int32 InSeed, int32 InNumElements, const TArray<float>& InWeights, int32 InNoRepeatOrder);
+		void InitSharedState(InitSharedStateArgs& InArgs);
 		int32 NextValue(uint32 InSharedStateId);
 		void SetSeed(uint32 InSharedStateId, int32 InSeed);
 		void SetNoRepeatOrder(uint32 InSharedStateId, int32 InNoRepeatOrder);
+		void SetRandomWeights(uint32 InSharedStateId, const TArray<float>& InRandomWeights);
 		void ResetSeed(uint32 InSharedStateId);
 
 	private:
@@ -94,7 +105,7 @@ namespace Metasound
 	public:
 		using FArrayDataReadReference = TDataReadReference<ArrayType>;
 		using FArrayWeightReadReference = TDataReadReference<TArray<float>>;
-		using WeightArrayType = TArray<float>;
+		using WeightsArrayType = TArray<float>;
 		using ElementType = typename MetasoundArrayNodesPrivate::TArrayElementType<ArrayType>::Type;
 		using FElementTypeWriteReference = TDataWriteReference<ElementType>;
 
@@ -107,7 +118,7 @@ namespace Metasound
 					TInputDataVertexModel<FTrigger>(InputTriggerNextName, InputTriggerNextTooltip),
 					TInputDataVertexModel<FTrigger>(InputTriggerResetName, InputTriggerResetTooltip),
 					TInputDataVertexModel<ArrayType>(InputRandomArrayName, InputRandomArrayTooltip),
-					TInputDataVertexModel<WeightArrayType>(InputWeightsName, InputWeightsTooltip),
+					TInputDataVertexModel<WeightsArrayType>(InputWeightsName, InputWeightsTooltip),
 					TInputDataVertexModel<int32>(InputSeedName, InputSeedTooltip, -1),
 					TInputDataVertexModel<int32>(InputNoRepeatOrderName, InputNoRepeatOrderTooltip, 1),
 					TInputDataVertexModel<bool>(InputEnableSharedStateName, InputEnableSharedStateTooltip, false)
@@ -150,7 +161,7 @@ namespace Metasound
 			FTriggerReadRef InTriggerNext = InParams.InputDataReferences.GetDataReadReferenceOrConstructWithVertexDefault<FTrigger>(Inputs, InputTriggerNextName, InParams.OperatorSettings);
 			FTriggerReadRef InTriggerReset = InParams.InputDataReferences.GetDataReadReferenceOrConstructWithVertexDefault<FTrigger>(Inputs, InputTriggerResetName, InParams.OperatorSettings);
 			FArrayDataReadReference InInputArray = InParams.InputDataReferences.GetDataReadReferenceOrConstructWithVertexDefault<ArrayType>(Inputs, InputRandomArrayName, InParams.OperatorSettings);
-			FArrayWeightReadReference InInputWeightsArray = InParams.InputDataReferences.GetDataReadReferenceOrConstructWithVertexDefault<WeightArrayType>(Inputs, InputWeightsName, InParams.OperatorSettings);
+			FArrayWeightReadReference InInputWeightsArray = InParams.InputDataReferences.GetDataReadReferenceOrConstructWithVertexDefault<WeightsArrayType>(Inputs, InputWeightsName, InParams.OperatorSettings);
 			FInt32ReadRef InSeedValue = InParams.InputDataReferences.GetDataReadReferenceOrConstructWithVertexDefault<int32>(Inputs, InputSeedName, InParams.OperatorSettings);
 			FInt32ReadRef InNoRepeatOrder = InParams.InputDataReferences.GetDataReadReferenceOrConstructWithVertexDefault<int32>(Inputs, InputNoRepeatOrderName, InParams.OperatorSettings);
 			FBoolReadRef bInEnableSharedState = InParams.InputDataReferences.GetDataReadReferenceOrConstructWithVertexDefault<bool>(Inputs, InputEnableSharedStateName, InParams.OperatorSettings);
@@ -163,7 +174,7 @@ namespace Metasound
 			const FTriggerReadRef& InTriggerNext,
 			const FTriggerReadRef& InTriggerReset,
 			const FArrayDataReadReference& InInputArray,
-			const TDataReadReference<WeightArrayType>& InInputWeightsArray,
+			const TDataReadReference<WeightsArrayType>& InInputWeightsArray,
 			const FInt32ReadRef& InSeedValue,
 			const FInt32ReadRef& InNoRepeatOrder,
 			const FBoolReadRef& bInEnableSharedState)
@@ -183,11 +194,11 @@ namespace Metasound
 			PrevSeedValue = *SeedValue;
 			PrevNoRepeatOrder = FMath::Max(*NoRepeatOrder, 0);
 
+			WeightsArray = *InputWeightsArray;
+
 			const ArrayType& InputArrayRef = *InputArray;
 			int32 ArraySize = InputArrayRef.Num();
 
-			const TArray<float>& InputWeightsArrayRef = *InputWeightsArray;
-
 			if (ArraySize > 0)
 			{
 				if (*bEnableSharedState)
@@ -196,12 +207,23 @@ namespace Metasound
 					SharedStateUniqueId = InParams.Environment.GetValue<uint32>(TEXT("SoundUniqueId"));
 					check(SharedStateUniqueId != INDEX_NONE);
 
+					bIsPreviewSound = InParams.Environment.GetValue<bool>(TEXT("IsPreviewSound"));
+
 					FSharedStateRandomGetManager& RGM = FSharedStateRandomGetManager::Get();
-					RGM.InitSharedState(SharedStateUniqueId, PrevSeedValue, ArraySize, InputWeightsArrayRef, PrevNoRepeatOrder);
+
+					InitSharedStateArgs Args;
+					Args.SharedStateId = SharedStateUniqueId;
+					Args.Seed = PrevSeedValue;
+					Args.NumElements = ArraySize;
+					Args.NoRepeatOrder = PrevNoRepeatOrder;
+					Args.bIsPreviewSound = bIsPreviewSound;
+					Args.Weights = WeightsArray;
+
+					RGM.InitSharedState(Args);
 				}
 				else
 				{
-					ArrayRandomGet = MakeUnique<FArrayRandomGet>(PrevSeedValue, ArraySize, InputWeightsArrayRef, PrevNoRepeatOrder);
+					ArrayRandomGet = MakeUnique<FArrayRandomGet>(PrevSeedValue, ArraySize, WeightsArray, PrevNoRepeatOrder);
 				}
 			}
 			else
@@ -278,6 +300,18 @@ namespace Metasound
 					ArrayRandomGet->SetNoRepeatOrder(PrevNoRepeatOrder);
 				}
 			}
+
+			WeightsArray = *InputWeightsArray;
+			if (SharedStateUniqueId != INDEX_NONE)
+			{
+				FSharedStateRandomGetManager& RGM = FSharedStateRandomGetManager::Get();
+				RGM.SetRandomWeights(SharedStateUniqueId, WeightsArray);
+			}
+			else
+			{
+				check(ArrayRandomGet.IsValid());
+				ArrayRandomGet->SetRandomWeights(WeightsArray);
+			}
  
 			// Don't do anything if our array is empty
 			if (InputArrayRef.Num() == 0)
@@ -341,7 +375,7 @@ namespace Metasound
 		FTriggerReadRef TriggerNext;
 		FTriggerReadRef TriggerReset;
 		FArrayDataReadReference InputArray;
-		TDataReadReference<WeightArrayType> InputWeightsArray;
+		TDataReadReference<WeightsArrayType> InputWeightsArray;
 		FInt32ReadRef SeedValue;
 		FInt32ReadRef NoRepeatOrder;
 		FBoolReadRef bEnableSharedState;
@@ -353,9 +387,11 @@ namespace Metasound
 
 		// Data
 		TUniquePtr<FArrayRandomGet> ArrayRandomGet;
+		TArray<float> WeightsArray;
 		int32 PrevSeedValue = INDEX_NONE;
 		int32 PrevNoRepeatOrder = INDEX_NONE;
 		uint32 SharedStateUniqueId = INDEX_NONE;
+		bool bIsPreviewSound = false;
 	};
 
 	template<typename ArrayType>
diff --git a/Engine/Source/Runtime/AudioMixer/Private/AudioMixerSource.cpp b/Engine/Source/Runtime/AudioMixer/Private/AudioMixerSource.cpp
--- a/Engine/Source/Runtime/AudioMixer/Private/AudioMixerSource.cpp
+++ b/Engine/Source/Runtime/AudioMixer/Private/AudioMixerSource.cpp
@@ -812,7 +812,16 @@ namespace Audio
 			InstanceID = WaveInstance->ActiveSound->GetAudioComponentID();
 		}
 
-		MixerSourceBuffer = FMixerSourceBuffer::Create(InstanceID, AudioDevice->GetSampleRate(), *MixerBuffer, SoundWave, InWaveInstance->LoopingMode, bIsSeeking);
+		FMixerSourceBufferInitArgs BufferInitArgs;
+		BufferInitArgs.InstanceID = InstanceID;
+		BufferInitArgs.SampleRate = AudioDevice->GetSampleRate();
+		BufferInitArgs.Buffer = MixerBuffer;
+		BufferInitArgs.SoundWave = &SoundWave;
+		BufferInitArgs.LoopingMode = InWaveInstance->LoopingMode;
+		BufferInitArgs.bIsSeeking = bIsSeeking;
+		BufferInitArgs.bIsPreviewSound = WaveInstance->ActiveSound->bIsPreviewSound;
+
+		MixerSourceBuffer = FMixerSourceBuffer::Create(BufferInitArgs);
 		
 		if (!MixerSourceBuffer.IsValid())
 		{
diff --git a/Engine/Source/Runtime/AudioMixer/Private/AudioMixerSourceBuffer.cpp b/Engine/Source/Runtime/AudioMixer/Private/AudioMixerSourceBuffer.cpp
--- a/Engine/Source/Runtime/AudioMixer/Private/AudioMixerSourceBuffer.cpp
+++ b/Engine/Source/Runtime/AudioMixer/Private/AudioMixerSourceBuffer.cpp
@@ -56,47 +56,49 @@ namespace Audio
 		return CurrentSample >= NumSamples;
 	}
 
-	TSharedPtr<FMixerSourceBuffer, ESPMode::ThreadSafe> FMixerSourceBuffer::Create(uint64 InInstanceID, int32 InSampleRate, FMixerBuffer& InBuffer, USoundWave& InWave, ELoopingMode InLoopingMode, bool bInIsSeeking, bool bInForceSyncDecode)
+	TSharedPtr<FMixerSourceBuffer, ESPMode::ThreadSafe> FMixerSourceBuffer::Create(FMixerSourceBufferInitArgs& InArgs)
 	{
 		LLM_SCOPE(ELLMTag::AudioMixer);
 
-		TSharedPtr<FMixerSourceBuffer, ESPMode::ThreadSafe> NewSourceBuffer = MakeShareable(new FMixerSourceBuffer(InInstanceID, InSampleRate, InBuffer, InWave, InLoopingMode, bInIsSeeking, bInForceSyncDecode));
+		TSharedPtr<FMixerSourceBuffer, ESPMode::ThreadSafe> NewSourceBuffer = MakeShareable(new FMixerSourceBuffer(InArgs));
 
 		return NewSourceBuffer;
 	}
 
-	FMixerSourceBuffer::FMixerSourceBuffer(uint64 InInstanceID, int32 InSampleRate, FMixerBuffer& InBuffer, USoundWave& InWave, ELoopingMode InLoopingMode, bool bInIsSeeking, bool bInForceSyncDecode)
+	FMixerSourceBuffer::FMixerSourceBuffer(FMixerSourceBufferInitArgs& InArgs)
 		: NumBuffersQeueued(0)
 		, CurrentBuffer(0)
-		, SoundWave(&InWave)
+		, SoundWave(InArgs.SoundWave)
 		, AsyncRealtimeAudioTask(nullptr)
 		, DecompressionState(nullptr)
-		, LoopingMode(InLoopingMode)
-		, NumChannels(InBuffer.NumChannels)
-		, BufferType(InBuffer.GetType())
-		, NumPrecacheFrames(InWave.NumPrecacheFrames)
+		, LoopingMode(InArgs.LoopingMode)
+		, NumChannels(InArgs.Buffer->NumChannels)
+		, BufferType(InArgs.Buffer->GetType())
+		, NumPrecacheFrames(InArgs.SoundWave->NumPrecacheFrames)
 		, bInitialized(false)
 		, bBufferFinished(false)
 		, bPlayedCachedBuffer(false)
-		, bIsSeeking(bInIsSeeking)
+		, bIsSeeking(InArgs.bIsSeeking)
 		, bLoopCallback(false)
-		, bProcedural(InWave.bProcedural)
-		, bIsBus(InWave.bIsSourceBus)
-		, bForceSyncDecode(bInForceSyncDecode)
+		, bProcedural(InArgs.SoundWave->bProcedural)
+		, bIsBus(InArgs.SoundWave->bIsSourceBus)
+		, bForceSyncDecode(InArgs.bForceSyncDecode)
 		, bHasError(false)
 	{
 		// TODO: remove the need to do this here. 1) remove need for decoders to depend on USoundWave and 2) remove need for procedural sounds to use USoundWaveProcedural
-		InWave.AddPlayingSource(this);
+		InArgs.SoundWave->AddPlayingSource(this);
 
 		// Retrieve a sound generator if this is a procedural sound wave
 		if (bProcedural)
 		{
 			FSoundGeneratorInitParams InitParams;
-			InitParams.SampleRate = InSampleRate;
+			InitParams.SampleRate = InArgs.SampleRate;
 			InitParams.NumChannels = NumChannels;
 			InitParams.NumFramesPerCallback = MONO_PCM_BUFFER_SAMPLES;
-			InitParams.InstanceID = InInstanceID;
-			SoundGenerator = InWave.CreateSoundGenerator(InitParams);
+			InitParams.InstanceID = InArgs.InstanceID;
+			InitParams.bIsPreviewSound = InArgs.bIsPreviewSound;
+
+			SoundGenerator = InArgs.SoundWave->CreateSoundGenerator(InitParams);
 		}
 
 		const uint32 TotalSamples = MONO_PCM_BUFFER_SAMPLES * NumChannels;
diff --git a/Engine/Source/Runtime/AudioMixer/Private/AudioMixerSourceBuffer.h b/Engine/Source/Runtime/AudioMixer/Private/AudioMixerSourceBuffer.h
--- a/Engine/Source/Runtime/AudioMixer/Private/AudioMixerSourceBuffer.h
+++ b/Engine/Source/Runtime/AudioMixer/Private/AudioMixerSourceBuffer.h
@@ -46,11 +46,23 @@ namespace Audio
 
 	using FMixerSourceBufferPtr = TSharedPtr<class FMixerSourceBuffer, ESPMode::ThreadSafe>;
 
+	struct FMixerSourceBufferInitArgs
+	{
+		uint32 InstanceID = 0;
+		int32 SampleRate = 0;
+		FMixerBuffer* Buffer = nullptr;
+		USoundWave* SoundWave = nullptr;
+		ELoopingMode LoopingMode = ELoopingMode::LOOP_Never;
+		bool bIsSeeking = false;
+		bool bForceSyncDecode = false;
+		bool bIsPreviewSound = false;
+	};
+
 	/** Class which handles decoding audio for a particular source buffer. */
 	class FMixerSourceBuffer : public ISoundWaveClient
 	{
 	public:
-		static FMixerSourceBufferPtr Create(uint64 InInstanceID, int32 InSampleRate, FMixerBuffer& InBuffer, USoundWave& InWave, ELoopingMode InLoopingMode, bool bInIsSeeking, bool bInForceSyncDecode = false);
+		static FMixerSourceBufferPtr Create(FMixerSourceBufferInitArgs& InArgs);
 
 		~FMixerSourceBuffer();
 
@@ -101,7 +113,7 @@ namespace Audio
 		bool IsGeneratorFinished() const;
 
 	private:
-		FMixerSourceBuffer(uint64 InInstanceID, int32 InSampleRate, FMixerBuffer& InBuffer, USoundWave& InWave, ELoopingMode InLoopingMode, bool bInIsSeeking, bool bInForceSyncDecode = false);
+		FMixerSourceBuffer(FMixerSourceBufferInitArgs& InArgs);
 
 		void SubmitInitialPCMBuffers();
 		void SubmitInitialRealtimeBuffers();
diff --git a/Engine/Source/Runtime/AudioMixer/Private/SoundWaveDecoder.cpp b/Engine/Source/Runtime/AudioMixer/Private/SoundWaveDecoder.cpp
--- a/Engine/Source/Runtime/AudioMixer/Private/SoundWaveDecoder.cpp
+++ b/Engine/Source/Runtime/AudioMixer/Private/SoundWaveDecoder.cpp
@@ -57,8 +57,16 @@ namespace Audio
 		bool bIsValid = false;
 		{
 			FScopeLock Lock(&MixerSourceBufferCritSec);
-			uint64 InstanceID = 0; // Unused.
-			MixerSourceBuffer = FMixerSourceBuffer::Create(InstanceID, InSampleRate, *MixerBuffer, *SoundWave, LoopingMode, bIsSeeking, bForceSyncDecode);
+
+			FMixerSourceBufferInitArgs Args;
+			Args.SampleRate = InSampleRate;
+			Args.Buffer = MixerBuffer;
+			Args.SoundWave = SoundWave;
+			Args.LoopingMode = LoopingMode;
+			Args.bIsSeeking = bIsSeeking;
+			Args.bForceSyncDecode = bForceSyncDecode;
+			MixerSourceBuffer = FMixerSourceBuffer::Create(Args);
+
 			bIsValid = MixerSourceBuffer.IsValid();
 		}
 
diff --git a/Engine/Source/Runtime/Engine/Classes/Sound/SoundGenerator.h b/Engine/Source/Runtime/Engine/Classes/Sound/SoundGenerator.h
--- a/Engine/Source/Runtime/Engine/Classes/Sound/SoundGenerator.h
+++ b/Engine/Source/Runtime/Engine/Classes/Sound/SoundGenerator.h
@@ -10,17 +10,11 @@
 // Parameters used for constructing a new ISoundGenerator.
 struct FSoundGeneratorInitParams
 {
-	float SampleRate;
-	int32 NumChannels;
-	int32 NumFramesPerCallback;
-	uint64 InstanceID;
-
-	FSoundGeneratorInitParams()
-		: SampleRate(0.0f)
-		, NumChannels(0)
-		, NumFramesPerCallback(0)
-		, InstanceID(0)
-	{}
+	float SampleRate = 0.0f;
+	int32 NumChannels = 0;
+	int32 NumFramesPerCallback = 0;
+	uint64 InstanceID = 0;
+	bool bIsPreviewSound = false;
 };
 
 class ENGINE_API ISoundGenerator
