Fix for missing FM modulation of bandlimited sawtooth gen

diff --git a/Engine/Plugins/Runtime/Metasound/Source/MetasoundStandardNodes/Private/MetasoundOscillatorNodes.cpp b/Engine/Plugins/Runtime/Metasound/Source/MetasoundStandardNodes/Private/MetasoundOscillatorNodes.cpp
--- a/Engine/Plugins/Runtime/Metasound/Source/MetasoundStandardNodes/Private/MetasoundOscillatorNodes.cpp
+++ b/Engine/Plugins/Runtime/Metasound/Source/MetasoundStandardNodes/Private/MetasoundOscillatorNodes.cpp
@@ -850,7 +850,7 @@ namespace Metasound
 			{
 			default:
 			case ESawGenerationType::Trivial: return MakeUnique<TOscillatorOperatorFM<FSawWithFm>>(OpParams, FmBuffer);
-			case ESawGenerationType::PolySmooth: return MakeUnique<TOscillatorOperatorFM<FSawPolysmooth>>(OpParams, FmBuffer);
+			case ESawGenerationType::PolySmooth: return MakeUnique<TOscillatorOperatorFM<FSawPolysmoothWithFm>>(OpParams, FmBuffer);
 			}
 		}
 		else
