Fixing metasound waveplayer

- Changing "On Done" and "On Nearly Done" to "On Finished" and "On Nearly Finished" to be consistent with the metasound output trigger. IMO "Done" is better, but it's even better to be consistent.

- Fixed playback progress

- Fixed missed On PLay output trigger

UE-112564, UE-112565

diff --git a/Engine/Plugins/Runtime/Metasound/Source/MetasoundEngine/Private/MetasoundWavePlayerNode.cpp b/Engine/Plugins/Runtime/Metasound/Source/MetasoundEngine/Private/MetasoundWavePlayerNode.cpp
--- a/Engine/Plugins/Runtime/Metasound/Source/MetasoundEngine/Private/MetasoundWavePlayerNode.cpp
+++ b/Engine/Plugins/Runtime/Metasound/Source/MetasoundEngine/Private/MetasoundWavePlayerNode.cpp
@@ -29,8 +29,8 @@ namespace Metasound
 		//static const TCHAR* InputStartTimeAsPercentName = TEXT("Start/Loop Times As Percent");
 
 		static const TCHAR* OutputTriggerOnPlayName = TEXT("On Play");
-		static const TCHAR* OutputTriggerOnDoneName = TEXT("On Done");
-		static const TCHAR* OutputTriggerOnNearlyDoneName = TEXT("On Nearly Done");
+		static const TCHAR* OutputTriggerOnDoneName = TEXT("On Finished");
+		static const TCHAR* OutputTriggerOnNearlyDoneName = TEXT("On Nearly Finished");
 		static const TCHAR* OutputTriggerOnLoopedName = TEXT("On Looped");
 		static const TCHAR* OutputTriggerOnCuePointName = TEXT("On Cue Point");
 		static const TCHAR* OutputCuePointIDName = TEXT("Cue Point ID");
@@ -276,6 +276,8 @@ namespace Metasound
 					{
 						ExecuteSeekRequest();
 					}
+					TriggerOnPlay->TriggerFrame(CurrAudioFrame);
+
 					StartPlaying();
 					++PlayTrigIndex;
 				}
@@ -485,7 +487,7 @@ namespace Metasound
 					int32 PrevConsumedFrameCount = CurrentConsumedFrameCount;
 					CurrentConsumedFrameCount = (CurrentConsumedFrameCount + NumFramesConsumed) % SoundAssetNumFrames;
 
-					*PlaybackLocation = SoundAssetDurationSeconds * (float)(CurrentConsumedFrameCount / SoundAssetNumFrames);
+					*PlaybackLocation = SoundAssetDurationSeconds * ((float)CurrentConsumedFrameCount / SoundAssetNumFrames);
 
 					// Check for any cue trigger events 
 					FSoundWaveCuePoint OutCuePoint;
