Metasound node cleanup, recategorization of a few nodes.

- removing header implementation for trigger any node

diff --git a/Engine/Plugins/Runtime/Metasound/Source/MetasoundStandardNodes/Private/MetasoundClampNode.cpp b/Engine/Plugins/Runtime/Metasound/Source/MetasoundStandardNodes/Private/MetasoundClampNode.cpp
--- a/Engine/Plugins/Runtime/Metasound/Source/MetasoundStandardNodes/Private/MetasoundClampNode.cpp
+++ b/Engine/Plugins/Runtime/Metasound/Source/MetasoundStandardNodes/Private/MetasoundClampNode.cpp
@@ -35,7 +35,7 @@ namespace Metasound
 				PluginAuthor,
 				PluginNodeMissingPrompt,
 				InDefaultInterface,
-				{LOCTEXT("ClampCategory", "Utils")},
+				{LOCTEXT("ClampCategory", "Math")},
 				{TEXT("Clamp")},
 				FNodeDisplayStyle{}
 			};
diff --git a/Engine/Plugins/Runtime/Metasound/Source/MetasoundStandardNodes/Private/MetasoundMaxNode.cpp b/Engine/Plugins/Runtime/Metasound/Source/MetasoundStandardNodes/Private/MetasoundMaxNode.cpp
--- a/Engine/Plugins/Runtime/Metasound/Source/MetasoundStandardNodes/Private/MetasoundMaxNode.cpp
+++ b/Engine/Plugins/Runtime/Metasound/Source/MetasoundStandardNodes/Private/MetasoundMaxNode.cpp
@@ -35,7 +35,7 @@ namespace Metasound
 				PluginAuthor,
 				PluginNodeMissingPrompt,
 				InDefaultInterface,
-				{LOCTEXT("MaxCategory", "Utils")},
+				{LOCTEXT("MaxCategory", "Math")},
 				{TEXT("Max")},
 				FNodeDisplayStyle{}
 			};
diff --git a/Engine/Plugins/Runtime/Metasound/Source/MetasoundStandardNodes/Private/MetasoundMinNode.cpp b/Engine/Plugins/Runtime/Metasound/Source/MetasoundStandardNodes/Private/MetasoundMinNode.cpp
--- a/Engine/Plugins/Runtime/Metasound/Source/MetasoundStandardNodes/Private/MetasoundMinNode.cpp
+++ b/Engine/Plugins/Runtime/Metasound/Source/MetasoundStandardNodes/Private/MetasoundMinNode.cpp
@@ -35,7 +35,7 @@ namespace Metasound
 				PluginAuthor,
 				PluginNodeMissingPrompt,
 				InDefaultInterface,
-				{LOCTEXT("MaxCategory", "Utils")},
+				{LOCTEXT("MaxCategory", "Math")},
 				{TEXT("Min")},
 				FNodeDisplayStyle{}
 			};
diff --git a/Engine/Plugins/Runtime/Metasound/Source/MetasoundStandardNodes/Private/MetasoundTriggerRouteNode.cpp b/Engine/Plugins/Runtime/Metasound/Source/MetasoundStandardNodes/Private/MetasoundTriggerRouteNode.cpp
--- a/Engine/Plugins/Runtime/Metasound/Source/MetasoundStandardNodes/Private/MetasoundTriggerRouteNode.cpp
+++ b/Engine/Plugins/Runtime/Metasound/Source/MetasoundStandardNodes/Private/MetasoundTriggerRouteNode.cpp
@@ -17,7 +17,7 @@
 
 #define LOCTEXT_NAMESPACE "MetasoundStandardNodes"
 
-#define REGISTER_TRIGGER_ROUTE_NOTE(DataType, Number) \
+#define REGISTER_TRIGGER_ROUTE_NODE(DataType, Number) \
 	using FTriggerRouteNode##DataType##_##Number = TTriggerRouteNode<DataType, Number>; \
 	METASOUND_REGISTER_NODE(FTriggerRouteNode##DataType##_##Number) \
 
@@ -250,37 +250,37 @@ namespace Metasound
 		virtual ~TTriggerRouteNode() = default;
 	};
 
-	REGISTER_TRIGGER_ROUTE_NOTE(int32, 2)
-	REGISTER_TRIGGER_ROUTE_NOTE(int32, 3)
-	REGISTER_TRIGGER_ROUTE_NOTE(int32, 4)
-	REGISTER_TRIGGER_ROUTE_NOTE(int32, 5)
-	REGISTER_TRIGGER_ROUTE_NOTE(int32, 6)
-	REGISTER_TRIGGER_ROUTE_NOTE(int32, 7)
-	REGISTER_TRIGGER_ROUTE_NOTE(int32, 8)
-
-	REGISTER_TRIGGER_ROUTE_NOTE(float, 2)
-	REGISTER_TRIGGER_ROUTE_NOTE(float, 3)
-	REGISTER_TRIGGER_ROUTE_NOTE(float, 4)
-	REGISTER_TRIGGER_ROUTE_NOTE(float, 5)
-	REGISTER_TRIGGER_ROUTE_NOTE(float, 6)
-	REGISTER_TRIGGER_ROUTE_NOTE(float, 7)
-	REGISTER_TRIGGER_ROUTE_NOTE(float, 8)
-
-	REGISTER_TRIGGER_ROUTE_NOTE(bool, 2)
-	REGISTER_TRIGGER_ROUTE_NOTE(bool, 3)
-	REGISTER_TRIGGER_ROUTE_NOTE(bool, 4)
-	REGISTER_TRIGGER_ROUTE_NOTE(bool, 5)
-	REGISTER_TRIGGER_ROUTE_NOTE(bool, 6)
-	REGISTER_TRIGGER_ROUTE_NOTE(bool, 7)
-	REGISTER_TRIGGER_ROUTE_NOTE(bool, 8)
-
-	REGISTER_TRIGGER_ROUTE_NOTE(FAudioBuffer, 2)
-	REGISTER_TRIGGER_ROUTE_NOTE(FAudioBuffer, 3)
-	REGISTER_TRIGGER_ROUTE_NOTE(FAudioBuffer, 4)
-	REGISTER_TRIGGER_ROUTE_NOTE(FAudioBuffer, 5)
-	REGISTER_TRIGGER_ROUTE_NOTE(FAudioBuffer, 6)
-	REGISTER_TRIGGER_ROUTE_NOTE(FAudioBuffer, 7)
-	REGISTER_TRIGGER_ROUTE_NOTE(FAudioBuffer, 8)
+	REGISTER_TRIGGER_ROUTE_NODE(int32, 2)
+	REGISTER_TRIGGER_ROUTE_NODE(int32, 3)
+	REGISTER_TRIGGER_ROUTE_NODE(int32, 4)
+	REGISTER_TRIGGER_ROUTE_NODE(int32, 5)
+	REGISTER_TRIGGER_ROUTE_NODE(int32, 6)
+	REGISTER_TRIGGER_ROUTE_NODE(int32, 7)
+	REGISTER_TRIGGER_ROUTE_NODE(int32, 8)
+
+	REGISTER_TRIGGER_ROUTE_NODE(float, 2)
+	REGISTER_TRIGGER_ROUTE_NODE(float, 3)
+	REGISTER_TRIGGER_ROUTE_NODE(float, 4)
+	REGISTER_TRIGGER_ROUTE_NODE(float, 5)
+	REGISTER_TRIGGER_ROUTE_NODE(float, 6)
+	REGISTER_TRIGGER_ROUTE_NODE(float, 7)
+	REGISTER_TRIGGER_ROUTE_NODE(float, 8)
+
+	REGISTER_TRIGGER_ROUTE_NODE(bool, 2)
+	REGISTER_TRIGGER_ROUTE_NODE(bool, 3)
+	REGISTER_TRIGGER_ROUTE_NODE(bool, 4)
+	REGISTER_TRIGGER_ROUTE_NODE(bool, 5)
+	REGISTER_TRIGGER_ROUTE_NODE(bool, 6)
+	REGISTER_TRIGGER_ROUTE_NODE(bool, 7)
+	REGISTER_TRIGGER_ROUTE_NODE(bool, 8)
+
+	REGISTER_TRIGGER_ROUTE_NODE(FAudioBuffer, 2)
+	REGISTER_TRIGGER_ROUTE_NODE(FAudioBuffer, 3)
+	REGISTER_TRIGGER_ROUTE_NODE(FAudioBuffer, 4)
+	REGISTER_TRIGGER_ROUTE_NODE(FAudioBuffer, 5)
+	REGISTER_TRIGGER_ROUTE_NODE(FAudioBuffer, 6)
+	REGISTER_TRIGGER_ROUTE_NODE(FAudioBuffer, 7)
+	REGISTER_TRIGGER_ROUTE_NODE(FAudioBuffer, 8)
 }
 
 #undef LOCTEXT_NAMESPACE // MetasoundTriggerDelayNode
