Make metasound asset action top level, reduce sound cue into a legacy category

diff --git a/Engine/Plugins/Runtime/Metasound/Source/MetasoundEditor/Private/MetasoundAssetTypeActions.cpp b/Engine/Plugins/Runtime/Metasound/Source/MetasoundEditor/Private/MetasoundAssetTypeActions.cpp
--- a/Engine/Plugins/Runtime/Metasound/Source/MetasoundEditor/Private/MetasoundAssetTypeActions.cpp
+++ b/Engine/Plugins/Runtime/Metasound/Source/MetasoundEditor/Private/MetasoundAssetTypeActions.cpp
@@ -13,21 +13,11 @@ namespace Metasound
 {
 	namespace Editor
 	{
-		static const TArray<FText> AssetTypeActionSubMenu
-		{
-			LOCTEXT("AssetMetaSoundSubMenu", "MetaSound")
-		};
-
 		UClass* FAssetTypeActions_Metasound::GetSupportedClass() const
 		{
 			return UMetaSound::StaticClass();
 		}
 
-		const TArray<FText>& FAssetTypeActions_Metasound::GetSubMenus() const
-		{
-			return AssetTypeActionSubMenu;
-		}
-
 		void FAssetTypeActions_Metasound::OpenAssetEditor(const TArray<UObject*>& InObjects, TSharedPtr<IToolkitHost> ToolkitHost)
 		{
 			const EToolkitMode::Type Mode = ToolkitHost.IsValid() ? EToolkitMode::WorldCentric : EToolkitMode::Standalone;
@@ -60,11 +50,6 @@ namespace Metasound
 				}
 			}
 		}
-
-		const TArray<FText>& FAssetTypeActions_MetasoundSource::GetSubMenus() const
-		{
-			return AssetTypeActionSubMenu;
-		}
 	} // namespace Editor
 } // namespace Metasound
 
diff --git a/Engine/Plugins/Runtime/Metasound/Source/MetasoundEditor/Private/MetasoundAssetTypeActions.h b/Engine/Plugins/Runtime/Metasound/Source/MetasoundEditor/Private/MetasoundAssetTypeActions.h
--- a/Engine/Plugins/Runtime/Metasound/Source/MetasoundEditor/Private/MetasoundAssetTypeActions.h
+++ b/Engine/Plugins/Runtime/Metasound/Source/MetasoundEditor/Private/MetasoundAssetTypeActions.h
@@ -21,21 +21,18 @@ namespace Metasound
 			virtual UClass* GetSupportedClass() const override;
 			virtual uint32 GetCategories() override { return EAssetTypeCategories::Sounds; }
 			virtual void OpenAssetEditor(const TArray<UObject*>& InObjects, TSharedPtr<IToolkitHost> ToolkitHost) override;
-
-			virtual const TArray<FText>& GetSubMenus() const override;
 		};
 
 		class FAssetTypeActions_MetasoundSource : public FAssetTypeActions_Base
 		{
 		public:
 			// IAssetTypeActions Implementation
-			virtual FText GetName() const override { return NSLOCTEXT("AssetTypeActions", "AssetTypeActions_MetaSoundSource", "MetaSoundSource"); }
+			virtual FText GetName() const override { return NSLOCTEXT("AssetTypeActions", "AssetTypeActions_MetaSoundSource", "MetaSound"); }
 			virtual FColor GetTypeColor() const override { return FColor(103, 214, 66); }
 			virtual UClass* GetSupportedClass() const override;
 			virtual uint32 GetCategories() override { return EAssetTypeCategories::Sounds; }
 			virtual void OpenAssetEditor(const TArray<UObject*>& InObjects, TSharedPtr<IToolkitHost> ToolkitHost) override;
 
-			virtual const TArray<FText>& GetSubMenus() const override;
-		};
+	};
 	} // namespace Editor
 } // namespace Metasound
diff --git a/Engine/Plugins/Runtime/SoundCueTemplates/Source/SoundCueTemplatesEditor/Private/AssetTypeActions_SoundCueTemplate.cpp b/Engine/Plugins/Runtime/SoundCueTemplates/Source/SoundCueTemplatesEditor/Private/AssetTypeActions_SoundCueTemplate.cpp
--- a/Engine/Plugins/Runtime/SoundCueTemplates/Source/SoundCueTemplatesEditor/Private/AssetTypeActions_SoundCueTemplate.cpp
+++ b/Engine/Plugins/Runtime/SoundCueTemplates/Source/SoundCueTemplatesEditor/Private/AssetTypeActions_SoundCueTemplate.cpp
@@ -59,6 +59,23 @@ void FAssetTypeActions_SoundCueTemplate::ExecuteCopyToSoundCue(TArray<TWeakObjec
 	}
 }
 
+const TArray<FText>& FAssetTypeActions_SoundCueTemplate::GetSubMenus() const
+{
+	if (FModuleManager::Get().IsModuleLoaded(TEXT("MetasoundEditor")))
+	{
+		static const TArray<FText> AssetTypeActionSubMenu
+		{
+			LOCTEXT("AssetSoudnCueSubMenu", "Legacy")
+		};
+		return AssetTypeActionSubMenu;
+	}
+	else
+	{
+		static const TArray<FText> AssetTypeActionSubMenu;
+		return AssetTypeActionSubMenu;
+	}
+}
+
 void FAssetActionExtender_SoundCueTemplate::RegisterMenus()
 {
 	if (!UToolMenus::IsToolMenuUIEnabled())
@@ -116,4 +133,8 @@ void FAssetActionExtender_SoundCueTemplate::ExecuteCreateSoundCueTemplate(const
 	FContentBrowserModule& ContentBrowserModule = FModuleManager::Get().LoadModuleChecked<FContentBrowserModule>("ContentBrowser");
 	ContentBrowserModule.Get().CreateNewAsset(Name, FPackageName::GetLongPackagePath(PackagePath), USoundCueTemplate::StaticClass(), Factory);
 }
+
+
+
+
 #undef LOCTEXT_NAMESPACE
diff --git a/Engine/Plugins/Runtime/SoundCueTemplates/Source/SoundCueTemplatesEditor/Private/AssetTypeActions_SoundCueTemplate.h b/Engine/Plugins/Runtime/SoundCueTemplates/Source/SoundCueTemplatesEditor/Private/AssetTypeActions_SoundCueTemplate.h
--- a/Engine/Plugins/Runtime/SoundCueTemplates/Source/SoundCueTemplatesEditor/Private/AssetTypeActions_SoundCueTemplate.h
+++ b/Engine/Plugins/Runtime/SoundCueTemplates/Source/SoundCueTemplatesEditor/Private/AssetTypeActions_SoundCueTemplate.h
@@ -17,6 +17,7 @@ public:
 	virtual UClass* GetSupportedClass() const override;
 	virtual void GetActions(const TArray<UObject*>& InObjects, struct FToolMenuSection& Section) override;
 	virtual bool CanFilter() override { return true; }
+	virtual const TArray<FText>& GetSubMenus() const override;
 
 protected:
 	/** Converts the provided SoundCue Template to a fully-modifiable SoundCue */
diff --git a/Engine/Source/Editor/AudioEditor/Private/AssetTypeActions/AssetTypeActions_SoundCue.cpp b/Engine/Source/Editor/AudioEditor/Private/AssetTypeActions/AssetTypeActions_SoundCue.cpp
--- a/Engine/Source/Editor/AudioEditor/Private/AssetTypeActions/AssetTypeActions_SoundCue.cpp
+++ b/Engine/Source/Editor/AudioEditor/Private/AssetTypeActions/AssetTypeActions_SoundCue.cpp
@@ -121,4 +121,21 @@ void FAssetTypeActions_SoundCue::ExecuteConsolidateAttenuation(TArray<TWeakObjec
 	}
 }
 
+const TArray<FText>& FAssetTypeActions_SoundCue::GetSubMenus() const
+{
+	if (FModuleManager::Get().IsModuleLoaded(TEXT("MetasoundEditor")))
+	{
+		static const TArray<FText> AssetTypeActionSubMenu
+		{
+			LOCTEXT("AssetSoudnCueSubMenu", "Legacy")
+		};
+		return AssetTypeActionSubMenu;
+	}
+	else
+	{
+		static const TArray<FText> AssetTypeActionSubMenu;
+		return AssetTypeActionSubMenu;
+	}
+}
+
 #undef LOCTEXT_NAMESPACE
diff --git a/Engine/Source/Editor/AudioEditor/Private/AssetTypeActions/AssetTypeActions_SoundCue.h b/Engine/Source/Editor/AudioEditor/Private/AssetTypeActions/AssetTypeActions_SoundCue.h
--- a/Engine/Source/Editor/AudioEditor/Private/AssetTypeActions/AssetTypeActions_SoundCue.h
+++ b/Engine/Source/Editor/AudioEditor/Private/AssetTypeActions/AssetTypeActions_SoundCue.h
@@ -18,6 +18,7 @@ public:
 	virtual void GetActions(const TArray<UObject*>& InObjects, struct FToolMenuSection& Section) override;
 	virtual void OpenAssetEditor( const TArray<UObject*>& InObjects, TSharedPtr<class IToolkitHost> EditWithinLevelEditor = TSharedPtr<IToolkitHost>() ) override;
 	virtual bool CanFilter() override { return true; }
+	virtual const TArray<FText>& GetSubMenus() const override;
 
 private:
 	/** Take selected SoundCues and combine, as much as possible, them to using shared attenuation settings */
