Metasound Crossfade node

diff --git a/Engine/Plugins/Runtime/Metasound/Source/MetasoundStandardNodes/Private/MetasoundCrossfadeNode.cpp b/Engine/Plugins/Runtime/Metasound/Source/MetasoundStandardNodes/Private/MetasoundCrossfadeNode.cpp
new file mode 100644
--- /dev/null
+++ b/Engine/Plugins/Runtime/Metasound/Source/MetasoundStandardNodes/Private/MetasoundCrossfadeNode.cpp
@@ -0,0 +1,329 @@
+// Copyright Epic Games, Inc. All Rights Reserved.
+
+#include "MetasoundNodeRegistrationMacro.h"
+#include "MetasoundAudioBuffer.h"
+#include "CoreMinimal.h"
+#include "MetasoundFacade.h"
+#include "MetasoundExecutableOperator.h"
+#include "MetasoundPrimitives.h"
+#include "MetasoundStandardNodesNames.h"
+#include "MetasoundTrigger.h"
+#include "Internationalization/Text.h"
+#include "MetasoundParamHelper.h"
+#include "DSP/BufferVectorOperations.h"
+
+#define LOCTEXT_NAMESPACE "MetasoundStandardNodes_CrossfadeNode"
+
+#define REGISTER_CROSSFADE_NODE(DataType, Number) \
+	using FCrossfadeNode##DataType##_##Number = TCrossfadeNode<DataType, Number>; \
+	METASOUND_REGISTER_NODE(FCrossfadeNode##DataType##_##Number) \
+
+
+namespace Metasound
+{
+	namespace CrossfadeVertexNames
+	{
+		const FString GetInputCrossfadeValueName()
+		{
+			static const FString CrossfadeName = TEXT("Crossfade Value");
+			return CrossfadeName;
+		}
+
+		const FText GetInputCrossfadeValueDescription()
+		{
+			static const FText Desc = LOCTEXT("InputCrossfadeValueDesc", "Crossfade value to crossfade across inputs. Output will be the float value between adjacent whole number values.");
+			return Desc;
+		}
+
+		const FString GetInputName(uint32 InIndex)
+		{
+			return FText::Format(LOCTEXT("CrossfadeInputName", "In {0}"), InIndex).ToString();
+		}
+
+		const FText GetInputDescription(uint32 InIndex)
+		{
+			return FText::Format(LOCTEXT("CrossfadeInputDesc", "Cross fade {0} input."), InIndex);
+		}
+
+		const FString& GetOutputName()
+		{
+			static const FString Name = TEXT("Out");
+			return Name;
+		}
+
+		const FText& GetOutputDescription()
+		{
+			static const FText Desc = LOCTEXT("TriggerAccumulateOutputTriggerDesc", "Triggered when all input triggers have been triggered. Call Reset to reset the state or use \"Auto Reset\"");
+			return Desc;
+		}
+	}
+
+	template<typename ValueType, uint32 NumInputs>
+	class TCrossfadeHelper
+	{
+	};
+
+	// Partial specialization for float
+	template<uint32 NumInputs>
+	class TCrossfadeHelper<float, NumInputs>
+	{
+	public:
+		TCrossfadeHelper(int32 NumFramesPerBlock) {}
+
+		void GetCrossfadeOutput(int32 IndexA, int32 IndexB, float Alpha, const TArray<FFloatReadRef>& InCurrentValues, float& OutValue)
+		{
+			const FFloatReadRef& InA = InCurrentValues[IndexA];
+			const FFloatReadRef& InB = InCurrentValues[IndexB];
+			OutValue = FMath::Lerp(*InA, *InB, Alpha);
+		}
+	};
+
+	// Partial specialization for FAudioBuffer
+	template<uint32 NumInputs>
+	class TCrossfadeHelper<FAudioBuffer, NumInputs>
+	{
+	public:
+		TCrossfadeHelper(int32 InNumFramesPerBlock)
+			: NumFramesPerBlock(InNumFramesPerBlock)
+		{
+			PrevGains.AddZeroed(NumInputs);
+			CurrentGains.AddZeroed(NumInputs);
+			NeedsMixing.AddZeroed(NumInputs);
+		}
+
+		void GetCrossfadeOutput(int32 IndexA, int32 IndexB, float Alpha, const TArray<FAudioBufferReadRef>& InAudioBuffersValues, FAudioBuffer& OutAudioBuffer)
+		{
+			// Determine the gains
+			for (int32 i = 0; i < NumInputs; ++i)
+			{
+				if (i == IndexA)
+				{
+					CurrentGains[i] = 1.0f - Alpha;
+					NeedsMixing[i] = true;
+				}
+				else if (i == IndexB)
+				{
+					CurrentGains[i] = Alpha;
+					NeedsMixing[i] = true;
+				}
+				else
+				{
+					CurrentGains[i] = 0.0f;
+
+					// If we were already at 0.0f, don't need to do any mixing!
+					if (PrevGains[i] == 0.0f)
+					{
+						NeedsMixing[i] = false;
+					}
+				}
+			}
+
+			// Zero the output buffer so we can mix into it
+			OutAudioBuffer.Zero();
+			float* OutAudioBufferPtr = OutAudioBuffer.GetData();
+
+			// Now write to the scratch buffers w/ fade buffer fast given the new inputs
+			for (int32 i = 0; i < NumInputs; ++i)
+			{
+				// Only need to do anything on an input if either curr or prev is non-zero
+				if (NeedsMixing[i])
+				{
+					// Copy the input to the output
+					const FAudioBufferReadRef& InBuff = InAudioBuffersValues[i];
+					const float* BufferPtr = (*InBuff).GetData();
+
+					// mix in and fade to the target gain values
+					Audio::MixInBufferFast(BufferPtr, OutAudioBufferPtr, NumFramesPerBlock, PrevGains[i], CurrentGains[i]);
+				}
+			}
+
+			// Copy the CurrentGains to PrevGains
+			PrevGains = CurrentGains;
+		}
+
+	private:
+		int32 NumFramesPerBlock = 0;
+		TArray<float> PrevGains;
+		TArray<float> CurrentGains;
+		TArray<bool> NeedsMixing;
+	};
+
+	template<typename ValueType, uint32 NumInputs>
+	class TCrossfadeOperator : public TExecutableOperator<TCrossfadeOperator<ValueType, NumInputs>>
+	{
+	public:
+		static const FVertexInterface& GetVertexInterface()
+		{
+			using namespace CrossfadeVertexNames;
+
+			auto CreateDefaultInterface = []() -> FVertexInterface
+			{
+				FInputVertexInterface InputInterface;
+
+				InputInterface.Add(TInputDataVertexModel<float>(GetInputCrossfadeValueName(), GetInputCrossfadeValueDescription()));
+
+				for (uint32 i = 0; i < NumInputs; ++i)
+				{
+					InputInterface.Add(TInputDataVertexModel<ValueType>(GetInputName(i), GetInputDescription(i)));
+				}
+
+				FOutputVertexInterface OutputInterface;
+				OutputInterface.Add(TOutputDataVertexModel<ValueType>(GetOutputName(), GetOutputDescription()));
+
+				return FVertexInterface(InputInterface, OutputInterface);
+			};
+
+			static const FVertexInterface DefaultInterface = CreateDefaultInterface();
+			return DefaultInterface;
+		}
+
+		static const FNodeClassMetadata& GetNodeInfo()
+		{
+			auto CreateNodeClassMetadata = []() -> FNodeClassMetadata
+			{
+				FName DataTypeName = GetMetasoundDataTypeName<ValueType>();
+				FName OperatorName = *FString::Printf(TEXT("Trigger Route (%s, %d)"), *DataTypeName.ToString(), NumInputs);
+				FText NodeDisplayName = FText::Format(LOCTEXT("CrossfadeDisplayNamePattern", "Crossfade ({0}, {1})"), FText::FromString(GetMetasoundDataTypeString<ValueType>()), NumInputs);
+				FText NodeDescription = LOCTEXT("CrossfadeDescription", "Crossfades inputs to outputs.");
+				FVertexInterface NodeInterface = GetVertexInterface();
+
+				FNodeClassMetadata Metadata
+				{
+					FNodeClassName{FName("Crossfade"), OperatorName, DataTypeName},
+					1, // Major Version
+					0, // Minor Version
+					NodeDisplayName,
+					NodeDescription,
+					PluginAuthor,
+					PluginNodeMissingPrompt,
+					NodeInterface,
+					{LOCTEXT("CrossfadeCategory", "Math")},
+					{TEXT("Crossfade")},
+					FNodeDisplayStyle{}
+				};
+				return Metadata;
+			};
+
+			static const FNodeClassMetadata Metadata = CreateNodeClassMetadata();
+			return Metadata;
+		}
+
+		static TUniquePtr<IOperator> CreateOperator(const FCreateOperatorParams& InParams, TArray<TUniquePtr<IOperatorBuildError>>& OutErrors)
+		{
+			using namespace CrossfadeVertexNames;
+
+			const FInputVertexInterface& InputInterface = InParams.Node.GetVertexInterface().GetInputInterface();
+			const FDataReferenceCollection& InputCollection = InParams.InputDataReferences;
+
+			FFloatReadRef CrossfadeValue = InputCollection.GetDataReadReferenceOrConstructWithVertexDefault<float>(InputInterface, GetInputCrossfadeValueName(), InParams.OperatorSettings);
+
+			TArray<TDataReadReference<ValueType>> InputValues;
+			for (uint32 i = 0; i < NumInputs; ++i)
+			{
+				InputValues.Add(InputCollection.GetDataReadReferenceOrConstructWithVertexDefault<ValueType>(InputInterface, GetInputName(i), InParams.OperatorSettings));
+			}
+
+			return MakeUnique<TCrossfadeOperator<ValueType, NumInputs>>(InParams.OperatorSettings, CrossfadeValue, MoveTemp(InputValues));
+		}
+
+
+		TCrossfadeOperator(const FOperatorSettings& InSettings, const FFloatReadRef& InCrossfadeValue, TArray<TDataReadReference<ValueType>>&& InInputValues)
+			: CrossfadeValue(InCrossfadeValue)
+			, InputValues(MoveTemp(InInputValues))
+			, OutputValue(TDataWriteReferenceFactory<ValueType>::CreateAny(InSettings))
+			, Crossfader(InSettings.GetNumFramesPerBlock())
+		{
+			PerformCrossfadeOutput();
+		}
+
+		virtual ~TCrossfadeOperator() = default;
+
+		virtual FDataReferenceCollection GetInputs() const override
+		{
+			using namespace CrossfadeVertexNames;
+			FDataReferenceCollection Inputs;
+			Inputs.AddDataReadReference(GetInputCrossfadeValueName(), CrossfadeValue);
+			for (uint32 i = 0; i < NumInputs; ++i)
+			{
+				Inputs.AddDataReadReference(GetInputName(i), InputValues[i]);
+			}
+			return Inputs;
+		}
+
+		virtual FDataReferenceCollection GetOutputs() const override
+		{
+			using namespace CrossfadeVertexNames;
+
+			FDataReferenceCollection Outputs;
+			Outputs.AddDataReadReference(GetOutputName(), OutputValue);
+
+			return Outputs;
+		}
+
+		void PerformCrossfadeOutput()
+		{
+			// Clamp the cross fade value based on the number of inputs
+			float CurrentCrossfadeValue = FMath::Clamp(*CrossfadeValue, 0.0f, (float)(NumInputs - 1));
+
+			// Only update the cross fade state if anything has changed
+			if (!FMath::IsNearlyEqual(CurrentCrossfadeValue, PrevCrossfadeValue))
+			{
+				PrevCrossfadeValue = CurrentCrossfadeValue;
+				IndexA = (int32)FMath::Floor(CurrentCrossfadeValue);
+				IndexB = IndexA + 1;
+				Alpha = CurrentCrossfadeValue - (float)IndexA;
+			}
+
+			// Need to call this each block in case inputs have changed
+			Crossfader.GetCrossfadeOutput(IndexA, IndexB, Alpha, InputValues, *OutputValue);
+		}
+
+		void Execute()
+		{
+			PerformCrossfadeOutput();
+		}
+
+	private:
+		FFloatReadRef CrossfadeValue;
+		TArray<TDataReadReference<ValueType>> InputValues;
+		TDataWriteReference<ValueType> OutputValue;
+
+		float PrevCrossfadeValue = -1.0f;
+		int32 IndexA = 0;
+		int32 IndexB = 0;
+		float Alpha = 0.0f;
+		TCrossfadeHelper<ValueType, NumInputs> Crossfader;
+	};
+
+	template<typename ValueType, uint32 NumInputs>
+	class TCrossfadeNode : public FNodeFacade
+	{
+	public:
+		/**
+		 * Constructor used by the Metasound Frontend.
+		 */
+		TCrossfadeNode(const FNodeInitData& InInitData)
+			: FNodeFacade(InInitData.InstanceName, InInitData.InstanceID, TFacadeOperatorClass<TCrossfadeOperator<ValueType, NumInputs>>())
+		{}
+
+		virtual ~TCrossfadeNode() = default;
+	};
+
+	REGISTER_CROSSFADE_NODE(float, 2);
+	REGISTER_CROSSFADE_NODE(float, 3);
+	REGISTER_CROSSFADE_NODE(float, 4);
+	REGISTER_CROSSFADE_NODE(float, 5);
+	REGISTER_CROSSFADE_NODE(float, 6);
+	REGISTER_CROSSFADE_NODE(float, 7);
+	REGISTER_CROSSFADE_NODE(float, 8);
+
+	REGISTER_CROSSFADE_NODE(FAudioBuffer, 2);
+	REGISTER_CROSSFADE_NODE(FAudioBuffer, 3);
+	REGISTER_CROSSFADE_NODE(FAudioBuffer, 4);
+	REGISTER_CROSSFADE_NODE(FAudioBuffer, 5);
+	REGISTER_CROSSFADE_NODE(FAudioBuffer, 6);
+	REGISTER_CROSSFADE_NODE(FAudioBuffer, 7);
+	REGISTER_CROSSFADE_NODE(FAudioBuffer, 8);
+}
+
+#undef LOCTEXT_NAMESPACE
