Making audio rate map range

- Removing header

diff --git a/Engine/Plugins/Runtime/Metasound/Source/MetasoundStandardNodes/Private/MetasoundMapRangeNode.cpp b/Engine/Plugins/Runtime/Metasound/Source/MetasoundStandardNodes/Private/MetasoundMapRangeNode.cpp
--- a/Engine/Plugins/Runtime/Metasound/Source/MetasoundStandardNodes/Private/MetasoundMapRangeNode.cpp
+++ b/Engine/Plugins/Runtime/Metasound/Source/MetasoundStandardNodes/Private/MetasoundMapRangeNode.cpp
@@ -1,88 +1,390 @@
 // Copyright Epic Games, Inc. All Rights Reserved.
 
-#include "MetasoundMapRangeNode.h"
-
 #include "MetasoundNodeRegistrationMacro.h"
+#include "MetasoundAudioBuffer.h"
+#include "CoreMinimal.h"
+#include "MetasoundFacade.h"
+#include "MetasoundExecutableOperator.h"
+#include "MetasoundPrimitives.h"
+#include "MetasoundStandardNodesNames.h"
+#include "MetasoundTrigger.h"
+#include "Internationalization/Text.h"
+#include "MetasoundParamHelper.h"
 
-#define LOCTEXT_NAMESPACE "MetasoundStandardNodes"
+#define LOCTEXT_NAMESPACE "MetasoundStandardNodes_MapRangeNode"
 
 namespace Metasound
 {
+	namespace MapRangeVertexNames
+	{
+		METASOUND_PARAM(InputValueName, "In", "Input value to map.");
+		METASOUND_PARAM(InputInRangeAName, "In Range A", "The min input value range.");
+		METASOUND_PARAM(InputInRangeBName, "In Range B", "The max input value range.");
+		METASOUND_PARAM(InputOutRangeAName, "Out Range A", "The min output value range.");
+		METASOUND_PARAM(InputOutRangeBName, "Out Range B", "The max output value range.");
+		METASOUND_PARAM(InputClampedName, "Clamped", "Whether or not to clamp the input to the specified input range.");
+
+		METASOUND_PARAM(OutputValueName, "Out Value", "Triggers when the envelope is triggered.");
+	}
+
 	namespace MetasoundMapRangeNodePrivate
 	{
-		FNodeClassMetadata CreateNodeClassMetadata(const FName& InDataTypeName, const FName& InOperatorName, const FText& InDisplayName, const FText& InDescription, const FVertexInterface& InDefaultInterface)
+		using namespace MapRangeVertexNames;
+
+		class FIntRange
 		{
-			FNodeClassMetadata Metadata
-			{
-				FNodeClassName{FName("MapRange"), InOperatorName, InDataTypeName},
-				1, // Major Version
-				0, // Minor Version
-				InDisplayName,
-				InDescription,
-				PluginAuthor,
-				PluginNodeMissingPrompt,
-				InDefaultInterface,
-				{LOCTEXT("MapRangeCategory", "Utils")},
-				{TEXT("MapRange")},
-				FNodeDisplayStyle{}
-			};
+		public:
+			static FName GetDataTypeName()
+			{
+				return TEXT("Int32");
+			}
 
-			return Metadata;
-		}
-	}
+			static FText GetNodeName()
+			{
+				return LOCTEXT("MapRange_Int32Name", "Map Range (Int32)");
+			}
 
+			static FVertexInterface GetVertexInterface()
+			{
+				static const FVertexInterface DefaultInterface(
+					FInputVertexInterface(
+						TInputDataVertexModel<int32>(METASOUND_GET_PARAM_NAME_AND_TT(InputValueName)),
+						TInputDataVertexModel<int32>(METASOUND_GET_PARAM_NAME_AND_TT(InputInRangeAName), 0),
+						TInputDataVertexModel<int32>(METASOUND_GET_PARAM_NAME_AND_TT(InputInRangeBName), 100),
+						TInputDataVertexModel<int32>(METASOUND_GET_PARAM_NAME_AND_TT(InputOutRangeAName), 0),
+						TInputDataVertexModel<int32>(METASOUND_GET_PARAM_NAME_AND_TT(InputOutRangeBName), 100),
+						TInputDataVertexModel<bool>(METASOUND_GET_PARAM_NAME_AND_TT(InputClampedName), true)
+					),
+					FOutputVertexInterface(
+						TOutputDataVertexModel<int32>(METASOUND_GET_PARAM_NAME_AND_TT(OutputValueName))
+					)
+				);
+				return DefaultInterface;
+			}
 
-	namespace MapRangeVertexNames
-	{
-		METASOUNDSTANDARDNODES_API const FString& GetInputValueName()
+			FIntRange(const FOperatorSettings& OperatorSettings, const FInputVertexInterface& InputInterface, const FDataReferenceCollection& InputCollection)
+				: Value(InputCollection.GetDataReadReferenceOrConstructWithVertexDefault<int32>(InputInterface, METASOUND_GET_PARAM_NAME(InputValueName), OperatorSettings))
+				, InRangeA(InputCollection.GetDataReadReferenceOrConstructWithVertexDefault<int32>(InputInterface, METASOUND_GET_PARAM_NAME(InputInRangeAName), OperatorSettings))
+				, InRangeB(InputCollection.GetDataReadReferenceOrConstructWithVertexDefault<int32>(InputInterface, METASOUND_GET_PARAM_NAME(InputInRangeBName), OperatorSettings))
+				, OutRangeA(InputCollection.GetDataReadReferenceOrConstructWithVertexDefault<int32>(InputInterface, METASOUND_GET_PARAM_NAME(InputOutRangeAName), OperatorSettings))
+				, OutRangeB(InputCollection.GetDataReadReferenceOrConstructWithVertexDefault<int32>(InputInterface, METASOUND_GET_PARAM_NAME(InputOutRangeBName), OperatorSettings))
+				, bClamped(InputCollection.GetDataReadReferenceOrConstructWithVertexDefault<bool>(InputInterface, METASOUND_GET_PARAM_NAME(InputClampedName), OperatorSettings))
+				, OutputValue(TDataWriteReferenceFactory<int32>::CreateAny(OperatorSettings))
+			{
+			}
+
+			FDataReferenceCollection GetInputs() const
+			{
+				FDataReferenceCollection Inputs;
+				Inputs.AddDataReadReference(METASOUND_GET_PARAM_NAME(InputValueName), Value);
+				Inputs.AddDataReadReference(METASOUND_GET_PARAM_NAME(InputInRangeAName), InRangeA);
+				Inputs.AddDataReadReference(METASOUND_GET_PARAM_NAME(InputInRangeBName), InRangeB);
+				Inputs.AddDataReadReference(METASOUND_GET_PARAM_NAME(InputOutRangeAName), OutRangeA);
+				Inputs.AddDataReadReference(METASOUND_GET_PARAM_NAME(InputOutRangeBName), OutRangeB);
+				Inputs.AddDataReadReference(METASOUND_GET_PARAM_NAME(InputClampedName), bClamped);
+
+				return Inputs;
+			}
+
+			FDataReferenceCollection GetOutputs() const
+			{
+				FDataReferenceCollection Outputs;
+				Outputs.AddDataReadReference(METASOUND_GET_PARAM_NAME(OutputValueName), OutputValue);
+				return Outputs;
+			}
+
+			void DoMapping()
+			{
+				if (*bClamped)
+				{
+					*OutputValue = (int32)FMath::GetMappedRangeValueClamped({ (float)*InRangeA, (float)*InRangeB }, { (float)*OutRangeA, (float)*OutRangeB }, (float)*Value);
+				}
+				else
+				{
+					*OutputValue = (int32)FMath::GetMappedRangeValueUnclamped({ (float)*InRangeA, (float)*InRangeB }, { (float)*OutRangeA, (float)*OutRangeB }, (float)*Value);
+				}
+			}
+
+		private:
+			TDataReadReference<int32> Value;
+			TDataReadReference<int32> InRangeA;
+			TDataReadReference<int32> InRangeB;
+			TDataReadReference<int32> OutRangeA;
+			TDataReadReference<int32> OutRangeB;
+			FBoolReadRef bClamped;
+			TDataWriteReference<int32> OutputValue;
+		};
+
+		class FFloatRange
 		{
-			static const FString Name = TEXT("In");
-			return Name;
-		}
+		public:
+			static FName GetDataTypeName()
+			{
+				return TEXT("Float");
+			}
+
+			static FText GetNodeName()
+			{
+				return LOCTEXT("MapRange_FloatName", "Map Range (Float)");
+			}
 
-		METASOUNDSTANDARDNODES_API const FString& GetInputInRangeAName()
+			static FVertexInterface GetVertexInterface()
+			{
+				static const FVertexInterface DefaultInterface(
+					FInputVertexInterface(
+						TInputDataVertexModel<float>(METASOUND_GET_PARAM_NAME_AND_TT(InputValueName)),
+						TInputDataVertexModel<float>(METASOUND_GET_PARAM_NAME_AND_TT(InputInRangeAName), 0.0f),
+						TInputDataVertexModel<float>(METASOUND_GET_PARAM_NAME_AND_TT(InputInRangeBName), 1.0f),
+						TInputDataVertexModel<float>(METASOUND_GET_PARAM_NAME_AND_TT(InputOutRangeAName), 0.0f),
+						TInputDataVertexModel<float>(METASOUND_GET_PARAM_NAME_AND_TT(InputOutRangeBName), 1.0f),
+						TInputDataVertexModel<bool>(METASOUND_GET_PARAM_NAME_AND_TT(InputClampedName), true)
+					),
+					FOutputVertexInterface(
+						TOutputDataVertexModel<float>(METASOUND_GET_PARAM_NAME_AND_TT(OutputValueName))
+					)
+				);
+				return DefaultInterface;
+			}
+
+			FFloatRange(const FOperatorSettings& OperatorSettings, const FInputVertexInterface& InputInterface, const FDataReferenceCollection& InputCollection)
+				: Value(InputCollection.GetDataReadReferenceOrConstructWithVertexDefault<float>(InputInterface,		METASOUND_GET_PARAM_NAME(InputValueName),		OperatorSettings))
+				, InRangeA(InputCollection.GetDataReadReferenceOrConstructWithVertexDefault<float>(InputInterface,	METASOUND_GET_PARAM_NAME(InputInRangeAName),	OperatorSettings))
+				, InRangeB(InputCollection.GetDataReadReferenceOrConstructWithVertexDefault<float>(InputInterface,	METASOUND_GET_PARAM_NAME(InputInRangeBName),	OperatorSettings))
+				, OutRangeA(InputCollection.GetDataReadReferenceOrConstructWithVertexDefault<float>(InputInterface, METASOUND_GET_PARAM_NAME(InputOutRangeAName),	OperatorSettings))
+				, OutRangeB(InputCollection.GetDataReadReferenceOrConstructWithVertexDefault<float>(InputInterface, METASOUND_GET_PARAM_NAME(InputOutRangeBName),	OperatorSettings))
+				, bClamped(InputCollection.GetDataReadReferenceOrConstructWithVertexDefault<bool>(InputInterface,	METASOUND_GET_PARAM_NAME(InputClampedName),		OperatorSettings))
+				, OutputValue(TDataWriteReferenceFactory<float>::CreateAny(OperatorSettings))
+			{
+			}
+
+			FDataReferenceCollection GetInputs() const
+			{
+				FDataReferenceCollection Inputs;
+				Inputs.AddDataReadReference(METASOUND_GET_PARAM_NAME(InputValueName),Value);
+				Inputs.AddDataReadReference(METASOUND_GET_PARAM_NAME(InputInRangeAName), InRangeA);
+				Inputs.AddDataReadReference(METASOUND_GET_PARAM_NAME(InputInRangeBName), InRangeB);
+				Inputs.AddDataReadReference(METASOUND_GET_PARAM_NAME(InputOutRangeAName), OutRangeA);
+				Inputs.AddDataReadReference(METASOUND_GET_PARAM_NAME(InputOutRangeBName), OutRangeB);
+				Inputs.AddDataReadReference(METASOUND_GET_PARAM_NAME(InputClampedName), bClamped);
+
+				return Inputs;
+			}
+
+			FDataReferenceCollection GetOutputs() const
+			{
+				FDataReferenceCollection Outputs;
+				Outputs.AddDataReadReference(METASOUND_GET_PARAM_NAME(OutputValueName), OutputValue);
+				return Outputs;
+			}
+
+			void DoMapping()
+			{
+				if (*bClamped)
+				{
+					*OutputValue = FMath::GetMappedRangeValueClamped({ *InRangeA, *InRangeB }, { *OutRangeA, *OutRangeB }, *Value);
+				}
+				else
+				{
+					*OutputValue = FMath::GetMappedRangeValueUnclamped({ *InRangeA, *InRangeB }, { *OutRangeA, *OutRangeB }, *Value);
+				}
+			}
+
+		private:
+			TDataReadReference<float> Value;
+			TDataReadReference<float> InRangeA;
+			TDataReadReference<float> InRangeB;
+			TDataReadReference<float> OutRangeA;
+			TDataReadReference<float> OutRangeB;
+			FBoolReadRef bClamped;
+			TDataWriteReference<float> OutputValue;
+		};
+
+		class FAudioRange
 		{
-			static const FString Name = TEXT("In Range A");
-			return Name;
+		public:
+			static FName GetDataTypeName()
+			{
+				return TEXT("Audio");
+			}
+
+			static FText GetNodeName()
+			{
+				return LOCTEXT("MapRange_AudioName", "Map Range (Audio)");
+			}
+
+			static FVertexInterface GetVertexInterface()
+			{
+				static const FVertexInterface DefaultInterface(
+					FInputVertexInterface(
+						TInputDataVertexModel<FAudioBuffer>(METASOUND_GET_PARAM_NAME_AND_TT(InputValueName)),
+						TInputDataVertexModel<float>(METASOUND_GET_PARAM_NAME_AND_TT(InputInRangeAName), -1.0f),
+						TInputDataVertexModel<float>(METASOUND_GET_PARAM_NAME_AND_TT(InputInRangeBName), 1.0f),
+						TInputDataVertexModel<float>(METASOUND_GET_PARAM_NAME_AND_TT(InputOutRangeAName), -1.0f),
+						TInputDataVertexModel<float>(METASOUND_GET_PARAM_NAME_AND_TT(InputOutRangeBName), 1.0f),
+						TInputDataVertexModel<bool>(METASOUND_GET_PARAM_NAME_AND_TT(InputClampedName), true)
+					),
+					FOutputVertexInterface(
+						TOutputDataVertexModel<FAudioBuffer>(METASOUND_GET_PARAM_NAME_AND_TT(OutputValueName))
+					)
+				);
+
+				return DefaultInterface;
+			}
+
+			FAudioRange(const FOperatorSettings& OperatorSettings, const FInputVertexInterface& InputInterface, const FDataReferenceCollection& InputCollection)
+				: Value(InputCollection.GetDataReadReferenceOrConstructWithVertexDefault<FAudioBuffer>(InputInterface, METASOUND_GET_PARAM_NAME(InputValueName), OperatorSettings))
+				, InRangeA(InputCollection.GetDataReadReferenceOrConstructWithVertexDefault<float>(InputInterface, METASOUND_GET_PARAM_NAME(InputInRangeAName), OperatorSettings))
+				, InRangeB(InputCollection.GetDataReadReferenceOrConstructWithVertexDefault<float>(InputInterface, METASOUND_GET_PARAM_NAME(InputInRangeBName), OperatorSettings))
+				, OutRangeA(InputCollection.GetDataReadReferenceOrConstructWithVertexDefault<float>(InputInterface, METASOUND_GET_PARAM_NAME(InputOutRangeAName), OperatorSettings))
+				, OutRangeB(InputCollection.GetDataReadReferenceOrConstructWithVertexDefault<float>(InputInterface, METASOUND_GET_PARAM_NAME(InputOutRangeBName), OperatorSettings))
+				, bClamped(InputCollection.GetDataReadReferenceOrConstructWithVertexDefault<bool>(InputInterface, METASOUND_GET_PARAM_NAME(InputClampedName), OperatorSettings))
+				, OutputValue(TDataWriteReferenceFactory<FAudioBuffer>::CreateAny(OperatorSettings))
+			{
+			}
+
+			FDataReferenceCollection GetInputs() const
+			{
+				FDataReferenceCollection Inputs;
+				Inputs.AddDataReadReference(METASOUND_GET_PARAM_NAME(InputValueName), Value);
+				Inputs.AddDataReadReference(METASOUND_GET_PARAM_NAME(InputInRangeAName), InRangeA);
+				Inputs.AddDataReadReference(METASOUND_GET_PARAM_NAME(InputInRangeBName), InRangeB);
+				Inputs.AddDataReadReference(METASOUND_GET_PARAM_NAME(InputOutRangeAName), OutRangeA);
+				Inputs.AddDataReadReference(METASOUND_GET_PARAM_NAME(InputOutRangeBName), OutRangeB);
+				Inputs.AddDataReadReference(METASOUND_GET_PARAM_NAME(InputClampedName), bClamped);
+
+				return Inputs;
+			}
+
+			FDataReferenceCollection GetOutputs() const
+			{
+				FDataReferenceCollection Outputs;
+				Outputs.AddDataReadReference(METASOUND_GET_PARAM_NAME(OutputValueName), OutputValue);
+				return Outputs;
+			}
+
+			void DoMapping()
+			{
+				FAudioBuffer& OutBuffer = *OutputValue;
+				float* OutBufferPtr = OutBuffer.GetData();
+				const FAudioBuffer& InBuffer = *Value;
+				const float* InBufferPtr = InBuffer.GetData();
+				int32 NumSamples = InBuffer.Num();
+
+				FVector2D InputRange = { *InRangeA, *InRangeB };
+				FVector2D OutputRange = { *OutRangeA, *OutRangeB };
+
+				// TODO: SIMD this
+				for (int32 i = 0; i < NumSamples; ++i)
+				{
+					OutBufferPtr[i] = FMath::GetMappedRangeValueClamped(InputRange, OutputRange, InBufferPtr[i]);
+				}
+			}
+
+		private:
+			TDataReadReference<FAudioBuffer> Value;
+			TDataReadReference<float> InRangeA;
+			TDataReadReference<float> InRangeB;
+			TDataReadReference<float> OutRangeA;
+			TDataReadReference<float> OutRangeB;
+			FBoolReadRef bClamped;
+			TDataWriteReference<FAudioBuffer> OutputValue;
+		};
+	}
+
+	template<typename FMappingClass>
+	class TMapRangeOperator : public TExecutableOperator<TMapRangeOperator<FMappingClass>>
+	{
+	public:
+		static const FNodeClassMetadata& GetNodeInfo()
+		{
+			auto CreateNodeClassMetadata = []() -> FNodeClassMetadata
+			{
+				FName OperatorName = TEXT("MapRange");
+				FText NodeDescription = LOCTEXT("MapRangeDescription", "Maps an input value in the given input range to the given output range.");
+				FVertexInterface NodeInterface = FMappingClass::GetVertexInterface();
+
+				FNodeClassMetadata Metadata
+				{
+					FNodeClassName{FName("MapRange"), OperatorName, FMappingClass::GetDataTypeName()},
+					1, // Major Version
+					0, // Minor Version
+					FMappingClass::GetNodeName(),
+					NodeDescription,
+					PluginAuthor,
+					PluginNodeMissingPrompt,
+					NodeInterface,
+					{LOCTEXT("MapRangeCategory", "Math")},
+					{TEXT("MapRange")},
+					FNodeDisplayStyle{}
+				};
+				return Metadata;
+			};
+
+			static const FNodeClassMetadata Metadata = CreateNodeClassMetadata();
+			return Metadata;
 		}
 
-		METASOUNDSTANDARDNODES_API const FString& GetInputInRangeBName()
+		static TUniquePtr<IOperator> CreateOperator(const FCreateOperatorParams& InParams, TArray<TUniquePtr<IOperatorBuildError>>& OutErrors)
 		{
-			static const FString Name = TEXT("In Range B");
-			return Name;
+			const FInputVertexInterface& InputInterface = InParams.Node.GetVertexInterface().GetInputInterface();
+			const FDataReferenceCollection& InputCollection = InParams.InputDataReferences;
+
+			FMappingClass MappingObject(InParams.OperatorSettings, InputInterface, InputCollection);
+
+			return MakeUnique<TMapRangeOperator<FMappingClass>>(MappingObject);
 		}
 
-		METASOUNDSTANDARDNODES_API const FString& GetInputOutRangeAName()
+
+		TMapRangeOperator(const FMappingClass& InMappingObject)
+			: MappingObject(InMappingObject)
 		{
-			static const FString Name = TEXT("Out Range A");
-			return Name;
+			MappingObject.DoMapping();
 		}
 
-		METASOUNDSTANDARDNODES_API const FString& GetInputOutRangeBName()
+		virtual ~TMapRangeOperator() = default;
+
+		virtual FDataReferenceCollection GetInputs() const override
 		{
-			static const FString Name = TEXT("Out Range B");
-			return Name;
+			return MappingObject.GetInputs();
 		}
 
-		METASOUNDSTANDARDNODES_API const FString& GetInputClampedName()
+		virtual FDataReferenceCollection GetOutputs() const override
 		{
-			static const FString Name = TEXT("Clamped");
-			return Name;
+			return MappingObject.GetOutputs();
 		}
 
-		METASOUNDSTANDARDNODES_API const FString& GetOutputValueName()
+		void Execute()
 		{
-			static const FString Name = TEXT("Value");
-			return Name;
+			MappingObject.DoMapping();
 		}
-	}
 
+	private:
+		FMappingClass MappingObject;
+	};
 
-	using FMapRangeNodeInt32 = TMapRangeNode<int32>;
+	template<typename FMappingClass>
+	class TMapRangeNode : public FNodeFacade
+	{
+	public:
+		/**
+		 * Constructor used by the Metasound Frontend.
+		 */
+		TMapRangeNode(const FNodeInitData& InInitData)
+			: FNodeFacade(InInitData.InstanceName, InInitData.InstanceID, TFacadeOperatorClass<TMapRangeOperator<FMappingClass>>())
+		{}
+
+		virtual ~TMapRangeNode() = default;
+	};
+
+	using FMapRangeNodeInt32 = TMapRangeNode<MetasoundMapRangeNodePrivate::FIntRange>;
  	METASOUND_REGISTER_NODE(FMapRangeNodeInt32)
 
-	using FMapRangeNodeFloat = TMapRangeNode<float>;
+	using FMapRangeNodeFloat = TMapRangeNode<MetasoundMapRangeNodePrivate::FFloatRange>;
 	METASOUND_REGISTER_NODE(FMapRangeNodeFloat)
+
+	using FMapRangeNodeAudio = TMapRangeNode<MetasoundMapRangeNodePrivate::FAudioRange>;
+	METASOUND_REGISTER_NODE(FMapRangeNodeAudio)
+
 }
 
 #undef LOCTEXT_NAMESPACE
diff --git a/Engine/Plugins/Runtime/Metasound/Source/MetasoundStandardNodes/Public/MetasoundMapRangeNode.h b/Engine/Plugins/Runtime/Metasound/Source/MetasoundStandardNodes/Public/MetasoundMapRangeNode.h
deleted file mode 100644
