Sample and hold node

diff --git a/Engine/Plugins/Runtime/Metasound/Source/MetasoundStandardNodes/Private/MetasoundSampleAndHoldNode.cpp b/Engine/Plugins/Runtime/Metasound/Source/MetasoundStandardNodes/Private/MetasoundSampleAndHoldNode.cpp
new file mode 100644
--- /dev/null
+++ b/Engine/Plugins/Runtime/Metasound/Source/MetasoundStandardNodes/Private/MetasoundSampleAndHoldNode.cpp
@@ -0,0 +1,183 @@
+// Copyright Epic Games, Inc. All Rights Reserved.
+
+#include "Internationalization/Text.h"
+#include "MetasoundExecutableOperator.h"
+#include "MetasoundNodeRegistrationMacro.h"
+#include "MetasoundPrimitives.h"
+#include "MetasoundStandardNodesNames.h"
+#include "MetasoundTrigger.h"
+#include "MetasoundTime.h"
+#include "MetasoundAudioBuffer.h"
+#include "DSP/Delay.h"
+#include "MetasoundStandardNodesCategories.h"
+#include "MetasoundFacade.h"
+#include "MetasoundParamHelper.h"
+
+#define LOCTEXT_NAMESPACE "MetasoundStandardNodes_SampleAndHold"
+
+namespace Metasound
+{
+	namespace SampleAndHoldVertexNames
+	{
+		METASOUND_PARAM(InputTriggerSampleAndHold, "Sample And Hold", "Trigger to sample and hold the input audio.");
+		METASOUND_PARAM(InputAudio, "In", "The audio input to sample.");
+
+		METASOUND_PARAM(OutputOnSampleAndHold, "On Sample And Hold", "Triggers when the input sample and hold is triggered.");
+		METASOUND_PARAM(OutputAudio, "Out", "Sampled output value.");
+	}
+
+	class FSampleAndHoldOperator : public TExecutableOperator<FSampleAndHoldOperator>
+	{
+	public:
+
+		static const FNodeClassMetadata& GetNodeInfo()
+		{
+			auto CreateNodeClassMetadata = []() -> FNodeClassMetadata
+			{
+				FVertexInterface NodeInterface = GetVertexInterface();
+
+				FNodeClassMetadata Metadata
+				{
+					FNodeClassName{FName("SampleAndHold"), TEXT("Sample And Hold"), TEXT("")},
+					1, // Major Version
+					0, // Minor Version
+					LOCTEXT("SampleAndHoldDisplayName", "Sample And Hold"),
+					LOCTEXT("SampleAndHoldDesc", "Will output a single value of the input audio signal when triggered."),
+					PluginAuthor,
+					PluginNodeMissingPrompt,
+					NodeInterface,
+					{StandardNodes::Filters},
+					{TEXT("SampleAndHold")},
+					FNodeDisplayStyle{}
+				};
+
+				return Metadata;
+			};
+
+			static const FNodeClassMetadata Metadata = CreateNodeClassMetadata();
+			return Metadata;
+		}
+
+		static const FVertexInterface& GetVertexInterface()
+		{
+			using namespace SampleAndHoldVertexNames;
+
+			static const FVertexInterface Interface(
+				FInputVertexInterface(
+					TInputDataVertexModel<FTrigger>(METASOUND_GET_PARAM_NAME_AND_TT(InputTriggerSampleAndHold)),
+					TInputDataVertexModel<FAudioBuffer>(METASOUND_GET_PARAM_NAME_AND_TT(InputAudio))
+				),
+				FOutputVertexInterface(
+					TOutputDataVertexModel<FTrigger>(METASOUND_GET_PARAM_NAME_AND_TT(OutputOnSampleAndHold)),
+					TOutputDataVertexModel<FAudioBuffer>(METASOUND_GET_PARAM_NAME_AND_TT(OutputAudio))
+				)
+			);
+
+			return Interface;
+		}
+		
+		static TUniquePtr<IOperator> CreateOperator(const FCreateOperatorParams& InParams, FBuildErrorArray& OutErrors)
+		{
+			using namespace SampleAndHoldVertexNames;
+			const FDataReferenceCollection& Inputs = InParams.InputDataReferences;
+
+			FTriggerReadRef InputTriggerSampleAndHold = Inputs.GetDataReadReferenceOrConstruct<FTrigger>(METASOUND_GET_PARAM_NAME(InputTriggerSampleAndHold), InParams.OperatorSettings);
+			FAudioBufferReadRef InputAudio = Inputs.GetDataReadReferenceOrConstruct<FAudioBuffer>(METASOUND_GET_PARAM_NAME(InputAudio), InParams.OperatorSettings);
+
+			return MakeUnique<FSampleAndHoldOperator>(InParams.OperatorSettings, InputTriggerSampleAndHold, InputAudio);
+		}
+
+		FSampleAndHoldOperator(const FOperatorSettings& InSettings, const FTriggerReadRef& InTriggerSampleAndHold, const FAudioBufferReadRef& InAudioInput)
+			: AudioInput(InAudioInput)
+			, TriggerSampleAndHold(InTriggerSampleAndHold)
+			, OnTriggerSampleAndHold(FTriggerWriteRef::CreateNew(InSettings))
+			, AudioOutput(FAudioBufferWriteRef::CreateNew(InSettings))
+		{
+			// Init the hold value to the first sample in the audio buffer
+			HoldValue = AudioInput->GetData()[0];
+		}
+
+		virtual FDataReferenceCollection GetInputs() const override
+		{
+			using namespace SampleAndHoldVertexNames;
+
+			FDataReferenceCollection InputDataReferences;
+			InputDataReferences.AddDataReadReference(METASOUND_GET_PARAM_NAME(InputAudio), AudioInput);
+			InputDataReferences.AddDataReadReference(METASOUND_GET_PARAM_NAME(InputTriggerSampleAndHold), TriggerSampleAndHold);
+
+			return InputDataReferences;
+
+		}
+
+		virtual FDataReferenceCollection GetOutputs() const override
+		{
+			using namespace SampleAndHoldVertexNames;
+			FDataReferenceCollection OutputDataReferences;
+			OutputDataReferences.AddDataReadReference(METASOUND_GET_PARAM_NAME(OutputOnSampleAndHold), OnTriggerSampleAndHold);
+			OutputDataReferences.AddDataReadReference(METASOUND_GET_PARAM_NAME(OutputAudio), AudioOutput);
+			return OutputDataReferences;
+		}
+		
+		void Execute()
+		{
+			OnTriggerSampleAndHold->AdvanceBlock();
+
+			// Lambda is used for both pre and post trigger frames
+			auto SetToHoldLambda = [this](int32 StartFrame, int32 EndFrame)
+			{
+				const int32 NumSamples = EndFrame - StartFrame;
+				const int32 NonSIMDCount = NumSamples % AUDIO_SIMD_FLOAT_ALIGNMENT;
+				const int32 SIMDCount = NumSamples - NonSIMDCount;
+
+				float* OutputBufferPtr = AudioOutput->GetData();
+
+				float* HoldStartPtr = &OutputBufferPtr[StartFrame];
+				Audio::BufferSetToConstantInplace(HoldStartPtr, SIMDCount, HoldValue);
+
+				for (int32 i = StartFrame + SIMDCount; i < EndFrame; ++i)
+				{
+					OutputBufferPtr[i] = HoldValue;
+				}
+			};
+
+			TriggerSampleAndHold->ExecuteBlock(
+				// On Pre-Trigger, continue outputting the held value
+				SetToHoldLambda, 
+				// On Trigger, get a new hold value, then output for the rest of the block
+				[this, &SetToHoldLambda](int32 StartFrame, int32 EndFrame)
+				{
+					// Get a new value to sample and hold
+					HoldValue = AudioInput->GetData()[StartFrame];
+					SetToHoldLambda(StartFrame, EndFrame);
+
+					// Output that we were triggered
+					OnTriggerSampleAndHold->TriggerFrame(StartFrame);
+				}
+			);
+		}
+
+	private:
+
+		FAudioBufferReadRef AudioInput;
+		FTriggerReadRef TriggerSampleAndHold;
+
+		FTriggerWriteRef OnTriggerSampleAndHold;
+		FAudioBufferWriteRef AudioOutput;
+
+		float HoldValue = 0.0f;
+	};
+
+	class FSampleAndHoldNode : public FNodeFacade
+	{
+	public:
+		FSampleAndHoldNode(const FNodeInitData& InitData)
+			: FNodeFacade(InitData.InstanceName, InitData.InstanceID, TFacadeOperatorClass<FSampleAndHoldOperator>())
+		{
+		}
+	};
+
+
+	METASOUND_REGISTER_NODE(FSampleAndHoldNode)
+}
+
+#undef LOCTEXT_NAMESPACE
