Templatize Audio mixer nodes

diff --git a/Engine/Plugins/Runtime/Metasound/Source/MetasoundStandardNodes/Private/MetasoundMixerNode.cpp b/Engine/Plugins/Runtime/Metasound/Source/MetasoundStandardNodes/Private/MetasoundMixerNode.cpp
--- a/Engine/Plugins/Runtime/Metasound/Source/MetasoundStandardNodes/Private/MetasoundMixerNode.cpp
+++ b/Engine/Plugins/Runtime/Metasound/Source/MetasoundStandardNodes/Private/MetasoundMixerNode.cpp
@@ -1,7 +1,7 @@
 // Copyright Epic Games, Inc. All Rights Reserved.
-#include "MetasoundMixerNode.h"
-
 #include "DSP/Dsp.h"
+#include "DSP/BufferVectorOperations.h"
+
 #include "MetasoundAudioBuffer.h"
 #include "MetasoundExecutableOperator.h"
 #include "MetasoundFacade.h"
@@ -11,166 +11,325 @@
 #include "MetasoundStandardNodesNames.h"
 #include "MetasoundTrigger.h"
 #include "MetasoundVertex.h"
+#include "MetasoundStandardNodesCategories.h"
 
-#define LOCTEXT_NAMESPACE "MetasoundStandardNodes"
+#define LOCTEXT_NAMESPACE "MetasoundStandardNodes_MixerNode"
 
 
 namespace Metasound
 {
-	class FMixerOperator : public TExecutableOperator<FMixerOperator>
+# pragma region Operator Declaration
+	template<uint32 NumInputs, uint32 NumChannels>
+	class TAudioMixerNodeOperator : public TExecutableOperator<TAudioMixerNodeOperator<NumInputs, NumChannels>>
 	{
 	public:
-		static TUniquePtr<IOperator> CreateOperator(const FCreateOperatorParams& InParams, FBuildErrorArray& OutErrors);
-		static const FNodeClassMetadata& GetNodeInfo();
-		static const FVertexInterface& GetVertexInterface();	
+		// ctor
+		TAudioMixerNodeOperator(const FOperatorSettings& InSettings, const TArray<FAudioBufferReadRef>&& InInputBuffers, const TArray<FFloatReadRef>&& InGainValues)
+			: Gains(InGainValues)
+			, Inputs (InInputBuffers)
+			, Settings(InSettings)
+		{
+			// create write refs
+			for (uint32 i = 0; i < NumChannels; ++i)
+			{
+				Outputs.Add(FAudioBufferWriteRef::CreateNew(InSettings));
+			}
+
+			// init previous gains to current values
+			PrevGains.Reset();
+			PrevGains.AddUninitialized(NumInputs);
+			for (uint32 i = 0; i < NumInputs; ++i)
+			{
+				PrevGains[i] = *Gains[i];
+			}
+		}
+
+		// dtor
+		virtual ~TAudioMixerNodeOperator() = default;
+
+		static const FVertexInterface& GetDefaultInterface()
+		{
+			auto CreateDefaultInterface = []()-> FVertexInterface
+			{
+				// inputs
+				FInputVertexInterface InputInterface;
+				for (uint32 InputIndex = 0; InputIndex < NumInputs; ++InputIndex)
+				{
+					// audio channels
+					for (uint32 ChanIndex = 0; ChanIndex < NumChannels; ++ChanIndex)
+					{
+						InputInterface.Add(TInputDataVertexModel<FAudioBuffer>(GetAudioInputName(InputIndex, ChanIndex), GetAudioInputDescription(InputIndex, ChanIndex)));
+					}
 
-		FMixerOperator(const FOperatorSettings& InSettings, TArray<FAudioBufferReadRef>&& InBuffers, TArray<FFloatReadRef>&& InGains);
+					// gain scalar
+					InputInterface.Add(TInputDataVertexModel<float>(GetGainInputName(InputIndex), GetGainInputDescription(InputIndex)));
+				}
 
-		virtual FDataReferenceCollection GetInputs() const override;
-		virtual FDataReferenceCollection GetOutputs() const override;
+				// outputs
+				FOutputVertexInterface OutputInterface;
+				for (uint32 i = 0; i < NumChannels; ++i)
+				{
+					OutputInterface.Add(TOutputDataVertexModel<FAudioBuffer>(GetAudioOutputName(i), GetAudioOutputDescription(i)));
+				}
 
-		void Execute();
+				return FVertexInterface(InputInterface, OutputInterface);
+			}; // end lambda: CreateDefaultInterface()
 
-	private:		
-		// TODO: Remove when we get dynamic pin support. (if you update this, update GetVertexInterface also).
-		static constexpr const int32 NumHardcodedInputs = 4;
+			static const FVertexInterface DefaultInterface = CreateDefaultInterface();
+			return DefaultInterface;
+		}
 
-		static FString GetInputAudioPinName(int32 InPinIndex)
+		static const FNodeClassMetadata& GetNodeInfo()
 		{
-			return FString::Printf(TEXT("In %d"), InPinIndex + 1);
+			// used if NumChannels == 1
+			auto CreateNodeClassMetadataMono = []() -> FNodeClassMetadata
+			{
+				FName OperatorName = *FString::Printf(TEXT("Audio Mixer (Mono, %d)"), NumInputs);
+				FText NodeDisplayName = FText::Format(LOCTEXT("AudioMixerDisplayNamePattern", "Mono Mixer ({0})"), NumInputs);
+				FText NodeDescription = LOCTEXT("MixerDescription", "Will scale input channels by their corresponding gain value and sum them together.");
+				FVertexInterface NodeInterface = GetDefaultInterface();
+
+				return CreateNodeClassMetadata(OperatorName, NodeDisplayName, NodeDescription, NodeInterface);
+			};
+
+			// used if NumChannels == 2
+			auto CreateNodeClassMetadataStereo = []() -> FNodeClassMetadata
+			{
+				FName OperatorName = *FString::Printf(TEXT("Audio Mixer (Stereo, %d)"), NumInputs);
+				FText NodeDisplayName = FText::Format(LOCTEXT("AudioMixerDisplayNamePattern", "Stereo Mixer ({0})"), NumInputs);
+				FText NodeDescription = LOCTEXT("MixerDescription", "Will scale input channels by their corresponding gain value and sum them together.");
+				FVertexInterface NodeInterface = GetDefaultInterface();
+
+				return  CreateNodeClassMetadata(OperatorName, NodeDisplayName, NodeDescription, NodeInterface);
+			};
+
+			// used if NumChannels > 2
+			auto CreateNodeClassMetadataMultiChan = []() -> FNodeClassMetadata
+			{
+				FName OperatorName = *FString::Printf(TEXT("Audio Mixer (%d-Channel, %d)"), NumChannels, NumInputs);
+				FText NodeDisplayName = FText::Format(LOCTEXT("AudioMixerDisplayNamePattern", "{0}-channel Mixer ({1})"), NumChannels, NumInputs);
+				FText NodeDescription = LOCTEXT("MixerDescription", "Will scale input audio by their corresponding gain value and sum them together.");
+				FVertexInterface NodeInterface = GetDefaultInterface();
+
+				return  CreateNodeClassMetadata(OperatorName, NodeDisplayName, NodeDescription, NodeInterface);
+			};
+
+			static const FNodeClassMetadata Metadata = (NumChannels == 1)? CreateNodeClassMetadataMono()
+														: (NumChannels == 2)? CreateNodeClassMetadataStereo() : CreateNodeClassMetadataMultiChan();
+			return Metadata;
 		}
-		static FString GetInputGainPinName(int32 InPinIndex)
+
+		static TUniquePtr<IOperator> CreateOperator(const FCreateOperatorParams& InParams, TArray<TUniquePtr<IOperatorBuildError>>& OutErrors)
 		{
-			return FString::Printf(TEXT("Gain %d"), InPinIndex + 1);
+			const FInputVertexInterface& InputInterface = InParams.Node.GetVertexInterface().GetInputInterface();
+			const FDataReferenceCollection& InputCollection = InParams.InputDataReferences;
+
+			TArray<FAudioBufferReadRef> InputBuffers;
+			TArray<FFloatReadRef> InputGains;
+
+			for (uint32 i = 0; i < NumInputs; ++i)
+			{
+				for (uint32 Chan = 0; Chan < NumChannels; ++Chan)
+				{
+					InputBuffers.Add(InputCollection.GetDataReadReferenceOrConstruct<FAudioBuffer>(GetAudioInputName(i, Chan), InParams.OperatorSettings));
+				}
+
+				InputGains.Add(InputCollection.GetDataReadReferenceOrConstruct<float>(GetGainInputName(i)));
+			}
+
+			return MakeUnique<TAudioMixerNodeOperator<NumInputs, NumChannels>>(InParams.OperatorSettings, MoveTemp(InputBuffers), MoveTemp(InputGains));
 		}
 
-		TArray<FAudioBufferReadRef> AudioBuffersIn;
-		TArray<FFloatReadRef> BufferMixGains;
-		FAudioBufferWriteRef AudioBufferOut;
-		
-		static constexpr const TCHAR* AudioOutPinName = TEXT("Out");
-	};
+		virtual FDataReferenceCollection GetInputs() const override
+		{
+			FDataReferenceCollection InputPins;
+			for (uint32 i = 0; i < NumInputs; ++i)
+			{
+				for (uint32 Chan = 0; Chan < NumChannels; ++Chan)
+				{
+					InputPins.AddDataReadReference(GetAudioInputName(i, Chan), Inputs[i * NumChannels + Chan]);
+				}
 
-	FMixerOperator::FMixerOperator(const FOperatorSettings& InSettings, TArray<FAudioBufferReadRef>&& InBuffers, TArray<FFloatReadRef>&& InGains)
-		: AudioBuffersIn(MoveTemp(InBuffers))
-		, BufferMixGains(MoveTemp(InGains))
-		, AudioBufferOut(FAudioBufferWriteRef::CreateNew(InSettings))
-	{
-		check(AudioBufferOut->Num() == InSettings.GetNumFramesPerBlock());
-		check(BufferMixGains.Num() == AudioBuffersIn.Num());
-	}
+				InputPins.AddDataReadReference(GetGainInputName(i), Gains[i]);
+			}
 
-	FDataReferenceCollection FMixerOperator::GetInputs() const
-	{
-		FDataReferenceCollection InputDataReferences;
-		for (int32 i = 0; i < AudioBuffersIn.Num(); ++i)
+			return InputPins;
+		}
+
+		virtual FDataReferenceCollection GetOutputs() const override
 		{
-			InputDataReferences.AddDataReadReference(GetInputAudioPinName(i), AudioBuffersIn[i]);
-			InputDataReferences.AddDataReadReference(GetInputGainPinName(i), BufferMixGains[i]);
+			FDataReferenceCollection OutputPins;
+
+			for (uint32 i = 0; i < NumChannels; ++i)
+			{
+				OutputPins.AddDataReadReference(GetAudioOutputName(i), Outputs[i]);
+			}
+
+			return OutputPins;
 		}
-		return InputDataReferences;
-	}
 
-	FDataReferenceCollection FMixerOperator::GetOutputs() const
-	{
-		FDataReferenceCollection OutputDataReferences;
-		OutputDataReferences.AddDataReadReference(AudioOutPinName, FAudioBufferReadRef(AudioBufferOut));
-		return OutputDataReferences;
-	}
-
-	void FMixerOperator::Execute()
-	{	
-		AudioBufferOut->Zero();
-		for (int32 i = 0; i < AudioBuffersIn.Num(); ++i)
+		void Execute()
 		{
-			Audio::MixInBufferFast(AudioBuffersIn[i]->GetData(), AudioBufferOut->GetData(), AudioBufferOut->Num(), *BufferMixGains[i]);
+			// zero the outputs
+			for (uint32 i = 0; i < NumChannels; ++i)
+			{
+				FMemory::Memzero(Outputs[i]->GetData(), sizeof(float) * Outputs[i]->Num());
+			}
+
+			// for each input
+			for (uint32 InputIndex = 0; InputIndex < NumInputs; ++InputIndex)
+			{
+				// for each channel of audio
+				for (uint32 ChanIndex = 0; ChanIndex < NumChannels; ++ChanIndex)
+				{
+					// Outputs[Chan] += Gains[i] * Inputs[i][Chan]
+					const float* InputPtr = Inputs[InputIndex * NumChannels + ChanIndex]->GetData();
+					const float NextGain = *Gains[InputIndex];
+					const float PrevGain = PrevGains[InputIndex];
+					float* OutputPtr = Outputs[ChanIndex]->GetData();
+
+					Audio::MixInBufferFast(InputPtr,  OutputPtr, Settings.GetNumFramesPerBlock() , PrevGain, NextGain);
+
+					PrevGains[InputIndex] = NextGain;
+				}
+			}
 		}
-	}
 
-	const FVertexInterface& FMixerOperator::GetVertexInterface()
-	{
-		static const FVertexInterface Interface(
-			FInputVertexInterface(
-				// 1
-				TInputDataVertexModel<FAudioBuffer>(GetInputAudioPinName(0), LOCTEXT("MixerInputDescription1", "Audio Input 1 of the Mixer")),		// Can't have dynamic pins yet, so just hard code 4.
-				TInputDataVertexModel<float>(GetInputGainPinName(0), LOCTEXT("MixerGainDescription1", ""), 1.0f),
-
-				// 2
-				TInputDataVertexModel<FAudioBuffer>(GetInputAudioPinName(1), LOCTEXT("MixerInputDescription2", "Audio Input 2 of the Mixer")),
-				TInputDataVertexModel<float>(GetInputGainPinName(1), LOCTEXT("MixerGainDescription2", ""), 1.0f),
-
-				// 3
-				TInputDataVertexModel<FAudioBuffer>(GetInputAudioPinName(2), LOCTEXT("MixerInputDescription3", "Audio Input 3 of the Mixer")),
-				TInputDataVertexModel<float>(GetInputGainPinName(2), LOCTEXT("MixerGainDescription3", ""), 1.0f),
-
-				// 4
-				TInputDataVertexModel<FAudioBuffer>(GetInputAudioPinName(3), LOCTEXT("MixerInputDescription3", "Audio Input 4 of the Mixer")),	
-				TInputDataVertexModel<float>(GetInputGainPinName(3), LOCTEXT("MixerGainDescription4", ""), 1.0f)
-			),
-			FOutputVertexInterface(
-				TOutputDataVertexModel<FAudioBuffer>(AudioOutPinName, LOCTEXT("AudioOutTooltip", "Audio Ouput from the mixer"))
-			)
-		);
-
-		return Interface;
-	}
-
-	const FNodeClassMetadata& FMixerOperator::GetNodeInfo()
-	{
-		auto InitNodeInfo = []() -> FNodeClassMetadata
+
+	private:
+		TArray<FFloatReadRef> Gains;
+		TArray<FAudioBufferReadRef> Inputs;
+		TArray<FAudioBufferWriteRef> Outputs;
+
+		TArray<float> PrevGains;
+
+		FOperatorSettings Settings;
+
+		static FNodeClassMetadata CreateNodeClassMetadata(const FName& InOperatorName, const FText& InDisplayName, const FText& InDescription, const FVertexInterface& InDefaultInterface)
 		{
-			FNodeClassMetadata Info;
-			Info.ClassName = { Metasound::StandardNodes::Namespace, TEXT("Mixer"), Metasound::StandardNodes::AudioVariant };
-			Info.MajorVersion = 1;
-			Info.MinorVersion = 0;
-			Info.DisplayName = LOCTEXT("Metasound_MixerNodeDisplayName", "Mixer");
-			Info.Description = LOCTEXT("Metasound_MixerNodeDescription", "Mixes input audio together");
-			Info.Author = PluginAuthor;
-			Info.PromptIfMissing = PluginNodeMissingPrompt;
-			Info.DefaultInterface = GetVertexInterface();
-
-			return Info;
-		};
-
-		static const FNodeClassMetadata Info = InitNodeInfo();
-
-		return Info;
-	}
-
-	FMixerNode::FMixerNode(const FNodeInitData& InInitData)
-		: FNodeFacade(InInitData.InstanceName, InInitData.InstanceID, TFacadeOperatorClass<FMixerOperator>())
-	{
-	}
+			FNodeClassMetadata Metadata
+			{
+				FNodeClassName{FName("AudioMixer"), InOperatorName, TEXT("")},
+				1, // Major Version
+				0, // Minor Version
+				InDisplayName,
+				InDescription,
+				PluginAuthor,
+				PluginNodeMissingPrompt,
+				InDefaultInterface,
+				{ StandardNodes::Audio },
+				{TEXT("AudioMixer")},
+				FNodeDisplayStyle{}
+			};
 
-	TUniquePtr<IOperator> FMixerOperator::CreateOperator(const FCreateOperatorParams& InParams, FBuildErrorArray& OutErrors)
-	{
-		const FMixerNode& Node = static_cast<const FMixerNode&>(InParams.Node);
-		const FDataReferenceCollection& InputCol = InParams.InputDataReferences;
-		const FOperatorSettings& Settings = InParams.OperatorSettings;
-		const FInputVertexInterface& InputInterface = GetVertexInterface().GetInputInterface();
-
-		// TODO: Remove this string manipulation when we get dynamic pin support.
-		TArray<FAudioBufferReadRef> InputBuffers;
-		TArray<FFloatReadRef> InputGains;
-		for (int32 i = 0; i < NumHardcodedInputs; ++i)
+			return Metadata;
+		}
+
+#pragma region Name Gen
+		static const FString GetAudioInputName(uint32 InputIndex, uint32 ChannelIndex)
+		{
+			if (NumChannels == 1)
+			{
+				return FText::Format(LOCTEXT("AudioMixerAudioInputName", "In {0}"), InputIndex).ToString();
+			}
+			else if (NumChannels == 2)
+			{
+				return FString::Printf(TEXT("%s %s"), *FText::Format(LOCTEXT("AudioMixerAudioInputName", "In {0}"), InputIndex).ToString(), (ChannelIndex == 0) ? "L" : "R");
+			}
+
+			return FText::Format(LOCTEXT("AudioMixerAudioInputName", "In {0}, {1}"), InputIndex, ChannelIndex).ToString();
+		}
+
+
+		static const FText GetAudioInputDescription(uint32 InputIndex, uint32 ChannelIndex)
+		{
+			return FText::Format(LOCTEXT("AudioMixerAudioInputDescription", "Audio Input #: {0}, Channel: {1}"), InputIndex, ChannelIndex);
+		}
+
+
+		static const FString GetGainInputName(uint32 InputIndex)
+		{
+			return FText::Format(LOCTEXT("AudioMixerGainInputName", "Gain {0}"), InputIndex).ToString();
+		}
+
+		static const FText GetGainInputDescription(uint32 InputIndex)
+		{
+			return FText::Format(LOCTEXT("AudioMixerGainInputDescription", "Gain Input #: {0}"), InputIndex);
+		}
+
+		static const FString GetAudioOutputName(uint32 ChannelIndex)
 		{
-			const FString AudioPinName = GetInputAudioPinName(i);
-			const FString GainPinName = GetInputGainPinName(i);
-			if (InputCol.ContainsDataReadReference<FAudioBuffer>(AudioPinName))
-			{
-				// Only Create buffers if there's something connected to it.
-				InputBuffers.Emplace(InputCol.GetDataReadReference<FAudioBuffer>(AudioPinName));
-				
-				// Make sure for every valid/connected pin, we have a corresponding gain, even if its defaulted.
-				InputGains.Emplace(InputCol.GetDataReadReferenceOrConstructWithVertexDefault<float>(InputInterface, GainPinName, InParams.OperatorSettings));
+			if (NumChannels == 1)
+			{
+				return LOCTEXT("AudioMixerAudioOUtputName", "Out").ToString();
 			}
+			else if (NumChannels == 2)
+			{
+				return FString::Printf(TEXT("%s %s"), *LOCTEXT("AudioMixerAudioOutputName", "Out").ToString(), (ChannelIndex == 0) ? TEXT("L") : TEXT("R"));
+			}
+
+			return FText::Format(LOCTEXT("AudioMixerAudioOUtputName", "Out {0}"), ChannelIndex).ToString();
 		}
-		
-		// Make a node even if we don't have any inputs to make sure we don't fail graph compilation.
-		return MakeUnique<FMixerOperator>(InParams.OperatorSettings, MoveTemp(InputBuffers), MoveTemp(InputGains));
-	}
 
-	METASOUND_REGISTER_NODE(FMixerNode);
-}
-#undef LOCTEXT_NAMESPACE //MetasoundMixerNode
+		static const FText GetAudioOutputDescription(uint32 ChannelIndex)
+		{
+			return FText::Format(LOCTEXT("AudioMixerAudioOUtputName", "Summed output for channel: {0}"), ChannelIndex);
+		}
+#pragma endregion
+	}; // class TAudioMixerNodeOperator
+#pragma endregion
+
+
+
+#pragma region Node Definition
+	template<uint32 NumInputs, uint32 NumChannels>
+	class METASOUNDSTANDARDNODES_API TAudioMixerNode : public FNodeFacade
+	{
+	public:
+		/**
+		 * Constructor used by the Metasound Frontend.
+		 */
+		TAudioMixerNode(const FNodeInitData& InInitData)
+			: FNodeFacade(InInitData.InstanceName, InInitData.InstanceID, TFacadeOperatorClass<TAudioMixerNodeOperator<NumInputs, NumChannels>>())
+		{}
+
+		virtual ~TAudioMixerNode() = default;
+	};
+#pragma endregion
+
+
+#pragma region Node Registration
+	#define REGISTER_AUDIOMIXER_NODE(A, B) \
+		using FAudioMixerNode_##A ## _ ##B = TAudioMixerNode<A, B>; \
+		METASOUND_REGISTER_NODE(FAudioMixerNode_##A ## _ ##B) \
+
+
+	// mono
+	REGISTER_AUDIOMIXER_NODE(2, 1)
+	REGISTER_AUDIOMIXER_NODE(3, 1)
+	REGISTER_AUDIOMIXER_NODE(4, 1)
+	REGISTER_AUDIOMIXER_NODE(5, 1)
+	REGISTER_AUDIOMIXER_NODE(6, 1)
+	REGISTER_AUDIOMIXER_NODE(7, 1)
+	REGISTER_AUDIOMIXER_NODE(8, 1)
+
+	// stereo
+ 	REGISTER_AUDIOMIXER_NODE(2, 2)
+	REGISTER_AUDIOMIXER_NODE(3, 2)
+	REGISTER_AUDIOMIXER_NODE(4, 2)
+	REGISTER_AUDIOMIXER_NODE(5, 2)
+	REGISTER_AUDIOMIXER_NODE(6, 2)
+	REGISTER_AUDIOMIXER_NODE(7, 2)
+	REGISTER_AUDIOMIXER_NODE(8, 2)
+
+	// test
+//	REGISTER_AUDIOMIXER_NODE(8, 6)
+
+#pragma endregion
+
+
+
+
+} // namespace Metasound
 
+#undef LOCTEXT_NAMESPACE // "MetasoundStandardNodes_MixerNode"
\ No newline at end of file
diff --git a/Engine/Plugins/Runtime/Metasound/Source/MetasoundStandardNodes/Private/MetasoundStandardNodesCategories.cpp b/Engine/Plugins/Runtime/Metasound/Source/MetasoundStandardNodes/Private/MetasoundStandardNodesCategories.cpp
--- a/Engine/Plugins/Runtime/Metasound/Source/MetasoundStandardNodes/Private/MetasoundStandardNodesCategories.cpp
+++ b/Engine/Plugins/Runtime/Metasound/Source/MetasoundStandardNodes/Private/MetasoundStandardNodesCategories.cpp
@@ -13,6 +13,7 @@ namespace Metasound
 		const FText TriggerUtils = { LOCTEXT("Metasound_TriggerUtilsCategory", "Triggers") };
 		const FText Music = { LOCTEXT("Metasound_MusicCategory", "Music") };
 		const FText RandomUtils = { LOCTEXT("Metasound_RandomCategory", "Random") };
+		const FText Audio = { LOCTEXT("Metasound_AudioCategory", "Audio") };
 		const FText Conversions  = { LOCTEXT("Metasound_ConversionsCategory", "Conversions") };
 		const FText DebugUtils = { LOCTEXT("Metasound_DebugUtilsCategory", "Debug") };
 		const FText Io = { LOCTEXT("Metasound_IoCategory", "External IO") };
diff --git a/Engine/Plugins/Runtime/Metasound/Source/MetasoundStandardNodes/Public/MetasoundMixerNode.h b/Engine/Plugins/Runtime/Metasound/Source/MetasoundStandardNodes/Public/MetasoundMixerNode.h
deleted file mode 100644
diff --git a/Engine/Plugins/Runtime/Metasound/Source/MetasoundStandardNodes/Public/MetasoundStandardNodesCategories.h b/Engine/Plugins/Runtime/Metasound/Source/MetasoundStandardNodes/Public/MetasoundStandardNodesCategories.h
--- a/Engine/Plugins/Runtime/Metasound/Source/MetasoundStandardNodes/Public/MetasoundStandardNodesCategories.h
+++ b/Engine/Plugins/Runtime/Metasound/Source/MetasoundStandardNodes/Public/MetasoundStandardNodesCategories.h
@@ -12,6 +12,7 @@ namespace Metasound
 		extern const FText Filters;
 		extern const FText TriggerUtils;
 		extern const FText Music;
+		extern const FText Audio;
 		extern const FText RandomUtils;
 		extern const FText Conversions;
 		extern const FText DebugUtils;
