Fix ensure when proxy attempting to be created from null UObject

diff --git a/Engine/Plugins/Runtime/Metasound/Source/MetasoundFrontend/Private/MetasoundFrontendRegistries.cpp b/Engine/Plugins/Runtime/Metasound/Source/MetasoundFrontend/Private/MetasoundFrontendRegistries.cpp
--- a/Engine/Plugins/Runtime/Metasound/Source/MetasoundFrontend/Private/MetasoundFrontendRegistries.cpp
+++ b/Engine/Plugins/Runtime/Metasound/Source/MetasoundFrontend/Private/MetasoundFrontendRegistries.cpp
@@ -344,20 +344,17 @@ namespace Metasound
 			{
 				if (ensureAlwaysMsgf(DataTypeRegistry.Contains(InDataType), TEXT("Couldn't find data type %s!"), *InDataType.ToString()))
 				{
-					 Audio::IProxyDataPtr ProxyPtr = DataTypeRegistry[InDataType].Callbacks.CreateAudioProxy(InObject);
-					 if (ensureAlwaysMsgf(ProxyPtr.IsValid(), TEXT("UObject failed to create a valid proxy!")))
-					 {
-						 return Metasound::FLiteral(MoveTemp(ProxyPtr));
-					 }
-					 else
-					 {
-						 return Metasound::FLiteral();
-					 }
-				}
-				else
-				{
-					return Metasound::FLiteral();
+					if (Audio::IProxyDataPtr ProxyPtr = DataTypeRegistry[InDataType].Callbacks.CreateAudioProxy(InObject))
+					{
+						if (InObject)
+						{
+							ensureAlwaysMsgf(ProxyPtr.IsValid(), TEXT("UObject failed to create a valid proxy!"));
+						}
+						return Metasound::FLiteral(MoveTemp(ProxyPtr));
+					}
 				}
+
+				return Metasound::FLiteral();
 			}
 
 			Metasound::FLiteral FRegistryContainerImpl::GenerateLiteralForUObjectArray(const FName& InDataType, TArray<UObject*> InObjectArray)
