Only allow exact matches of classes for Metasound proxies in editor

diff --git a/Engine/Plugins/Runtime/Metasound/Source/MetasoundEditor/Private/MetasoundNodeDetailCustomization.cpp b/Engine/Plugins/Runtime/Metasound/Source/MetasoundEditor/Private/MetasoundNodeDetailCustomization.cpp
--- a/Engine/Plugins/Runtime/Metasound/Source/MetasoundEditor/Private/MetasoundNodeDetailCustomization.cpp
+++ b/Engine/Plugins/Runtime/Metasound/Source/MetasoundEditor/Private/MetasoundNodeDetailCustomization.cpp
@@ -313,7 +313,7 @@ namespace Metasound
 				{
 					if (UClass* Class = Object->GetClass())
 					{
-						return Class->IsChildOf(InProxyGenClass.Get());
+						return Class == InProxyGenClass.Get();
 					}
 				}
 
