Check that editor connections do not cause loops

diff --git a/Engine/Plugins/Runtime/Metasound/Source/MetasoundEditor/Private/MetasoundEditorGraphBuilder.cpp b/Engine/Plugins/Runtime/Metasound/Source/MetasoundEditor/Private/MetasoundEditorGraphBuilder.cpp
--- a/Engine/Plugins/Runtime/Metasound/Source/MetasoundEditor/Private/MetasoundEditorGraphBuilder.cpp
+++ b/Engine/Plugins/Runtime/Metasound/Source/MetasoundEditor/Private/MetasoundEditorGraphBuilder.cpp
@@ -1104,6 +1104,28 @@ namespace Metasound
 			return false;
 		}
 
+		void FGraphBuilder::DepthFirstTraversal(UEdGraphNode* InInitialNode, FDepthFirstVisitFunction InVisitFunction)
+		{
+			// Non recursive depth first traversal.
+			TArray<UEdGraphNode*> Stack({InInitialNode});
+			TSet<UEdGraphNode*> Visited;
+
+			while (Stack.Num() > 0)
+			{
+				UEdGraphNode* CurrentNode = Stack.Pop();
+				if (Visited.Contains(CurrentNode))
+				{
+					// Do not revisit a node that has already been visited. 
+					continue;
+				}
+
+				TArray<UEdGraphNode*> Children = InVisitFunction(CurrentNode).Array();
+				Stack.Append(Children);
+
+				Visited.Add(CurrentNode);
+			}
+		}
+
 		UEdGraphPin* FGraphBuilder::AddPinToNode(UMetasoundEditorGraphNode& InEditorNode, Frontend::FInputHandle InInputHandle)
 		{
 			using namespace Frontend;
diff --git a/Engine/Plugins/Runtime/Metasound/Source/MetasoundEditor/Private/MetasoundEditorGraphBuilder.h b/Engine/Plugins/Runtime/Metasound/Source/MetasoundEditor/Private/MetasoundEditorGraphBuilder.h
--- a/Engine/Plugins/Runtime/Metasound/Source/MetasoundEditor/Private/MetasoundEditorGraphBuilder.h
+++ b/Engine/Plugins/Runtime/Metasound/Source/MetasoundEditor/Private/MetasoundEditorGraphBuilder.h
@@ -183,6 +183,19 @@ namespace Metasound
 
 			// Returns true if the FOutputHandle and UEdGraphPin match each other.
 			static bool IsMatchingOutputHandleAndPin(const Frontend::FOutputHandle& InOutputHandle, const UEdGraphPin& InEditorPin);
+
+			// Function signature for visiting a node doing depth first traversal.
+			//
+			// Functions accept a UEdGraphNode* and return a TSet<UEdGraphNode*> which
+			// represent all the children of the node. 
+			using FDepthFirstVisitFunction = TFunctionRef<TSet<UEdGraphNode*> (UEdGraphNode*)>; 
+
+			// Traverse depth first starting at the InInitialNode and calling the InVisitFunction
+			// for each node. 
+			//
+			// This implementation avoids recursive function calls to support deep
+			// graphs.
+			static void DepthFirstTraversal(UEdGraphNode* InInitialNode, FDepthFirstVisitFunction InVisitFunction);
 		};
 	} // namespace Editor
 } // namespace Metasound
diff --git a/Engine/Plugins/Runtime/Metasound/Source/MetasoundEditor/Private/MetasoundEditorGraphSchema.cpp b/Engine/Plugins/Runtime/Metasound/Source/MetasoundEditor/Private/MetasoundEditorGraphSchema.cpp
--- a/Engine/Plugins/Runtime/Metasound/Source/MetasoundEditor/Private/MetasoundEditorGraphSchema.cpp
+++ b/Engine/Plugins/Runtime/Metasound/Source/MetasoundEditor/Private/MetasoundEditorGraphSchema.cpp
@@ -369,8 +369,47 @@ UMetasoundEditorGraphSchema::UMetasoundEditorGraphSchema(const FObjectInitialize
 
 bool UMetasoundEditorGraphSchema::ConnectionCausesLoop(const UEdGraphPin* InputPin, const UEdGraphPin* OutputPin) const
 {
-	// TODO: Implement loop check
-	return false;
+	using namespace Metasound::Editor;
+
+	bool bCausesLoop = false;
+
+	if ((nullptr != InputPin) && (nullptr != OutputPin))
+	{
+		UEdGraphNode* InputNode = InputPin->GetOwningNode();
+		UEdGraphNode* OutputNode = OutputPin->GetOwningNode();
+
+		// Sets bCausesLoop if the input node already has a path to the output node
+		//
+		FGraphBuilder::DepthFirstTraversal(InputNode, [&](UEdGraphNode* Node) -> TSet<UEdGraphNode*>
+			{
+				TSet<UEdGraphNode*> Children;
+
+				if (OutputNode == Node)
+				{
+					// If the input node can already reach the output node, then this 
+					// connection will cause a loop.
+					bCausesLoop = true;
+				}
+
+				if (!bCausesLoop)
+				{
+					// Only produce children if no loop exists to avoid wasting unnecessary CPU
+					if (nullptr != Node)
+					{
+						Node->ForEachNodeDirectlyConnectedToOutputs([&](UEdGraphNode* ChildNode) 
+							{ 
+								Children.Add(ChildNode);
+							}
+						);
+					}
+				}
+
+				return Children;
+			}
+		);
+	}
+	
+	return bCausesLoop;
 }
 
 void UMetasoundEditorGraphSchema::GetPaletteActions(FGraphActionMenuBuilder& ActionMenuBuilder) const
