Fix Metasound object proxy asset filters (FOR GOOD!) to not include child class types. (Post EA, will make an option per proxy, but for now do all as SoundWave is only type).

diff --git a/Engine/Plugins/Runtime/Metasound/Source/MetasoundEditor/Private/MetasoundEditorGraphNode.cpp b/Engine/Plugins/Runtime/Metasound/Source/MetasoundEditor/Private/MetasoundEditorGraphNode.cpp
--- a/Engine/Plugins/Runtime/Metasound/Source/MetasoundEditor/Private/MetasoundEditorGraphNode.cpp
+++ b/Engine/Plugins/Runtime/Metasound/Source/MetasoundEditor/Private/MetasoundEditorGraphNode.cpp
@@ -14,6 +14,7 @@
 #include "MetasoundEditorCommands.h"
 #include "MetasoundEditorModule.h"
 #include "MetasoundFrontend.h"
+#include "MetasoundFrontendRegistries.h"
 #include "MetasoundUObjectRegistry.h"
 #include "ScopedTransaction.h"
 #include "ToolMenus.h"
@@ -232,6 +233,80 @@ void UMetasoundEditorGraphNode::PinDefaultValueChanged(UEdGraphPin* Pin)
 	}
 }
 
+FString UMetasoundEditorGraphNode::GetPinMetaData(FName InPinName, FName InKey)
+{
+	using namespace Metasound::Editor;
+	using namespace Metasound::Frontend;
+
+	if (InKey == "DisallowedClasses")
+	{
+		if (UEdGraphPin* Pin = FindPin(InPinName, EGPD_Input))
+		{
+			FInputHandle Handle = FGraphBuilder::GetInputHandleFromPin(Pin);
+
+			FMetasoundFrontendRegistryContainer* Registry = FMetasoundFrontendRegistryContainer::Get();
+			if (!ensure(Registry))
+			{
+				return FString();
+			}
+
+			Metasound::FDataTypeRegistryInfo DataTypeInfo;
+			if (!ensure(Registry->GetInfoForDataType(Handle->GetDataType(), DataTypeInfo)))
+			{
+				return FString();
+			}
+
+			const EMetasoundFrontendLiteralType LiteralType = GetMetasoundFrontendLiteralType(DataTypeInfo.PreferredLiteralType);
+			if (LiteralType != EMetasoundFrontendLiteralType::UObject && LiteralType != EMetasoundFrontendLiteralType::UObjectArray)
+			{
+				return FString();
+			}
+
+			UClass* ProxyGenClass = DataTypeInfo.ProxyGeneratorClass;
+			if (!ProxyGenClass)
+			{
+				return FString();
+			}
+
+			FString DisallowedClasses;
+			const FName ClassName = ProxyGenClass->GetFName();
+			for (TObjectIterator<UClass> ClassIt; ClassIt; ++ClassIt)
+			{
+				UClass* Class = *ClassIt;
+				if (!Class->IsNative())
+				{
+					continue;
+				}
+
+				if (Class->HasAnyClassFlags(CLASS_Abstract | CLASS_Deprecated | CLASS_NewerVersionExists))
+				{
+					continue;
+				}
+
+				if (ClassIt->GetFName() == ClassName)
+				{
+					continue;
+				}
+
+				if (Class->IsChildOf(ProxyGenClass))
+				{
+					if (!DisallowedClasses.IsEmpty())
+					{
+						DisallowedClasses += TEXT(",");
+					}
+					DisallowedClasses += *ClassIt->GetName();
+				}
+			}
+
+			return DisallowedClasses;
+		}
+
+		return FString();
+	}
+
+	return Super::GetPinMetaData(InPinName, InKey);
+}
+
 void UMetasoundEditorGraphNode::PostEditImport()
 {
 }
diff --git a/Engine/Plugins/Runtime/Metasound/Source/MetasoundEditor/Private/MetasoundNodeDetailCustomization.cpp b/Engine/Plugins/Runtime/Metasound/Source/MetasoundEditor/Private/MetasoundNodeDetailCustomization.cpp
--- a/Engine/Plugins/Runtime/Metasound/Source/MetasoundEditor/Private/MetasoundNodeDetailCustomization.cpp
+++ b/Engine/Plugins/Runtime/Metasound/Source/MetasoundEditor/Private/MetasoundNodeDetailCustomization.cpp
@@ -320,16 +320,6 @@ namespace Metasound
 				return false;
 			};
 
-			auto CommitAsset = [InPropertyHandle = PropertyHandle, InProxyGenClass = ProxyGenClass](const FAssetData& InAssetData)
-			{
-				// if we've hit this code, the presumption is that the datatype for this parameter has already defined a corresponding UClass that can be used to set it.
-				ensureAlways(InProxyGenClass.IsValid());
-
-				UObject* InObject = InAssetData.GetAsset();
-
-				InPropertyHandle->SetValue(InObject);
-			};
-
 			auto GetAssetPath = [PropertyHandle = PropertyHandle]()
 			{
 				UObject* Object = nullptr;
@@ -343,12 +333,29 @@ namespace Metasound
 			TArray<const UClass*> AllowedClasses;
 			AllowedClasses.Add(ProxyGenClass.Get());
 
+			auto FilterAsset = [InProxyGenClass = ProxyGenClass](const FAssetData& InAsset)
+			{
+				if (InProxyGenClass.IsValid())
+				{
+					if (UObject* Object = InAsset.GetAsset())
+					{
+						if (UClass* Class = Object->GetClass())
+						{
+							return Class != InProxyGenClass.Get();
+						}
+					}
+				}
+
+				return true;
+			};
+
 			return SNew(SObjectPropertyEntryBox)
 				.ObjectPath_Lambda(GetAssetPath)
 				.AllowedClass(ProxyGenClass.Get())
 				.OnShouldSetAsset_Lambda(ValidateAsset)
-				.OnObjectChanged_Lambda(CommitAsset)
-				.AllowClear(false)
+				.OnShouldFilterAsset_Lambda(FilterAsset)
+				.PropertyHandle(PropertyHandle)
+				.AllowClear(true)
 				.DisplayUseSelected(true)
 				.DisplayBrowse(true)
 				.DisplayThumbnail(true)
diff --git a/Engine/Plugins/Runtime/Metasound/Source/MetasoundEditor/Public/MetasoundEditorGraphNode.h b/Engine/Plugins/Runtime/Metasound/Source/MetasoundEditor/Public/MetasoundEditorGraphNode.h
--- a/Engine/Plugins/Runtime/Metasound/Source/MetasoundEditor/Public/MetasoundEditorGraphNode.h
+++ b/Engine/Plugins/Runtime/Metasound/Source/MetasoundEditor/Public/MetasoundEditorGraphNode.h
@@ -55,6 +55,7 @@ public:
 	virtual FText GetTooltipText() const override;
 	virtual void PinDefaultValueChanged(UEdGraphPin* Pin) override;
 	virtual void ReconstructNode() override;
+	virtual FString GetPinMetaData(FName InPinName, FName InKey) override;
 	// End of UEdGraphNode interface
 
 	// UObject interface
