Hide and disable unneeded types and nodes
- Added enable structs for arrays, auto conversions, send/receive to opt-out when desired.
- Lots of header include fixes to get build working.
- Comment out "int64" and "double" registered data types

diff --git a/Engine/Plugins/Runtime/Metasound/Source/MetasoundFrontend/Private/MetasoundFrontend.cpp b/Engine/Plugins/Runtime/Metasound/Source/MetasoundFrontend/Private/MetasoundFrontend.cpp
--- a/Engine/Plugins/Runtime/Metasound/Source/MetasoundFrontend/Private/MetasoundFrontend.cpp
+++ b/Engine/Plugins/Runtime/Metasound/Source/MetasoundFrontend/Private/MetasoundFrontend.cpp
@@ -21,7 +21,6 @@
 #include "StructSerializer.h"
 #include "Serialization/MemoryReader.h"
 
-
 namespace Metasound
 {
 	namespace Frontend
@@ -224,17 +223,4 @@ class FMetasoundFrontendModule : public IModuleInterface
 	}
 };
 
-
-REGISTER_METASOUND_DATATYPE(bool, "Bool", ::Metasound::ELiteralType::Boolean)
-REGISTER_METASOUND_DATATYPE(int32, "Int32", ::Metasound::ELiteralType::Integer)
-REGISTER_METASOUND_DATATYPE(int64, "Int64", ::Metasound::ELiteralType::Integer)
-REGISTER_METASOUND_DATATYPE(float, "Float", ::Metasound::ELiteralType::Float)
-REGISTER_METASOUND_DATATYPE(double, "Double", ::Metasound::ELiteralType::Float)
-REGISTER_METASOUND_DATATYPE(FString, "String", ::Metasound::ELiteralType::String)
-
-REGISTER_METASOUND_DATATYPE(Metasound::FTrigger, "Trigger", ::Metasound::ELiteralType::Boolean)
-REGISTER_METASOUND_DATATYPE(Metasound::FTime, "Time", ::Metasound::ELiteralType::Float)
-REGISTER_METASOUND_DATATYPE(Metasound::FAudioBuffer, "Audio")
-REGISTER_METASOUND_DATATYPE(Metasound::FSendAddress, "Transmission:Address", ::Metasound::ELiteralType::String)
-
 IMPLEMENT_MODULE(FMetasoundFrontendModule, MetasoundFrontend);
diff --git a/Engine/Plugins/Runtime/Metasound/Source/MetasoundFrontend/Private/MetasoundPrimitives.cpp b/Engine/Plugins/Runtime/Metasound/Source/MetasoundFrontend/Private/MetasoundPrimitives.cpp
new file mode 100644
--- /dev/null
+++ b/Engine/Plugins/Runtime/Metasound/Source/MetasoundFrontend/Private/MetasoundPrimitives.cpp
@@ -0,0 +1,52 @@
+// Copyright Epic Games, Inc. All Rights Reserved.
+
+#include "MetasoundPrimitives.h"
+
+#include "MetasoundAudioBuffer.h"
+#include "MetasoundDataTypeRegistrationMacro.h"
+#include "MetasoundRouter.h"
+#include "MetasoundTime.h"
+#include "MetasoundTrigger.h"
+
+#include <type_traits>
+
+namespace Metasound
+{
+	// Disable auto-conversion using the FAudioBuffer(int32) constructor
+	template<typename FromDataType>
+	struct TIsAutoConvertible<FromDataType, FAudioBuffer>
+	{
+		static constexpr bool bIsConvertible = std::is_convertible<FromDataType, FAudioBuffer>::value;
+
+		static constexpr bool bIsFromArithmeticType = std::is_arithmetic<FromDataType>::value;
+
+	public:
+
+		static constexpr bool Value = bIsConvertible && (!bIsFromArithmeticType);
+	};
+
+	// Disable auto-conversions based on FTrigger implicit converters
+	template<typename ToDataType>
+	struct TIsAutoConvertible<FTrigger, ToDataType>
+	{
+		static constexpr bool bIsConvertible = std::is_convertible<FTrigger, ToDataType>::value;
+
+		static constexpr bool bIsToArithmeticType = std::is_arithmetic<ToDataType>::value;
+
+	public:
+
+		static constexpr bool Value = bIsConvertible && (!bIsToArithmeticType);
+	};
+}
+
+REGISTER_METASOUND_DATATYPE(bool, "Bool", ::Metasound::ELiteralType::Boolean)
+REGISTER_METASOUND_DATATYPE(int32, "Int32", ::Metasound::ELiteralType::Integer)
+//REGISTER_METASOUND_DATATYPE(int64, "Int64", ::Metasound::ELiteralType::Integer)
+REGISTER_METASOUND_DATATYPE(float, "Float", ::Metasound::ELiteralType::Float)
+//REGISTER_METASOUND_DATATYPE(double, "Double", ::Metasound::ELiteralType::Float)
+REGISTER_METASOUND_DATATYPE(FString, "String", ::Metasound::ELiteralType::String)
+
+REGISTER_METASOUND_DATATYPE(Metasound::FTrigger, "Trigger", ::Metasound::ELiteralType::Boolean)
+REGISTER_METASOUND_DATATYPE(Metasound::FTime, "Time", ::Metasound::ELiteralType::Float)
+REGISTER_METASOUND_DATATYPE(Metasound::FAudioBuffer, "Audio")
+REGISTER_METASOUND_DATATYPE(Metasound::FSendAddress, "Transmission:Address", ::Metasound::ELiteralType::String)
diff --git a/Engine/Plugins/Runtime/Metasound/Source/MetasoundFrontend/Public/MetasoundArrayNodes.h b/Engine/Plugins/Runtime/Metasound/Source/MetasoundFrontend/Public/MetasoundArrayNodes.h
--- a/Engine/Plugins/Runtime/Metasound/Source/MetasoundFrontend/Public/MetasoundArrayNodes.h
+++ b/Engine/Plugins/Runtime/Metasound/Source/MetasoundFrontend/Public/MetasoundArrayNodes.h
@@ -10,17 +10,14 @@
 #include "MetasoundFacade.h"
 #include "MetasoundLog.h"
 #include "MetasoundNodeInterface.h"
-#include "MetasoundNodeRegistrationMacro.h"
 #include "MetasoundOperatorInterface.h"
 #include "MetasoundPrimitives.h"
 #include "MetasoundTrigger.h"
-#include "Containers/CircularQueue.h"
 
 #include <type_traits>
 
 #define LOCTEXT_NAMESPACE "MetasoundFrontend"
 
-
 namespace Metasound
 {
 	namespace MetasoundArrayNodesPrivate
diff --git a/Engine/Plugins/Runtime/Metasound/Source/MetasoundFrontend/Public/MetasoundArrayNodesRegistration.h b/Engine/Plugins/Runtime/Metasound/Source/MetasoundFrontend/Public/MetasoundArrayNodesRegistration.h
--- a/Engine/Plugins/Runtime/Metasound/Source/MetasoundFrontend/Public/MetasoundArrayNodesRegistration.h
+++ b/Engine/Plugins/Runtime/Metasound/Source/MetasoundFrontend/Public/MetasoundArrayNodesRegistration.h
@@ -11,6 +11,11 @@
 
 namespace Metasound
 {
+	template<typename ... ElementType>
+	struct TEnableArrayNodes
+	{
+		static constexpr bool Value = true;
+	};
 
 	namespace MetasoundArrayNodesPrivate
 	{
@@ -22,30 +27,37 @@ namespace Metasound
 		template<typename ArrayType>
 		struct TArrayNodeSupport
 		{
-			using ElementType = typename TArrayElementType<ArrayType>::Type;
+		private:
+			using ElementType = typename MetasoundArrayNodesPrivate::TArrayElementType<ArrayType>::Type;
+
+			static constexpr bool bIsElementParsableAndAssignable = TIsParsable<ElementType>::Value && std::is_copy_assignable<ElementType>::value;
+
+			static constexpr bool bEnabled = TEnableArrayNodes<ElementType>::Value;
+
+		public:
 			
 			// Array num is supported for all array types.
-			static constexpr bool bIsArrayNumSupported = true;
+			static constexpr bool bIsArrayNumSupported = bEnabled;
 
 			// Element must be default parsable to create get operator because a
 			// value must be returned even if the index is invalid. Also values are
 			// assigned by copy.
-			static constexpr bool bIsArrayGetSupported = TIsParsable<ElementType>::Value && std::is_copy_assignable<ElementType>::value;
+			static constexpr bool bIsArrayGetSupported = bEnabled && bIsElementParsableAndAssignable;
 
 			// Element must be copy assignable to set the value.
-			static constexpr bool bIsArraySetSupported = std::is_copy_assignable<ElementType>::value && std::is_copy_constructible<ElementType>::value;
+			static constexpr bool bIsArraySetSupported = bEnabled && std::is_copy_assignable<ElementType>::value && std::is_copy_constructible<ElementType>::value;
 
 			// Elements must be copy constructible
-			static constexpr bool bIsArrayConcatSupported = std::is_copy_constructible<ElementType>::value;
+			static constexpr bool bIsArrayConcatSupported = bEnabled && std::is_copy_constructible<ElementType>::value;
 
 			// Elements must be copy constructible
-			static constexpr bool bIsArraySubsetSupported = std::is_copy_constructible<ElementType>::value;
+			static constexpr bool bIsArraySubsetSupported = bEnabled && std::is_copy_constructible<ElementType>::value;
 
 			// Array shuffle is supported for all types that get is supported for.
-			static constexpr bool bIsArrayShuffleSupported = bIsArrayGetSupported;
+			static constexpr bool bIsArrayShuffleSupported = bEnabled && bIsElementParsableAndAssignable;
 
 			// Random get is supported for all types that get is supported for.
-			static constexpr bool bIsArrayRandomGetSupported = bIsArrayGetSupported;
+			static constexpr bool bIsArrayRandomGetSupported = bEnabled && bIsElementParsableAndAssignable;
 		};
 
 		template<typename ArrayType, typename std::enable_if<TArrayNodeSupport<ArrayType>::bIsArrayGetSupported, bool>::type = true>
@@ -113,14 +125,19 @@ namespace Metasound
 			return true;
 		}
 
-		template<typename ArrayType>
+		template<typename ArrayType, typename std::enable_if<TArrayNodeSupport<ArrayType>::bIsArrayNumSupported, bool>::type = true>
 		bool RegisterArrayNumNode()
 		{
-			static_assert(TArrayNodeSupport<ArrayType>::bIsArrayNumSupported, "TArrayNumNode<> is not supported by array type");
-
 			return ensureAlways(RegisterNodeWithFrontend<Metasound::TArrayNumNode<ArrayType>>());
 		}
 
+		template<typename ArrayType, typename std::enable_if<!TArrayNodeSupport<ArrayType>::bIsArrayNumSupported, bool>::type = true>
+		bool RegisterArrayNumNode()
+		{
+			// No op if not supported
+			return true;
+		}
+
 		template<typename ArrayType, typename std::enable_if<TArrayNodeSupport<ArrayType>::bIsArrayShuffleSupported, bool>::type = true>
 		bool RegisterArrayShuffleNode()
 		{
@@ -143,7 +160,7 @@ namespace Metasound
 		}
 
 		template<typename ArrayType, typename std::enable_if<!TArrayNodeSupport<ArrayType>::bIsArrayRandomGetSupported, bool>::type = true>
-		bool RegisterArrayShuffleNode()
+		bool RegisterArrayRandomGetNode()
 		{
 			// No op if not supported
 			return true;
diff --git a/Engine/Plugins/Runtime/Metasound/Source/MetasoundFrontend/Public/MetasoundArrayRandomNode.h b/Engine/Plugins/Runtime/Metasound/Source/MetasoundFrontend/Public/MetasoundArrayRandomNode.h
--- a/Engine/Plugins/Runtime/Metasound/Source/MetasoundFrontend/Public/MetasoundArrayRandomNode.h
+++ b/Engine/Plugins/Runtime/Metasound/Source/MetasoundFrontend/Public/MetasoundArrayRandomNode.h
@@ -2,11 +2,13 @@
 
 #pragma once
 
+#include "Containers/CircularQueue.h"
 #include "CoreMinimal.h"
+
+#include "MetasoundArrayNodes.h"
+#include "MetasoundExecutableOperator.h"
 #include "MetasoundFacade.h"
 #include "MetasoundNodeInterface.h"
-#include "MetasoundExecutableOperator.h"
-#include "MetasoundArrayNodes.h"
 #include "MetasoundTrigger.h"
 
 #define LOCTEXT_NAMESPACE "MetasoundFrontend_RandomArrayGet"
@@ -409,4 +411,4 @@ namespace Metasound
 
 }
 
-#undef LOCTEXT_NAMESPACE
\ No newline at end of file
+#undef LOCTEXT_NAMESPACE
diff --git a/Engine/Plugins/Runtime/Metasound/Source/MetasoundFrontend/Public/MetasoundArrayShuffleNode.h b/Engine/Plugins/Runtime/Metasound/Source/MetasoundFrontend/Public/MetasoundArrayShuffleNode.h
--- a/Engine/Plugins/Runtime/Metasound/Source/MetasoundFrontend/Public/MetasoundArrayShuffleNode.h
+++ b/Engine/Plugins/Runtime/Metasound/Source/MetasoundFrontend/Public/MetasoundArrayShuffleNode.h
@@ -10,12 +10,11 @@
 #include "MetasoundFacade.h"
 #include "MetasoundLog.h"
 #include "MetasoundNodeInterface.h"
-#include "MetasoundNodeRegistrationMacro.h"
 #include "MetasoundOperatorInterface.h"
 #include "MetasoundPrimitives.h"
 #include "MetasoundTrigger.h"
-#include "Containers/CircularQueue.h"
 #include "MetasoundArrayNodes.h"
+#include "Misc/ScopeLock.h"
 
 #include <type_traits>
 
diff --git a/Engine/Plugins/Runtime/Metasound/Source/MetasoundFrontend/Public/MetasoundDataTypeRegistrationMacro.h b/Engine/Plugins/Runtime/Metasound/Source/MetasoundFrontend/Public/MetasoundDataTypeRegistrationMacro.h
--- a/Engine/Plugins/Runtime/Metasound/Source/MetasoundFrontend/Public/MetasoundDataTypeRegistrationMacro.h
+++ b/Engine/Plugins/Runtime/Metasound/Source/MetasoundFrontend/Public/MetasoundDataTypeRegistrationMacro.h
@@ -28,6 +28,50 @@
 
 namespace Metasound
 {
+	// Helper utility to test if we can transmit a datatype between a send and a receive node.
+	template <typename TDataType>
+	struct TIsTransmittable
+	{
+	private:
+		static constexpr bool bIsCopyConstructible = std::is_copy_constructible<TDataType>::value;
+		static constexpr bool bIsCopyAssignable = std::is_copy_assignable<TDataType>::value;
+
+		// TODO: audio types were intended to be send/receive nodes but they require 
+		// template specialization.  TIsTransmittable should ask the Send and Receive nodes
+		// if they are transmittable rather than attempting to do the logic for all
+		// types here. 
+		//static constexpr bool bIsAudioDataType = TIsDerivedFrom<TDataType, IAudioDataType>::Value;
+
+		static constexpr bool bCanBeTransmitted = bIsCopyConstructible && bIsCopyAssignable;
+
+	public:
+
+		static constexpr bool Value = bCanBeTransmitted;
+	};
+
+	// Specialization of TIsTransmittable<> for TArray to handle lax TArray copy constructor 
+	// definition. This can be removed once updates to TArray are merged into the same codebase.
+	// TODO: delete me eventually.
+	template<typename TElementType>
+	struct TIsTransmittable<TArray<TElementType>>
+	{
+		// Depend on copy constructor of elements to determine whether TArray is copy constructible
+		// and copy assignable. Definitions for TArray copy construction and copy assignment are
+		// defined by default whether or not the TArray element types support them. This logic
+		// skips checking the TArray object itselt and instead checks the TArray element type
+		// directly.
+
+		static constexpr bool Value = TIsTransmittable<TElementType>::Value;
+	};
+
+	// Determines whether an auto converter node will be registered to convert 
+	// between two types. 
+	template<typename TFromDataType, typename TToDataType>
+	struct TIsAutoConvertible
+	{
+		static constexpr bool Value = std::is_convertible<TFromDataType, TToDataType>::value;
+	};
+
 	namespace MetasoundDataTypeRegistrationPrivate
 	{
 		// Returns the Array version of a literal type if it exists.
@@ -103,47 +147,6 @@ namespace Metasound
 			return nullptr;
 		}
 		
-		// Helper utility to test if we can transmit a datatype between a send and a receive node.
-		template <typename TDataType>
-		struct TIsTransmittable
-		{
-		private:
-			static constexpr bool bIsCopyConstructible = std::is_copy_constructible<TDataType>::value;
-			static constexpr bool bIsCopyAssignable = std::is_copy_assignable<TDataType>::value;
-
-			// TODO: audio types were intended to be send/receive nodes but they require 
-			// template specialization.  TIsTransmittable should ask the Send and Receive nodes
-			// if they are transmittable rather than attempting to do the logic for all
-			// types here. 
-			//static constexpr bool bIsAudioDataType = TIsDerivedFrom<TDataType, IAudioDataType>::Value;
-
-			static constexpr bool bCanBeTransmitted = bIsCopyConstructible && bIsCopyAssignable;
-
-		public:
-
-			static constexpr bool Value = bCanBeTransmitted;
-		};
-
-		// Specialization of TIsTransmittable<> for TArray to handle lax TArray copy constructor 
-		// definition. This can be removed once updates to TArray are merged into the same codebase.
-		// TODO: delete me eventually.
-		template<typename TElementType>
-		struct TIsTransmittable<TArray<TElementType>>
-		{
-		private:
-			// Depend on copy constructor of elements to determine whether TArray is copy constructible
-			// and copy assignable. Definitions for TArray copy construction and copy assignment are
-			// defined by default whether or not the TArray element types support them. This logic
-			// skips checking the TArray object itselt and instead checks the TArray element type
-			// directly.
-			static constexpr bool bIsElementCopyConstructible = std::is_copy_constructible<TElementType>::value;
-
-			static constexpr bool bCanBeTransmitted = bIsElementCopyConstructible;
-
-		public:
-
-			static constexpr bool Value = bCanBeTransmitted;
-		};
 
 		// This utility function can be used to optionally check to see if we can transmit a data type, and autogenerate send and receive nodes for that datatype.
 		template<typename TDataType, typename TEnableIf<TIsTransmittable<TDataType>::Value, bool>::Type = true>
@@ -162,7 +165,7 @@ namespace Metasound
 		}
 
 		// This utility function can be used to check to see if we can static cast between two types, and autogenerate a node for that static cast.
-		template<typename TFromDataType, typename TToDataType, typename std::enable_if<std::is_convertible<TFromDataType, TToDataType>::value, int>::type = 0>
+		template<typename TFromDataType, typename TToDataType, typename std::enable_if<TIsAutoConvertible<TFromDataType, TToDataType>::Value, bool>::type = true>
 		void AttemptToRegisterConverter()
 		{
 			using FConverterNode = Metasound::TAutoConverterNode<TFromDataType, TToDataType>;
@@ -179,7 +182,7 @@ namespace Metasound
 			}
 		}
 
-		template<typename TFromDataType, typename TToDataType, typename std::enable_if<!std::is_convertible<TFromDataType, TToDataType>::value, int>::type = 0>
+		template<typename TFromDataType, typename TToDataType, typename std::enable_if<!TIsAutoConvertible<TFromDataType, TToDataType>::Value, int>::type = 0>
 		void AttemptToRegisterConverter()
 		{
 			// This implementation intentionally noops, because static_cast<TFromDataType>(TToDataType&) is invalid.
@@ -245,10 +248,6 @@ namespace Metasound
 		template<typename TDataType, typename std::enable_if<!TEnumTraits<TDataType>::bIsEnum, bool>::type = true>
 		bool RegisterEnumDataTypeWithFrontend() { return false; }
 
-
-
-
-
 		template<typename TDataType, ELiteralType PreferredArgType = ELiteralType::None, typename UClassToUse = UObject>
 		bool RegisterDataTypeWithFrontendInternal()
 		{
@@ -397,7 +396,6 @@ namespace Metasound
 	{
 		using namespace MetasoundDataTypeRegistrationPrivate;
 
-
 		// Register TDataType as a metasound data type.
 		bool bSuccess = RegisterDataTypeWithFrontendInternal<TDataType, PreferredArgType, UClassToUse>();
 		ensure(bSuccess);
diff --git a/Engine/Plugins/Runtime/Metasound/Source/MetasoundFrontend/Public/MetasoundEnum.h b/Engine/Plugins/Runtime/Metasound/Source/MetasoundFrontend/Public/MetasoundEnum.h
--- a/Engine/Plugins/Runtime/Metasound/Source/MetasoundFrontend/Public/MetasoundEnum.h
+++ b/Engine/Plugins/Runtime/Metasound/Source/MetasoundFrontend/Public/MetasoundEnum.h
@@ -2,7 +2,8 @@
 
 #pragma once
 
-#include "MetasoundDataReference.h"
+#include "CoreMinimal.h"
+
 #include <type_traits>
 
 // NOTE: Metasound Enum types are defined outside of Engine so can't use the UENUM type reflection here.
@@ -206,58 +207,3 @@ namespace Metasound
 	};	
 }
 
-// Helper macros: for defining/declaring a new Enum to be used in Metasounds.
-// * Declares/defines a typedef FEnum<YourEnum> wrapper around the enum that you pass it
-// * Declares/defines a specialization of TEnumStringHelper which is used to validate and convert between values/names.
-
-// These helps generate a couple of boiler plate functions:
-//  - GetNamespace() which will return the name of the Enum "EMyEnum" etc. 
-//  - GetAllEntries() which returns all posible entries in the enum. 
-
-/** DECLARE_METASOUND_ENUM
- * @param ENUMNAME - The typename of your EnumType you want to use for Metasounds. e.g. MyEnum
- * @param DEFAULT - A fully qualified default Enum value. e.g. EMyEnum::One
- * @param API - The module API this is declared inside e.g. METASOUNDSTANDARDNODES_API 
- * @param ENUMTYPEDEF - The name of the TEnum<YourType> wrapper type
- * @param TYPEINFO - The name of the TypeInfo type you want to define e.g. FMyEnumTypeInfo
- * @param READREF - The name of the Read Reference type you want to define. e.g FMyEnumReadRef
- * @param WRITEREF -The name of the Write Reference type you want to define e.g. FMyEnumWriteRef
- */
-#define DECLARE_METASOUND_ENUM(ENUMNAME, DEFAULT, API, ENUMTYPEDEF, TYPEINFO, READREF, WRITEREF)\
-	using ENUMTYPEDEF = Metasound::TEnum<ENUMNAME, DEFAULT>; \
-	DECLARE_METASOUND_DATA_REFERENCE_TYPES(ENUMTYPEDEF, API, TYPEINFO, READREF, WRITEREF);\
-	template<> struct API Metasound::TEnumStringHelper<ENUMNAME> : Metasound::TEnumStringHelperBase<Metasound::TEnumStringHelper<ENUMNAME>, ENUMNAME>\
-	{\
-		static FName GetNamespace()\
-		{\
-			static const FName ThisName { TEXT(#ENUMNAME) };\
-			return ThisName;\
-		}\
-		static TArrayView<const Metasound::TEnumEntry<ENUMNAME>> GetAllEntries();\
-	};
-
-/** DEFINE_METASOUND_ENUM_BEGIN
- * @param ENUMNAME - The typename of your raw EnumType you want to use for Metasounds. e.g. EMyType
- * @param ENUMTYPEDEF - The name of the TEnum<YourType> wrapper type
- * @param DATATYPENAMESTRING - The string that will the data type name "Enum:<string>" e.g. "MyEnum"
- */
-#define DEFINE_METASOUND_ENUM_BEGIN(ENUMNAME,ENUMTYPEDEF,DATATYPENAMESTRING)\
-	REGISTER_METASOUND_DATATYPE(ENUMTYPEDEF, "Enum:" DATATYPENAMESTRING, ::Metasound::ELiteralType::Integer);\
-	TArrayView<const Metasound::TEnumEntry<ENUMNAME>> Metasound::TEnumStringHelper<ENUMNAME>::GetAllEntries()\
-	{\
-		static const Metasound::TEnumEntry<ENUMNAME> Entries[] = {
-
-/** DEFINE_METASOUND_ENUM_ENTRY - defines a single Enum Entry
-  * @param ENTRY - Fully Qualified Name of Entry of the Enum. (e.g. EMyType::One)
-  * @param DISPLAYNAME - FText Presented to User
-  * @param TOOLTIP - FText displayed as a tooltip
-  */
-	#define DEFINE_METASOUND_ENUM_ENTRY(ENTRY, DISPLAYNAME, TOOLTIP) { ENTRY, TEXT(#ENTRY), DISPLAYNAME, TOOLTIP }
-
-/** DEFINE_METASOUND_ENUM_END - macro which ends the function body of the GetAllEntries function
-   */
-#define DEFINE_METASOUND_ENUM_END() \
-		};\
-		return Entries;\
-	};
-
diff --git a/Engine/Plugins/Runtime/Metasound/Source/MetasoundFrontend/Public/MetasoundEnumRegistrationMacro.h b/Engine/Plugins/Runtime/Metasound/Source/MetasoundFrontend/Public/MetasoundEnumRegistrationMacro.h
new file mode 100644
--- /dev/null
+++ b/Engine/Plugins/Runtime/Metasound/Source/MetasoundFrontend/Public/MetasoundEnumRegistrationMacro.h
@@ -0,0 +1,92 @@
+// Copyright Epic Games, Inc. All Rights Reserved.
+
+#pragma once
+
+#include "MetasoundArrayNodesRegistration.h"
+#include "MetasoundDataReference.h"
+#include "MetasoundEnum.h"
+#include "MetasoundDataTypeRegistrationMacro.h"
+
+namespace Metasound
+{
+	template<typename FromDataType, typename EnumType, EnumType DefaultValue>
+	struct TIsAutoConvertible<FromDataType, TEnum<EnumType, DefaultValue>>
+	{
+		static constexpr bool Value = false;
+	};
+
+	template<typename ToDataType, typename EnumType, EnumType DefaultValue>
+	struct TIsAutoConvertible<TEnum<EnumType, DefaultValue>, ToDataType>
+	{
+		static constexpr bool Value = false;
+	};
+
+	template<typename EnumType, EnumType DefaultValue>
+	struct TEnableArrayNodes<TEnum<EnumType, DefaultValue>>
+	{
+		static constexpr bool Value = false;
+	};
+
+	// Specialization of TIsTransmittable<> to disable transmission of enums.
+	template<typename EnumType, EnumType DefaultValue>
+	struct TIsTransmittable<TEnum<EnumType, DefaultValue>>
+	{
+	public:
+		static constexpr bool Value = false;
+	};
+}
+
+// Helper macros: for defining/declaring a new Enum to be used in Metasounds.
+// * Declares/defines a typedef FEnum<YourEnum> wrapper around the enum that you pass it
+// * Declares/defines a specialization of TEnumStringHelper which is used to validate and convert between values/names.
+
+// These helps generate a couple of boiler plate functions:
+//  - GetNamespace() which will return the name of the Enum "EMyEnum" etc. 
+//  - GetAllEntries() which returns all posible entries in the enum. 
+
+/** DECLARE_METASOUND_ENUM
+ * @param ENUMNAME - The typename of your EnumType you want to use for Metasounds. e.g. MyEnum
+ * @param DEFAULT - A fully qualified default Enum value. e.g. EMyEnum::One
+ * @param API - The module API this is declared inside e.g. METASOUNDSTANDARDNODES_API 
+ * @param ENUMTYPEDEF - The name of the TEnum<YourType> wrapper type
+ * @param TYPEINFO - The name of the TypeInfo type you want to define e.g. FMyEnumTypeInfo
+ * @param READREF - The name of the Read Reference type you want to define. e.g FMyEnumReadRef
+ * @param WRITEREF -The name of the Write Reference type you want to define e.g. FMyEnumWriteRef
+ */
+#define DECLARE_METASOUND_ENUM(ENUMNAME, DEFAULT, API, ENUMTYPEDEF, TYPEINFO, READREF, WRITEREF)\
+	using ENUMTYPEDEF = Metasound::TEnum<ENUMNAME, DEFAULT>; \
+	DECLARE_METASOUND_DATA_REFERENCE_TYPES(ENUMTYPEDEF, API, TYPEINFO, READREF, WRITEREF);\
+	template<> struct API Metasound::TEnumStringHelper<ENUMNAME> : Metasound::TEnumStringHelperBase<Metasound::TEnumStringHelper<ENUMNAME>, ENUMNAME>\
+	{\
+		static FName GetNamespace()\
+		{\
+			static const FName ThisName { TEXT(#ENUMNAME) };\
+			return ThisName;\
+		}\
+		static TArrayView<const Metasound::TEnumEntry<ENUMNAME>> GetAllEntries();\
+	};
+
+/** DEFINE_METASOUND_ENUM_BEGIN
+ * @param ENUMNAME - The typename of your raw EnumType you want to use for Metasounds. e.g. EMyType
+ * @param ENUMTYPEDEF - The name of the TEnum<YourType> wrapper type
+ * @param DATATYPENAMESTRING - The string that will the data type name "Enum:<string>" e.g. "MyEnum"
+ */
+#define DEFINE_METASOUND_ENUM_BEGIN(ENUMNAME,ENUMTYPEDEF,DATATYPENAMESTRING)\
+	REGISTER_METASOUND_DATATYPE(ENUMTYPEDEF, "Enum:" DATATYPENAMESTRING, ::Metasound::ELiteralType::Integer);\
+	TArrayView<const Metasound::TEnumEntry<ENUMNAME>> Metasound::TEnumStringHelper<ENUMNAME>::GetAllEntries()\
+	{\
+		static const Metasound::TEnumEntry<ENUMNAME> Entries[] = {
+
+/** DEFINE_METASOUND_ENUM_ENTRY - defines a single Enum Entry
+  * @param ENTRY - Fully Qualified Name of Entry of the Enum. (e.g. EMyType::One)
+  * @param DISPLAYNAME - FText Presented to User
+  * @param TOOLTIP - FText displayed as a tooltip
+  */
+	#define DEFINE_METASOUND_ENUM_ENTRY(ENTRY, DISPLAYNAME, TOOLTIP) { ENTRY, TEXT(#ENTRY), DISPLAYNAME, TOOLTIP }
+
+/** DEFINE_METASOUND_ENUM_END - macro which ends the function body of the GetAllEntries function
+   */
+#define DEFINE_METASOUND_ENUM_END() \
+		};\
+		return Entries;\
+	};
diff --git a/Engine/Plugins/Runtime/Metasound/Source/MetasoundFrontend/Public/MetasoundFrontend.h b/Engine/Plugins/Runtime/Metasound/Source/MetasoundFrontend/Public/MetasoundFrontend.h
--- a/Engine/Plugins/Runtime/Metasound/Source/MetasoundFrontend/Public/MetasoundFrontend.h
+++ b/Engine/Plugins/Runtime/Metasound/Source/MetasoundFrontend/Public/MetasoundFrontend.h
@@ -7,7 +7,6 @@
 #include "MetasoundFrontendRegistries.h"
 #include "MetasoundFrontendDocument.h"
 #include "MetasoundBuilderInterface.h"
-#include "UObject/WeakObjectPtrTemplates.h"
 
 // Forward Declarations
 class FMetasoundAssetBase;
diff --git a/Engine/Plugins/Runtime/Metasound/Source/MetasoundFrontend/Public/MetasoundNodeRegistrationMacro.h b/Engine/Plugins/Runtime/Metasound/Source/MetasoundFrontend/Public/MetasoundNodeRegistrationMacro.h
--- a/Engine/Plugins/Runtime/Metasound/Source/MetasoundFrontend/Public/MetasoundNodeRegistrationMacro.h
+++ b/Engine/Plugins/Runtime/Metasound/Source/MetasoundFrontend/Public/MetasoundNodeRegistrationMacro.h
@@ -8,7 +8,6 @@
 #include "MetasoundFrontendRegistries.h"
 #include "MetasoundLog.h"
 #include "MetasoundNodeInterface.h"
-#include "Templates/IntegralConstant.h"
 
 namespace Metasound
 {
diff --git a/Engine/Plugins/Runtime/Metasound/Source/MetasoundFrontend/Public/MetasoundPrimitives.h b/Engine/Plugins/Runtime/Metasound/Source/MetasoundFrontend/Public/MetasoundPrimitives.h
--- a/Engine/Plugins/Runtime/Metasound/Source/MetasoundFrontend/Public/MetasoundPrimitives.h
+++ b/Engine/Plugins/Runtime/Metasound/Source/MetasoundFrontend/Public/MetasoundPrimitives.h
@@ -11,11 +11,11 @@ namespace Metasound
 	
 	DECLARE_METASOUND_DATA_REFERENCE_TYPES(int32, METASOUNDFRONTEND_API, FInt32TypeInfo, FInt32ReadRef, FInt32WriteRef);
 
-	DECLARE_METASOUND_DATA_REFERENCE_TYPES(int64, METASOUNDFRONTEND_API, FInt64TypeInfo, FInt64ReadRef, FInt64WriteRef);
+	//DECLARE_METASOUND_DATA_REFERENCE_TYPES(int64, METASOUNDFRONTEND_API, FInt64TypeInfo, FInt64ReadRef, FInt64WriteRef);
 
 	DECLARE_METASOUND_DATA_REFERENCE_TYPES(float, METASOUNDFRONTEND_API, FFloatTypeInfo, FFloatReadRef, FFloatWriteRef);
 
-	DECLARE_METASOUND_DATA_REFERENCE_TYPES(double, METASOUNDFRONTEND_API, FDoubleTypeInfo, FDoubleReadRef, FDoubleWriteRef);
+	//DECLARE_METASOUND_DATA_REFERENCE_TYPES(double, METASOUNDFRONTEND_API, FDoubleTypeInfo, FDoubleReadRef, FDoubleWriteRef);
 
 	DECLARE_METASOUND_DATA_REFERENCE_TYPES(FString, METASOUNDFRONTEND_API, FStringTypeInfo, FStringReadRef, FStringWriteRef);
 }
diff --git a/Engine/Plugins/Runtime/Metasound/Source/MetasoundFrontend/Public/MetasoundTransmissionRegistration.h b/Engine/Plugins/Runtime/Metasound/Source/MetasoundFrontend/Public/MetasoundTransmissionRegistration.h
--- a/Engine/Plugins/Runtime/Metasound/Source/MetasoundFrontend/Public/MetasoundTransmissionRegistration.h
+++ b/Engine/Plugins/Runtime/Metasound/Source/MetasoundFrontend/Public/MetasoundTransmissionRegistration.h
@@ -11,6 +11,13 @@ namespace Metasound
 {
 	namespace MetasoundTransmissionPrivate
 	{
+		// Determines whether send/receive nodes are enabled for a specific data type.
+		template<typename DataType>
+		struct TEnableTransmission
+		{
+			static constexpr bool Value = true;
+		};
+
 		// TTransmissionSupport determines whether the send/receive system is 
 		// supported for a given data type. It is used to add send/receive
 		// support when possible, and avoid errors when registering data types
@@ -19,6 +26,7 @@ namespace Metasound
 		struct TTransmissionSupport
 		{
 		private:
+			static constexpr bool bEnabled = TEnableTransmission<DataType>::Value;
 
 			// All types that support copy assignment and copy construction can be
 			// used in the transmission system. 
@@ -30,7 +38,7 @@ namespace Metasound
 
 		public:
 			
-			static constexpr bool bIsTransmissionSupported = bIsCopyable || bIsAudioDataType; 
+			static constexpr bool bIsTransmissionSupported = bEnabled && (bIsCopyable || bIsAudioDataType);
 		};
 
 		// At the time of writing this code, TArray incorrectly defines a copy constructor
@@ -43,9 +51,11 @@ namespace Metasound
 		{
 			static constexpr bool bIsCopyable = std::is_copy_assignable<ElementType>::value && std::is_copy_constructible<ElementType>::value;
 
+			static constexpr bool bEnabled = TEnableTransmission<TArray<ElementType>>::Value;
+
 		public:
 			
-			static constexpr bool bIsTransmissionSupported = bIsCopyable;
+			static constexpr bool bIsTransmissionSupported = bIsCopyable && bEnabled;
 		};
 	}
 
diff --git a/Engine/Plugins/Runtime/Metasound/Source/MetasoundStandardNodes/Private/MetasoundAudioFormats.cpp b/Engine/Plugins/Runtime/Metasound/Source/MetasoundStandardNodes/Private/MetasoundAudioFormats.cpp
--- a/Engine/Plugins/Runtime/Metasound/Source/MetasoundStandardNodes/Private/MetasoundAudioFormats.cpp
+++ b/Engine/Plugins/Runtime/Metasound/Source/MetasoundStandardNodes/Private/MetasoundAudioFormats.cpp
@@ -3,6 +3,7 @@
 #include "MetasoundAudioFormats.h"
 
 #include "CoreMinimal.h"
+#include "MetasoundArrayNodesRegistration.h"
 #include "MetasoundAudioBuffer.h"
 #include "MetasoundDataReference.h"
 #include "MetasoundDataTypeRegistrationMacro.h"
@@ -636,6 +637,58 @@ namespace Metasound
 		TSharedRef<FInputOperatorFactory, ESPMode::ThreadSafe> Factory;
 
 	};
+
+	// Disable arrays of audio formats. 
+	template<>
+	struct TEnableArrayNodes<FMonoAudioFormat>
+	{
+		static constexpr bool Value = false;
+	};
+
+	template<>
+	struct TEnableArrayNodes<FStereoAudioFormat>
+	{
+		static constexpr bool Value = false;
+	};
+
+	// Disable transmission of audio formats
+	template<>
+	struct TIsTransmittable<FMonoAudioFormat>
+	{
+		static constexpr bool Value = false;
+	};
+
+	template<>
+	struct TIsTransmittable<FStereoAudioFormat>
+	{
+		static constexpr bool Value = false;
+	};
+
+	// Disable auto converts using audio format construtors
+	template<typename FromDataType>
+	struct TIsAutoConvertible<FromDataType, FMonoAudioFormat>
+	{
+		static constexpr bool bIsConvertible = std::is_convertible<FromDataType, FAudioBuffer>::value;
+
+		static constexpr bool bIsFromArithmeticType = std::is_arithmetic<FromDataType>::value;
+
+	public:
+
+		static constexpr bool Value = bIsConvertible && (!bIsFromArithmeticType);
+	};
+
+	template<typename FromDataType>
+	struct TIsAutoConvertible<FromDataType, FStereoAudioFormat>
+	{
+		static constexpr bool bIsConvertible = std::is_convertible<FromDataType, FAudioBuffer>::value;
+
+		static constexpr bool bIsFromArithmeticType = std::is_arithmetic<FromDataType>::value;
+
+	public:
+
+		static constexpr bool Value = bIsConvertible && (!bIsFromArithmeticType);
+	};
+
 }
 
 // Data type registration has to happen after TInputNode<> and TOutputNode<> specializations
diff --git a/Engine/Plugins/Runtime/Metasound/Source/MetasoundStandardNodes/Private/MetasoundBasicFilters.cpp b/Engine/Plugins/Runtime/Metasound/Source/MetasoundStandardNodes/Private/MetasoundBasicFilters.cpp
--- a/Engine/Plugins/Runtime/Metasound/Source/MetasoundStandardNodes/Private/MetasoundBasicFilters.cpp
+++ b/Engine/Plugins/Runtime/Metasound/Source/MetasoundStandardNodes/Private/MetasoundBasicFilters.cpp
@@ -3,6 +3,7 @@
 #include "CoreMinimal.h"
 #include "Internationalization/Text.h"
 
+#include "MetasoundEnumRegistrationMacro.h"
 #include "MetasoundFacade.h"
 #include "MetasoundPrimitives.h"
 #include "MetasoundAudioBuffer.h"
@@ -983,4 +984,4 @@ namespace Metasound
 
 
 
-#undef LOCTEXT_NAMESPACE //MetasoundBasicFilterNodes
\ No newline at end of file
+#undef LOCTEXT_NAMESPACE //MetasoundBasicFilterNodes
diff --git a/Engine/Plugins/Runtime/Metasound/Source/MetasoundStandardNodes/Private/MetasoundMathNodes.cpp b/Engine/Plugins/Runtime/Metasound/Source/MetasoundStandardNodes/Private/MetasoundMathNodes.cpp
--- a/Engine/Plugins/Runtime/Metasound/Source/MetasoundStandardNodes/Private/MetasoundMathNodes.cpp
+++ b/Engine/Plugins/Runtime/Metasound/Source/MetasoundStandardNodes/Private/MetasoundMathNodes.cpp
@@ -817,6 +817,7 @@ namespace Metasound
 		}
 	};
 
+#if 0
 	// int64 has to be explicitly implemented as they are PODs that are not implemented as a literal base type.
 	template <>
 	class TMathOpAdd<int64>
@@ -863,6 +864,7 @@ namespace Metasound
 			}
 		}
 	};
+#endif
 
 	template <>
 	class TMathOpSubtract<FAudioBuffer>
@@ -983,6 +985,7 @@ namespace Metasound
 		}
 	};
 
+#if 0
 	// int64 has to be explicitly implemented as they are PODs that are not implemented as a literal base type.
 	template <>
 	class TMathOpSubtract<int64>
@@ -1029,6 +1032,7 @@ namespace Metasound
 			}
 		}
 	};
+#endif
 
 	template <>
 	class TMathOpMultiply<FAudioBuffer>
@@ -1218,6 +1222,7 @@ namespace Metasound
 		}
 	};
 
+#if 0
 	// int64 has to be explicitly implemented as they are PODs that are not implemented as a literal base type.
 	template <>
 	class TMathOpMultiply<int64>
@@ -1264,6 +1269,7 @@ namespace Metasound
 			}
 		}
 	};
+#endif
 
 	template <>
 	class TMathOpDivide<FTime, float>
@@ -1317,6 +1323,7 @@ namespace Metasound
 		}
 	};
 
+#if 0
 	// int64 has to be explicitly implemented as they are PODs that are not implemented as a literal base type.
 	template <>
 	class TMathOpDivide<int64>
@@ -1425,11 +1432,12 @@ namespace Metasound
 			}
 		}
 	};
+#endif
 
 	// Definitions
 	DEFINE_METASOUND_MATHOP(Add, Float, float, LOCTEXT("Metasound_MathAddFloatNodeDescription", "Adds floats."))
 	DEFINE_METASOUND_MATHOP(Add, Int32, int32, LOCTEXT("Metasound_MathAddInt32NodeDescription", "Adds int32s."))
-	DEFINE_METASOUND_MATHOP(Add, Int64, int64, LOCTEXT("Metasound_MathAddInt64NodeDescription", "Adds int64s."))
+	//DEFINE_METASOUND_MATHOP(Add, Int64, int64, LOCTEXT("Metasound_MathAddInt64NodeDescription", "Adds int64s."))
 	DEFINE_METASOUND_MATHOP(Add, Audio, FAudioBuffer, LOCTEXT("Metasound_MathAddBufferNodeDescription", "Adds buffers together by sample."))
 	DEFINE_METASOUND_MATHOP(Add, Time, FTime, LOCTEXT("Metasound_MathAddTimeNodeDescription", "Adds time values."))
 
@@ -1437,13 +1445,13 @@ namespace Metasound
 
 	DEFINE_METASOUND_MATHOP(Subtract, Float, float, LOCTEXT("Metasound_MathSubractFloatNodeDescription", "Subtracts floats."))
 	DEFINE_METASOUND_MATHOP(Subtract, Int32, int32, LOCTEXT("Metasound_MathSubractInt32NodeDescription", "Subtracts int32s."))
-	DEFINE_METASOUND_MATHOP(Subtract, Int64, int64, LOCTEXT("Metasound_MathSubractInt64NodeDescription", "Subtracts int64s."))
+	//DEFINE_METASOUND_MATHOP(Subtract, Int64, int64, LOCTEXT("Metasound_MathSubractInt64NodeDescription", "Subtracts int64s."))
 	DEFINE_METASOUND_MATHOP(Subtract, Audio, FAudioBuffer, LOCTEXT("Metasound_MathSubtractBufferNodeDescription", "Subtracts buffers sample-by-sample."))
 	DEFINE_METASOUND_MATHOP(Subtract, Time, FTime, LOCTEXT("Metasound_MathSubractTimeNodeDescription", "Subtracts time values."))
 
 	DEFINE_METASOUND_MATHOP(Multiply, Float, float, LOCTEXT("Metasound_MathMultiplyFloatNodeDescription", "Multiplies floats."))
 	DEFINE_METASOUND_MATHOP(Multiply, Int32, int32, LOCTEXT("Metasound_MathMultiplyInt32NodeDescription", "Multiplies int32s."))
-	DEFINE_METASOUND_MATHOP(Multiply, Int64, int64, LOCTEXT("Metasound_MathMultiplyInt64NodeDescription", "Multiplies int64s."))
+	//DEFINE_METASOUND_MATHOP(Multiply, Int64, int64, LOCTEXT("Metasound_MathMultiplyInt64NodeDescription", "Multiplies int64s."))
 	DEFINE_METASOUND_MATHOP(Multiply, Audio, FAudioBuffer, LOCTEXT("Metasound_MathMultiplyBufferNodeDescription", "Multiplies buffers together sample-by-sample."))
 
 	DEFINE_METASOUND_OPERAND_TYPED_MATHOP(Multiply, Audio, FAudioBuffer, Float, float, LOCTEXT("Metasound_MathMultiplyAudioByFloatDescription", "Multiplies buffer by float scalars."))
@@ -1451,12 +1459,12 @@ namespace Metasound
 
 	DEFINE_METASOUND_MATHOP(Divide, Float, float, LOCTEXT("Metasound_MathDivideFloatNodeDescription", "Divide float by another float."))
 	DEFINE_METASOUND_MATHOP(Divide, Int32, int32, LOCTEXT("Metasound_MathDivideInt32NodeDescription", "Divide int32 by another int32."))
-	DEFINE_METASOUND_MATHOP(Divide, Int64, int64, LOCTEXT("Metasound_MathDivideInt64NodeDescription", "Divide int64 by another int64."))
+	//DEFINE_METASOUND_MATHOP(Divide, Int64, int64, LOCTEXT("Metasound_MathDivideInt64NodeDescription", "Divide int64 by another int64."))
 
 	DEFINE_METASOUND_OPERAND_TYPED_MATHOP(Divide, Time, FTime, Float, float, LOCTEXT("Metasound_MathDivideTimeNodeDescription", "Divides time by floats."))
 
 	DEFINE_METASOUND_MATHOP(Modulo, Int32, int32, LOCTEXT("Metasound_MathModulusInt32NodeDescription", "Modulo int32 by another int32."))
-	DEFINE_METASOUND_MATHOP(Modulo, Int64, int64, LOCTEXT("Metasound_MathModulusInt64NodeDescription", "Modulo int64 by another int64."))
+	//DEFINE_METASOUND_MATHOP(Modulo, Int64, int64, LOCTEXT("Metasound_MathModulusInt64NodeDescription", "Modulo int64 by another int64."))
 
 	DEFINE_METASOUND_MATHOP(Power, Float, float, LOCTEXT("Metasound_MathPowerFloatNodeDescription", "Raise float to the power of another float."))
 
@@ -1464,14 +1472,14 @@ namespace Metasound
 
 	METASOUND_REGISTER_NODE(FAddFloatNode)
 	METASOUND_REGISTER_NODE(FAddInt32Node)
-	METASOUND_REGISTER_NODE(FAddInt64Node)
+	//METASOUND_REGISTER_NODE(FAddInt64Node)
 	METASOUND_REGISTER_NODE(FAddTimeNode)
 	METASOUND_REGISTER_NODE(FAddAudioNode)
 	METASOUND_REGISTER_NODE(FAddAudioFloatNode)
 
 	METASOUND_REGISTER_NODE(FSubtractFloatNode)
 	METASOUND_REGISTER_NODE(FSubtractInt32Node)
-	METASOUND_REGISTER_NODE(FSubtractInt64Node)
+	//METASOUND_REGISTER_NODE(FSubtractInt64Node)
 	METASOUND_REGISTER_NODE(FSubtractTimeNode)
 	METASOUND_REGISTER_NODE(FSubtractAudioNode)
 
@@ -1479,16 +1487,16 @@ namespace Metasound
 	METASOUND_REGISTER_NODE(FMultiplyAudioFloatNode)
 	METASOUND_REGISTER_NODE(FMultiplyFloatNode)
 	METASOUND_REGISTER_NODE(FMultiplyInt32Node)
-	METASOUND_REGISTER_NODE(FMultiplyInt64Node)
+	//METASOUND_REGISTER_NODE(FMultiplyInt64Node)
 	METASOUND_REGISTER_NODE(FMultiplyTimeFloatNode)
 
 	METASOUND_REGISTER_NODE(FDivideFloatNode)
 	METASOUND_REGISTER_NODE(FDivideInt32Node)
-	METASOUND_REGISTER_NODE(FDivideInt64Node)
+	//METASOUND_REGISTER_NODE(FDivideInt64Node)
 	METASOUND_REGISTER_NODE(FDivideTimeFloatNode)
 
 	METASOUND_REGISTER_NODE(FModuloInt32Node)
-	METASOUND_REGISTER_NODE(FModuloInt64Node)
+	//METASOUND_REGISTER_NODE(FModuloInt64Node)
 
 	METASOUND_REGISTER_NODE(FPowerFloatNode)
 
diff --git a/Engine/Plugins/Runtime/Metasound/Source/MetasoundStandardNodes/Private/MetasoundMidiScaleToArrayNode.cpp b/Engine/Plugins/Runtime/Metasound/Source/MetasoundStandardNodes/Private/MetasoundMidiScaleToArrayNode.cpp
--- a/Engine/Plugins/Runtime/Metasound/Source/MetasoundStandardNodes/Private/MetasoundMidiScaleToArrayNode.cpp
+++ b/Engine/Plugins/Runtime/Metasound/Source/MetasoundStandardNodes/Private/MetasoundMidiScaleToArrayNode.cpp
@@ -4,6 +4,7 @@
 #include "DSP/MidiNoteQuantizer.h"
 #include "Internationalization/Text.h"
 
+#include "MetasoundEnumRegistrationMacro.h"
 #include "MetasoundFacade.h"
 #include "MetasoundPrimitives.h"
 #include "MetasoundScaleEnums.h"
@@ -210,4 +211,4 @@ namespace Metasound
 
 } // namespace Metasound
 
-#undef LOCTEXT_NAMESPACE //MetasoundBasicFilterNodes
\ No newline at end of file
+#undef LOCTEXT_NAMESPACE //MetasoundBasicFilterNodes
diff --git a/Engine/Plugins/Runtime/Metasound/Source/MetasoundStandardNodes/Private/MetasoundNoiseGenerator.cpp b/Engine/Plugins/Runtime/Metasound/Source/MetasoundStandardNodes/Private/MetasoundNoiseGenerator.cpp
--- a/Engine/Plugins/Runtime/Metasound/Source/MetasoundStandardNodes/Private/MetasoundNoiseGenerator.cpp
+++ b/Engine/Plugins/Runtime/Metasound/Source/MetasoundStandardNodes/Private/MetasoundNoiseGenerator.cpp
@@ -2,6 +2,7 @@
 #include "MetasoundNoiseGenerator.h"
 
 #include "MetasoundAudioBuffer.h"
+#include "MetasoundEnumRegistrationMacro.h"
 #include "MetasoundExecutableOperator.h"
 #include "MetasoundFacade.h"
 #include "MetasoundNodeRegistrationMacro.h"
diff --git a/Engine/Plugins/Runtime/Metasound/Source/MetasoundStandardNodes/Private/MetasoundOscillatorNodes.cpp b/Engine/Plugins/Runtime/Metasound/Source/MetasoundStandardNodes/Private/MetasoundOscillatorNodes.cpp
--- a/Engine/Plugins/Runtime/Metasound/Source/MetasoundStandardNodes/Private/MetasoundOscillatorNodes.cpp
+++ b/Engine/Plugins/Runtime/Metasound/Source/MetasoundStandardNodes/Private/MetasoundOscillatorNodes.cpp
@@ -2,17 +2,19 @@
 #include "MetasoundOscillatorNodes.h"
 
 #include "DSP/Dsp.h"
+
 #include "MetasoundAudioBuffer.h"
+#include "MetasoundEnumRegistrationMacro.h"
 #include "MetasoundExecutableOperator.h"
 #include "MetasoundFacade.h"
 #include "MetasoundNodeRegistrationMacro.h"
 #include "MetasoundDataTypeRegistrationMacro.h"
 #include "MetasoundPrimitives.h"
+#include "MetasoundStandardNodesCategories.h"
 #include "MetasoundStandardNodesNames.h"
+#include "MetasoundTime.h"
 #include "MetasoundTrigger.h"
 #include "MetasoundVertex.h"
-#include "MetasoundTime.h"
-#include "MetasoundStandardNodesCategories.h"
 
 #define LOCTEXT_NAMESPACE "MetasoundStandardNodes"
 
diff --git a/Engine/Plugins/Runtime/Metasound/Source/MetasoundStandardNodes/Private/MetasoundStereoDelayNode.cpp b/Engine/Plugins/Runtime/Metasound/Source/MetasoundStandardNodes/Private/MetasoundStereoDelayNode.cpp
--- a/Engine/Plugins/Runtime/Metasound/Source/MetasoundStandardNodes/Private/MetasoundStereoDelayNode.cpp
+++ b/Engine/Plugins/Runtime/Metasound/Source/MetasoundStandardNodes/Private/MetasoundStereoDelayNode.cpp
@@ -1,6 +1,7 @@
 // Copyright Epic Games, Inc. All Rights Reserved.
 
 #include "Internationalization/Text.h"
+#include "MetasoundEnumRegistrationMacro.h"
 #include "MetasoundExecutableOperator.h"
 #include "MetasoundNodeRegistrationMacro.h"
 #include "MetasoundDataTypeRegistrationMacro.h"
diff --git a/Engine/Plugins/Runtime/Metasound/Source/MetasoundStandardNodes/Private/MetasoundStereoPannerNode.cpp b/Engine/Plugins/Runtime/Metasound/Source/MetasoundStandardNodes/Private/MetasoundStereoPannerNode.cpp
--- a/Engine/Plugins/Runtime/Metasound/Source/MetasoundStandardNodes/Private/MetasoundStereoPannerNode.cpp
+++ b/Engine/Plugins/Runtime/Metasound/Source/MetasoundStandardNodes/Private/MetasoundStereoPannerNode.cpp
@@ -1,6 +1,7 @@
 // Copyright Epic Games, Inc. All Rights Reserved.
 
 #include "Internationalization/Text.h"
+#include "MetasoundEnumRegistrationMacro.h"
 #include "MetasoundExecutableOperator.h"
 #include "MetasoundNodeRegistrationMacro.h"
 #include "MetasoundDataTypeRegistrationMacro.h"
diff --git a/Engine/Plugins/Runtime/Metasound/Source/MetasoundStandardNodes/Private/MetasoundTriggerOnThresholdNode.cpp b/Engine/Plugins/Runtime/Metasound/Source/MetasoundStandardNodes/Private/MetasoundTriggerOnThresholdNode.cpp
--- a/Engine/Plugins/Runtime/Metasound/Source/MetasoundStandardNodes/Private/MetasoundTriggerOnThresholdNode.cpp
+++ b/Engine/Plugins/Runtime/Metasound/Source/MetasoundStandardNodes/Private/MetasoundTriggerOnThresholdNode.cpp
@@ -2,6 +2,7 @@
 #include "MetasoundTriggerOnThresholdNode.h"
 
 #include "MetasoundAudioBuffer.h"
+#include "MetasoundEnumRegistrationMacro.h"
 #include "MetasoundExecutableOperator.h"
 #include "MetasoundFacade.h"
 #include "MetasoundNodeRegistrationMacro.h"
diff --git a/Engine/Plugins/Runtime/Metasound/Source/MetasoundStandardNodes/Private/MetasoundWaveShaperNode.cpp b/Engine/Plugins/Runtime/Metasound/Source/MetasoundStandardNodes/Private/MetasoundWaveShaperNode.cpp
--- a/Engine/Plugins/Runtime/Metasound/Source/MetasoundStandardNodes/Private/MetasoundWaveShaperNode.cpp
+++ b/Engine/Plugins/Runtime/Metasound/Source/MetasoundStandardNodes/Private/MetasoundWaveShaperNode.cpp
@@ -1,6 +1,7 @@
 // Copyright Epic Games, Inc. All Rights Reserved.
 
 #include "Internationalization/Text.h"
+#include "MetasoundEnumRegistrationMacro.h"
 #include "MetasoundExecutableOperator.h"
 #include "MetasoundNodeRegistrationMacro.h"
 #include "MetasoundDataTypeRegistrationMacro.h"
diff --git a/Engine/Plugins/Runtime/Metasound/Source/MetasoundStandardNodes/Private/MetasoundWaveWriterNode.cpp b/Engine/Plugins/Runtime/Metasound/Source/MetasoundStandardNodes/Private/MetasoundWaveWriterNode.cpp
--- a/Engine/Plugins/Runtime/Metasound/Source/MetasoundStandardNodes/Private/MetasoundWaveWriterNode.cpp
+++ b/Engine/Plugins/Runtime/Metasound/Source/MetasoundStandardNodes/Private/MetasoundWaveWriterNode.cpp
@@ -3,12 +3,13 @@
 #include "MetasoundWaveWriterNode.h"
 
 #include "HAL/FileManager.h"
+#include "MetasoundBuildError.h"
+#include "MetasoundEnumRegistrationMacro.h"
+#include "MetasoundExecutableOperator.h"
 #include "MetasoundNodeRegistrationMacro.h"
 #include "MetasoundPrimitives.h"
-#include "MetasoundStandardNodesNames.h"
 #include "MetasoundStandardNodesCategories.h"
-#include "MetasoundExecutableOperator.h"
-#include "MetasoundBuildError.h"
+#include "MetasoundStandardNodesNames.h"
 #include "Misc/Paths.h"
 #include "Misc/ScopeLock.h"
 
diff --git a/Engine/Plugins/Runtime/Metasound/Source/MetasoundStandardNodes/Public/MetasoundRandomNode.h b/Engine/Plugins/Runtime/Metasound/Source/MetasoundStandardNodes/Public/MetasoundRandomNode.h
--- a/Engine/Plugins/Runtime/Metasound/Source/MetasoundStandardNodes/Public/MetasoundRandomNode.h
+++ b/Engine/Plugins/Runtime/Metasound/Source/MetasoundStandardNodes/Public/MetasoundRandomNode.h
@@ -508,3 +508,4 @@ namespace Metasound
 
 } // namespace Metasound
 
+#undef LOCTEXT_NAMESPACE
diff --git a/Engine/Plugins/Runtime/Metasound/Source/MetasoundStandardNodes/Public/MetasoundScaleEnums.h b/Engine/Plugins/Runtime/Metasound/Source/MetasoundStandardNodes/Public/MetasoundScaleEnums.h
--- a/Engine/Plugins/Runtime/Metasound/Source/MetasoundStandardNodes/Public/MetasoundScaleEnums.h
+++ b/Engine/Plugins/Runtime/Metasound/Source/MetasoundStandardNodes/Public/MetasoundScaleEnums.h
@@ -4,6 +4,7 @@
 
 #include "CoreMinimal.h"
 #include "MetasoundEnum.h"
+#include "MetasoundEnumRegistrationMacro.h"
 #include "DSP/MidiNoteQuantizer.h"
 #include "Internationalization/Text.h"
 #include "MetasoundDataTypeRegistrationMacro.h"
diff --git a/Engine/Plugins/Runtime/Metasound/Source/MetasoundStandardNodes/Public/MetasoundTriggerAccumulatorNode.h b/Engine/Plugins/Runtime/Metasound/Source/MetasoundStandardNodes/Public/MetasoundTriggerAccumulatorNode.h
--- a/Engine/Plugins/Runtime/Metasound/Source/MetasoundStandardNodes/Public/MetasoundTriggerAccumulatorNode.h
+++ b/Engine/Plugins/Runtime/Metasound/Source/MetasoundStandardNodes/Public/MetasoundTriggerAccumulatorNode.h
@@ -12,7 +12,6 @@
 
 #define LOCTEXT_NAMESPACE "MetasoundStandardNodes"
 
-
 namespace Metasound
 {
 	namespace MetasoundTriggerAccumulatorNodePrivate
@@ -209,3 +208,5 @@ namespace Metasound
 	};
 
 } // namespace Metasound
+
+#undef LOCTEXT_NAMESPACE
diff --git a/Engine/Plugins/Runtime/Metasound/Source/MetasoundStandardNodes/Public/MetasoundTriggerAnyNode.h b/Engine/Plugins/Runtime/Metasound/Source/MetasoundStandardNodes/Public/MetasoundTriggerAnyNode.h
--- a/Engine/Plugins/Runtime/Metasound/Source/MetasoundStandardNodes/Public/MetasoundTriggerAnyNode.h
+++ b/Engine/Plugins/Runtime/Metasound/Source/MetasoundStandardNodes/Public/MetasoundTriggerAnyNode.h
@@ -162,3 +162,5 @@ namespace Metasound
 	};
 
 } // namespace Metasound
+
+#undef LOCTEXT_NAMESPACE
diff --git a/Engine/Plugins/Runtime/Metasound/Source/MetasoundStandardNodes/Public/MetasoundTriggerCompareNode.h b/Engine/Plugins/Runtime/Metasound/Source/MetasoundStandardNodes/Public/MetasoundTriggerCompareNode.h
--- a/Engine/Plugins/Runtime/Metasound/Source/MetasoundStandardNodes/Public/MetasoundTriggerCompareNode.h
+++ b/Engine/Plugins/Runtime/Metasound/Source/MetasoundStandardNodes/Public/MetasoundTriggerCompareNode.h
@@ -5,6 +5,7 @@
 #include "CoreMinimal.h"
 #include "MetasoundBuilderInterface.h"
 #include "MetasoundDataReferenceCollection.h"
+#include "MetasoundEnumRegistrationMacro.h"
 #include "MetasoundExecutableOperator.h"
 #include "MetasoundFacade.h"
 #include "MetasoundNode.h"
@@ -238,4 +239,4 @@ namespace Metasound
 	};
 }
 
-#undef LOCTEXT_NAMESPACE
\ No newline at end of file
+#undef LOCTEXT_NAMESPACE
diff --git a/Engine/Plugins/Runtime/Metasound/Source/MetasoundStandardNodes/Public/MetasoundValueNode.h b/Engine/Plugins/Runtime/Metasound/Source/MetasoundStandardNodes/Public/MetasoundValueNode.h
--- a/Engine/Plugins/Runtime/Metasound/Source/MetasoundStandardNodes/Public/MetasoundValueNode.h
+++ b/Engine/Plugins/Runtime/Metasound/Source/MetasoundStandardNodes/Public/MetasoundValueNode.h
@@ -196,3 +196,4 @@ namespace Metasound
 	};
 } // namespace Metasound
 
+#undef LOCTEXT_NAMESPACE
diff --git a/Engine/Plugins/Runtime/Metasound/Source/MetasoundStandardNodes/Public/MetasoundWaveWriterNode.h b/Engine/Plugins/Runtime/Metasound/Source/MetasoundStandardNodes/Public/MetasoundWaveWriterNode.h
--- a/Engine/Plugins/Runtime/Metasound/Source/MetasoundStandardNodes/Public/MetasoundWaveWriterNode.h
+++ b/Engine/Plugins/Runtime/Metasound/Source/MetasoundStandardNodes/Public/MetasoundWaveWriterNode.h
@@ -5,8 +5,6 @@
 #include "CoreMinimal.h"
 #include "MetasoundFacade.h"
 
-#define LOCTEXT_NAMESPACE "MetasoundStandardNodes"
-
 namespace Metasound
 {
 	class METASOUNDSTANDARDNODES_API FWaveWriterNode : public FNodeFacade
@@ -16,4 +14,5 @@ namespace Metasound
 		FWaveWriterNode(const FNodeInitData& InInitData);
 	private:
 	};
-} // namespace Metasound
\ No newline at end of file
+} // namespace Metasound
+
