- Fix triggered inputs firing on initialization in subsequent SoundGenerator instances
    - Caused by leak in Transmission System's Global DataChannels, wherein AudioComponents were not removing instance data, which caused new sounds played from that Component taking on last value set from prior instance.
    - Remove custom input preview Sender logic in MetasoundEditor to avoid managing multiple senders (can just use the one housed on the preview AudioComponent).

diff --git a/Engine/Plugins/Runtime/Metasound/Source/MetasoundEditor/Private/MetasoundEditor.cpp b/Engine/Plugins/Runtime/Metasound/Source/MetasoundEditor/Private/MetasoundEditor.cpp
--- a/Engine/Plugins/Runtime/Metasound/Source/MetasoundEditor/Private/MetasoundEditor.cpp
+++ b/Engine/Plugins/Runtime/Metasound/Source/MetasoundEditor/Private/MetasoundEditor.cpp
@@ -208,11 +208,25 @@ namespace Metasound
 			InTabManager->UnregisterTabSpawner(TabFactory::Names::Metasound);
 		}
 
+		bool FEditor::IsPlaying() const
+		{
+			if (Metasound)
+			{
+				FMetasoundAssetBase* MetasoundAsset = IMetasoundUObjectRegistry::Get().GetObjectAsAssetBase(Metasound);
+				check(MetasoundAsset);
+
+				if (UMetasoundEditorGraph* Graph = Cast<UMetasoundEditorGraph>(MetasoundAsset->GetGraph()))
+				{
+					return Graph->IsPreviewing();
+				}
+			}
+
+			return false;
+		}
+
 		FEditor::~FEditor()
 		{
-			// Stop any playing sounds when the editor closes
-			UAudioComponent* Component = GEditor->GetPreviewAudioComponent();
-			if (Component && Component->IsPlaying())
+			if (IsPlaying())
 			{
 				Stop();
 			}
@@ -712,44 +726,30 @@ namespace Metasound
 
 		void FEditor::Play()
 		{
-			check(GEditor);
-
 			if (USoundBase* MetasoundToPlay = Cast<USoundBase>(Metasound))
 			{
 				// Set the send to the audio bus that is used for analyzing the metasound output
+				check(GEditor);
 				if (UAudioComponent* PreviewComp = GEditor->PlayPreviewSound(MetasoundToPlay))
 				{
 					PlayTime = 0.0;
 
-					if (MetasoundAudioBus.IsValid())
-					{
-						PreviewComp->SetAudioBusSendPostEffect(MetasoundAudioBus.Get(), 1.0f);
-					}
-
-					TUniquePtr<IAudioInstanceTransmitter> MetasoundComm;
-
-					// Communicating with playing component requires audio component ID and sample rate. 
-					if (FAudioDevice* PreviewAudioDevice = PreviewComp->GetAudioDevice())
+					UObject* CommObject = PreviewComp->GetCommunicationInterface().GetObject();
+					if (ensure(CommObject))
 					{
-						if(FMetasoundAssetBase* MetasoundAsset = IMetasoundUObjectRegistry::Get().GetObjectAsAssetBase(Metasound))
-						{
-							const uint64 AudioComponentID = PreviewComp->GetAudioComponentID();
-							const float SampleRate = PreviewAudioDevice->GetSampleRate();
-							MetasoundComm = MetasoundToPlay->CreateInstanceTransmitter(FAudioInstanceTransmitterInitParams { AudioComponentID, SampleRate });
-						}
+						SetPreviewID(CommObject->GetUniqueID());
 					}
 
-					if (UMetasoundEditorGraph* Graph = CastChecked<UMetasoundEditorGraph>(MetasoundGraphEditor->GetCurrentGraph()))
+					if (MetasoundAudioBus.IsValid())
 					{
-						Graph->SetMetasoundInstanceTransmitter(MoveTemp(MetasoundComm));
+						PreviewComp->SetAudioBusSendPostEffect(MetasoundAudioBus.Get(), 1.0f);
 					}
 				}
 
 				MetasoundGraphEditor->RegisterActiveTimer(0.0f,
 					FWidgetActiveTimerDelegate::CreateLambda([this](double InCurrentTime, float InDeltaTime)
 					{
-						UAudioComponent* PreviewComp = GEditor->GetPreviewAudioComponent();
-						if (PreviewComp && PreviewComp->IsPlaying())
+						if (IsPlaying())
 						{
 							if (PlayTimeWidget.IsValid())
 							{
@@ -764,6 +764,7 @@ namespace Metasound
 						}
 						else
 						{
+							SetPreviewID(INDEX_NONE);
 							PlayTime = 0.0;
 							PlayTimeWidget->SetText(FText::GetEmpty());
 							return EActiveTimerReturnType::Stop;
@@ -773,6 +774,26 @@ namespace Metasound
 			}
 		}
 
+		void FEditor::SetPreviewID(uint32 InPreviewID)
+		{
+			if (!Metasound)
+			{
+				return;
+			}
+
+			FMetasoundAssetBase* MetasoundAsset = IMetasoundUObjectRegistry::Get().GetObjectAsAssetBase(Metasound);
+			check(MetasoundAsset);
+
+			UEdGraph* Graph = MetasoundAsset->GetGraph();
+			if (!Graph)
+			{
+				return;
+			}
+			UMetasoundEditorGraph* MetasoundGraph = CastChecked<UMetasoundEditorGraph>(MetasoundAsset->GetGraph());
+
+			MetasoundGraph->SetPreviewID(InPreviewID);
+		}
+
 		void FEditor::ExecuteNode()
 		{
 			const FGraphPanelSelectionSet SelectedNodes = MetasoundGraphEditor->GetSelectedNodes();
@@ -796,14 +817,14 @@ namespace Metasound
 		{
 			check(GEditor);
 			GEditor->ResetPreviewAudioComponent();
+			SetPreviewID(INDEX_NONE);
 		}
 
 		void FEditor::TogglePlayback()
 		{
 			check(GEditor);
 
-			UAudioComponent* Component = GEditor->GetPreviewAudioComponent();
-			if (Component && Component->IsPlaying())
+			if (IsPlaying())
 			{
 				Stop();
 			}
@@ -815,17 +836,22 @@ namespace Metasound
 
 		void FEditor::ExecuteNode(UEdGraphNode* Node)
 		{
+			if (!IsPlaying())
+			{
+				return;
+			}
+
 			if (UMetasoundEditorGraphInputNode* InputNode = Cast<UMetasoundEditorGraphInputNode>(Node))
 			{
-				// TODO: fix how identifying the parameter to update is determined. It should not be done
-				// with a "DisplayName" but rather the vertex Guid.
-				if (UMetasoundEditorGraph* MetasoundGraph = Cast<UMetasoundEditorGraph>(InputNode->GetOuter()))
+				if (UAudioComponent* PreviewComponent = GEditor->GetPreviewAudioComponent())
 				{
-					if (IAudioInstanceTransmitter* Transmitter = MetasoundGraph->GetMetasoundInstanceTransmitter())
+					// TODO: fix how identifying the parameter to update is determined. It should not be done
+					// with a "DisplayName" but rather the vertex Guid.
+					if (TScriptInterface<IAudioCommunicationInterface> CommInterface = PreviewComponent->GetCommunicationInterface())
 					{
 						Metasound::Frontend::FConstNodeHandle NodeHandle = InputNode->GetConstNodeHandle();
 						Metasound::FVertexKey VertexKey = Metasound::FVertexKey(NodeHandle->GetDisplayName().ToString());
-						Transmitter->SetParameter(*VertexKey, true);
+						CommInterface->Trigger(*VertexKey);
 					}
 				}
 			}
diff --git a/Engine/Plugins/Runtime/Metasound/Source/MetasoundEditor/Private/MetasoundEditor.h b/Engine/Plugins/Runtime/Metasound/Source/MetasoundEditor/Private/MetasoundEditor.h
--- a/Engine/Plugins/Runtime/Metasound/Source/MetasoundEditor/Private/MetasoundEditor.h
+++ b/Engine/Plugins/Runtime/Metasound/Source/MetasoundEditor/Private/MetasoundEditor.h
@@ -182,6 +182,8 @@ namespace Metasound
 			void RedoGraphAction();
 
 		private:
+			void SetPreviewID(uint32 InPreviewID);
+
 			/** FNotifyHook interface */
 			virtual void NotifyPostChange(const FPropertyChangedEvent& PropertyChangedEvent, FProperty* PropertyThatChanged) override;
 
@@ -251,6 +253,8 @@ namespace Metasound
 
 			void NotifyUserModifiedBySync();
 
+			bool IsPlaying() const;
+
 			/** List of open tool panels; used to ensure only one exists at any one time */
 			TMap<FName, TWeakPtr<SDockableTab>> SpawnedToolPanels;
 
diff --git a/Engine/Plugins/Runtime/Metasound/Source/MetasoundEditor/Private/MetasoundEditorGraph.cpp b/Engine/Plugins/Runtime/Metasound/Source/MetasoundEditor/Private/MetasoundEditorGraph.cpp
--- a/Engine/Plugins/Runtime/Metasound/Source/MetasoundEditor/Private/MetasoundEditorGraph.cpp
+++ b/Engine/Plugins/Runtime/Metasound/Source/MetasoundEditor/Private/MetasoundEditorGraph.cpp
@@ -1,6 +1,7 @@
 // Copyright Epic Games, Inc. All Rights Reserved.
 #include "MetasoundEditorGraph.h"
 
+#include "Components/AudioComponent.h"
 #include "EdGraph/EdGraphNode.h"
 #include "MetasoundEditorGraphBuilder.h"
 #include "MetasoundEditorGraphNode.h"
@@ -8,6 +9,7 @@
 #include "MetasoundEditorModule.h"
 #include "MetasoundUObjectRegistry.h"
 #include "ScopedTransaction.h"
+#include "Sound/AudioCommunicationInterface.h"
 
 #define LOCTEXT_NAMESPACE "MetaSoundEditor"
 
@@ -159,11 +161,11 @@ void UMetasoundEditorGraphInput::UpdateDocumentInput(bool bPostTransaction)
 	GraphHandle->SetDefaultInput(VertexID, Literal->GetDefault());
 }
 
-void UMetasoundEditorGraphInput::UpdatePreviewInstance(const Metasound::FVertexKey& InParameterName, IAudioInstanceTransmitter& InInstanceTransmitter) const
+void UMetasoundEditorGraphInput::UpdatePreviewInstance(const Metasound::FVertexKey& InParameterName, TScriptInterface<IAudioCommunicationInterface>& InCommInterface) const
 {
 	if (ensure(Literal))
 	{
-		Literal->UpdatePreviewInstance(InParameterName, InInstanceTransmitter);
+		Literal->UpdatePreviewInstance(InParameterName, InCommInterface);
 	}
 }
 
@@ -206,15 +208,19 @@ void UMetasoundEditorGraphInput::OnLiteralChanged(bool bPostTransaction)
 {
 	UpdateDocumentInput(bPostTransaction);
 
-	if (UMetasoundEditorGraph* MetasoundGraph = Cast<UMetasoundEditorGraph>(GetOuter()))
+	const bool bIsPreviewing = CastChecked<UMetasoundEditorGraph>(GetOuter())->IsPreviewing();
+	if (bIsPreviewing)
 	{
-		if (IAudioInstanceTransmitter* Transmitter = MetasoundGraph->GetMetasoundInstanceTransmitter())
+		UAudioComponent* PreviewComponent = GEditor->GetPreviewAudioComponent();
+		check(PreviewComponent);
+
+		if (TScriptInterface<IAudioCommunicationInterface> CommInterface = PreviewComponent->GetCommunicationInterface())
 		{
 			// TODO: fix how identifying the parameter to update is determined. It should not be done
 			// with a "DisplayName" but rather the vertex Guid.
 			Metasound::Frontend::FConstNodeHandle NodeHandle = GetConstNodeHandle();
 			Metasound::FVertexKey VertexKey = Metasound::FVertexKey(NodeHandle->GetDisplayName().ToString());
-			UpdatePreviewInstance(VertexKey, *Transmitter);
+			UpdatePreviewInstance(VertexKey, CommInterface);
 		}
 	}
 }
@@ -296,21 +302,6 @@ UObject& UMetasoundEditorGraph::GetMetasoundChecked() const
 	return *ParentMetasound;
 }
 
-void UMetasoundEditorGraph::SetMetasoundInstanceTransmitter(TUniquePtr<IAudioInstanceTransmitter>&& InTransmitter)
-{
-	Transmitter = MoveTemp(InTransmitter);
-}
-
-IAudioInstanceTransmitter* UMetasoundEditorGraph::GetMetasoundInstanceTransmitter()
-{
-	return Transmitter.Get();
-}
-
-const IAudioInstanceTransmitter* UMetasoundEditorGraph::GetMetasoundInstanceTransmitter() const
-{
-	return Transmitter.Get();
-}
-
 void UMetasoundEditorGraph::Synchronize()
 {
 	if (ParentMetasound)
@@ -451,6 +442,33 @@ void UMetasoundEditorGraph::IterateInputs(TUniqueFunction<void(UMetasoundEditorG
 	}
 }
 
+void UMetasoundEditorGraph::SetPreviewID(uint32 InPreviewID)
+{
+	PreviewID = InPreviewID;
+}
+
+bool UMetasoundEditorGraph::IsPreviewing() const
+{
+	UAudioComponent* PreviewComponent = GEditor->GetPreviewAudioComponent();
+	if (!PreviewComponent)
+	{
+		return false;
+	}
+
+	if (!PreviewComponent->IsPlaying())
+	{
+		return false;
+	}
+
+	UObject* CommInterface = PreviewComponent->GetCommunicationInterface().GetObject();
+	if (!CommInterface)
+	{
+		return false;
+	}
+
+	return CommInterface->GetUniqueID() == PreviewID;
+}
+
 void UMetasoundEditorGraph::IterateOutputs(TUniqueFunction<void(UMetasoundEditorGraphOutput&)> InFunction) const
 {
 	for (UMetasoundEditorGraphOutput* Output : Outputs)
diff --git a/Engine/Plugins/Runtime/Metasound/Source/MetasoundEditor/Private/MetasoundNodeDetailCustomization.cpp b/Engine/Plugins/Runtime/Metasound/Source/MetasoundEditor/Private/MetasoundNodeDetailCustomization.cpp
--- a/Engine/Plugins/Runtime/Metasound/Source/MetasoundEditor/Private/MetasoundNodeDetailCustomization.cpp
+++ b/Engine/Plugins/Runtime/Metasound/Source/MetasoundEditor/Private/MetasoundNodeDetailCustomization.cpp
@@ -1,6 +1,7 @@
 // Copyright Epic Games, Inc. All Rights Reserved.
 #include "MetasoundNodeDetailCustomization.h"
 
+#include "Components/AudioComponent.h"
 #include "Containers/Set.h"
 #include "CoreMinimal.h"
 #include "Delegates/Delegate.h"
@@ -82,15 +83,24 @@ namespace Metasound
 					return;
 				}
 
-				// TODO: fix how identifying the parameter to update is determined. It should not be done
-				// with a "DisplayName" but rather the vertex Guid.
-				if (UMetasoundEditorGraph* MetasoundGraph = Cast<UMetasoundEditorGraph>(Input->GetOuter()))
+				// If modifying graph not currently being previewed, do not forward request
+				if (UMetasoundEditorGraph* Graph = Cast<UMetasoundEditorGraph>(Input->GetOuter()))
 				{
-					if (IAudioInstanceTransmitter* Transmitter = MetasoundGraph->GetMetasoundInstanceTransmitter())
+					if (!Graph->IsPreviewing())
 					{
+						return;
+					}
+				}
+
+				if (UAudioComponent* PreviewComponent = GEditor->GetPreviewAudioComponent())
+				{
+					if (TScriptInterface<IAudioCommunicationInterface> CommInterface = PreviewComponent->GetCommunicationInterface())
+					{
+						// TODO: fix how identifying the parameter to update is determined. It should not be done
+						// with a "DisplayName" but rather the vertex Guid.
 						Metasound::Frontend::FConstNodeHandle NodeHandle = Input->GetConstNodeHandle();
 						Metasound::FVertexKey VertexKey = Metasound::FVertexKey(NodeHandle->GetDisplayName().ToString());
-						Transmitter->SetParameter(*VertexKey, true);
+						CommInterface->Trigger(*VertexKey);
 					}
 				}
 			}
diff --git a/Engine/Plugins/Runtime/Metasound/Source/MetasoundEditor/Public/MetasoundEditorGraph.h b/Engine/Plugins/Runtime/Metasound/Source/MetasoundEditor/Public/MetasoundEditorGraph.h
--- a/Engine/Plugins/Runtime/Metasound/Source/MetasoundEditor/Public/MetasoundEditorGraph.h
+++ b/Engine/Plugins/Runtime/Metasound/Source/MetasoundEditor/Public/MetasoundEditorGraph.h
@@ -8,7 +8,9 @@
 #include "MetasoundFrontendLiteral.h"
 #include "MetasoundSource.h"
 #include "MetasoundUObjectRegistry.h"
+#include "Sound/AudioCommunicationInterface.h"
 #include "UObject/ObjectMacros.h"
+#include "UObject/ScriptInterface.h"
 
 #include "MetasoundEditorGraph.generated.h"
 
@@ -57,7 +59,7 @@ class METASOUNDEDITOR_API UMetasoundEditorGraphInputLiteral : public UObject
 	GENERATED_BODY()
 
 public:
-	virtual void UpdatePreviewInstance(const Metasound::FVertexKey& InParameterName, IAudioInstanceTransmitter& InInstanceTransmitter) const
+	virtual void UpdatePreviewInstance(const Metasound::FVertexKey& InParameterName, TScriptInterface<IAudioCommunicationInterface>& InCommInterface) const
 	{
 	}
 
@@ -94,7 +96,7 @@ public:
 	UMetasoundEditorGraphInputLiteral* Literal;
 
 	void UpdateDocumentInput(bool bPostTransaction = true);
-	void UpdatePreviewInstance(const Metasound::FVertexKey& InParameterName, IAudioInstanceTransmitter& InInstanceTransmitter) const;
+	void UpdatePreviewInstance(const Metasound::FVertexKey& InParameterName, TScriptInterface<IAudioCommunicationInterface>& InCommInterface) const;
 
 	void OnDataTypeChanged() override;
 	void OnLiteralChanged(bool bPostTransaction = true);
@@ -130,18 +132,16 @@ public:
 	bool ContainsInput(UMetasoundEditorGraphInput* InInput) const;
 	bool ContainsOutput(UMetasoundEditorGraphOutput* InOutput) const;
 
-	virtual void Synchronize();
-
-	// Sets the transmitter interface for the current metasound preview.
-	void SetMetasoundInstanceTransmitter(TUniquePtr<IAudioInstanceTransmitter>&& InTransmitter);
-
-	// Gets the transmitter interface for the current metasound preview.
-	IAudioInstanceTransmitter* GetMetasoundInstanceTransmitter();
+	void SetPreviewID(uint32 InPreviewID);
+	bool IsPreviewing() const;
 
-	// Gets the transmitter interface for the current metasound preview.
-	const IAudioInstanceTransmitter* GetMetasoundInstanceTransmitter() const;
+	virtual void Synchronize();
 
 private:
+	// Preview ID is the Unique ID provided by the UObject that implements
+	// a sound's CommunicationInterface when a sound begins playing.
+	uint32 PreviewID = INDEX_NONE;
+
 	UPROPERTY()
 	TObjectPtr<UObject> ParentMetasound;
 
@@ -162,9 +162,6 @@ public:
 
 	bool RemoveVariable(UMetasoundEditorGraphVariable& InVariable);
 
-	TUniquePtr<IAudioInstanceTransmitter> Transmitter;
-
-// 	friend class Metasound::Editor::FEditor;
 	friend class UMetasoundFactory;
 	friend class UMetasoundSourceFactory;
 };
diff --git a/Engine/Plugins/Runtime/Metasound/Source/MetasoundEditor/Public/MetasoundEditorGraphInputNodes.h b/Engine/Plugins/Runtime/Metasound/Source/MetasoundEditor/Public/MetasoundEditorGraphInputNodes.h
--- a/Engine/Plugins/Runtime/Metasound/Source/MetasoundEditor/Public/MetasoundEditorGraphInputNodes.h
+++ b/Engine/Plugins/Runtime/Metasound/Source/MetasoundEditor/Public/MetasoundEditorGraphInputNodes.h
@@ -15,6 +15,7 @@
 #include "MetasoundFrontendRegistries.h"
 #include "MetasoundPrimitives.h"
 #include "Misc/Guid.h"
+#include "Sound/AudioCommunicationInterface.h"
 #include "Sound/SoundWave.h"
 #include "UObject/ObjectMacros.h"
 #include "UObject/SoftObjectPath.h"
@@ -74,11 +75,11 @@ public:
 		return Super::GetClassName();
 	}
 
-	virtual void UpdatePreviewInstance(const Metasound::FVertexKey& InParameterName, IAudioInstanceTransmitter& InInstanceTransmitter) const
+	virtual void UpdatePreviewInstance(const Metasound::FVertexKey& InParameterName, TScriptInterface<IAudioCommunicationInterface>& InCommInterface) const
 	{
 		if (Input)
 		{
-			Input->UpdatePreviewInstance(InParameterName, InInstanceTransmitter);
+			Input->UpdatePreviewInstance(InParameterName, InCommInterface);
 		}
 	}
 	
@@ -151,9 +152,9 @@ public:
 		Metasound::Editor::InputPrivate::ConvertLiteral<bool>(InLiteral, Default.Value);
 	}
 
-	virtual void UpdatePreviewInstance(const Metasound::FVertexKey& InParameterName, IAudioInstanceTransmitter& InInstanceTransmitter) const override
+	virtual void UpdatePreviewInstance(const Metasound::FVertexKey& InParameterName, TScriptInterface<IAudioCommunicationInterface>& InCommInterface) const override
 	{
-		InInstanceTransmitter.SetParameter(*InParameterName, Default.Value);
+		InCommInterface->SetBool(*InParameterName, Default.Value);
 	}
 };
 
@@ -188,12 +189,11 @@ public:
 		Metasound::Editor::InputPrivate::ConvertLiteralToArray<bool, FMetasoundEditorGraphInputBoolRef>(InLiteral, Default);
 	}
 
-	virtual void UpdatePreviewInstance(const Metasound::FVertexKey& InParameterName, IAudioInstanceTransmitter& InInstanceTransmitter) const override
+	virtual void UpdatePreviewInstance(const Metasound::FVertexKey& InParameterName, TScriptInterface<IAudioCommunicationInterface>& InCommInterface) const override
 	{
 		TArray<bool> BoolArray;
 		Algo::Transform(Default, BoolArray, [](const FMetasoundEditorGraphInputBoolRef& InValue) { return InValue.Value; });
-
-		InInstanceTransmitter.SetParameter(*InParameterName, MoveTemp(BoolArray));
+		InCommInterface->SetBoolArray(*InParameterName, BoolArray);
 	}
 };
 
@@ -235,9 +235,9 @@ public:
 		Metasound::Editor::InputPrivate::ConvertLiteral<int32>(InLiteral, Default.Value);
 	}
 
-	virtual void UpdatePreviewInstance(const Metasound::FVertexKey& InParameterName, IAudioInstanceTransmitter& InInstanceTransmitter) const override
+	virtual void UpdatePreviewInstance(const Metasound::FVertexKey& InParameterName, TScriptInterface<IAudioCommunicationInterface>& InCommInterface) const override
 	{
-		InInstanceTransmitter.SetParameter(*InParameterName, Default.Value);
+		InCommInterface->SetInt(*InParameterName, Default.Value);
 	}
 };
 
@@ -273,12 +273,11 @@ public:
 		Metasound::Editor::InputPrivate::ConvertLiteralToArray<int32, FMetasoundEditorGraphInputIntRef>(InLiteral, Default);
 	}
 
-	void UpdatePreviewInstance(const Metasound::FVertexKey& InParameterName, IAudioInstanceTransmitter& InInstanceTransmitter) const override
+	void UpdatePreviewInstance(const Metasound::FVertexKey& InParameterName, TScriptInterface<IAudioCommunicationInterface>& InCommInterface) const override
 	{
 		TArray<int32> IntArray;
 		Algo::Transform(Default, IntArray, [](const FMetasoundEditorGraphInputIntRef& InValue) { return InValue.Value; });
-
-		InInstanceTransmitter.SetParameter(*InParameterName, MoveTemp(IntArray));
+		InCommInterface->SetIntArray(*InParameterName, IntArray);
 	}
 };
 
@@ -310,9 +309,9 @@ public:
 		Metasound::Editor::InputPrivate::ConvertLiteral<float>(InLiteral, Default);
 	}
 
-	virtual void UpdatePreviewInstance(const Metasound::FVertexKey& InParameterName, IAudioInstanceTransmitter& InInstanceTransmitter) const override
+	virtual void UpdatePreviewInstance(const Metasound::FVertexKey& InParameterName, TScriptInterface<IAudioCommunicationInterface>& InCommInterface) const override
 	{
-		InInstanceTransmitter.SetParameter(*InParameterName, Default);
+		InCommInterface->SetFloat(*InParameterName, Default);
 	}
 };
 
@@ -344,9 +343,9 @@ public:
 		Metasound::Editor::InputPrivate::ConvertLiteralToArray<float>(InLiteral, Default);
 	}
 
-	virtual void UpdatePreviewInstance(const Metasound::FVertexKey& InParameterName, IAudioInstanceTransmitter& InInstanceTransmitter) const override
+	virtual void UpdatePreviewInstance(const Metasound::FVertexKey& InParameterName, TScriptInterface<IAudioCommunicationInterface>& InCommInterface) const override
 	{
-		InInstanceTransmitter.SetParameter(*InParameterName, TArray<float> { Default });
+		InCommInterface->SetFloatArray(*InParameterName, Default);
 	}
 };
 
@@ -378,9 +377,9 @@ public:
 		Metasound::Editor::InputPrivate::ConvertLiteral<FString>(InLiteral, Default);
 	}
 
-	virtual void UpdatePreviewInstance(const Metasound::FVertexKey& InParameterName, IAudioInstanceTransmitter& InInstanceTransmitter) const override
+	virtual void UpdatePreviewInstance(const Metasound::FVertexKey& InParameterName, TScriptInterface<IAudioCommunicationInterface>& InCommInterface) const override
 	{
-		InInstanceTransmitter.SetParameter(*InParameterName, FString { Default });
+		InCommInterface->SetString(*InParameterName, Default);
 	}
 };
 
@@ -412,9 +411,9 @@ public:
 		Metasound::Editor::InputPrivate::ConvertLiteralToArray<FString>(InLiteral, Default);
 	}
 
-	virtual void UpdatePreviewInstance(const Metasound::FVertexKey& InParameterName, IAudioInstanceTransmitter& InInstanceTransmitter) const override
+	virtual void UpdatePreviewInstance(const Metasound::FVertexKey& InParameterName, TScriptInterface<IAudioCommunicationInterface>& InCommInterface) const override
 	{
-		InInstanceTransmitter.SetParameter(*InParameterName, TArray<FString> { Default });
+		InCommInterface->SetStringArray(*InParameterName, Default);
 	}
 };
 
@@ -456,10 +455,9 @@ public:
 // 		Metasound::Editor::InputPrivate::ConvertLiteral<UObject*>(InLiteral, Default.Object);
 	}
 
-	virtual void UpdatePreviewInstance(const Metasound::FVertexKey& InParameterName, IAudioInstanceTransmitter& InInstanceTransmitter) const override
+	virtual void UpdatePreviewInstance(const Metasound::FVertexKey& InParameterName, TScriptInterface<IAudioCommunicationInterface>& InCommInterface) const override
 	{
 		// TODO. We need proxy object here safely.
-		//InInstanceTransmitter.SetParameter(*InParameterName, Default.Object);
 	}
 };
 
@@ -494,12 +492,10 @@ public:
 // 		Metasound::Editor::InputPrivate::ConvertLiteralToArray<UObject*, FMetasoundEditorGraphInputObjectRef>(InLiteral, Default);
 	}
 
-	virtual void UpdatePreviewInstance(const Metasound::FVertexKey& InParameterName, IAudioInstanceTransmitter& InInstanceTransmitter) const override
+	virtual void UpdatePreviewInstance(const Metasound::FVertexKey& InParameterName, TScriptInterface<IAudioCommunicationInterface>& InCommInterface) const override
 	{
 		TArray<UObject*> ObjectArray;
 		Algo::Transform(Default, ObjectArray, [](const FMetasoundEditorGraphInputObjectRef& InValue) { return InValue.Object; });
-		
 		// TODO. We need proxy object here safely.
-		//InInstanceTransmitter.SetParameter(*InParameterName, Default.Object);
 	}
 };
\ No newline at end of file
diff --git a/Engine/Plugins/Runtime/Metasound/Source/MetasoundFrontend/Private/MetasoundInstanceTransmitter.cpp b/Engine/Plugins/Runtime/Metasound/Source/MetasoundFrontend/Private/MetasoundInstanceTransmitter.cpp
--- a/Engine/Plugins/Runtime/Metasound/Source/MetasoundFrontend/Private/MetasoundInstanceTransmitter.cpp
+++ b/Engine/Plugins/Runtime/Metasound/Source/MetasoundFrontend/Private/MetasoundInstanceTransmitter.cpp
@@ -11,6 +11,18 @@ namespace Metasound
 	{
 	}
 
+	bool FMetasoundInstanceTransmitter::Shutdown()
+	{
+		bool bSuccess = true;
+
+		for (const FSendInfo& SendInfo : SendInfos)
+		{
+			bSuccess &= FDataTransmissionCenter::Get().UnregisterDataChannel(SendInfo.TypeName, SendInfo.Address.ChannelName);
+		}
+
+		return bSuccess;
+	}
+
 	uint64 FMetasoundInstanceTransmitter::GetInstanceID() const
 	{
 		return InstanceID;
@@ -104,26 +116,20 @@ namespace Metasound
 		return nullptr;
 	}
 
-	TUniquePtr<ISender> FMetasoundInstanceTransmitter::CreateSender(const FSendInfo& InInfo) const
-	{
-		// TODO: likely want to remove this and opt for different protocols having different behaviors.
-		const float DelayTimeInSeconds = 0.1f; //< This not used for non-audio routing.
-		FSenderInitParams InitParams = {OperatorSettings, DelayTimeInSeconds};
-
-		return FDataTransmissionCenter::Get().RegisterNewSender(InInfo.TypeName, InInfo.Address, InitParams);
-	}
-
 	ISender* FMetasoundInstanceTransmitter::AddSender(const FSendInfo& InInfo)
 	{
-		TUniquePtr<ISender> Sender = CreateSender(InInfo);
+		// TODO: likely want to remove this and opt for different protocols having different behaviors.
+		const float DelayTimeInSeconds = 0.1f; // This not used for non-audio routing.
+		const FSenderInitParams InitParams = { OperatorSettings, DelayTimeInSeconds };
 
+		TUniquePtr<ISender> Sender = FDataTransmissionCenter::Get().RegisterNewSender(InInfo.TypeName, InInfo.Address, InitParams);
 		if (ensureMsgf(Sender.IsValid(), TEXT("Failed to create sender [DataType:%s, Address:%s]"), *InInfo.TypeName.ToString(), *InInfo.Address.ChannelName.ToString()))
 		{
 			ISender* Ptr = Sender.Get();
 			InputSends.Add(InInfo.ParameterName, MoveTemp(Sender));
 			return Ptr;
 		}
-		
+
 		return nullptr;
 	}
 }
diff --git a/Engine/Plugins/Runtime/Metasound/Source/MetasoundFrontend/Private/MetasoundRouter.cpp b/Engine/Plugins/Runtime/Metasound/Source/MetasoundFrontend/Private/MetasoundRouter.cpp
--- a/Engine/Plugins/Runtime/Metasound/Source/MetasoundFrontend/Private/MetasoundRouter.cpp
+++ b/Engine/Plugins/Runtime/Metasound/Source/MetasoundFrontend/Private/MetasoundRouter.cpp
@@ -200,6 +200,11 @@ namespace Metasound
 		return MoveTemp(Receiver);
 	}
 
+	bool FDataTransmissionCenter::UnregisterDataChannel(const FName& InDataTypeName, const FName& InChannelName)
+	{
+		return GlobalRouter.UnregisterDataChannel(InDataTypeName, InChannelName);
+	}
+
 	bool FDataTransmissionCenter::PushLiteral(FName DataType, FName GlobalChannelName, const FLiteral& InParam)
 	{
 		return GlobalRouter.PushLiteral(DataType, GlobalChannelName, InParam);
@@ -289,9 +294,29 @@ namespace Metasound
 		}
 	}
 
-	TUniquePtr<IReceiver> FAddressRouter::RegisterNewReceiver(const FName& InDataTypeName, const FName& InChannelName, const FReceiverInitParams& InitParams)
+	bool FAddressRouter::UnregisterDataChannel(const FName& InDataTypeName, const FName& InChannelName)
 	{
+		FScopeLock ScopeLock(&DataChannelMapMutationLock);
+		const FDataChannelKey ChannelKey = GetDataChannelKey(InDataTypeName, InChannelName);
+
+		if (TSharedRef<IDataChannel, ESPMode::ThreadSafe>* Channel = DataChannelMap.Find(ChannelKey))
+		{
+			if (const int32 NumReceiversActive = Channel->Get().GetNumActiveReceivers())
+			{
+				UE_LOG(LogMetaSound, Verbose, TEXT("DataChannel '%s' of type '%s' shutting down with %d receivers active."), *InChannelName.ToString(), *InDataTypeName.ToString(), NumReceiversActive);
+			}
 
+			if (const int32 NumSendersActive = Channel->Get().GetNumActiveSenders())
+			{
+				UE_LOG(LogMetaSound, Verbose, TEXT("DataChannel '%s' of type '%s' shutting down with %d senders active."), *InChannelName.ToString(), *InDataTypeName.ToString(), NumSendersActive);
+			}
+		}
+
+		return DataChannelMap.Remove(ChannelKey) > 0;
+	}
+
+	TUniquePtr<IReceiver> FAddressRouter::RegisterNewReceiver(const FName& InDataTypeName, const FName& InChannelName, const FReceiverInitParams& InitParams)
+	{
 		TSharedPtr<IDataChannel, ESPMode::ThreadSafe> DataChannel = GetDataChannel(InDataTypeName, InChannelName, InitParams.OperatorSettings);
 
 		if (DataChannel.IsValid())
@@ -335,5 +360,4 @@ namespace Metasound
 			return NewRouter.RegisterNewSender(InDataTypeName, InChannelName, InitParams);
 		}
 	}
-
 }
diff --git a/Engine/Plugins/Runtime/Metasound/Source/MetasoundFrontend/Public/MetasoundInstanceTransmitter.h b/Engine/Plugins/Runtime/Metasound/Source/MetasoundFrontend/Public/MetasoundInstanceTransmitter.h
--- a/Engine/Plugins/Runtime/Metasound/Source/MetasoundFrontend/Public/MetasoundInstanceTransmitter.h
+++ b/Engine/Plugins/Runtime/Metasound/Source/MetasoundFrontend/Public/MetasoundInstanceTransmitter.h
@@ -63,6 +63,8 @@ namespace Metasound
 		FMetasoundInstanceTransmitter(const FMetasoundInstanceTransmitter::FInitParams& InInitParams);
 		virtual ~FMetasoundInstanceTransmitter() = default;
 
+		bool Shutdown() override;
+
 		/** Returns ID of the MetaSound instance associated with this transmitter. */
 		uint64 GetInstanceID() const override;
 		
@@ -179,9 +181,6 @@ namespace Metasound
 		// Create and store a new ISender for the given FSendInfo.
 		ISender* AddSender(const FSendInfo& InInfo);
 
-		// Create a new ISender from FSendInfo.
-		TUniquePtr<ISender> CreateSender(const FSendInfo& InInfo) const;
-
 		TArray<FSendInfo> SendInfos;
 		FOperatorSettings OperatorSettings;
 		uint64 InstanceID;
@@ -189,4 +188,3 @@ namespace Metasound
 		TMap<FName, TUniquePtr<ISender>> InputSends;
 	};
 }
-
diff --git a/Engine/Plugins/Runtime/Metasound/Source/MetasoundFrontend/Public/MetasoundRouter.h b/Engine/Plugins/Runtime/Metasound/Source/MetasoundFrontend/Public/MetasoundRouter.h
--- a/Engine/Plugins/Runtime/Metasound/Source/MetasoundFrontend/Public/MetasoundRouter.h
+++ b/Engine/Plugins/Runtime/Metasound/Source/MetasoundFrontend/Public/MetasoundRouter.h
@@ -8,9 +8,11 @@
 #include "MetasoundAudioBuffer.h"
 #include "MetasoundDataFactory.h"
 #include "MetasoundDataReference.h"
+#include "MetasoundExecutableOperator.h"
 #include "MetasoundLiteral.h"
 #include "MetasoundOperatorInterface.h"
 #include "MetasoundOperatorSettings.h"
+#include "MetasoundTrigger.h"
 #include "Misc/Guid.h"
 #include "UObject/NameTypes.h"
 
@@ -372,7 +374,7 @@ namespace Metasound
 			}
 			else
 			{
-				return Push(TDataTypeLiteralFactory<TDataType>::CreateExplicitArgs(Params.OperatorSettings, InLiteral));
+				return TSender<TDataType>::Push(TDataTypeLiteralFactory<TDataType>::CreateExplicitArgs(Params.OperatorSettings, InLiteral));
 			}
 		}
 
@@ -458,7 +460,17 @@ namespace Metasound
 
 		bool CanPop() const
 		{
-			return DataChannel->GetPayloadID() != LastPayloadID && !DataChannel->IsEmpty();
+			if (DataChannel->GetPayloadID() == LastPayloadID)
+			{
+				return false;
+			}
+
+			if (DataChannel->IsEmpty())
+			{
+				return false;
+			}
+
+			return true;
 		}
 
 		// Pop the latest value from the data channel.
@@ -480,6 +492,13 @@ namespace Metasound
 			: IReceiver(InDataChannel, GetDataTypeName())
 			, OperatorSettings(InitParams.OperatorSettings)
 		{
+			// Executable DataTypes (currently just triggers) do not
+			// sync to past state provided by the data channel,
+			// so mark them as up-to-date.
+			if (TExecutableDataType<TDataType>::bIsExecutable)
+			{
+				LastPayloadID = InDataChannel->GetPayloadID();
+			}
 		}
 
 		void PushPayload(int64 InPayloadID, const TDataType& InDataPayload)
@@ -962,6 +981,8 @@ namespace Metasound
 			}
 		}
 
+		bool UnregisterDataChannel(const FName& InDataTypeName, const FName& InChannelName);
+
 		FAddressRouter(const FAddressRouter& Other)
 			: DataChannelMap(Other.DataChannelMap)
 		{
@@ -973,7 +994,6 @@ namespace Metasound
 	private:
 		FDataChannelKey GetDataChannelKey(const FName& InDataTypeName, const FName& InChannelName) const;
 
-
 		TMap<FDataChannelKey, TSharedRef<IDataChannel, ESPMode::ThreadSafe>> DataChannelMap;
 		FCriticalSection DataChannelMapMutationLock;
 	};
@@ -1130,6 +1150,9 @@ namespace Metasound
 			return Receiver;
 		}
 
+		// Unregisters DataChannel irrespective of number of receivers or senders still active.
+		bool UnregisterDataChannel(const FName& InDataTypeName, const FName& InChannelName);
+
 		// Pushes a literal parameter to a specific data channel in the global router.
 		// returns false if the literal type isn't supported.
 		bool PushLiteral(FName DataTypeName, FName GlobalChannelName, const FLiteral& InParam);
diff --git a/Engine/Source/Runtime/AudioExtensions/Public/IAudioExtensionPlugin.h b/Engine/Source/Runtime/AudioExtensions/Public/IAudioExtensionPlugin.h
--- a/Engine/Source/Runtime/AudioExtensions/Public/IAudioExtensionPlugin.h
+++ b/Engine/Source/Runtime/AudioExtensions/Public/IAudioExtensionPlugin.h
@@ -626,6 +626,8 @@ class IAudioInstanceTransmitter
 	public:
 		virtual ~IAudioInstanceTransmitter() = default;
 
+		virtual bool Shutdown() = 0;
+
 		// Return the instance ID
 		virtual uint64 GetInstanceID() const = 0;
 	
diff --git a/Engine/Source/Runtime/Engine/Classes/Sound/AudioCommunicationInterface.h b/Engine/Source/Runtime/Engine/Classes/Sound/AudioCommunicationInterface.h
--- a/Engine/Source/Runtime/Engine/Classes/Sound/AudioCommunicationInterface.h
+++ b/Engine/Source/Runtime/Engine/Classes/Sound/AudioCommunicationInterface.h
@@ -18,6 +18,8 @@ class IAudioCommunicationInterface : public IInterface
 	GENERATED_IINTERFACE_BODY()
 
 public:
+	virtual void Shutdown() = 0;
+
 	// Triggers a named trigger 
 	UFUNCTION(BlueprintCallable, Category = "Communication|Triggers")
 	virtual void Trigger(FName Name) = 0;
diff --git a/Engine/Source/Runtime/Engine/Classes/Sound/AudioComponentCommuncation.h b/Engine/Source/Runtime/Engine/Classes/Sound/AudioComponentCommuncation.h
--- a/Engine/Source/Runtime/Engine/Classes/Sound/AudioComponentCommuncation.h
+++ b/Engine/Source/Runtime/Engine/Classes/Sound/AudioComponentCommuncation.h
@@ -14,7 +14,11 @@ class ENGINE_API UAudioComponentCommunication	: public UObject
 	GENERATED_BODY()
 
 protected:
+	virtual void BeginDestroy() override;
+
 	// IAudioCommunicationInterface
+	void Shutdown() override;
+
 	void Trigger(FName InName) override;
 	void SetBool(FName InName, bool InValue) override;
 	void SetBoolArray(FName InName, const TArray<bool>& InValue) override;
diff --git a/Engine/Source/Runtime/Engine/Private/AudioComponentCommunication.cpp b/Engine/Source/Runtime/Engine/Private/AudioComponentCommunication.cpp
--- a/Engine/Source/Runtime/Engine/Private/AudioComponentCommunication.cpp
+++ b/Engine/Source/Runtime/Engine/Private/AudioComponentCommunication.cpp
@@ -1,11 +1,39 @@
 // Copyright Epic Games, Inc. All Rights Reserved.
 
-#include "Sound/AudioComponentCommuncation.h"
+#include "Audio.h"
 #include "AudioDevice.h"
+#include "ActiveSound.h"
 #include "AudioThread.h"
 #include "Components/AudioComponent.h"
-#include "ActiveSound.h"
 #include "IAudioExtensionPlugin.h"
+#include "Sound/AudioComponentCommuncation.h"
+
+void UAudioComponentCommunication::Shutdown()
+{
+	if (UAudioComponent* OwningComponent = Cast<UAudioComponent>(GetOuter()))
+	{
+		if (FAudioDevice* AudioDevice = OwningComponent->GetAudioDevice())
+		{
+			const uint64 MyAudioComponentID = OwningComponent->GetAudioComponentID();
+			FAudioThread::RunCommandOnAudioThread([AudioDevice, MyAudioComponentID]() mutable
+			{
+				if (FActiveSound* ActiveSound = AudioDevice->FindActiveSound(MyAudioComponentID))
+				{
+					if (IAudioInstanceTransmitter* Transmitter = ActiveSound->GetTransmitter())
+					{
+						Transmitter->Shutdown();
+					}
+				}
+			});
+		}
+	}
+}
+
+void UAudioComponentCommunication::BeginDestroy()
+{
+	Shutdown();
+	Super::BeginDestroy();
+}
 
 void UAudioComponentCommunication::Trigger(FName InName)
 {
@@ -116,7 +144,7 @@ void UAudioComponentCommunication::SetValue(FName InName, T&& InX)
 						{
 							if (!Transmitter->SetParameter(InName, MoveTempIfPossible(InX)))
 							{
-								UE_LOG(LogTemp, Warning, TEXT("Failed to SetParameter '%s'"), *InName.ToString() );
+								UE_LOG(LogAudio, Warning, TEXT("Failed to SetParameter '%s'"), *InName.ToString() );
 							}
 						}
 					}
diff --git a/Engine/Source/Runtime/Engine/Private/Components/AudioComponent.cpp b/Engine/Source/Runtime/Engine/Private/Components/AudioComponent.cpp
--- a/Engine/Source/Runtime/Engine/Private/Components/AudioComponent.cpp
+++ b/Engine/Source/Runtime/Engine/Private/Components/AudioComponent.cpp
@@ -86,10 +86,6 @@ UAudioComponent::UAudioComponent(const FObjectInitializer& ObjectInitializer)
 		FScopeLock Lock(&AudioIDToComponentMapLock);
 		AudioIDToComponentMap.Add(AudioComponentID, this);
 	}
-
-	// Create communications object.
-	static const FName NAME_CommunicationSubObj{ TEXT("AudioComponentCommunication") };
-	CommunicationInterface = CreateDefaultSubobject<UAudioComponentCommunication>(NAME_CommunicationSubObj, true);
 }
 
 UAudioComponent* UAudioComponent::GetAudioComponentFromID(uint64 AudioComponentID)
@@ -103,7 +99,10 @@ UAudioComponent* UAudioComponent::GetAudioComponentFromID(uint64 AudioComponentI
 
 void UAudioComponent::BeginDestroy()
 {
-	Super::BeginDestroy();
+	if (CommunicationInterface)
+	{
+		CommunicationInterface->Shutdown();
+	}
 
 	if (IsActive() && Sound && Sound->IsLooping())
 	{
@@ -111,8 +110,12 @@ void UAudioComponent::BeginDestroy()
 		Stop();
 	}
 
-	FScopeLock Lock(&AudioIDToComponentMapLock);
-	AudioIDToComponentMap.Remove(AudioComponentID);
+	{
+		FScopeLock Lock(&AudioIDToComponentMapLock);
+		AudioIDToComponentMap.Remove(AudioComponentID);
+	}
+
+	Super::BeginDestroy();
 }
 
 FString UAudioComponent::GetDetailedInfoInternal() const
@@ -450,7 +453,6 @@ void UAudioComponent::PlayInternal(const PlayInternalRequestData& InPlayRequestD
 		bAutoDestroy = bCurrentAutoDestroy;
 	}
 
-	// Whether or not we managed to actually try to play the sound
 	if (Sound && (World == nullptr || World->bAllowAudioPlayback))
 	{
 		if (FAudioDevice* AudioDevice = GetAudioDevice())
@@ -502,6 +504,13 @@ void UAudioComponent::PlayInternal(const PlayInternalRequestData& InPlayRequestD
 
 			AudioDevice->GetMaxDistanceAndFocusFactor(Sound, World, Location, AttenuationSettingsToApply, MaxDistance, FocusFactor);
 
+			// Create communications object.
+			if (!CommunicationInterface)
+			{
+				static const FName InterfaceName("AudioComponentCommunication");
+				CommunicationInterface = NewObject<UAudioComponentCommunication>(this, InterfaceName);
+			}
+
 			FActiveSound NewActiveSound;
 			NewActiveSound.SetAudioComponent(*this);
 			NewActiveSound.SetWorld(GetWorld());
@@ -753,6 +762,15 @@ void UAudioComponent::Stop()
 		return;
 	}
 
+	if (bIsPreviewSound)
+	{
+		if (CommunicationInterface)
+		{
+			CommunicationInterface->Shutdown();
+			CommunicationInterface = nullptr;
+		}
+	}
+
 	// Set this to immediately be inactive
 	SetActiveFlag(false);
 
