Miscellaneous tool tip fixes and conversions to headerless-Nodes

diff --git a/Engine/Plugins/Runtime/Metasound/Source/MetasoundStandardNodes/Private/MetasoundClampNode.cpp b/Engine/Plugins/Runtime/Metasound/Source/MetasoundStandardNodes/Private/MetasoundClampNode.cpp
--- a/Engine/Plugins/Runtime/Metasound/Source/MetasoundStandardNodes/Private/MetasoundClampNode.cpp
+++ b/Engine/Plugins/Runtime/Metasound/Source/MetasoundStandardNodes/Private/MetasoundClampNode.cpp
@@ -9,7 +9,7 @@
 #include "Internationalization/Text.h"
 #include "MetasoundParamHelper.h"
 
-#define LOCTEXT_NAMESPACE "MetasoundStandardNodes"
+#define LOCTEXT_NAMESPACE "MetasoundStandardNodes_ClampNode"
 
 namespace Metasound
 {
@@ -176,7 +176,7 @@ namespace Metasound
 				FName DataTypeName = GetMetasoundDataTypeName<ValueType>();
 				FName OperatorName = TEXT("Clamp");
 				FText NodeDisplayName = FText::Format(LOCTEXT("ClampDisplayNamePattern", "Clamp ({0})"), FText::FromString(GetMetasoundDataTypeString<ValueType>()));
-				FText NodeDescription = LOCTEXT("MinDesc", "Returns the clamped value of the input within the given value range.");
+				FText NodeDescription = LOCTEXT("ClampDesc", "Returns the clamped value of the input within the given value range.");
 				FVertexInterface NodeInterface = GetDefaultInterface();
 
 				return MetasoundClampNodePrivate::CreateNodeClassMetadata(DataTypeName, OperatorName, NodeDisplayName, NodeDescription, NodeInterface);
@@ -284,4 +284,4 @@ namespace Metasound
 
 }
 
-#undef LOCTEXT_NAMESPACE
+#undef LOCTEXT_NAMESPACE //MetaSoundStandardNodes_ClampNode
diff --git a/Engine/Plugins/Runtime/Metasound/Source/MetasoundStandardNodes/Private/MetasoundInterpToNode.cpp b/Engine/Plugins/Runtime/Metasound/Source/MetasoundStandardNodes/Private/MetasoundInterpToNode.cpp
--- a/Engine/Plugins/Runtime/Metasound/Source/MetasoundStandardNodes/Private/MetasoundInterpToNode.cpp
+++ b/Engine/Plugins/Runtime/Metasound/Source/MetasoundStandardNodes/Private/MetasoundInterpToNode.cpp
@@ -1,6 +1,6 @@
 // Copyright Epic Games, Inc. All Rights Reserved.
 
-#include "MetasoundInterpToNode.h"
+#include "MetasoundFacade.h"
 
 #include "Internationalization/Text.h"
 #include "MetasoundExecutableOperator.h"
@@ -11,7 +11,7 @@
 #include "MetasoundTime.h"
 #include "DSP/VolumeFader.h"
 
-#define LOCTEXT_NAMESPACE "MetasoundStandardNodes"
+#define LOCTEXT_NAMESPACE "MetasoundStandardNodes_InterpNode"
 
 namespace Metasound
 {
@@ -22,6 +22,19 @@ namespace Metasound
 		static const TCHAR* OutParamNameValue = TEXT("Value");
 	}
 
+	/** FInterpToNode
+	*
+	*  Interpolates to a target value over a given time.
+	*/
+	class METASOUNDSTANDARDNODES_API FInterpToNode : public FNodeFacade
+	{
+	public:
+		/**
+		 * Constructor used by the Metasound Frontend.
+		 */
+		FInterpToNode(const FNodeInitData& InitData);
+	};
+
 	class FInterpToOperator : public TExecutableOperator<FInterpToOperator>
 	{
 	public:
@@ -144,6 +157,12 @@ namespace Metasound
 		return Info;
 	}
 
+
+	FInterpToNode::FInterpToNode(const FNodeInitData& InitData)
+		: FNodeFacade(InitData.InstanceName, InitData.InstanceID, TFacadeOperatorClass<FInterpToOperator>())
+	{
+	}
+
 	TUniquePtr<IOperator> FInterpToOperator::CreateOperator(const FCreateOperatorParams& InParams, FBuildErrorArray& OutErrors)
 	{
 		const FInterpToNode& InterpToNode = static_cast<const FInterpToNode&>(InParams.Node);
@@ -156,12 +175,8 @@ namespace Metasound
 		return MakeUnique<FInterpToOperator>(InParams.OperatorSettings, TargetValue, InterpTime);
 	}
 
-	FInterpToNode::FInterpToNode(const FNodeInitData& InitData)
-		: FNodeFacade(InitData.InstanceName, InitData.InstanceID, TFacadeOperatorClass<FInterpToOperator>())
-	{
-	}
 
 	METASOUND_REGISTER_NODE(FInterpToNode)
 }
 
-#undef LOCTEXT_NAMESPACE
+#undef LOCTEXT_NAMESPACE //MetasoundStandardNodes_InterpNode
diff --git a/Engine/Plugins/Runtime/Metasound/Source/MetasoundStandardNodes/Private/MetasoundMathNodes.cpp b/Engine/Plugins/Runtime/Metasound/Source/MetasoundStandardNodes/Private/MetasoundMathNodes.cpp
--- a/Engine/Plugins/Runtime/Metasound/Source/MetasoundStandardNodes/Private/MetasoundMathNodes.cpp
+++ b/Engine/Plugins/Runtime/Metasound/Source/MetasoundStandardNodes/Private/MetasoundMathNodes.cpp
@@ -501,7 +501,7 @@ namespace Metasound
 					TInputDataVertexModel<TOperandDataClass>(MathOpNames::AdditionalOperandsName, LOCTEXT("MathOpModuloDivisorsTooltip", "Divisor(s)."), static_cast<TOperandDataClass>(1))
 				),
 				FOutputVertexInterface(
-					TOutputDataVertexModel<TDataClass>(TEXT("Out"), LOCTEXT("MathOpOutTooltip", "Math operation result"))
+					TOutputDataVertexModel<TDataClass>(TEXT("Out"), LOCTEXT("MathOpOutTooltip", "Resulting value"))
 				)
 			);
 
@@ -1389,8 +1389,8 @@ namespace Metasound
 		{
 			static const FVertexInterface Interface(
 				FInputVertexInterface(
-					TInputDataVertexModel<int64>(MathOpNames::PrimaryOperandName, LOCTEXT("MathOpModuloDividendTooltip", "Dividend of modulo operation."), 1),
-					TInputDataVertexModel<int64>(MathOpNames::AdditionalOperandsName, LOCTEXT("MathOpModuluDivisorTooltip", "Divisor of modulo operation."), 1)
+					TInputDataVertexModel<int64>(MathOpNames::PrimaryOperandName, LOCTEXT("MathOpModuloDividendTooltip", "Dividend"), 1),
+					TInputDataVertexModel<int64>(MathOpNames::AdditionalOperandsName, LOCTEXT("MathOpModuluDivisorTooltip", "Divisor(s)"), 1)
 				),
 				FOutputVertexInterface(
 					TOutputDataVertexModel<int64>(TEXT("Out"), LOCTEXT("MathOpModuloOutTooltip", "Resulting value"))
diff --git a/Engine/Plugins/Runtime/Metasound/Source/MetasoundStandardNodes/Private/MetasoundStereoDelayNode.cpp b/Engine/Plugins/Runtime/Metasound/Source/MetasoundStandardNodes/Private/MetasoundStereoDelayNode.cpp
--- a/Engine/Plugins/Runtime/Metasound/Source/MetasoundStandardNodes/Private/MetasoundStereoDelayNode.cpp
+++ b/Engine/Plugins/Runtime/Metasound/Source/MetasoundStandardNodes/Private/MetasoundStereoDelayNode.cpp
@@ -290,7 +290,7 @@ namespace Metasound
 				TInputDataVertexModel<FTime>(StereoDelay::InParamNameDelayTime, LOCTEXT("DelayTimeTooltip", "The amount of time to delay the audio."), 1.0f),
 				TInputDataVertexModel<float>(StereoDelay::InParamNameDelayRatio, LOCTEXT("DelayRatioTooltip", "Delay spread for left and right channels. Allows left and right channels to have differential delay amounts. Useful for stereo channel decorrelation"), 0.0f),
 				TInputDataVertexModel<float>(StereoDelay::InParamNameDryLevel, LOCTEXT("DryLevelTooltip", "The dry level of the delay."), 0.0f),
-				TInputDataVertexModel<float>(StereoDelay::InParamNameWetLevel, LOCTEXT("FeedbackTooltip", "The wet level of the delay."), 1.0f),
+				TInputDataVertexModel<float>(StereoDelay::InParamNameWetLevel, LOCTEXT("WetLevelTooltip", "The wet level of the delay."), 1.0f),
 				TInputDataVertexModel<float>(StereoDelay::InParamNameFeedbackAmount, LOCTEXT("FeedbackTooltip", "Feedback amount."), 0.0f)
 			),
 			FOutputVertexInterface(
@@ -357,4 +357,4 @@ namespace Metasound
 	METASOUND_REGISTER_NODE(FStereoDelayNode)
 }
 
-#undef LOCTEXT_NAMESPACE
+#undef LOCTEXT_NAMESPACE //MetasoundStandardNodes_DelayNode
diff --git a/Engine/Plugins/Runtime/Metasound/Source/MetasoundStandardNodes/Private/MetasoundTriggerAnyNode.cpp b/Engine/Plugins/Runtime/Metasound/Source/MetasoundStandardNodes/Private/MetasoundTriggerAnyNode.cpp
--- a/Engine/Plugins/Runtime/Metasound/Source/MetasoundStandardNodes/Private/MetasoundTriggerAnyNode.cpp
+++ b/Engine/Plugins/Runtime/Metasound/Source/MetasoundStandardNodes/Private/MetasoundTriggerAnyNode.cpp
@@ -1,13 +1,14 @@
 // Copyright Epic Games, Inc. All Rights Reserved.
-#include "MetasoundTriggerAnyNode.h"
-
+#include "MetasoundTrigger.h"
+#include "MetasoundFacade.h"
+#include "MetasoundExecutableOperator.h"
 #include "Internationalization/Text.h"
 #include "MetasoundNodeRegistrationMacro.h"
 #include "MetasoundStandardNodesNames.h"
 #include "MetasoundVertex.h"
 #include "MetasoundStandardNodesCategories.h"
 
-#define LOCTEXT_NAMESPACE "MetasoundStandardNodes"
+#define LOCTEXT_NAMESPACE "MetasoundStandardNodes_TriggerAny"
 
 #define REGISTER_TRIGGER_ANY_NODE(Number) \
 	using FTriggerAnyNode_##Number = TTriggerAnyNode<Number>; \
@@ -17,7 +18,7 @@ namespace Metasound
 {
 	namespace MetasoundTriggerAnyNodePrivate
 	{
-		FNodeClassMetadata CreateNodeClassMetadata(const FName& InOperatorName, const FText& InDisplayName, const FText& InDescription, const FVertexInterface& InDefaultInterface)
+		METASOUNDSTANDARDNODES_API FNodeClassMetadata CreateNodeClassMetadata(const FName& InOperatorName, const FText& InDisplayName, const FText& InDescription, const FVertexInterface& InDefaultInterface)
 		{
 			FNodeClassMetadata Metadata
 			{
@@ -40,29 +41,163 @@ namespace Metasound
 
 	namespace TriggerAnyVertexNames
 	{
-		const FString GetInputTriggerName(uint32 InIndex)
+		METASOUNDSTANDARDNODES_API const FString GetInputTriggerName(uint32 InIndex)
 		{
 			return FText::Format(LOCTEXT("TriggerAnyTriggerInputName", "In {0}"), InIndex).ToString();
 		}
 
-		const FText GetInputTriggerDescription(uint32 InIndex)
+		METASOUNDSTANDARDNODES_API const FText GetInputTriggerDescription(uint32 InIndex)
 		{
-			return FText::Format(LOCTEXT("TriggerAnyInputTriggerDesc", "Trigger {0} input. All trigger inputs must be triggered before the output trigger is hit."), InIndex);
+			return FText::Format(LOCTEXT("TriggerAnyInputTriggerDesc", "Trigger {0} input. The output trigger is hit when any of the input triggers are hit."), InIndex);
 		}
 
-		const FString& GetOutputTriggerName()
+		METASOUNDSTANDARDNODES_API const FString& GetOutputTriggerName()
 		{
 			static const FString Name = TEXT("Out");
 			return Name;
 		}
 
-		const FText& GetOutputTriggerDescription()
+		METASOUNDSTANDARDNODES_API const FText& GetOutputTriggerDescription()
 		{
-			static const FText Desc = LOCTEXT("TriggerAccumulateOutputTriggerDesc", "Triggered when all input triggers have been triggered. Call Reset to reset the state or use \"Auto Reset\"");
+			static const FText Desc = LOCTEXT("TriggerAnyOutputTriggerDesc", "Triggered when any of the input triggers have been triggered. ");
 			return Desc;
 		}
 	}
 
+	template<uint32 NumInputs>
+	class TTriggerAnyOperator : public TExecutableOperator<TTriggerAnyOperator<NumInputs>>
+	{
+	public:
+		static const FVertexInterface& GetDefaultInterface()
+		{
+			using namespace TriggerAnyVertexNames;
+
+			auto CreateDefaultInterface = []() -> FVertexInterface
+			{
+				FInputVertexInterface InputInterface;
+
+				for (uint32 i = 0; i < NumInputs; ++i)
+				{
+					InputInterface.Add(TInputDataVertexModel<FTrigger>(GetInputTriggerName(i), GetInputTriggerDescription(i)));
+				}
+
+				FOutputVertexInterface OutputInterface;
+				OutputInterface.Add(TOutputDataVertexModel<FTrigger>(GetOutputTriggerName(), GetOutputTriggerDescription()));
+
+				return FVertexInterface(InputInterface, OutputInterface);
+			};
+
+			static const FVertexInterface DefaultInterface = CreateDefaultInterface();
+			return DefaultInterface;
+		}
+
+		static const FNodeClassMetadata& GetNodeInfo()
+		{
+			auto CreateNodeClassMetadata = []() -> FNodeClassMetadata
+			{
+				FName OperatorName = *FString::Printf(TEXT("Trigger Any (%d)"), NumInputs);
+				FText NodeDisplayName = FText::Format(LOCTEXT("TriggerAnyDisplayNamePattern", "Trigger Any ({0})"), NumInputs);
+				FText NodeDescription = LOCTEXT("TriggerAnyDescription", "Will trigger output on any of the input triggers.");
+				FVertexInterface NodeInterface = GetDefaultInterface();
+
+				return MetasoundTriggerAnyNodePrivate::CreateNodeClassMetadata(OperatorName, NodeDisplayName, NodeDescription, NodeInterface);
+			};
+
+			static const FNodeClassMetadata Metadata = CreateNodeClassMetadata();
+			return Metadata;
+		}
+
+		static TUniquePtr<IOperator> CreateOperator(const FCreateOperatorParams& InParams, TArray<TUniquePtr<IOperatorBuildError>>& OutErrors)
+		{
+			using namespace TriggerAnyVertexNames;
+
+			const FInputVertexInterface& InputInterface = InParams.Node.GetVertexInterface().GetInputInterface();
+			const FDataReferenceCollection& InputCollection = InParams.InputDataReferences;
+
+			TArray<FTriggerReadRef> InputTriggers;
+
+			for (uint32 i = 0; i < NumInputs; ++i)
+			{
+				InputTriggers.Add(InputCollection.GetDataReadReferenceOrConstruct<FTrigger>(GetInputTriggerName(i), InParams.OperatorSettings));
+			}
+
+			return MakeUnique<TTriggerAnyOperator<NumInputs>>(InParams.OperatorSettings, MoveTemp(InputTriggers));
+		}
+
+		TTriggerAnyOperator(const FOperatorSettings& InSettings, const TArray<FTriggerReadRef>&& InInputTriggers)
+			: InputTriggers(InInputTriggers)
+			, OutputTrigger(FTriggerWriteRef::CreateNew(InSettings))
+		{
+		}
+
+		virtual ~TTriggerAnyOperator() = default;
+
+
+		virtual FDataReferenceCollection GetInputs() const override
+		{
+			using namespace TriggerAnyVertexNames;
+
+			FDataReferenceCollection Inputs;
+			for (uint32 i = 0; i < NumInputs; ++i)
+			{
+				Inputs.AddDataReadReference(GetInputTriggerName(i), InputTriggers[i]);
+			}
+
+			return Inputs;
+		}
+
+		virtual FDataReferenceCollection GetOutputs() const override
+		{
+			using namespace TriggerAnyVertexNames;
+
+			FDataReferenceCollection Outputs;
+			Outputs.AddDataReadReference(GetOutputTriggerName(), OutputTrigger);
+
+			return Outputs;
+		}
+
+		void Execute()
+		{
+			OutputTrigger->AdvanceBlock();
+
+			for (uint32 i = 0; i < NumInputs; ++i)
+			{
+				InputTriggers[i]->ExecuteBlock(
+					[&](int32 StartFrame, int32 EndFrame)
+					{
+					},
+					[this, i](int32 StartFrame, int32 EndFrame)
+					{
+						OutputTrigger->TriggerFrame(StartFrame);
+					}
+					);
+			}
+		}
+
+	private:
+
+		TArray<FTriggerReadRef> InputTriggers;
+		FTriggerWriteRef OutputTrigger;
+	};
+
+	/** TTriggerAnyNode
+	*
+	*  Will output a trigger whenever any of its input triggers are set.
+	*/
+	template<uint32 NumInputs>
+	class METASOUNDSTANDARDNODES_API TTriggerAnyNode : public FNodeFacade
+	{
+	public:
+		/**
+		 * Constructor used by the Metasound Frontend.
+		 */
+		TTriggerAnyNode(const FNodeInitData& InInitData)
+			: FNodeFacade(InInitData.InstanceName, InInitData.InstanceID, TFacadeOperatorClass<TTriggerAnyOperator<NumInputs>>())
+		{}
+
+		virtual ~TTriggerAnyNode() = default;
+	};
+
 	REGISTER_TRIGGER_ANY_NODE(2)
 	REGISTER_TRIGGER_ANY_NODE(3)
 	REGISTER_TRIGGER_ANY_NODE(4)
diff --git a/Engine/Plugins/Runtime/Metasound/Source/MetasoundStandardNodes/Private/MetasoundTriggerRepeatNode.cpp b/Engine/Plugins/Runtime/Metasound/Source/MetasoundStandardNodes/Private/MetasoundTriggerRepeatNode.cpp
--- a/Engine/Plugins/Runtime/Metasound/Source/MetasoundStandardNodes/Private/MetasoundTriggerRepeatNode.cpp
+++ b/Engine/Plugins/Runtime/Metasound/Source/MetasoundStandardNodes/Private/MetasoundTriggerRepeatNode.cpp
@@ -18,7 +18,7 @@
 #include "MetasoundStandardNodesCategories.h"
 #include "MetasoundParamHelper.h"
 
-#define LOCTEXT_NAMESPACE "MetasoundStandardNodes"
+#define LOCTEXT_NAMESPACE "MetasoundStandardNodes_RepeatNode"
 
 namespace Metasound
 {
@@ -29,7 +29,7 @@ namespace Metasound
 		METASOUND_PARAM(InputStart, "Start", "Starts executing periodic output triggers.");
 		METASOUND_PARAM(InputStop, "Stop", "Stops executing periodic output triggers.");
 		METASOUND_PARAM(InputPeriod, "Period", "The period to trigger in seconds.");
-		METASOUND_PARAM(OutputOnTrigger, "Out", "The periodically generated output trigger.");
+		METASOUND_PARAM(RepeatOutputOnTrigger, "RepeatOut", "The periodically generated output trigger.");
 	}
 
 	class FTriggerRepeatOperator : public TExecutableOperator<FTriggerRepeatOperator>
@@ -88,7 +88,7 @@ namespace Metasound
 		using namespace TriggerRepeatVertexNames;
 
 		FDataReferenceCollection OutputDataReferences;
-		OutputDataReferences.AddDataReadReference(METASOUND_GET_PARAM_NAME(OutputOnTrigger), TriggerOut);
+		OutputDataReferences.AddDataReadReference(METASOUND_GET_PARAM_NAME(RepeatOutputOnTrigger), TriggerOut);
 
 		return OutputDataReferences;
 	}
@@ -157,7 +157,7 @@ namespace Metasound
 				TInputDataVertexModel<FTime>(METASOUND_GET_PARAM_NAME_AND_TT(InputPeriod), 0.2f)
 			),
 			FOutputVertexInterface(
-				TOutputDataVertexModel<FTrigger>(METASOUND_GET_PARAM_NAME_AND_TT(OutputOnTrigger))
+				TOutputDataVertexModel<FTrigger>(METASOUND_GET_PARAM_NAME_AND_TT(RepeatOutputOnTrigger))
 			)
 		);
 
diff --git a/Engine/Plugins/Runtime/Metasound/Source/MetasoundStandardNodes/Public/MetasoundInterpToNode.h b/Engine/Plugins/Runtime/Metasound/Source/MetasoundStandardNodes/Public/MetasoundInterpToNode.h
deleted file mode 100644
diff --git a/Engine/Plugins/Runtime/Metasound/Source/MetasoundStandardNodes/Public/MetasoundTriggerAccumulatorNode.h b/Engine/Plugins/Runtime/Metasound/Source/MetasoundStandardNodes/Public/MetasoundTriggerAccumulatorNode.h
--- a/Engine/Plugins/Runtime/Metasound/Source/MetasoundStandardNodes/Public/MetasoundTriggerAccumulatorNode.h
+++ b/Engine/Plugins/Runtime/Metasound/Source/MetasoundStandardNodes/Public/MetasoundTriggerAccumulatorNode.h
@@ -10,7 +10,7 @@
 #include "Internationalization/Text.h"
 #include "MetasoundParamHelper.h"
 
-#define LOCTEXT_NAMESPACE "MetasoundStandardNodes"
+#define LOCTEXT_NAMESPACE "MetasoundStandardNodes_AccumulatorNode"
 
 namespace Metasound
 {
@@ -23,7 +23,7 @@ namespace Metasound
 	{
 		METASOUND_PARAM(InputAutoReset, "Auto Reset", "Input trigger which results in a delayed trigger.");	
 		METASOUND_PARAM(InputTrigger, "In {0}", "Trigger {0} input. All trigger inputs must be triggered before the output trigger is hit.");
-		METASOUND_PARAM(OutputOnTrigger, "Out", "Triggered when all input triggers have been triggered. Call Reset to reset the state or use \"Auto Reset\".");
+		METASOUND_PARAM(AccumulateOutputOnTrigger, "Out", "Triggered when all input triggers have been triggered. Call Reset to reset the state or use \"Auto Reset\".");
 	}
 
 	template<uint32 NumInputs>
@@ -46,7 +46,7 @@ namespace Metasound
 				InputInterface.Add(TInputDataVertexModel<bool>(METASOUND_GET_PARAM_NAME_AND_TT(InputAutoReset)));
 
 				FOutputVertexInterface OutputInterface;
-				OutputInterface.Add(TOutputDataVertexModel<FTrigger>(METASOUND_GET_PARAM_NAME_AND_TT(OutputOnTrigger)));
+				OutputInterface.Add(TOutputDataVertexModel<FTrigger>(METASOUND_GET_PARAM_NAME_AND_TT(AccumulateOutputOnTrigger)));
 
 				return FVertexInterface(InputInterface, OutputInterface);
 			};
@@ -121,7 +121,7 @@ namespace Metasound
 			using namespace TriggerAccumulatorVertexNames;
 
 			FDataReferenceCollection Outputs;
-			Outputs.AddDataReadReference(METASOUND_GET_PARAM_NAME(OutputOnTrigger), OutputTrigger);
+			Outputs.AddDataReadReference(METASOUND_GET_PARAM_NAME(AccumulateOutputOnTrigger), OutputTrigger);
 
 			return Outputs;
 		}
@@ -209,4 +209,4 @@ namespace Metasound
 
 } // namespace Metasound
 
-#undef LOCTEXT_NAMESPACE
+#undef LOCTEXT_NAMESPACE //MetasoundStandardNodes_AccumulatorNode
diff --git a/Engine/Plugins/Runtime/Metasound/Source/MetasoundStandardNodes/Public/MetasoundTriggerAnyNode.h b/Engine/Plugins/Runtime/Metasound/Source/MetasoundStandardNodes/Public/MetasoundTriggerAnyNode.h
deleted file mode 100644
