metasound: Fix for silent output by default in audio mixer node

diff --git a/Engine/Plugins/Runtime/Metasound/Source/MetasoundStandardNodes/Private/MetasoundMixerNode.cpp b/Engine/Plugins/Runtime/Metasound/Source/MetasoundStandardNodes/Private/MetasoundMixerNode.cpp
--- a/Engine/Plugins/Runtime/Metasound/Source/MetasoundStandardNodes/Private/MetasoundMixerNode.cpp
+++ b/Engine/Plugins/Runtime/Metasound/Source/MetasoundStandardNodes/Private/MetasoundMixerNode.cpp
@@ -62,7 +62,7 @@ namespace Metasound
 					}
 
 					// gain scalar
-					InputInterface.Add(TInputDataVertexModel<float>(GetGainInputName(InputIndex), GetGainInputDescription(InputIndex)));
+					InputInterface.Add(TInputDataVertexModel<float>(GetGainInputName(InputIndex), GetGainInputDescription(InputIndex), 1.0f));
 				}
 
 				// outputs
