Removing more unwanted types and hiding blacklist of input types

diff --git a/Engine/Plugins/Runtime/Metasound/Source/MetasoundEditor/Private/MetasoundEditorModule.cpp b/Engine/Plugins/Runtime/Metasound/Source/MetasoundEditor/Private/MetasoundEditorModule.cpp
--- a/Engine/Plugins/Runtime/Metasound/Source/MetasoundEditor/Private/MetasoundEditorModule.cpp
+++ b/Engine/Plugins/Runtime/Metasound/Source/MetasoundEditor/Private/MetasoundEditorModule.cpp
@@ -163,6 +163,11 @@ namespace Metasound
 				return DataTypeInfo.FindChecked(InDataTypeName);
 			}
 
+			virtual bool IsRegisteredDataType(FName InDataTypeName) const override
+			{
+				return DataTypeInfo.Contains(InDataTypeName);
+			}
+
 			virtual void IterateDataTypes(TUniqueFunction<void(const FEditorDataType&)> InDataTypeFunction) const override
 			{
 				for (const TPair<FName, FEditorDataType>& Pair : DataTypeInfo)
diff --git a/Engine/Plugins/Runtime/Metasound/Source/MetasoundEditor/Private/MetasoundEditorModule.h b/Engine/Plugins/Runtime/Metasound/Source/MetasoundEditor/Private/MetasoundEditorModule.h
--- a/Engine/Plugins/Runtime/Metasound/Source/MetasoundEditor/Private/MetasoundEditorModule.h
+++ b/Engine/Plugins/Runtime/Metasound/Source/MetasoundEditor/Private/MetasoundEditorModule.h
@@ -37,6 +37,7 @@ namespace Metasound
 		public:
 			virtual void RegisterDataType(FName InPinCategoryName, FName InPinSubCategoryName, const FDataTypeRegistryInfo& InRegistryInfo) = 0;
 			virtual const FEditorDataType& FindDataType(FName InDataTypeName) const = 0;
+			virtual bool IsRegisteredDataType(FName InDataTypeName) const = 0;
 			virtual void IterateDataTypes(TUniqueFunction<void(const FEditorDataType&)> InDataTypeFunction) const = 0;
 
 			virtual const TSubclassOf<UMetasoundEditorGraphInputLiteral> FindInputLiteralClass(EMetasoundFrontendLiteralType InLiteralType) const = 0;
diff --git a/Engine/Plugins/Runtime/Metasound/Source/MetasoundEditor/Private/MetasoundNodeDetailCustomization.cpp b/Engine/Plugins/Runtime/Metasound/Source/MetasoundEditor/Private/MetasoundNodeDetailCustomization.cpp
--- a/Engine/Plugins/Runtime/Metasound/Source/MetasoundEditor/Private/MetasoundNodeDetailCustomization.cpp
+++ b/Engine/Plugins/Runtime/Metasound/Source/MetasoundEditor/Private/MetasoundNodeDetailCustomization.cpp
@@ -70,6 +70,14 @@ namespace Metasound
 			static const FName DataTypeNameIdentifier = "DataTypeName";
 			static const FName ProxyGeneratorClassNameIdentifier = "GeneratorClass";
 
+			/** Set of input types which are valid registered types, but should
+			 * not show up as an input type option in the MetaSound editor. */
+			static const TSet<FName> HiddenInputTypeNames =
+			{
+				"Audio:Mono",
+				"Audio:Stereo"
+			};
+
 			void ExecuteInputTrigger(UMetasoundEditorGraphInputLiteral* Literal)
 			{
 				if (!Literal)
@@ -524,10 +532,17 @@ namespace Metasound
 				CurrentTypeName.LeftChopInline(VariableCustomizationPrivate::ArrayIdentifier.Len());
 			}
 
-			DataTypeNames.Reset();
 			IMetasoundEditorModule& EditorModule = FModuleManager::GetModuleChecked<IMetasoundEditorModule>("MetaSoundEditor");
+
+			// Not all types have an equivalent array type. Base types without array
+			// types should have the "Is Array" checkbox disabled. 
+			const bool bIsArrayTypeRegisteredForCurrentType = EditorModule.IsRegisteredDataType(FName(CurrentTypeName + VariableCustomizationPrivate::ArrayIdentifier));
+
+			DataTypeNames.Reset();
 			EditorModule.IterateDataTypes([&](const FEditorDataType& EditorDataType)
 			{
+				using namespace VariableCustomizationPrivate;
+
 				const FString TypeName = EditorDataType.RegistryInfo.DataTypeName.ToString();
 
 				// Array types are handles separately via checkbox
@@ -536,6 +551,12 @@ namespace Metasound
 					return;
 				}
 
+				// Hidden input types should be omitted from the drop down.
+				if (HiddenInputTypeNames.Contains(EditorDataType.RegistryInfo.DataTypeName))
+				{
+					return;
+				}
+
 				TSharedPtr<FString> TypeStrPtr = MakeShared<FString>(TypeName);
 				if (TypeName == CurrentTypeName)
 				{
@@ -589,6 +610,7 @@ namespace Metasound
 				.VAlign(VAlign_Center)
 				[
 					SAssignNew(DataTypeArrayCheckbox, SCheckBox)
+					.IsEnabled(bIsArrayTypeRegisteredForCurrentType)
 					.IsChecked_Lambda([this, InGraphVariable]()
 					{
 						return OnGetDataTypeArrayCheckState(InGraphVariable);
@@ -786,17 +808,45 @@ namespace Metasound
 		{
 			if (ItemSelected.IsValid() && !ItemSelected->IsEmpty() && InGraphVariable.IsValid())
 			{
-				FString DataTypeString = *ItemSelected.Get();
+				IMetasoundEditorModule& EditorModule = FModuleManager::GetModuleChecked<IMetasoundEditorModule>("MetaSoundEditor");
+
+				FName BaseDataTypeName = FName(*ItemSelected.Get());
+				FName ArrayDataTypeName = FName(*ItemSelected.Get() + VariableCustomizationPrivate::ArrayIdentifier);
+
+				FName NewDataTypeName;
+
+				// Update data type based on "Is Array" checkbox and support for arrays.
+				// If an array type is not supported, default to the base data type.
 				if (DataTypeArrayCheckbox->GetCheckedState() == ECheckBoxState::Checked)
 				{
-					DataTypeString += VariableCustomizationPrivate::ArrayIdentifier;
+					if (EditorModule.IsRegisteredDataType(ArrayDataTypeName))
+					{
+						NewDataTypeName = ArrayDataTypeName;
+					}
+					else
+					{
+						check(EditorModule.IsRegisteredDataType(BaseDataTypeName));
+						NewDataTypeName = BaseDataTypeName;
+					}
+				}
+				else
+				{
+					if (EditorModule.IsRegisteredDataType(BaseDataTypeName))
+					{
+						NewDataTypeName = BaseDataTypeName;
+					}
+					else
+					{
+						check(EditorModule.IsRegisteredDataType(ArrayDataTypeName));
+						NewDataTypeName = ArrayDataTypeName;
+					}
 				}
 
 				// Have to stop playback to avoid attempting to change live edit data on invalid input type.
 				check(GEditor);
 				GEditor->ResetPreviewAudioComponent();
 
-				InGraphVariable->SetDataType(FName(DataTypeString));
+				InGraphVariable->SetDataType(NewDataTypeName);
 
 				// Required to rebuild the literal details customization.
 				// This is seemingly dangerous (as the Builder's raw ptr is cached),
diff --git a/Engine/Plugins/Runtime/Metasound/Source/MetasoundEditor/Public/MetasoundEditorGraphSchema.h b/Engine/Plugins/Runtime/Metasound/Source/MetasoundEditor/Public/MetasoundEditorGraphSchema.h
--- a/Engine/Plugins/Runtime/Metasound/Source/MetasoundEditor/Public/MetasoundEditorGraphSchema.h
+++ b/Engine/Plugins/Runtime/Metasound/Source/MetasoundEditor/Public/MetasoundEditorGraphSchema.h
@@ -7,7 +7,6 @@
 #include "CoreMinimal.h"
 #include "EdGraph/EdGraphSchema.h"
 #include "EdGraphUtilities.h"
-#include "MetasoundEditorModule.h"
 #include "MetasoundFrontend.h"
 #include "MetasoundFrontendController.h"
 #include "Templates/Function.h"
diff --git a/Engine/Plugins/Runtime/Metasound/Source/MetasoundFrontend/Private/MetasoundPrimitives.cpp b/Engine/Plugins/Runtime/Metasound/Source/MetasoundFrontend/Private/MetasoundPrimitives.cpp
--- a/Engine/Plugins/Runtime/Metasound/Source/MetasoundFrontend/Private/MetasoundPrimitives.cpp
+++ b/Engine/Plugins/Runtime/Metasound/Source/MetasoundFrontend/Private/MetasoundPrimitives.cpp
@@ -12,30 +12,25 @@
 
 namespace Metasound
 {
+	// Disable audio buffer arrays.
+	template<>
+	struct TEnableAutoArrayTypeRegistration<FAudioBuffer>
+	{
+		static constexpr bool Value = false;
+	};
+
 	// Disable auto-conversion using the FAudioBuffer(int32) constructor
 	template<typename FromDataType>
-	struct TIsAutoConvertible<FromDataType, FAudioBuffer>
+	struct TEnableAutoConverterNodeRegistration<FromDataType, FAudioBuffer>
 	{
-		static constexpr bool bIsConvertible = std::is_convertible<FromDataType, FAudioBuffer>::value;
-
-		static constexpr bool bIsFromArithmeticType = std::is_arithmetic<FromDataType>::value;
-
-	public:
-
-		static constexpr bool Value = bIsConvertible && (!bIsFromArithmeticType);
+		static constexpr bool Value = !std::is_arithmetic<FromDataType>::value;
 	};
 
 	// Disable auto-conversions based on FTrigger implicit converters
 	template<typename ToDataType>
-	struct TIsAutoConvertible<FTrigger, ToDataType>
+	struct TEnableAutoConverterNodeRegistration<FTrigger, ToDataType>
 	{
-		static constexpr bool bIsConvertible = std::is_convertible<FTrigger, ToDataType>::value;
-
-		static constexpr bool bIsToArithmeticType = std::is_arithmetic<ToDataType>::value;
-
-	public:
-
-		static constexpr bool Value = bIsConvertible && (!bIsToArithmeticType);
+		static constexpr bool Value = !std::is_arithmetic<ToDataType>::value;
 	};
 }
 
diff --git a/Engine/Plugins/Runtime/Metasound/Source/MetasoundFrontend/Public/MetasoundDataTypeRegistrationMacro.h b/Engine/Plugins/Runtime/Metasound/Source/MetasoundFrontend/Public/MetasoundDataTypeRegistrationMacro.h
--- a/Engine/Plugins/Runtime/Metasound/Source/MetasoundFrontend/Public/MetasoundDataTypeRegistrationMacro.h
+++ b/Engine/Plugins/Runtime/Metasound/Source/MetasoundFrontend/Public/MetasoundDataTypeRegistrationMacro.h
@@ -28,52 +28,80 @@
 
 namespace Metasound
 {
-	// Helper utility to test if we can transmit a datatype between a send and a receive node.
-	template <typename TDataType>
-	struct TIsTransmittable
+	/** Enables or disables automatic registration of array types given a 
+	 * MetaSound data type. By default this is true, and all data types will have
+	 * an associated TArray<DataType> registered if it is supported. */
+	template<typename ... Type>
+	struct TEnableAutoArrayTypeRegistration
 	{
-	private:
-		static constexpr bool bIsCopyConstructible = std::is_copy_constructible<TDataType>::value;
-		static constexpr bool bIsCopyAssignable = std::is_copy_assignable<TDataType>::value;
-
-		// TODO: audio types were intended to be send/receive nodes but they require 
-		// template specialization.  TIsTransmittable should ask the Send and Receive nodes
-		// if they are transmittable rather than attempting to do the logic for all
-		// types here. 
-		//static constexpr bool bIsAudioDataType = TIsDerivedFrom<TDataType, IAudioDataType>::Value;
-
-		static constexpr bool bCanBeTransmitted = bIsCopyConstructible && bIsCopyAssignable;
-
-	public:
-
-		static constexpr bool Value = bCanBeTransmitted;
+		static constexpr bool Value = true;
 	};
 
-	// Specialization of TIsTransmittable<> for TArray to handle lax TArray copy constructor 
-	// definition. This can be removed once updates to TArray are merged into the same codebase.
-	// TODO: delete me eventually.
-	template<typename TElementType>
-	struct TIsTransmittable<TArray<TElementType>>
+	/** Enables or disables automatic registration of auto conversion nodes given a 
+	 * MetaSound data type. By default this is true, and all data types will have
+	 * associated conversion nodes registered based upon the data types supported
+	 * constructors and implicit conversions. */
+	template<typename ... Type>
+	struct TEnableAutoConverterNodeRegistration
 	{
-		// Depend on copy constructor of elements to determine whether TArray is copy constructible
-		// and copy assignable. Definitions for TArray copy construction and copy assignment are
-		// defined by default whether or not the TArray element types support them. This logic
-		// skips checking the TArray object itselt and instead checks the TArray element type
-		// directly.
-
-		static constexpr bool Value = TIsTransmittable<TElementType>::Value;
+		static constexpr bool Value = true;
 	};
 
-	// Determines whether an auto converter node will be registered to convert 
-	// between two types. 
-	template<typename TFromDataType, typename TToDataType>
-	struct TIsAutoConvertible
+	/** Enables or disables send and receive node registration for a given MetaSound
+	 * data type. By default this is true and all data types supported by the transmission
+	 * system will have associated send and receive nodes. */
+	template<typename ... Type>
+	struct TEnableTransmissionNodeRegistration
 	{
-		static constexpr bool Value = std::is_convertible<TFromDataType, TToDataType>::value;
+		static constexpr bool Value = true;
 	};
 
 	namespace MetasoundDataTypeRegistrationPrivate
 	{
+
+		// Helper utility to test if we can transmit a datatype between a send and a receive node.
+		template <typename TDataType>
+		struct TIsTransmittable
+		{
+		private:
+			static constexpr bool bIsCopyConstructible = std::is_copy_constructible<TDataType>::value;
+			static constexpr bool bIsCopyAssignable = std::is_copy_assignable<TDataType>::value;
+
+			// TODO: audio types were intended to be send/receive nodes but they require 
+			// template specialization.  TIsTransmittable should ask the Send and Receive nodes
+			// if they are transmittable rather than attempting to do the logic for all
+			// types here. 
+			//static constexpr bool bIsAudioDataType = TIsDerivedFrom<TDataType, IAudioDataType>::Value;
+
+			static constexpr bool bCanBeTransmitted = bIsCopyConstructible && bIsCopyAssignable;
+
+		public:
+
+			static constexpr bool Value = bCanBeTransmitted;
+		};
+
+		// Specialization of TIsTransmittable<> for TArray to handle lax TArray copy constructor 
+		// definition. This can be removed once updates to TArray are merged into the same codebase.
+		// TODO: delete me eventually.
+		template<typename TElementType>
+		struct TIsTransmittable<TArray<TElementType>>
+		{
+			// Depend on copy constructor of elements to determine whether TArray is copy constructible
+			// and copy assignable. Definitions for TArray copy construction and copy assignment are
+			// defined by default whether or not the TArray element types support them. This logic
+			// skips checking the TArray object itselt and instead checks the TArray element type
+			// directly.
+
+			static constexpr bool Value = TIsTransmittable<TElementType>::Value;
+		};
+
+		// Determines whether an auto converter node will be registered to convert 
+		// between two types. 
+		template<typename TFromDataType, typename TToDataType>
+		struct TIsAutoConvertible
+		{
+			static constexpr bool Value = std::is_convertible<TFromDataType, TToDataType>::value;
+		};
 		// Returns the Array version of a literal type if it exists.
 		template<ELiteralType LiteralType>
 		struct TLiteralArrayEnum 
@@ -152,8 +180,11 @@ namespace Metasound
 		template<typename TDataType, typename TEnableIf<TIsTransmittable<TDataType>::Value, bool>::Type = true>
 		void AttemptToRegisterSendAndReceiveNodes()
 		{
-			ensureAlways(RegisterNodeWithFrontend<Metasound::TSendNode<TDataType>>());
-			ensureAlways(RegisterNodeWithFrontend<Metasound::TReceiveNode<TDataType>>());
+			if (TEnableTransmissionNodeRegistration<TDataType>::Value)
+			{
+				ensureAlways(RegisterNodeWithFrontend<Metasound::TSendNode<TDataType>>());
+				ensureAlways(RegisterNodeWithFrontend<Metasound::TReceiveNode<TDataType>>());
+			}
 		}
 
 		template<typename TDataType, typename TEnableIf<!TIsTransmittable<TDataType>::Value, bool>::Type = true>
@@ -170,15 +201,18 @@ namespace Metasound
 		{
 			using FConverterNode = Metasound::TAutoConverterNode<TFromDataType, TToDataType>;
 
-			const FNodeClassMetadata& Metadata = FConverterNode::GetAutoConverterNodeMetadata();
-			const Metasound::Frontend::FNodeRegistryKey Key = FMetasoundFrontendRegistryContainer::GetRegistryKey(Metadata);
-
-			if (!std::is_same<TFromDataType, TToDataType>::value && !FMetasoundFrontendRegistryContainer::Get()->IsNodeRegistered(Key))
+			if (TEnableAutoConverterNodeRegistration<TFromDataType, TToDataType>::Value)
 			{
-				ensureAlways(RegisterNodeWithFrontend<FConverterNode>(Metadata));
-				
-				bool bSucessfullyRegisteredConversionNode = RegisterConversionNode<FConverterNode, TFromDataType, TToDataType>(FConverterNode::GetInputName(), FConverterNode::GetOutputName(), Metadata);
-				ensureAlways(bSucessfullyRegisteredConversionNode);
+				const FNodeClassMetadata& Metadata = FConverterNode::GetAutoConverterNodeMetadata();
+				const Metasound::Frontend::FNodeRegistryKey Key = FMetasoundFrontendRegistryContainer::GetRegistryKey(Metadata);
+
+				if (!std::is_same<TFromDataType, TToDataType>::value && !FMetasoundFrontendRegistryContainer::Get()->IsNodeRegistered(Key))
+				{
+					ensureAlways(RegisterNodeWithFrontend<FConverterNode>(Metadata));
+					
+					bool bSucessfullyRegisteredConversionNode = RegisterConversionNode<FConverterNode, TFromDataType, TToDataType>(FConverterNode::GetInputName(), FConverterNode::GetOutputName(), Metadata);
+					ensureAlways(bSucessfullyRegisteredConversionNode);
+				}
 			}
 		}
 
@@ -384,10 +418,14 @@ namespace Metasound
 			using namespace MetasoundDataTypeRegistrationPrivate;
 			using TArrayType = TArray<TDataType>;
 
-			bool bSuccess = RegisterDataTypeWithFrontendInternal<TArrayType, TLiteralArrayEnum<PreferredArgType>::Value, UClassToUse>();
-			bSuccess = bSuccess && RegisterArrayNodes<TArrayType>();
+			if (TEnableAutoArrayTypeRegistration<TDataType>::Value)
+			{
+				bool bSuccess = RegisterDataTypeWithFrontendInternal<TArrayType, TLiteralArrayEnum<PreferredArgType>::Value, UClassToUse>();
+				bSuccess = bSuccess && RegisterArrayNodes<TArrayType>();
+				return bSuccess;
+			}
 
-			return bSuccess;
+			return true;
 		}
 	}
 	
diff --git a/Engine/Plugins/Runtime/Metasound/Source/MetasoundFrontend/Public/MetasoundEnumRegistrationMacro.h b/Engine/Plugins/Runtime/Metasound/Source/MetasoundFrontend/Public/MetasoundEnumRegistrationMacro.h
--- a/Engine/Plugins/Runtime/Metasound/Source/MetasoundFrontend/Public/MetasoundEnumRegistrationMacro.h
+++ b/Engine/Plugins/Runtime/Metasound/Source/MetasoundFrontend/Public/MetasoundEnumRegistrationMacro.h
@@ -9,18 +9,28 @@
 
 namespace Metasound
 {
+	// Disable registration of converter nodes which convert to enums
 	template<typename FromDataType, typename EnumType, EnumType DefaultValue>
-	struct TIsAutoConvertible<FromDataType, TEnum<EnumType, DefaultValue>>
+	struct TEnableAutoConverterNodeRegistration<FromDataType, TEnum<EnumType, DefaultValue>>
 	{
 		static constexpr bool Value = false;
 	};
 
+	// Disable registration of converter nodes which convert from enums
 	template<typename ToDataType, typename EnumType, EnumType DefaultValue>
-	struct TIsAutoConvertible<TEnum<EnumType, DefaultValue>, ToDataType>
+	struct TEnableAutoConverterNodeRegistration<TEnum<EnumType, DefaultValue>, ToDataType>
 	{
 		static constexpr bool Value = false;
 	};
 
+	// Disable arrays of enums
+	template<typename EnumType, EnumType DefaultValue>
+	struct TEnableAutoArrayTypeRegistration<TEnum<EnumType, DefaultValue>>
+	{
+		static constexpr bool Value = false;
+	};
+
+	// Disable array nodes of enums
 	template<typename EnumType, EnumType DefaultValue>
 	struct TEnableArrayNodes<TEnum<EnumType, DefaultValue>>
 	{
@@ -29,7 +39,7 @@ namespace Metasound
 
 	// Specialization of TIsTransmittable<> to disable transmission of enums.
 	template<typename EnumType, EnumType DefaultValue>
-	struct TIsTransmittable<TEnum<EnumType, DefaultValue>>
+	struct TEnableTransmissionNodeRegistration<TEnum<EnumType, DefaultValue>>
 	{
 	public:
 		static constexpr bool Value = false;
diff --git a/Engine/Plugins/Runtime/Metasound/Source/MetasoundStandardNodes/Private/MetasoundAudioFormats.cpp b/Engine/Plugins/Runtime/Metasound/Source/MetasoundStandardNodes/Private/MetasoundAudioFormats.cpp
--- a/Engine/Plugins/Runtime/Metasound/Source/MetasoundStandardNodes/Private/MetasoundAudioFormats.cpp
+++ b/Engine/Plugins/Runtime/Metasound/Source/MetasoundStandardNodes/Private/MetasoundAudioFormats.cpp
@@ -653,42 +653,42 @@ namespace Metasound
 
 	// Disable transmission of audio formats
 	template<>
-	struct TIsTransmittable<FMonoAudioFormat>
+	struct TEnableTransmissionNodeRegistration<FMonoAudioFormat>
 	{
 		static constexpr bool Value = false;
 	};
 
 	template<>
-	struct TIsTransmittable<FStereoAudioFormat>
+	struct TEnableTransmissionNodeRegistration<FStereoAudioFormat>
 	{
 		static constexpr bool Value = false;
 	};
 
-	// Disable auto converts using audio format construtors
+	// Disable auto converts using audio format constructors
 	template<typename FromDataType>
-	struct TIsAutoConvertible<FromDataType, FMonoAudioFormat>
+	struct TEnableAutoConverterNodeRegistration<FromDataType, FMonoAudioFormat>
 	{
-		static constexpr bool bIsConvertible = std::is_convertible<FromDataType, FAudioBuffer>::value;
-
-		static constexpr bool bIsFromArithmeticType = std::is_arithmetic<FromDataType>::value;
-
-	public:
-
-		static constexpr bool Value = bIsConvertible && (!bIsFromArithmeticType);
+		static constexpr bool Value = !std::is_arithmetic<FromDataType>::value;
 	};
 
 	template<typename FromDataType>
-	struct TIsAutoConvertible<FromDataType, FStereoAudioFormat>
+	struct TEnableAutoConverterNodeRegistration<FromDataType, FStereoAudioFormat>
 	{
-		static constexpr bool bIsConvertible = std::is_convertible<FromDataType, FAudioBuffer>::value;
-
-		static constexpr bool bIsFromArithmeticType = std::is_arithmetic<FromDataType>::value;
-
-	public:
+		static constexpr bool Value = !std::is_arithmetic<FromDataType>::value;
+	};
 
-		static constexpr bool Value = bIsConvertible && (!bIsFromArithmeticType);
+	// Disable arrays of audio formats
+	template<>
+	struct TEnableAutoArrayTypeRegistration<FMonoAudioFormat>
+	{
+		static constexpr bool Value = false;
 	};
 
+	template<>
+	struct TEnableAutoArrayTypeRegistration<FStereoAudioFormat>
+	{
+		static constexpr bool Value = false;
+	};
 }
 
 // Data type registration has to happen after TInputNode<> and TOutputNode<> specializations
