MetaSounds: fixing crash w/ sample & hold + midi quantizer

diff --git a/Engine/Plugins/Runtime/Metasound/Source/MetasoundStandardNodes/Private/MetasoundMidiNoteQuantizerNode.cpp b/Engine/Plugins/Runtime/Metasound/Source/MetasoundStandardNodes/Private/MetasoundMidiNoteQuantizerNode.cpp
--- a/Engine/Plugins/Runtime/Metasound/Source/MetasoundStandardNodes/Private/MetasoundMidiNoteQuantizerNode.cpp
+++ b/Engine/Plugins/Runtime/Metasound/Source/MetasoundStandardNodes/Private/MetasoundMidiNoteQuantizerNode.cpp
@@ -187,6 +187,7 @@ namespace Metasound
 		if ((*Scale).IsEmpty())
 		{
 			*MidiNoteOut = *MidiNoteIn;
+			return;
 		}
 
 		// calculate new output and cache values if needed
diff --git a/Engine/Plugins/Runtime/Metasound/Source/MetasoundStandardNodes/Private/MetasoundSampleAndHoldNode.cpp b/Engine/Plugins/Runtime/Metasound/Source/MetasoundStandardNodes/Private/MetasoundSampleAndHoldNode.cpp
--- a/Engine/Plugins/Runtime/Metasound/Source/MetasoundStandardNodes/Private/MetasoundSampleAndHoldNode.cpp
+++ b/Engine/Plugins/Runtime/Metasound/Source/MetasoundStandardNodes/Private/MetasoundSampleAndHoldNode.cpp
@@ -38,7 +38,7 @@ namespace Metasound
 
 				FNodeClassMetadata Metadata
 				{
-					FNodeClassName{FName("SampleAndHold"), TEXT("Sample And Hold"), TEXT("")},
+					FNodeClassName{FName("SampleAndHold"), TEXT("Sample And Hold"), Metasound::StandardNodes::AudioVariant},
 					1, // Major Version
 					0, // Minor Version
 					LOCTEXT("SampleAndHoldDisplayName", "Sample And Hold"),
@@ -90,7 +90,6 @@ namespace Metasound
 		FSampleAndHoldOperator(const FOperatorSettings& InSettings, const FTriggerReadRef& InTriggerSampleAndHold, const FAudioBufferReadRef& InAudioInput)
 			: AudioInput(InAudioInput)
 			, TriggerSampleAndHold(InTriggerSampleAndHold)
-			, OnTriggerSampleAndHold(FTriggerWriteRef::CreateNew(InSettings))
 			, AudioOutput(FAudioBufferWriteRef::CreateNew(InSettings))
 		{
 			// Init the hold value to the first sample in the audio buffer
@@ -113,30 +112,53 @@ namespace Metasound
 		{
 			using namespace SampleAndHoldVertexNames;
 			FDataReferenceCollection OutputDataReferences;
-			OutputDataReferences.AddDataReadReference(METASOUND_GET_PARAM_NAME(OutputOnSampleAndHold), OnTriggerSampleAndHold);
+			OutputDataReferences.AddDataReadReference(METASOUND_GET_PARAM_NAME(OutputOnSampleAndHold), TriggerSampleAndHold);
 			OutputDataReferences.AddDataReadReference(METASOUND_GET_PARAM_NAME(OutputAudio), AudioOutput);
 			return OutputDataReferences;
 		}
 		
 		void Execute()
 		{
-			OnTriggerSampleAndHold->AdvanceBlock();
-
 			// Lambda is used for both pre and post trigger frames
 			auto SetToHoldLambda = [this](int32 StartFrame, int32 EndFrame)
 			{
+				check(AudioOutput->Num() >= EndFrame);
+				float* OutputBufferPtr = AudioOutput->GetData();
+				float* CurrOutputPtr = &OutputBufferPtr[StartFrame];
+
 				const int32 NumSamples = EndFrame - StartFrame;
-				const int32 NonSIMDCount = NumSamples % AUDIO_SIMD_FLOAT_ALIGNMENT;
-				const int32 SIMDCount = NumSamples - NonSIMDCount;
 
-				float* OutputBufferPtr = AudioOutput->GetData();
+				// if CurrOutputPtr is on the last float of a (4)aligned block, StartAlignment will be 3,
+				// and we want LeadingNonSIMDCount below to be 1 (4-1)
+				const int32 StartAlignment = (uintptr_t)((void*)CurrOutputPtr) % AUDIO_SIMD_FLOAT_ALIGNMENT;
+
+				// if start alignment is 0, we are already aligned (hence the !!)
+				int32 SamplesRemaining = NumSamples;
+				const int32 LeadingNonSIMDCount = (StartAlignment == 0)? 0 : (AUDIO_SIMD_FLOAT_ALIGNMENT - StartAlignment);
+
+				SamplesRemaining -= LeadingNonSIMDCount;
+				const int32 SIMDCount = SamplesRemaining - (SamplesRemaining % AUDIO_SIMD_FLOAT_ALIGNMENT);
+
+				SamplesRemaining -= SIMDCount;
+				const int32 TrailingNonSIMDCount = SamplesRemaining;
 
-				float* HoldStartPtr = &OutputBufferPtr[StartFrame];
-				Audio::BufferSetToConstantInplace(HoldStartPtr, SIMDCount, HoldValue);
 
-				for (int32 i = StartFrame + SIMDCount; i < EndFrame; ++i)
+				// leading non-simd
+				for (int32 i = 0; i < LeadingNonSIMDCount; ++i)
 				{
-					OutputBufferPtr[i] = HoldValue;
+					*(CurrOutputPtr++) = HoldValue;
+				}
+
+				// simd
+				if (SIMDCount)
+				{
+					Audio::BufferSetToConstantInplace(CurrOutputPtr, SIMDCount, HoldValue);
+				}
+
+				// trailing non-simd
+				for (int32 i = 0; i < TrailingNonSIMDCount; ++i)
+				{
+					*(CurrOutputPtr++) = HoldValue;
 				}
 			};
 
@@ -149,9 +171,6 @@ namespace Metasound
 					// Get a new value to sample and hold
 					HoldValue = AudioInput->GetData()[StartFrame];
 					SetToHoldLambda(StartFrame, EndFrame);
-
-					// Output that we were triggered
-					OnTriggerSampleAndHold->TriggerFrame(StartFrame);
 				}
 			);
 		}
@@ -161,7 +180,6 @@ namespace Metasound
 		FAudioBufferReadRef AudioInput;
 		FTriggerReadRef TriggerSampleAndHold;
 
-		FTriggerWriteRef OnTriggerSampleAndHold;
 		FAudioBufferWriteRef AudioOutput;
 
 		float HoldValue = 0.0f;
diff --git a/Engine/Source/Runtime/SignalProcessing/Private/BufferVectorOperations.cpp b/Engine/Source/Runtime/SignalProcessing/Private/BufferVectorOperations.cpp
--- a/Engine/Source/Runtime/SignalProcessing/Private/BufferVectorOperations.cpp
+++ b/Engine/Source/Runtime/SignalProcessing/Private/BufferVectorOperations.cpp
@@ -168,6 +168,10 @@ namespace Audio
 
 	void BufferSetToConstantInplace(float* RESTRICT InBuffer, int32 NumSamples, float InConstant)
 	{
+		checkf(NumSamples >= 4, TEXT("Buffer must have atleast 4 elements."));
+		checkf(0 == (NumSamples % 4), TEXT("Buffer length be a multiple of 4."));
+		checkf(IsAligned<const float*>(InBuffer, AUDIO_SIMD_FLOAT_ALIGNMENT), TEXT("Memory must be aligned to use vector operations."));
+
 		const VectorRegister Constant = VectorLoadFloat1(&InConstant);
 
 		for (int32 i = 0; i < NumSamples; i += 4)
