- Fix for defaulted non-POD literal types not initializing correctly on new node creation. Issue was effectively manifesting with Enums (Triggers had a hack fix that is no longer necessary with this change).
- Fix for changing enum defaults not transacting.

diff --git a/Engine/Plugins/Runtime/Metasound/Source/MetasoundEditor/Private/MetasoundEditorGraphBuilder.cpp b/Engine/Plugins/Runtime/Metasound/Source/MetasoundEditor/Private/MetasoundEditorGraphBuilder.cpp
--- a/Engine/Plugins/Runtime/Metasound/Source/MetasoundEditor/Private/MetasoundEditorGraphBuilder.cpp
+++ b/Engine/Plugins/Runtime/Metasound/Source/MetasoundEditor/Private/MetasoundEditorGraphBuilder.cpp
@@ -368,66 +368,74 @@ namespace Metasound
 		{
 			using namespace Frontend;
 
-			FMetasoundFrontendLiteral PinDataTypeDefaultLiteral;
-			if (!GetPinDefaultLiteral(InInputPin, PinDataTypeDefaultLiteral))
+			FInputHandle InputHandle = GetInputHandleFromPin(&InInputPin);
+			if (!ensure(InputHandle->IsValid()))
 			{
 				return;
 			}
 
-			FInputHandle InputHandle = GetInputHandleFromPin(&InInputPin);
-			if (!ensure(InputHandle->IsValid()))
+			FOutputHandle OutputHandle = InputHandle->GetCurrentlyConnectedOutput();
+			if (OutputHandle->IsValid())
 			{
+				FNodeHandle InputNode = OutputHandle->GetOwningNode();
+				if (IsLiteralInput(InputNode))
+				{
+					Frontend::FGraphHandle GraphHandle = InputNode->GetOwningGraph();
+					const FGuid VertexID = GraphHandle->GetVertexIDForInputVertex(InputNode->GetNodeName());
+
+					FMetasoundFrontendLiteral PinDataTypeDefaultLiteral;
+					if (GetPinDefaultLiteral(InInputPin, PinDataTypeDefaultLiteral))
+					{
+						GraphHandle->SetDefaultInput(VertexID, PinDataTypeDefaultLiteral);
+					}
+				}
 				return;
 			}
 
-			FOutputHandle OutputHandle = InputHandle->GetCurrentlyConnectedOutput();
-			if (!OutputHandle->IsValid())
+			const FName TypeName = InputHandle->GetDataType();
+			const FString NewInputName = GenerateUniqueInputName(InMetasound);
+			Frontend::FNodeHandle NewInputNode = INodeController::GetInvalidHandle();
+
+			if (bForcePinValueAsDefault)
+			{
+				FMetasoundFrontendLiteral PinDataTypeDefaultLiteral;
+				if (GetPinDefaultLiteral(InInputPin, PinDataTypeDefaultLiteral))
+				{
+					NewInputNode = AddInputNodeHandle(InMetasound, NewInputName, TypeName, FText::GetEmpty(), EMetasoundFrontendNodeStyleDisplayVisibility::Hidden, &PinDataTypeDefaultLiteral);
+				}
+			}
+			else
 			{
-				const FName TypeName = InputHandle->GetDataType();
-				const FString NewInputName = GenerateUniqueInputName(InMetasound);
-				Frontend::FNodeHandle NewInputNode = AddInputNodeHandle(InMetasound, NewInputName, TypeName, FText::GetEmpty(), EMetasoundFrontendNodeStyleDisplayVisibility::Hidden, &PinDataTypeDefaultLiteral);
+				NewInputNode = AddInputNodeHandle(InMetasound, NewInputName, TypeName, FText::GetEmpty(), EMetasoundFrontendNodeStyleDisplayVisibility::Hidden);
 
 				Frontend::FGraphHandle GraphHandle = InNodeHandle->GetOwningGraph();
 				const FGuid VertexID = GraphHandle->GetVertexIDForInputVertex(NewInputNode->GetNodeName());
-				if (bForcePinValueAsDefault)
+
+				FMetasoundFrontendLiteral DefaultLiteral;
+
+				// If the node's input is set/not the base implementation (none), set the new graph input's value/pin's default value to the node input's default literal.
+				const FMetasoundFrontendLiteral* NodeInputDefault = InputHandle->GetDefaultLiteral();
+				if (NodeInputDefault && NodeInputDefault->GetType() != EMetasoundFrontendLiteralType::None)
 				{
-					GraphHandle->SetDefaultInput(VertexID, PinDataTypeDefaultLiteral);
+					DefaultLiteral = *NodeInputDefault;
 				}
-				// Hack fix for triggers having special aspect in that they hide the default
-				// literal picker, but require a default value for graph synchronization purposes.
-				else if (InInputPin.PinType.PinCategory == PinCategoryTrigger)
+				// Otherwise, get the class default for the given input node's type, and set the new graph input's value/input pin's default to that.
+				else
 				{
-					FMetasoundFrontendLiteral DefaultTriggerLiteral;
-					DefaultTriggerLiteral.Set(false);
-					GraphHandle->SetDefaultInput(VertexID, DefaultTriggerLiteral);
-					InInputPin.DefaultValue = DefaultTriggerLiteral.ToString();
-				}
-				else if (const FMetasoundFrontendLiteral* DefaultLiteral = InputHandle->GetDefaultLiteral())
-				{
-					if (DefaultLiteral->IsValid())
-					{
-						GraphHandle->SetDefaultInput(VertexID, *DefaultLiteral);
-						InInputPin.DefaultValue = DefaultLiteral->ToString();
-					}
+					const Metasound::FLiteral Literal = Frontend::GetDefaultParamForDataType(TypeName);
+					DefaultLiteral.SetFromLiteral(Literal);
 				}
+				InInputPin.DefaultValue = DefaultLiteral.ToString();
+				GraphHandle->SetDefaultInput(VertexID, DefaultLiteral);
+			}
 
+			if (NewInputNode->IsValid())
+			{
 				TArray<FOutputHandle> OutputHandles = NewInputNode->GetOutputs();
 				if (ensure(OutputHandles.Num() == 1))
 				{
 					OutputHandle = OutputHandles[0];
-				}
-
-				ensure(InputHandle->Connect(*OutputHandle));
-			}
-			else
-			{
-				FNodeHandle InputNode = OutputHandle->GetOwningNode();
-				if (IsLiteralInput(InputNode))
-				{
-					Frontend::FGraphHandle GraphHandle = InputNode->GetOwningGraph();
-					const FGuid VertexID = GraphHandle->GetVertexIDForInputVertex(InputNode->GetNodeName());
-
-					GraphHandle->SetDefaultInput(VertexID, PinDataTypeDefaultLiteral);
+					ensure(InputHandle->Connect(*OutputHandle));
 				}
 			}
 		}
@@ -662,10 +670,7 @@ namespace Metasound
 
 				if (InDefaultValue)
 				{
-					if (InDefaultValue->GetType() != EMetasoundFrontendLiteralType::None && InDefaultValue->GetType() != EMetasoundFrontendLiteralType::Invalid)
-					{
-						GraphHandle->SetDefaultInput(VertexID, *InDefaultValue);
-					}
+					GraphHandle->SetDefaultInput(VertexID, *InDefaultValue);
 				}
 				else
 				{
diff --git a/Engine/Plugins/Runtime/Metasound/Source/MetasoundEditor/Private/SMetasoundEnumPin.cpp b/Engine/Plugins/Runtime/Metasound/Source/MetasoundEditor/Private/SMetasoundEnumPin.cpp
--- a/Engine/Plugins/Runtime/Metasound/Source/MetasoundEditor/Private/SMetasoundEnumPin.cpp
+++ b/Engine/Plugins/Runtime/Metasound/Source/MetasoundEditor/Private/SMetasoundEnumPin.cpp
@@ -1,6 +1,7 @@
 // Copyright Epic Games, Inc. All Rights Reserved.
 
 #include "SMetasoundEnumPin.h"
+#include "MetasoundEditorGraph.h"
 #include "MetasoundEditorGraphNode.h"
 #include "ScopedTransaction.h"
 
@@ -82,6 +83,18 @@ void SMetasoundEnumPin::ComboBoxSelectionChanged(TSharedPtr<int32> NewSelection,
 		FString EnumValueString = FString::FromInt(EnumValue);
 		if (GraphPinObj->GetDefaultAsString() != EnumValueString)
 		{
+			const FScopedTransaction Transaction(NSLOCTEXT("MetaSoundEditor", "ChangeEnumPinValue", "Change MetaSound Node Default Input Enum Value"));
+			GraphPinObj->Modify();
+
+			if (UMetasoundEditorGraphNode* MetaSoundNode = Cast<UMetasoundEditorGraphNode>(GraphPinObj->GetOwningNode()))
+			{
+				if (UMetasoundEditorGraph* Graph = Cast<UMetasoundEditorGraph>(MetaSoundNode->GetGraph()))
+				{
+					Graph->Modify();
+					Graph->GetMetasoundChecked().Modify();
+				}
+			}
+
 			//Set new selection
 			GraphPinObj->GetSchema()->TrySetDefaultValue(*GraphPinObj, EnumValueString);
 		}
