- Rename AudioCommunicationInterface to AudioParameterInterface
- Add button to execute MetaSound trigger on node titlebar

diff --git a/Engine/Plugins/Runtime/Metasound/Source/MetasoundEditor/Private/MetasoundEditor.cpp b/Engine/Plugins/Runtime/Metasound/Source/MetasoundEditor/Private/MetasoundEditor.cpp
--- a/Engine/Plugins/Runtime/Metasound/Source/MetasoundEditor/Private/MetasoundEditor.cpp
+++ b/Engine/Plugins/Runtime/Metasound/Source/MetasoundEditor/Private/MetasoundEditor.cpp
@@ -734,10 +734,10 @@ namespace Metasound
 				{
 					PlayTime = 0.0;
 
-					UObject* CommObject = PreviewComp->GetCommunicationInterface().GetObject();
-					if (ensure(CommObject))
+					UObject* ParamInterfaceObject = PreviewComp->GetParameterInterface().GetObject();
+					if (ensure(ParamInterfaceObject))
 					{
-						SetPreviewID(CommObject->GetUniqueID());
+						SetPreviewID(ParamInterfaceObject->GetUniqueID());
 					}
 
 					if (MetasoundAudioBus.IsValid())
@@ -847,11 +847,11 @@ namespace Metasound
 				{
 					// TODO: fix how identifying the parameter to update is determined. It should not be done
 					// with a "DisplayName" but rather the vertex Guid.
-					if (TScriptInterface<IAudioCommunicationInterface> CommInterface = PreviewComponent->GetCommunicationInterface())
+					if (TScriptInterface<IAudioParameterInterface> ParamInterface = PreviewComponent->GetParameterInterface())
 					{
 						Metasound::Frontend::FConstNodeHandle NodeHandle = InputNode->GetConstNodeHandle();
 						Metasound::FVertexKey VertexKey = Metasound::FVertexKey(NodeHandle->GetDisplayName().ToString());
-						CommInterface->Trigger(*VertexKey);
+						ParamInterface->Trigger(*VertexKey);
 					}
 				}
 			}
diff --git a/Engine/Plugins/Runtime/Metasound/Source/MetasoundEditor/Private/MetasoundEditorGraph.cpp b/Engine/Plugins/Runtime/Metasound/Source/MetasoundEditor/Private/MetasoundEditorGraph.cpp
--- a/Engine/Plugins/Runtime/Metasound/Source/MetasoundEditor/Private/MetasoundEditorGraph.cpp
+++ b/Engine/Plugins/Runtime/Metasound/Source/MetasoundEditor/Private/MetasoundEditorGraph.cpp
@@ -1,6 +1,7 @@
 // Copyright Epic Games, Inc. All Rights Reserved.
 #include "MetasoundEditorGraph.h"
 
+#include "Audio/AudioParameterInterface.h"
 #include "Components/AudioComponent.h"
 #include "EdGraph/EdGraphNode.h"
 #include "MetasoundEditorGraphBuilder.h"
@@ -9,7 +10,6 @@
 #include "MetasoundEditorModule.h"
 #include "MetasoundUObjectRegistry.h"
 #include "ScopedTransaction.h"
-#include "Sound/AudioCommunicationInterface.h"
 
 #define LOCTEXT_NAMESPACE "MetaSoundEditor"
 
@@ -161,11 +161,11 @@ void UMetasoundEditorGraphInput::UpdateDocumentInput(bool bPostTransaction)
 	GraphHandle->SetDefaultInput(VertexID, Literal->GetDefault());
 }
 
-void UMetasoundEditorGraphInput::UpdatePreviewInstance(const Metasound::FVertexKey& InParameterName, TScriptInterface<IAudioCommunicationInterface>& InCommInterface) const
+void UMetasoundEditorGraphInput::UpdatePreviewInstance(const Metasound::FVertexKey& InParameterName, TScriptInterface<IAudioParameterInterface>& InParamInterface) const
 {
 	if (ensure(Literal))
 	{
-		Literal->UpdatePreviewInstance(InParameterName, InCommInterface);
+		Literal->UpdatePreviewInstance(InParameterName, InParamInterface);
 	}
 }
 
@@ -214,13 +214,13 @@ void UMetasoundEditorGraphInput::OnLiteralChanged(bool bPostTransaction)
 		UAudioComponent* PreviewComponent = GEditor->GetPreviewAudioComponent();
 		check(PreviewComponent);
 
-		if (TScriptInterface<IAudioCommunicationInterface> CommInterface = PreviewComponent->GetCommunicationInterface())
+		if (TScriptInterface<IAudioParameterInterface> ParamInterface = PreviewComponent->GetParameterInterface())
 		{
 			// TODO: fix how identifying the parameter to update is determined. It should not be done
 			// with a "DisplayName" but rather the vertex Guid.
 			Metasound::Frontend::FConstNodeHandle NodeHandle = GetConstNodeHandle();
 			Metasound::FVertexKey VertexKey = Metasound::FVertexKey(NodeHandle->GetDisplayName().ToString());
-			UpdatePreviewInstance(VertexKey, CommInterface);
+			UpdatePreviewInstance(VertexKey, ParamInterface);
 		}
 	}
 }
@@ -460,13 +460,13 @@ bool UMetasoundEditorGraph::IsPreviewing() const
 		return false;
 	}
 
-	UObject* CommInterface = PreviewComponent->GetCommunicationInterface().GetObject();
-	if (!CommInterface)
+	UObject* ParamInterfaceObject = PreviewComponent->GetParameterInterface().GetObject();
+	if (!ParamInterfaceObject)
 	{
 		return false;
 	}
 
-	return CommInterface->GetUniqueID() == PreviewID;
+	return ParamInterfaceObject->GetUniqueID() == PreviewID;
 }
 
 void UMetasoundEditorGraph::IterateOutputs(TUniqueFunction<void(UMetasoundEditorGraphOutput&)> InFunction) const
diff --git a/Engine/Plugins/Runtime/Metasound/Source/MetasoundEditor/Private/MetasoundNodeDetailCustomization.cpp b/Engine/Plugins/Runtime/Metasound/Source/MetasoundEditor/Private/MetasoundNodeDetailCustomization.cpp
--- a/Engine/Plugins/Runtime/Metasound/Source/MetasoundEditor/Private/MetasoundNodeDetailCustomization.cpp
+++ b/Engine/Plugins/Runtime/Metasound/Source/MetasoundEditor/Private/MetasoundNodeDetailCustomization.cpp
@@ -13,6 +13,7 @@
 #include "IDetailGroup.h"
 #include "Internationalization/Text.h"
 #include "MetasoundAssetBase.h"
+#include "MetasoundDataReference.h"
 #include "MetasoundEditorGraphBuilder.h"
 #include "MetasoundEditorGraphNode.h"
 #include "MetasoundEditorGraphInputNodes.h"
@@ -26,6 +27,7 @@
 #include "PropertyHandle.h"
 #include "PropertyRestriction.h"
 #include "SlateCore/Public/Styling/SlateColor.h"
+#include "SMetasoundGraphNode.h"
 #include "Templates/Casts.h"
 #include "Templates/SharedPointer.h"
 #include "UObject/WeakObjectPtr.h"
@@ -65,8 +67,6 @@ namespace Metasound
 			static const FText InputNameText = LOCTEXT("Input_Name", "Input Name");
 			static const FText OutputNameText = LOCTEXT("Output_Name", "Output Name");
 
-			static const FName TriggerTypeName = "Trigger";
-
 			static const FName DataTypeNameIdentifier = "DataTypeName";
 			static const FName ProxyGeneratorClassNameIdentifier = "GeneratorClass";
 
@@ -77,41 +77,6 @@ namespace Metasound
 				"Audio:Mono",
 				"Audio:Stereo"
 			};
-
-			void ExecuteInputTrigger(UMetasoundEditorGraphInputLiteral* Literal)
-			{
-				if (!Literal)
-				{
-					return;
-				}
-
-				UMetasoundEditorGraphInput* Input = Cast<UMetasoundEditorGraphInput>(Literal->GetOuter());
-				if (!ensure(Input))
-				{
-					return;
-				}
-
-				// If modifying graph not currently being previewed, do not forward request
-				if (UMetasoundEditorGraph* Graph = Cast<UMetasoundEditorGraph>(Input->GetOuter()))
-				{
-					if (!Graph->IsPreviewing())
-					{
-						return;
-					}
-				}
-
-				if (UAudioComponent* PreviewComponent = GEditor->GetPreviewAudioComponent())
-				{
-					if (TScriptInterface<IAudioCommunicationInterface> CommInterface = PreviewComponent->GetCommunicationInterface())
-					{
-						// TODO: fix how identifying the parameter to update is determined. It should not be done
-						// with a "DisplayName" but rather the vertex Guid.
-						Metasound::Frontend::FConstNodeHandle NodeHandle = Input->GetConstNodeHandle();
-						Metasound::FVertexKey VertexKey = Metasound::FVertexKey(NodeHandle->GetDisplayName().ToString());
-						CommInterface->Trigger(*VertexKey);
-					}
-				}
-			}
 		}
 
 		void FMetasoundInputBoolDetailCustomization::CacheProxyData(TSharedPtr<IPropertyHandle> ProxyHandle)
@@ -127,7 +92,7 @@ namespace Metasound
 
 		FText FMetasoundInputBoolDetailCustomization::GetPropertyNameOverride() const
 		{
-			if (DataTypeName == VariableCustomizationPrivate::TriggerTypeName)
+			if (DataTypeName == Metasound::GetMetasoundDataTypeName<Metasound::FTrigger>())
 			{
 				return LOCTEXT("TriggerInput_SimulateTitle", "Simulate");
 			}
@@ -145,41 +110,20 @@ namespace Metasound
 				if (ValueProperty.IsValid())
 				{
 					// Not a trigger, so just display as underlying literal type (bool)
-					if (DataTypeName != VariableCustomizationPrivate::TriggerTypeName)
+					if (DataTypeName != Metasound::GetMetasoundDataTypeName<Metasound::FTrigger>())
 					{
 						return ValueProperty->CreatePropertyValueWidget();
 					}
 
-					return
-						SNew(SHorizontalBox)
-						+ SHorizontalBox::Slot()
-						.AutoWidth()
-						.Padding(2.0f, 0.0f, 0.0f, 0.0f)
-						.VAlign(VAlign_Center)
-						[
-							SNew(SButton)
-							.ButtonStyle(FAppStyle::Get(), "SimpleButton")
-							.OnClicked_Lambda([ValueProperty]()
-							{
-								TArray<UObject*> OuterObjects;
-								ValueProperty->GetOuterObjects(OuterObjects);
-								for (UObject* Object : OuterObjects)
-								{
-									UMetasoundEditorGraphInputLiteral* Literal = Cast<UMetasoundEditorGraphInputLiteral>(Object);
-									VariableCustomizationPrivate::ExecuteInputTrigger(Literal);
-								}
-								return FReply::Handled();
-							})
-							.ToolTipText(LOCTEXT("TriggerTestToolTip", "Executes trigger if currently previewing Metasound."))
-							.ForegroundColor(FSlateColor::UseForeground())
-							.ContentPadding(0)
-							.IsFocusable(false)
-							[
-								SNew(SImage)
-								.Image(FAppStyle::Get().GetBrush("Icons.CircleArrowDown"))
-								.ColorAndOpacity(FSlateColor::UseForeground())
-							]
-						];
+					TArray<UObject*> OuterObjects;
+					ValueProperty->GetOuterObjects(OuterObjects);
+					for (UObject* Object : OuterObjects)
+					{
+						if (UMetasoundEditorGraphInputLiteral* Literal = Cast<UMetasoundEditorGraphInputLiteral>(Object))
+						{
+							return SMetasoundGraphNode::CreateTriggerSimulationWidget(*Literal);
+						}
+					}
 				}
 			}
 
diff --git a/Engine/Plugins/Runtime/Metasound/Source/MetasoundEditor/Private/SMetasoundGraphNode.cpp b/Engine/Plugins/Runtime/Metasound/Source/MetasoundEditor/Private/SMetasoundGraphNode.cpp
--- a/Engine/Plugins/Runtime/Metasound/Source/MetasoundEditor/Private/SMetasoundGraphNode.cpp
+++ b/Engine/Plugins/Runtime/Metasound/Source/MetasoundEditor/Private/SMetasoundGraphNode.cpp
@@ -1,6 +1,8 @@
 // Copyright Epic Games, Inc. All Rights Reserved.
 #include "SMetasoundGraphNode.h"
 
+#include "Audio/AudioParameterInterface.h"
+#include "Components/AudioComponent.h"
 #include "EditorStyleSet.h"
 #include "GraphEditorSettings.h"
 #include "IDocumentation.h"
@@ -10,11 +12,13 @@
 #include "KismetPins/SGraphPinNum.h"
 #include "KismetPins/SGraphPinObject.h"
 #include "KismetPins/SGraphPinString.h"
+#include "MetasoundEditorGraph.h"
 #include "MetasoundEditorGraphBuilder.h"
 #include "MetasoundEditorGraphNode.h"
 #include "MetasoundEditorGraphSchema.h"
 #include "MetasoundEditorModule.h"
 #include "MetasoundFrontendRegistries.h"
+#include "MetasoundTrigger.h"
 #include "NodeFactory.h"
 #include "PropertyCustomizationHelpers.h"
 #include "ScopedTransaction.h"
@@ -22,10 +26,12 @@
 #include "SGraphNode.h"
 #include "SGraphPin.h"
 #include "SLevelOfDetailBranchNode.h"
+#include "Styling/AppStyle.h"
 #include "Styling/SlateStyleRegistry.h"
 #include "Styling/SlateColor.h"
 #include "TutorialMetaData.h"
 #include "Widgets/Images/SImage.h"
+#include "Widgets/Input/SButton.h"
 #include "Widgets/Layout/SBox.h"
 #include "Widgets/Layout/SSpacer.h"
 #include "Widgets/SBoxPanel.h"
@@ -33,6 +39,7 @@
 #include "Widgets/SWidget.h"
 #include "SGraphPinComboBox.h"
 #include "SMetasoundEnumPin.h"
+#include "UObject/ScriptInterface.h"
 
 #define LOCTEXT_NAMESPACE "MetasoundGraphNode"
 
@@ -44,6 +51,66 @@ void SMetasoundGraphNode::Construct(const FArguments& InArgs, class UEdGraphNode
 	UpdateGraphNode();
 }
 
+void SMetasoundGraphNode::ExecuteInputTrigger(UMetasoundEditorGraphInputLiteral& Literal)
+{
+	UMetasoundEditorGraphInput* Input = Cast<UMetasoundEditorGraphInput>(Literal.GetOuter());
+	if (!ensure(Input))
+	{
+		return;
+	}
+
+	// If modifying graph not currently being previewed, do not forward request
+	if (UMetasoundEditorGraph* Graph = Cast<UMetasoundEditorGraph>(Input->GetOuter()))
+	{
+		if (!Graph->IsPreviewing())
+		{
+			return;
+		}
+	}
+
+	if (UAudioComponent* PreviewComponent = GEditor->GetPreviewAudioComponent())
+	{
+		if (TScriptInterface<IAudioParameterInterface> ParamInterface = PreviewComponent->GetParameterInterface())
+		{
+			// TODO: fix how identifying the parameter to update is determined. It should not be done
+			// with a "DisplayName" but rather the vertex Guid.
+			Metasound::Frontend::FConstNodeHandle NodeHandle = Input->GetConstNodeHandle();
+			Metasound::FVertexKey VertexKey = Metasound::FVertexKey(NodeHandle->GetDisplayName().ToString());
+			ParamInterface->Trigger(*VertexKey);
+		}
+	}
+}
+
+TSharedRef<SWidget> SMetasoundGraphNode::CreateTriggerSimulationWidget(UMetasoundEditorGraphInputLiteral& InputLiteral)
+{
+	return SNew(SHorizontalBox)
+	+ SHorizontalBox::Slot()
+	.Padding(2.0f, 0.0f, 0.0f, 0.0f)
+	.HAlign(HAlign_Left)
+	.VAlign(VAlign_Center)
+	[
+		SNew(SButton)
+		.ButtonStyle(FAppStyle::Get(), "SimpleButton")
+		.OnClicked_Lambda([LiteralPtr = TWeakObjectPtr<UMetasoundEditorGraphInputLiteral>(&InputLiteral)]()
+		{
+			if (LiteralPtr.IsValid())
+			{
+				ExecuteInputTrigger(*LiteralPtr.Get());
+			}
+			return FReply::Handled();
+		})
+		.ToolTipText(LOCTEXT("TriggerTestToolTip", "Executes trigger if currently previewing MetaSound."))
+		.ForegroundColor(FSlateColor::UseForeground())
+		.ContentPadding(0)
+		.IsFocusable(false)
+		[
+			SNew(SImage)
+			.Image(FAppStyle::Get().GetBrush("Icons.CircleArrowDown"))
+			.ColorAndOpacity(FSlateColor::UseForeground())
+		]
+	];
+}
+
 void SMetasoundGraphNode::CreateOutputSideAddButton(TSharedPtr<SVerticalBox> OutputBox)
 {
 	TSharedRef<SWidget> AddPinButton = AddPinButtonContent(
@@ -65,8 +132,7 @@ void SMetasoundGraphNode::CreateOutputSideAddButton(TSharedPtr<SVerticalBox> Out
 
 UMetasoundEditorGraphNode& SMetasoundGraphNode::GetMetasoundNode()
 {
-	check(GraphNode);
-	return *Cast<UMetasoundEditorGraphNode>(GraphNode);
+	return *CastChecked<UMetasoundEditorGraphNode>(GraphNode);
 }
 
 const UMetasoundEditorGraphNode& SMetasoundGraphNode::GetMetasoundNode() const
@@ -297,6 +363,42 @@ FReply SMetasoundGraphNode::OnAddPin()
 	return FReply::Handled();
 }
 
+FName SMetasoundGraphNode::GetLiteralDataType() const
+{
+	using namespace Metasound::Frontend;
+
+	FName TypeName;
+
+	// Just take last type.  If more than one, all types are the same.
+	const UMetasoundEditorGraphNode& Node = GetMetasoundNode();
+	Node.GetNodeHandle()->IterateConstOutputs([InTypeName = &TypeName](FConstOutputHandle OutputHandle)
+	{
+		*InTypeName = OutputHandle->GetDataType();
+	});
+
+	return TypeName;
+}
+
+TSharedRef<SWidget> SMetasoundGraphNode::CreateTitleRightWidget()
+{
+	const FName TypeName = GetLiteralDataType();
+	if (TypeName == Metasound::GetMetasoundDataTypeName<Metasound::FTrigger>())
+	{
+		if (UMetasoundEditorGraphInputNode* Node = Cast<UMetasoundEditorGraphInputNode>(&GetMetasoundNode()))
+		{
+			if (UMetasoundEditorGraphInput* Input = Node->Input)
+			{
+				if (Input->Literal)
+				{
+					return CreateTriggerSimulationWidget(*Input->Literal);
+				}
+			}
+		}
+	}
+
+	return SGraphNode::CreateTitleRightWidget();
+}
+
 TSharedRef<SWidget> SMetasoundGraphNode::CreateNodeContentArea()
 {
 	using namespace Metasound::Editor;
diff --git a/Engine/Plugins/Runtime/Metasound/Source/MetasoundEditor/Private/SMetasoundGraphNode.h b/Engine/Plugins/Runtime/Metasound/Source/MetasoundEditor/Private/SMetasoundGraphNode.h
--- a/Engine/Plugins/Runtime/Metasound/Source/MetasoundEditor/Private/SMetasoundGraphNode.h
+++ b/Engine/Plugins/Runtime/Metasound/Source/MetasoundEditor/Private/SMetasoundGraphNode.h
@@ -12,6 +12,7 @@
 // Forward Declarations
 class SGraphPin;
 class SVerticalBox;
+class UMetasoundEditorGraphInputLiteral;
 class UMetasoundEditorGraphNode;
 
 
@@ -34,11 +35,17 @@ protected:
 	virtual TSharedRef<SWidget> CreateNodeContentArea() override;
 	virtual TSharedRef<SWidget> CreateTitleWidget(TSharedPtr<SNodeTitle> NodeTitle) override;
 	virtual const FSlateBrush* GetNodeBodyBrush() const override;
+	virtual TSharedRef<SWidget> CreateTitleRightWidget() override;
 	virtual EVisibility IsAddPinButtonVisible() const override;
 	virtual FReply OnAddPin() override;
 	virtual void MoveTo(const FVector2D& NewPosition, FNodeSet& NodeFilter) override;
 	virtual void SetDefaultTitleAreaWidget(TSharedRef<SOverlay> DefaultTitleAreaWidget) override;
 
+	FName GetLiteralDataType() const;
 	UMetasoundEditorGraphNode& GetMetasoundNode();
 	const UMetasoundEditorGraphNode& GetMetasoundNode() const;
+
+public:
+	static void ExecuteInputTrigger(UMetasoundEditorGraphInputLiteral& Literal);
+	static TSharedRef<SWidget> CreateTriggerSimulationWidget(UMetasoundEditorGraphInputLiteral& Literal);
 };
diff --git a/Engine/Plugins/Runtime/Metasound/Source/MetasoundEditor/Public/MetasoundEditorGraph.h b/Engine/Plugins/Runtime/Metasound/Source/MetasoundEditor/Public/MetasoundEditorGraph.h
--- a/Engine/Plugins/Runtime/Metasound/Source/MetasoundEditor/Public/MetasoundEditorGraph.h
+++ b/Engine/Plugins/Runtime/Metasound/Source/MetasoundEditor/Public/MetasoundEditorGraph.h
@@ -1,6 +1,7 @@
 // Copyright Epic Games, Inc. All Rights Reserved.
 #pragma once
 
+#include "Audio/AudioParameterInterface.h"
 #include "CoreMinimal.h"
 #include "MetasoundEditor.h"
 #include "MetasoundFrontendController.h"
@@ -8,7 +9,6 @@
 #include "MetasoundFrontendLiteral.h"
 #include "MetasoundSource.h"
 #include "MetasoundUObjectRegistry.h"
-#include "Sound/AudioCommunicationInterface.h"
 #include "UObject/ObjectMacros.h"
 #include "UObject/ScriptInterface.h"
 
@@ -59,7 +59,7 @@ class METASOUNDEDITOR_API UMetasoundEditorGraphInputLiteral : public UObject
 	GENERATED_BODY()
 
 public:
-	virtual void UpdatePreviewInstance(const Metasound::FVertexKey& InParameterName, TScriptInterface<IAudioCommunicationInterface>& InCommInterface) const
+	virtual void UpdatePreviewInstance(const Metasound::FVertexKey& InParameterName, TScriptInterface<IAudioParameterInterface>& InParameterInterface) const
 	{
 	}
 
@@ -96,7 +96,7 @@ public:
 	UMetasoundEditorGraphInputLiteral* Literal;
 
 	void UpdateDocumentInput(bool bPostTransaction = true);
-	void UpdatePreviewInstance(const Metasound::FVertexKey& InParameterName, TScriptInterface<IAudioCommunicationInterface>& InCommInterface) const;
+	void UpdatePreviewInstance(const Metasound::FVertexKey& InParameterName, TScriptInterface<IAudioParameterInterface>& InParameterInterface) const;
 
 	void OnDataTypeChanged() override;
 	void OnLiteralChanged(bool bPostTransaction = true);
@@ -139,7 +139,7 @@ public:
 
 private:
 	// Preview ID is the Unique ID provided by the UObject that implements
-	// a sound's CommunicationInterface when a sound begins playing.
+	// a sound's ParameterInterface when a sound begins playing.
 	uint32 PreviewID = INDEX_NONE;
 
 	UPROPERTY()
diff --git a/Engine/Plugins/Runtime/Metasound/Source/MetasoundEditor/Public/MetasoundEditorGraphInputNodes.h b/Engine/Plugins/Runtime/Metasound/Source/MetasoundEditor/Public/MetasoundEditorGraphInputNodes.h
--- a/Engine/Plugins/Runtime/Metasound/Source/MetasoundEditor/Public/MetasoundEditorGraphInputNodes.h
+++ b/Engine/Plugins/Runtime/Metasound/Source/MetasoundEditor/Public/MetasoundEditorGraphInputNodes.h
@@ -3,6 +3,7 @@
 
 #include "Algo/Count.h"
 #include "Algo/Transform.h"
+#include "Audio/AudioParameterInterface.h"
 #include "CoreMinimal.h"
 #include "EdGraph/EdGraphNode.h"
 #include "MetasoundDataReference.h"
@@ -15,7 +16,6 @@
 #include "MetasoundFrontendRegistries.h"
 #include "MetasoundPrimitives.h"
 #include "Misc/Guid.h"
-#include "Sound/AudioCommunicationInterface.h"
 #include "Sound/SoundWave.h"
 #include "UObject/ObjectMacros.h"
 #include "UObject/SoftObjectPath.h"
@@ -75,11 +75,11 @@ public:
 		return Super::GetClassName();
 	}
 
-	virtual void UpdatePreviewInstance(const Metasound::FVertexKey& InParameterName, TScriptInterface<IAudioCommunicationInterface>& InCommInterface) const
+	virtual void UpdatePreviewInstance(const Metasound::FVertexKey& InParameterName, TScriptInterface<IAudioParameterInterface>& InParameterInterface) const
 	{
 		if (Input)
 		{
-			Input->UpdatePreviewInstance(InParameterName, InCommInterface);
+			Input->UpdatePreviewInstance(InParameterName, InParameterInterface);
 		}
 	}
 	
@@ -152,9 +152,9 @@ public:
 		Metasound::Editor::InputPrivate::ConvertLiteral<bool>(InLiteral, Default.Value);
 	}
 
-	virtual void UpdatePreviewInstance(const Metasound::FVertexKey& InParameterName, TScriptInterface<IAudioCommunicationInterface>& InCommInterface) const override
+	virtual void UpdatePreviewInstance(const Metasound::FVertexKey& InParameterName, TScriptInterface<IAudioParameterInterface>& InParameterInterface) const override
 	{
-		InCommInterface->SetBool(*InParameterName, Default.Value);
+		InParameterInterface->SetBool(*InParameterName, Default.Value);
 	}
 };
 
@@ -189,11 +189,11 @@ public:
 		Metasound::Editor::InputPrivate::ConvertLiteralToArray<bool, FMetasoundEditorGraphInputBoolRef>(InLiteral, Default);
 	}
 
-	virtual void UpdatePreviewInstance(const Metasound::FVertexKey& InParameterName, TScriptInterface<IAudioCommunicationInterface>& InCommInterface) const override
+	virtual void UpdatePreviewInstance(const Metasound::FVertexKey& InParameterName, TScriptInterface<IAudioParameterInterface>& InParameterInterface) const override
 	{
 		TArray<bool> BoolArray;
 		Algo::Transform(Default, BoolArray, [](const FMetasoundEditorGraphInputBoolRef& InValue) { return InValue.Value; });
-		InCommInterface->SetBoolArray(*InParameterName, BoolArray);
+		InParameterInterface->SetBoolArray(*InParameterName, BoolArray);
 	}
 };
 
@@ -235,9 +235,9 @@ public:
 		Metasound::Editor::InputPrivate::ConvertLiteral<int32>(InLiteral, Default.Value);
 	}
 
-	virtual void UpdatePreviewInstance(const Metasound::FVertexKey& InParameterName, TScriptInterface<IAudioCommunicationInterface>& InCommInterface) const override
+	virtual void UpdatePreviewInstance(const Metasound::FVertexKey& InParameterName, TScriptInterface<IAudioParameterInterface>& InParameterInterface) const override
 	{
-		InCommInterface->SetInt(*InParameterName, Default.Value);
+		InParameterInterface->SetInt(*InParameterName, Default.Value);
 	}
 };
 
@@ -273,11 +273,11 @@ public:
 		Metasound::Editor::InputPrivate::ConvertLiteralToArray<int32, FMetasoundEditorGraphInputIntRef>(InLiteral, Default);
 	}
 
-	void UpdatePreviewInstance(const Metasound::FVertexKey& InParameterName, TScriptInterface<IAudioCommunicationInterface>& InCommInterface) const override
+	void UpdatePreviewInstance(const Metasound::FVertexKey& InParameterName, TScriptInterface<IAudioParameterInterface>& InParameterInterface) const override
 	{
 		TArray<int32> IntArray;
 		Algo::Transform(Default, IntArray, [](const FMetasoundEditorGraphInputIntRef& InValue) { return InValue.Value; });
-		InCommInterface->SetIntArray(*InParameterName, IntArray);
+		InParameterInterface->SetIntArray(*InParameterName, IntArray);
 	}
 };
 
@@ -309,9 +309,9 @@ public:
 		Metasound::Editor::InputPrivate::ConvertLiteral<float>(InLiteral, Default);
 	}
 
-	virtual void UpdatePreviewInstance(const Metasound::FVertexKey& InParameterName, TScriptInterface<IAudioCommunicationInterface>& InCommInterface) const override
+	virtual void UpdatePreviewInstance(const Metasound::FVertexKey& InParameterName, TScriptInterface<IAudioParameterInterface>& InParameterInterface) const override
 	{
-		InCommInterface->SetFloat(*InParameterName, Default);
+		InParameterInterface->SetFloat(*InParameterName, Default);
 	}
 };
 
@@ -343,9 +343,9 @@ public:
 		Metasound::Editor::InputPrivate::ConvertLiteralToArray<float>(InLiteral, Default);
 	}
 
-	virtual void UpdatePreviewInstance(const Metasound::FVertexKey& InParameterName, TScriptInterface<IAudioCommunicationInterface>& InCommInterface) const override
+	virtual void UpdatePreviewInstance(const Metasound::FVertexKey& InParameterName, TScriptInterface<IAudioParameterInterface>& InParameterInterface) const override
 	{
-		InCommInterface->SetFloatArray(*InParameterName, Default);
+		InParameterInterface->SetFloatArray(*InParameterName, Default);
 	}
 };
 
@@ -377,9 +377,9 @@ public:
 		Metasound::Editor::InputPrivate::ConvertLiteral<FString>(InLiteral, Default);
 	}
 
-	virtual void UpdatePreviewInstance(const Metasound::FVertexKey& InParameterName, TScriptInterface<IAudioCommunicationInterface>& InCommInterface) const override
+	virtual void UpdatePreviewInstance(const Metasound::FVertexKey& InParameterName, TScriptInterface<IAudioParameterInterface>& InParameterInterface) const override
 	{
-		InCommInterface->SetString(*InParameterName, Default);
+		InParameterInterface->SetString(*InParameterName, Default);
 	}
 };
 
@@ -411,9 +411,9 @@ public:
 		Metasound::Editor::InputPrivate::ConvertLiteralToArray<FString>(InLiteral, Default);
 	}
 
-	virtual void UpdatePreviewInstance(const Metasound::FVertexKey& InParameterName, TScriptInterface<IAudioCommunicationInterface>& InCommInterface) const override
+	virtual void UpdatePreviewInstance(const Metasound::FVertexKey& InParameterName, TScriptInterface<IAudioParameterInterface>& InParameterInterface) const override
 	{
-		InCommInterface->SetStringArray(*InParameterName, Default);
+		InParameterInterface->SetStringArray(*InParameterName, Default);
 	}
 };
 
@@ -455,7 +455,7 @@ public:
 // 		Metasound::Editor::InputPrivate::ConvertLiteral<UObject*>(InLiteral, Default.Object);
 	}
 
-	virtual void UpdatePreviewInstance(const Metasound::FVertexKey& InParameterName, TScriptInterface<IAudioCommunicationInterface>& InCommInterface) const override
+	virtual void UpdatePreviewInstance(const Metasound::FVertexKey& InParameterName, TScriptInterface<IAudioParameterInterface>& InParameterInterface) const override
 	{
 		// TODO. We need proxy object here safely.
 	}
@@ -492,7 +492,7 @@ public:
 // 		Metasound::Editor::InputPrivate::ConvertLiteralToArray<UObject*, FMetasoundEditorGraphInputObjectRef>(InLiteral, Default);
 	}
 
-	virtual void UpdatePreviewInstance(const Metasound::FVertexKey& InParameterName, TScriptInterface<IAudioCommunicationInterface>& InCommInterface) const override
+	virtual void UpdatePreviewInstance(const Metasound::FVertexKey& InParameterName, TScriptInterface<IAudioParameterInterface>& InParameterInterface) const override
 	{
 		TArray<UObject*> ObjectArray;
 		Algo::Transform(Default, ObjectArray, [](const FMetasoundEditorGraphInputObjectRef& InValue) { return InValue.Object; });
diff --git a/Engine/Source/Runtime/Engine/Classes/Audio/AudioComponentParameterization.h b/Engine/Source/Runtime/Engine/Classes/Audio/AudioComponentParameterization.h
new file mode 100644
--- /dev/null
+++ b/Engine/Source/Runtime/Engine/Classes/Audio/AudioComponentParameterization.h
@@ -0,0 +1,35 @@
+// Copyright Epic Games, Inc. All Rights Reserved.
+#pragma once
+
+#include "Audio/AudioParameterInterface.h"
+#include "UObject/Object.h"
+
+#include "AudioComponentParameterization.generated.h"
+
+
+UCLASS()
+class ENGINE_API UAudioComponentParameterization : public UObject , public IAudioParameterInterface
+{
+	GENERATED_BODY()
+
+protected:
+	virtual void BeginDestroy() override;
+
+	// IAudioParameterInterface
+	void Shutdown() override;
+
+	void Trigger(FName InName) override;
+	void SetBool(FName InName, bool InValue) override;
+	void SetBoolArray(FName InName, const TArray<bool>& InValue) override;
+	void SetInt(FName InName, int32 InValue) override;
+	void SetIntArray(FName InName, const TArray<int32>& InValue) override;
+	void SetFloat(FName InName, float InValue) override;
+	void SetFloatArray(FName InName, const TArray<float>& InValue) override;
+	void SetString(FName InName, const FString& InValue) override;
+	void SetStringArray(FName InName, const TArray<FString>& InValue) override;
+	void SetObject(FName InName, UObject* InValue) override;
+	void SetObjectArray(FName InName, const TArray<UObject*>& InValue) override;
+
+	template<typename T>
+	void SetValue(FName InName, T&& InX);
+};
diff --git a/Engine/Source/Runtime/Engine/Classes/Audio/AudioParameterInterface.h b/Engine/Source/Runtime/Engine/Classes/Audio/AudioParameterInterface.h
new file mode 100644
--- /dev/null
+++ b/Engine/Source/Runtime/Engine/Classes/Audio/AudioParameterInterface.h
@@ -0,0 +1,66 @@
+// Copyright Epic Games, Inc. All Rights Reserved.
+
+#pragma once
+
+#include "CoreMinimal.h"
+#include "UObject/Interface.h"
+
+#include "AudioParameterInterface.generated.h"
+
+UINTERFACE(MinimalApi, BlueprintType, meta = (CannotImplementInterfaceInBlueprint))
+class UAudioParameterInterface : public UInterface
+{
+	GENERATED_UINTERFACE_BODY()
+};
+
+class IAudioParameterInterface : public IInterface
+{
+	GENERATED_IINTERFACE_BODY()
+
+public:
+	virtual void Shutdown() = 0;
+
+	// Triggers a named trigger 
+	UFUNCTION(BlueprintCallable, Category = "Audio|Parameter|Triggers")
+	virtual void Trigger(FName Name) = 0;
+
+	// Sets a named Boolean 
+	UFUNCTION(BlueprintCallable, Category = "Audio|Parameter|Primitive")
+	virtual void SetBool(FName Name, bool Value) = 0;
+	
+	// Sets a named Boolean Array 
+	UFUNCTION(BlueprintCallable, Category = "Audio|Parameter|Array")
+	virtual void SetBoolArray(FName Name,const TArray<bool>& Value) = 0;
+
+	// Sets a named Int32 
+	UFUNCTION(BlueprintCallable, Category = "Audio|Parameter|Primitive")
+	virtual void SetInt(FName Name, int32 Value) = 0;
+
+	// Sets a named Int32 Array
+	UFUNCTION(BlueprintCallable, Category = "Audio|Parameter|Array")
+	virtual void SetIntArray(FName Name, const TArray<int32>& Value) = 0;
+
+	// Sets a named Float
+	UFUNCTION(BlueprintCallable, Category = "Audio|Parameter|Primitive")
+	virtual void SetFloat(FName Name, float Value) = 0;
+	
+	// Sets a named Float Array
+	UFUNCTION(BlueprintCallable, Category = "Audio|Parameter|Array")
+	virtual void SetFloatArray(FName Name, const TArray<float>& Value) = 0;
+	
+	// Sets a named String
+	UFUNCTION(BlueprintCallable, Category = "Audio|Parameter|Primitive")
+	virtual void SetString(FName Name, const FString& Value) = 0;
+
+	// Sets a named String Array
+	UFUNCTION(BlueprintCallable, Category = "Audio|Parameter|Array")
+	virtual void SetStringArray(FName Name, const TArray<FString>& Value) = 0;
+	
+	// Sets a named UObject 
+	UFUNCTION(BlueprintCallable, Category = "Audio|Parameter|Primitive")
+	virtual void SetObject(FName Name, UObject* Value) = 0;
+
+	// Sets a named UObject Array
+	UFUNCTION(BlueprintCallable, Category = "Audio|Parameter|Array")
+	virtual void SetObjectArray(FName Name, const TArray<UObject*>& Value) = 0;
+};
\ No newline at end of file
diff --git a/Engine/Source/Runtime/Engine/Classes/Components/AudioComponent.h b/Engine/Source/Runtime/Engine/Classes/Components/AudioComponent.h
--- a/Engine/Source/Runtime/Engine/Classes/Components/AudioComponent.h
+++ b/Engine/Source/Runtime/Engine/Classes/Components/AudioComponent.h
@@ -1,6 +1,7 @@
 // Copyright Epic Games, Inc. All Rights Reserved.
 #pragma once
 
+#include "Audio/AudioParameterInterface.h"
 #include "Components/SceneComponent.h"
 #include "CoreMinimal.h"
 #include "Engine/EngineTypes.h"
@@ -13,7 +14,6 @@
 #include "Sound/QuartzQuantizationUtilities.h"
 #include "Quartz/AudioMixerClockHandle.h"
 #include "Quartz/AudioMixerQuantizedCommands.h"
-#include "Sound/AudioCommunicationInterface.h"
 
 #include "AudioComponent.generated.h"
 
@@ -530,10 +530,11 @@ public:
 	void SetFloatParameter(FName InName, float InFloat);
 
 	UPROPERTY(Transient)
-	TScriptInterface<IAudioCommunicationInterface> CommunicationInterface;
+	TScriptInterface<IAudioParameterInterface> ParameterInterface;
 
-	UFUNCTION(BlueprintCallable, Category = "Communications")
-	TScriptInterface<IAudioCommunicationInterface> GetCommunicationInterface() const;
+	/** Retrieves the parameter interface to set parameters on sound types that support parameterization (ex. MetaSounds if plugin is enabled). */
+	UFUNCTION(BlueprintCallable, Category = "Audio|Parameters")
+	TScriptInterface<IAudioParameterInterface> GetParameterInterface() const;
 
 	/** Allows the designer to set the Wave Parameter on the SoundCue whose name matches the name indicated.
 	 * @param InName The name of the Wave to set. It must match the name set in SoundCue's WaveParam Node
diff --git a/Engine/Source/Runtime/Engine/Classes/Sound/AudioCommunicationInterface.h b/Engine/Source/Runtime/Engine/Classes/Sound/AudioCommunicationInterface.h
deleted file mode 100644
diff --git a/Engine/Source/Runtime/Engine/Classes/Sound/AudioComponentCommuncation.h b/Engine/Source/Runtime/Engine/Classes/Sound/AudioComponentCommuncation.h
deleted file mode 100644
diff --git a/Engine/Source/Runtime/Engine/Private/Audio/AudioCommunicationInterface.cpp b/Engine/Source/Runtime/Engine/Private/Audio/AudioCommunicationInterface.cpp
deleted file mode 100644
diff --git a/Engine/Source/Runtime/Engine/Private/Audio/AudioComponentParameterization.cpp b/Engine/Source/Runtime/Engine/Private/Audio/AudioComponentParameterization.cpp
new file mode 100644
--- /dev/null
+++ b/Engine/Source/Runtime/Engine/Private/Audio/AudioComponentParameterization.cpp
@@ -0,0 +1,156 @@
+// Copyright Epic Games, Inc. All Rights Reserved.
+#include "Audio/AudioComponentParameterization.h"
+
+#include "ActiveSound.h"
+#include "Audio.h"
+#include "AudioDevice.h"
+#include "AudioThread.h"
+#include "Components/AudioComponent.h"
+#include "IAudioExtensionPlugin.h"
+
+
+void UAudioComponentParameterization::Shutdown()
+{
+	if (UAudioComponent* OwningComponent = Cast<UAudioComponent>(GetOuter()))
+	{
+		if (FAudioDevice* AudioDevice = OwningComponent->GetAudioDevice())
+		{
+			const uint64 MyAudioComponentID = OwningComponent->GetAudioComponentID();
+			FAudioThread::RunCommandOnAudioThread([AudioDevice, MyAudioComponentID]() mutable
+			{
+				if (FActiveSound* ActiveSound = AudioDevice->FindActiveSound(MyAudioComponentID))
+				{
+					if (IAudioInstanceTransmitter* Transmitter = ActiveSound->GetTransmitter())
+					{
+						Transmitter->Shutdown();
+					}
+				}
+			});
+		}
+	}
+}
+
+void UAudioComponentParameterization::BeginDestroy()
+{
+	Shutdown();
+	Super::BeginDestroy();
+}
+
+void UAudioComponentParameterization::Trigger(FName InName)
+{
+	// Trigger is just a (true) bool param currently.
+	SetValue(InName, true);
+}
+
+void UAudioComponentParameterization::SetBool(FName InName, bool InValue)
+{
+	SetValue(InName, InValue);
+}
+
+void UAudioComponentParameterization::SetBoolArray(FName InName, const TArray<bool>& InValue)
+{
+	SetValue(InName, InValue);
+}
+
+void UAudioComponentParameterization::SetInt(FName InName, int32 InValue)
+{
+	SetValue(InName, InValue);
+}
+
+void UAudioComponentParameterization::SetIntArray(FName InName, const TArray<int32>& InValue)
+{
+	SetValue(InName, InValue);
+}
+
+void UAudioComponentParameterization::SetFloat(FName InName, float InValue)
+{
+	SetValue(InName, InValue);
+}
+
+void UAudioComponentParameterization::SetFloatArray(FName InName, const TArray<float>& InValue)
+{
+	SetValue(InName, InValue);
+}
+
+void UAudioComponentParameterization::SetString(FName InName, const FString& InValue)
+{
+	SetValue(InName, InValue);
+}
+
+void UAudioComponentParameterization::SetStringArray(FName InName, const TArray<FString>& InValue)
+{
+	SetValue(InName, InValue);
+}
+
+void UAudioComponentParameterization::SetObject(FName InName, UObject* InValue)
+{
+	IAudioProxyDataFactory* ObjectAsFactory = Cast<USoundWave>(InValue);	// FIXME. Possible to query for IAudioProxy support? Metasound used SFINAE
+
+	// The creation of the Proxy object needs to happen where its safe to access UObject, 
+	// namely the Game thread. 
+
+	if (ensureAlways(ObjectAsFactory))
+	{
+		static FName ProxySubsystemName = TEXT("Metasound");
+		Audio::FProxyDataInitParams ProxyInitParams;
+		ProxyInitParams.NameOfFeatureRequestingProxy = ProxySubsystemName;
+
+		TUniquePtr<Audio::IProxyData> Proxy{ ObjectAsFactory->CreateNewProxyData(ProxyInitParams) };
+		SetValue(InName, MoveTemp(Proxy));
+	}
+}
+
+void UAudioComponentParameterization::SetObjectArray(FName InName, const TArray<UObject*>& InValue)
+{
+	static FName ProxySubsystemName = TEXT("Metasound");
+	Audio::FProxyDataInitParams ProxyInitParams;
+	ProxyInitParams.NameOfFeatureRequestingProxy = ProxySubsystemName;
+
+	// The creation of the Proxy object needs to happen where its safe to access UObject, 
+	// namely the Game thread. 
+
+	TArray<Audio::IProxyDataPtr> ProxiedInputs;
+	for (UObject* i : InValue)
+	{
+		if (IAudioProxyDataFactory* ProxyFactory = Cast<USoundWave>(i)) // FIXME. Query for IAudioProxyDataFactory
+		{
+			ProxiedInputs.Emplace(ProxyFactory->CreateNewProxyData(ProxyInitParams));
+		}
+	}
+	SetValue(InName, MoveTemp(ProxiedInputs));
+}
+
+template<typename T>
+void UAudioComponentParameterization::SetValue(FName InName, T&& InX)
+{
+	if (InName == NAME_None)
+	{
+		return;
+	}
+
+	if (UAudioComponent* OwningComponent = Cast<UAudioComponent>(GetOuter()))
+	{
+		if (OwningComponent->IsActive())
+		{
+			if (FAudioDevice* AudioDevice = OwningComponent->GetAudioDevice())
+			{
+				DECLARE_CYCLE_STAT(TEXT("FAudioThreadTask.SetValue"), STAT_AudioSetSoundParameter, STATGROUP_AudioThreadCommands);
+
+				const uint64 MyAudioComponentID = OwningComponent->GetAudioComponentID();
+				FAudioThread::RunCommandOnAudioThread([AudioDevice, MyAudioComponentID, InName, InX{ MoveTempIfPossible(InX) }]() mutable
+				{
+					if (FActiveSound* ActiveSound = AudioDevice->FindActiveSound(MyAudioComponentID))
+					{
+						if (IAudioInstanceTransmitter* Transmitter = ActiveSound->GetTransmitter())
+						{
+							if (!Transmitter->SetParameter(InName, MoveTempIfPossible(InX)))
+							{
+								UE_LOG(LogAudio, Warning, TEXT("Failed to SetParameter '%s'"), *InName.ToString() );
+							}
+						}
+					}
+				}, GET_STATID(STAT_AudioSetSoundParameter));
+			}
+		}
+	}
+}
diff --git a/Engine/Source/Runtime/Engine/Private/Audio/AudioParameterInterface.cpp b/Engine/Source/Runtime/Engine/Private/Audio/AudioParameterInterface.cpp
new file mode 100644
--- /dev/null
+++ b/Engine/Source/Runtime/Engine/Private/Audio/AudioParameterInterface.cpp
@@ -0,0 +1,9 @@
+// Copyright Epic Games, Inc. All Rights Reserved.
+
+#include "Audio/AudioParameterInterface.h"
+
+
+UAudioParameterInterface::UAudioParameterInterface(FObjectInitializer const& InOI)
+	: UInterface(InOI)
+{
+}
\ No newline at end of file
diff --git a/Engine/Source/Runtime/Engine/Private/AudioComponentCommunication.cpp b/Engine/Source/Runtime/Engine/Private/AudioComponentCommunication.cpp
deleted file mode 100644
diff --git a/Engine/Source/Runtime/Engine/Private/Components/AudioComponent.cpp b/Engine/Source/Runtime/Engine/Private/Components/AudioComponent.cpp
--- a/Engine/Source/Runtime/Engine/Private/Components/AudioComponent.cpp
+++ b/Engine/Source/Runtime/Engine/Private/Components/AudioComponent.cpp
@@ -1,19 +1,19 @@
 // Copyright Epic Games, Inc. All Rights Reserved.
 
 #include "Components/AudioComponent.h"
-#include "Audio.h"
-#include "Engine/Texture2D.h"
 #include "ActiveSound.h"
-#include "AudioThread.h"
+#include "Audio.h"
+#include "Audio/AudioComponentParameterization.h"
 #include "AudioDevice.h"
+#include "AudioThread.h"
+#include "Components/BillboardComponent.h"
 #include "DSP/VolumeFader.h"
-#include "Sound/SoundNodeAttenuation.h"
+#include "Engine/Texture2D.h"
+#include "Kismet/GameplayStatics.h"
+#include "Misc/App.h"
 #include "Sound/SoundCue.h"
-#include "Components/BillboardComponent.h"
+#include "Sound/SoundNodeAttenuation.h"
 #include "UObject/FrameworkObjectVersion.h"
-#include "Misc/App.h"
-#include "Kismet/GameplayStatics.h"
-#include "Sound/AudioComponentCommuncation.h"
 
 DECLARE_CYCLE_STAT(TEXT("AudioComponent Play"), STAT_AudioComp_Play, STATGROUP_Audio);
 
@@ -99,9 +99,9 @@ UAudioComponent* UAudioComponent::GetAudioComponentFromID(uint64 AudioComponentI
 
 void UAudioComponent::BeginDestroy()
 {
-	if (CommunicationInterface)
+	if (ParameterInterface)
 	{
-		CommunicationInterface->Shutdown();
+		ParameterInterface->Shutdown();
 	}
 
 	if (IsActive() && Sound && Sound->IsLooping())
@@ -504,11 +504,10 @@ void UAudioComponent::PlayInternal(const PlayInternalRequestData& InPlayRequestD
 
 			AudioDevice->GetMaxDistanceAndFocusFactor(Sound, World, Location, AttenuationSettingsToApply, MaxDistance, FocusFactor);
 
-			// Create communications object.
-			if (!CommunicationInterface)
+			if (!ParameterInterface)
 			{
-				static const FName InterfaceName("AudioComponentCommunication");
-				CommunicationInterface = NewObject<UAudioComponentCommunication>(this, InterfaceName);
+				static const FName InterfaceName("AudioComponentParameterization");
+				ParameterInterface = NewObject<UAudioComponentParameterization>(this, InterfaceName);
 			}
 
 			FActiveSound NewActiveSound;
@@ -764,10 +763,10 @@ void UAudioComponent::Stop()
 
 	if (bIsPreviewSound)
 	{
-		if (CommunicationInterface)
+		if (ParameterInterface)
 		{
-			CommunicationInterface->Shutdown();
-			CommunicationInterface = nullptr;
+			ParameterInterface->Shutdown();
+			ParameterInterface = nullptr;
 		}
 	}
 
@@ -1128,9 +1127,9 @@ void UAudioComponent::SetFloatParameter( const FName InName, const float InFloat
 	}
 }
 
-TScriptInterface<IAudioCommunicationInterface> UAudioComponent::GetCommunicationInterface() const
+TScriptInterface<IAudioParameterInterface> UAudioComponent::GetParameterInterface() const
 {
-	return CommunicationInterface;
+	return ParameterInterface;
 }
 
 void UAudioComponent::SetWaveParameter( const FName InName, USoundWave* InWave )
