MetaSounds: have Random Get (Array) handle array being empty on init more gracefully, and making Midi to scale node initailize its array in its constructor

diff --git a/Engine/Plugins/Runtime/Metasound/Source/MetasoundFrontend/Public/MetasoundArrayRandomNode.h b/Engine/Plugins/Runtime/Metasound/Source/MetasoundFrontend/Public/MetasoundArrayRandomNode.h
--- a/Engine/Plugins/Runtime/Metasound/Source/MetasoundFrontend/Public/MetasoundArrayRandomNode.h
+++ b/Engine/Plugins/Runtime/Metasound/Source/MetasoundFrontend/Public/MetasoundArrayRandomNode.h
@@ -199,9 +199,9 @@ namespace Metasound
 			WeightsArray = *InputWeightsArray;
 
 			const ArrayType& InputArrayRef = *InputArray;
-			int32 ArraySize = InputArrayRef.Num();
+			PrevArraySize = InputArrayRef.Num();
 
-			if (ArraySize > 0)
+			if (PrevArraySize > 0)
 			{
 				if (*bEnableSharedState)
 				{
@@ -216,7 +216,7 @@ namespace Metasound
 					InitSharedStateArgs Args;
 					Args.SharedStateId = SharedStateUniqueId;
 					Args.Seed = PrevSeedValue;
-					Args.NumElements = ArraySize;
+					Args.NumElements = PrevArraySize;
 					Args.NoRepeatOrder = PrevNoRepeatOrder;
 					Args.bIsPreviewSound = bIsPreviewSound;
 					Args.Weights = WeightsArray;
@@ -225,7 +225,7 @@ namespace Metasound
 				}
 				else
 				{
-					ArrayRandomGet = MakeUnique<FArrayRandomGet>(PrevSeedValue, ArraySize, WeightsArray, PrevNoRepeatOrder);
+					ArrayRandomGet = MakeUnique<FArrayRandomGet>(PrevSeedValue, PrevArraySize, WeightsArray, PrevNoRepeatOrder);
 				}
 			}
 			else
@@ -271,6 +271,30 @@ namespace Metasound
 
 			const ArrayType& InputArrayRef = *InputArray;
 
+			// If the array size has changed, we need to reinit before getting the next value
+			if (PrevArraySize != InputArrayRef.Num())
+			{
+				PrevArraySize = InputArrayRef.Num();
+				if (PrevArraySize != 0)
+				{
+					if (!ArrayRandomGet.IsValid())
+					{
+						ArrayRandomGet = MakeUnique<FArrayRandomGet>(PrevSeedValue, PrevArraySize, WeightsArray, PrevNoRepeatOrder);
+					}
+					ArrayRandomGet->Init(PrevSeedValue, PrevArraySize, WeightsArray, PrevNoRepeatOrder);
+				}
+			}
+
+			if (PrevArraySize == 0)
+			{
+				if (!bHasLoggedEmptyArrayWarning)
+				{
+					UE_LOG(LogMetaSound, Warning, TEXT("Array Random Get has an empty array input."));
+					bHasLoggedEmptyArrayWarning = true;
+				}
+				return;
+			}
+
 			// Check for a seed change
 			if (PrevSeedValue != *SeedValue)
 			{
@@ -393,7 +417,9 @@ namespace Metasound
 		int32 PrevSeedValue = INDEX_NONE;
 		int32 PrevNoRepeatOrder = INDEX_NONE;
 		uint32 SharedStateUniqueId = INDEX_NONE;
+		int32 PrevArraySize = 0;
 		bool bIsPreviewSound = false;
+		bool bHasLoggedEmptyArrayWarning = false;
 	};
 
 	template<typename ArrayType>
diff --git a/Engine/Plugins/Runtime/Metasound/Source/MetasoundStandardNodes/Private/MetasoundMidiScaleToArrayNode.cpp b/Engine/Plugins/Runtime/Metasound/Source/MetasoundStandardNodes/Private/MetasoundMidiScaleToArrayNode.cpp
--- a/Engine/Plugins/Runtime/Metasound/Source/MetasoundStandardNodes/Private/MetasoundMidiScaleToArrayNode.cpp
+++ b/Engine/Plugins/Runtime/Metasound/Source/MetasoundStandardNodes/Private/MetasoundMidiScaleToArrayNode.cpp
@@ -94,6 +94,8 @@ namespace Metasound
 		, ScaleArrayOutput(FArrayScaleDegreeWriteRef::CreateNew())
 		, Settings(InSettings)
 	{
+		// prime our output array
+		Execute();
 	}
 
 
