Reseting send/receive nodes when nothing is hooked up.

diff --git a/Engine/Plugins/Runtime/Metasound/Source/MetasoundFrontend/Private/MetasoundInstanceTransmitter.cpp b/Engine/Plugins/Runtime/Metasound/Source/MetasoundFrontend/Private/MetasoundInstanceTransmitter.cpp
--- a/Engine/Plugins/Runtime/Metasound/Source/MetasoundFrontend/Private/MetasoundInstanceTransmitter.cpp
+++ b/Engine/Plugins/Runtime/Metasound/Source/MetasoundFrontend/Private/MetasoundInstanceTransmitter.cpp
@@ -17,7 +17,7 @@ namespace Metasound
 
 		for (const FSendInfo& SendInfo : SendInfos)
 		{
-			bSuccess &= FDataTransmissionCenter::Get().UnregisterDataChannel(SendInfo.TypeName, SendInfo.Address.ChannelName);
+			bSuccess &= FDataTransmissionCenter::Get().UnregisterDataChannel(SendInfo.TypeName, SendInfo.Address);
 		}
 
 		return bSuccess;
diff --git a/Engine/Plugins/Runtime/Metasound/Source/MetasoundFrontend/Private/MetasoundRouter.cpp b/Engine/Plugins/Runtime/Metasound/Source/MetasoundFrontend/Private/MetasoundRouter.cpp
--- a/Engine/Plugins/Runtime/Metasound/Source/MetasoundFrontend/Private/MetasoundRouter.cpp
+++ b/Engine/Plugins/Runtime/Metasound/Source/MetasoundFrontend/Private/MetasoundRouter.cpp
@@ -200,9 +200,24 @@ namespace Metasound
 		return MoveTemp(Receiver);
 	}
 
-	bool FDataTransmissionCenter::UnregisterDataChannel(const FName& InDataTypeName, const FName& InChannelName)
+	bool FDataTransmissionCenter::UnregisterDataChannel(const FName& InDataTypeName, const FSendAddress& InAddress)
 	{
-		return GlobalRouter.UnregisterDataChannel(InDataTypeName, InChannelName);
+		// Only supporting global router. Other subsystems are not used and will be reworked. 
+		if (ensure(InAddress.Subsystem == GetSubsystemNameForSendScope(ETransmissionScope::Global)))
+		{
+			return GlobalRouter.UnregisterDataChannel(InDataTypeName, InAddress.ChannelName);
+		}
+		return false;
+	}
+
+	bool FDataTransmissionCenter::UnregisterDataChannelIfUnconnected(const FName& InDataTypeName, const FSendAddress& InAddress)
+	{
+		// Only supporting global router. Other subsystems are not used and will be reworked. 
+		if (ensure(InAddress.Subsystem == GetSubsystemNameForSendScope(ETransmissionScope::Global)))
+		{
+			return GlobalRouter.UnregisterDataChannelIfUnconnected(InDataTypeName, InAddress.ChannelName);
+		}
+		return false;
 	}
 
 	bool FDataTransmissionCenter::PushLiteral(FName DataType, FName GlobalChannelName, const FLiteral& InParam)
@@ -315,6 +330,25 @@ namespace Metasound
 		return DataChannelMap.Remove(ChannelKey) > 0;
 	}
 
+	bool FAddressRouter::UnregisterDataChannelIfUnconnected(const FName& InDataTypeName, const FName& InChannelName)
+	{
+		FScopeLock ScopeLock(&DataChannelMapMutationLock);
+		const FDataChannelKey ChannelKey = GetDataChannelKey(InDataTypeName, InChannelName);
+
+		if (TSharedRef<IDataChannel, ESPMode::ThreadSafe>* Channel = DataChannelMap.Find(ChannelKey))
+		{
+			if (0 == Channel->Get().GetNumActiveReceivers())
+			{
+				if (0 == Channel->Get().GetNumActiveSenders())
+				{
+					return DataChannelMap.Remove(ChannelKey) > 0;
+				}
+			}
+		}
+
+		return false;
+	}
+
 	TUniquePtr<IReceiver> FAddressRouter::RegisterNewReceiver(const FName& InDataTypeName, const FName& InChannelName, const FReceiverInitParams& InitParams)
 	{
 		TSharedPtr<IDataChannel, ESPMode::ThreadSafe> DataChannel = GetDataChannel(InDataTypeName, InChannelName, InitParams.OperatorSettings);
diff --git a/Engine/Plugins/Runtime/Metasound/Source/MetasoundFrontend/Public/MetasoundReceiveNode.h b/Engine/Plugins/Runtime/Metasound/Source/MetasoundFrontend/Public/MetasoundReceiveNode.h
--- a/Engine/Plugins/Runtime/Metasound/Source/MetasoundFrontend/Public/MetasoundReceiveNode.h
+++ b/Engine/Plugins/Runtime/Metasound/Source/MetasoundFrontend/Public/MetasoundReceiveNode.h
@@ -102,7 +102,10 @@ namespace Metasound
 				{
 				}
 
-				virtual ~TReceiverOperator() {}
+				virtual ~TReceiverOperator() 
+				{
+					ResetReceiverAndCleanupChannel();
+				}
 
 				virtual FDataReferenceCollection GetInputs() const override
 				{
@@ -125,15 +128,11 @@ namespace Metasound
 				{
 					if (SendAddress->ChannelName != CachedSendAddress.ChannelName)
 					{
+						ResetReceiverAndCleanupChannel();
 						CachedSendAddress = *SendAddress;
 						Receiver = FDataTransmissionCenter::Get().RegisterNewReceiver<TDataType>(CachedSendAddress, CachedReceiverParams);
 					}
 
-					if (bHasNotReceivedData)
-					{
-						*OutputData = *DefaultData;
-					}
-
 					bool bHasNewData = false;
 					if (ensure(Receiver.IsValid()))
 					{
@@ -145,6 +144,12 @@ namespace Metasound
 						}
 					}
 
+					if (bHasNotReceivedData)
+					{
+						*OutputData = *DefaultData;
+						bHasNewData = true;
+					}
+
 					if (TExecutableDataType<TDataType>::bIsExecutable)
 					{
 						TExecutableDataType<TDataType>::ExecuteInline(*OutputData, bHasNewData);
@@ -152,6 +157,13 @@ namespace Metasound
 				}
 
 			private:
+
+				void ResetReceiverAndCleanupChannel()
+				{
+					Receiver.Reset();
+					FDataTransmissionCenter::Get().UnregisterDataChannelIfUnconnected(GetMetasoundDataTypeName<TDataType>(), CachedSendAddress);
+				}
+
 				bool bHasNotReceivedData;
 
 				TDataReadReference<TDataType> DefaultData;
diff --git a/Engine/Plugins/Runtime/Metasound/Source/MetasoundFrontend/Public/MetasoundRouter.h b/Engine/Plugins/Runtime/Metasound/Source/MetasoundFrontend/Public/MetasoundRouter.h
--- a/Engine/Plugins/Runtime/Metasound/Source/MetasoundFrontend/Public/MetasoundRouter.h
+++ b/Engine/Plugins/Runtime/Metasound/Source/MetasoundFrontend/Public/MetasoundRouter.h
@@ -982,6 +982,7 @@ namespace Metasound
 		}
 
 		bool UnregisterDataChannel(const FName& InDataTypeName, const FName& InChannelName);
+		bool UnregisterDataChannelIfUnconnected(const FName& InDataTypeName, const FName& InChannelName);
 
 		FAddressRouter(const FAddressRouter& Other)
 			: DataChannelMap(Other.DataChannelMap)
@@ -1151,7 +1152,10 @@ namespace Metasound
 		}
 
 		// Unregisters DataChannel irrespective of number of receivers or senders still active.
-		bool UnregisterDataChannel(const FName& InDataTypeName, const FName& InChannelName);
+		bool UnregisterDataChannel(const FName& InDataTypeName, const FSendAddress& InAddress);
+
+		// Unregister a data channel if there are no senders or receivers
+		bool UnregisterDataChannelIfUnconnected(const FName& InDataTypeName, const FSendAddress& InAddress);
 
 		// Pushes a literal parameter to a specific data channel in the global router.
 		// returns false if the literal type isn't supported.
diff --git a/Engine/Plugins/Runtime/Metasound/Source/MetasoundFrontend/Public/MetasoundSendNode.h b/Engine/Plugins/Runtime/Metasound/Source/MetasoundFrontend/Public/MetasoundSendNode.h
--- a/Engine/Plugins/Runtime/Metasound/Source/MetasoundFrontend/Public/MetasoundSendNode.h
+++ b/Engine/Plugins/Runtime/Metasound/Source/MetasoundFrontend/Public/MetasoundSendNode.h
@@ -86,7 +86,10 @@ namespace Metasound
 				{
 				}
 
-				virtual ~TSendOperator() {}
+				virtual ~TSendOperator() 
+				{
+					ResetSenderAndCleanupChannel();
+				}
 
 				virtual FDataReferenceCollection GetInputs() const override
 				{
@@ -105,6 +108,7 @@ namespace Metasound
 				{
 					if (SendAddress->ChannelName != CachedSendAddress.ChannelName)
 					{
+						ResetSenderAndCleanupChannel();
 						CachedSendAddress = *SendAddress;
 						Sender = FDataTransmissionCenter::Get().RegisterNewSender<TDataType>(CachedSendAddress, CachedSenderParams);
 						check(Sender.IsValid());
@@ -114,6 +118,12 @@ namespace Metasound
 				}
 
 			private:
+				void ResetSenderAndCleanupChannel()
+				{
+					Sender.Reset();
+					FDataTransmissionCenter::Get().UnregisterDataChannelIfUnconnected(GetMetasoundDataTypeName<TDataType>(), CachedSendAddress);
+				}
+
 				TDataReadReference<TDataType> InputData;
 				TDataReadReference<FSendAddress> SendAddress;
 				FSendAddress CachedSendAddress;
