- Validate pin string literal representation against transformed string through the literal system during graph sync
- Fix incorrect default type for TriggerCounterNode

diff --git a/Engine/Plugins/Runtime/Metasound/Source/MetasoundEditor/Private/MetasoundEditorGraphBuilder.cpp b/Engine/Plugins/Runtime/Metasound/Source/MetasoundEditor/Private/MetasoundEditorGraphBuilder.cpp
--- a/Engine/Plugins/Runtime/Metasound/Source/MetasoundEditor/Private/MetasoundEditorGraphBuilder.cpp
+++ b/Engine/Plugins/Runtime/Metasound/Source/MetasoundEditor/Private/MetasoundEditorGraphBuilder.cpp
@@ -1310,12 +1310,17 @@ namespace Metasound
 							const FGuid VertexID = GraphHandle->GetVertexIDForInputVertex(InputName);
 							const FMetasoundFrontendLiteral DefaultLiteral = GraphHandle->GetDefaultInput(VertexID);
 							const FString DefaultLiteralString = DefaultLiteral.ToString();
-							if (MatchingPin->DefaultValue != DefaultLiteralString)
+
+							FMetasoundFrontendLiteral PinDataTypeDefaultLiteral;
+							if (GetPinDefaultLiteral(*MatchingPin, PinDataTypeDefaultLiteral))
 							{
-								const FText& OwningNodeName = NodeInput->GetOwningNode()->GetDisplayName();
-								UE_LOG(LogMetasoundEditor, Display, TEXT("Synchronizing Node '%s' Connection: Setting pin '%s' default value to '%s'"), *OwningNodeName.ToString(), *MatchingPin->GetName(), *DefaultLiteralString);
-								MatchingPin->DefaultValue = DefaultLiteralString;
-								bIsNodeDirty = true;
+								if (PinDataTypeDefaultLiteral.ToString() != DefaultLiteralString)
+								{
+									const FText& OwningNodeName = NodeInput->GetOwningNode()->GetDisplayName();
+									UE_LOG(LogMetasoundEditor, Display, TEXT("Synchronizing Node '%s' Connection: Setting pin '%s' default value to '%s'"), *OwningNodeName.ToString(), *MatchingPin->GetName(), *DefaultLiteralString);
+									MatchingPin->DefaultValue = DefaultLiteralString;
+									bIsNodeDirty = true;
+								}
 							}
 						}
 					}
diff --git a/Engine/Plugins/Runtime/Metasound/Source/MetasoundStandardNodes/Private/MetasoundTriggerCounterNode.cpp b/Engine/Plugins/Runtime/Metasound/Source/MetasoundStandardNodes/Private/MetasoundTriggerCounterNode.cpp
--- a/Engine/Plugins/Runtime/Metasound/Source/MetasoundStandardNodes/Private/MetasoundTriggerCounterNode.cpp
+++ b/Engine/Plugins/Runtime/Metasound/Source/MetasoundStandardNodes/Private/MetasoundTriggerCounterNode.cpp
@@ -185,7 +185,7 @@ namespace Metasound
 				TInputDataVertexModel<FTrigger>(METASOUND_GET_PARAM_NAME_AND_TT(InputInTrigger)),
 				TInputDataVertexModel<FTrigger>(METASOUND_GET_PARAM_NAME_AND_TT(InputReset)),
 				TInputDataVertexModel<float>(METASOUND_GET_PARAM_NAME_AND_TT(InputStartValue), 0.0f),
-				TInputDataVertexModel<float>(METASOUND_GET_PARAM_NAME_AND_TT(InputStepSize), 1),
+				TInputDataVertexModel<float>(METASOUND_GET_PARAM_NAME_AND_TT(InputStepSize), 1.0f),
 				TInputDataVertexModel<int32>(METASOUND_GET_PARAM_NAME_AND_TT(InputAutoResetCount), 0)
 			),
 			FOutputVertexInterface(
