Fix for crash when clicking out audio output node

diff --git a/Engine/Plugins/Runtime/Metasound/Source/MetasoundEditor/Private/MetasoundNodeDetailCustomization.cpp b/Engine/Plugins/Runtime/Metasound/Source/MetasoundEditor/Private/MetasoundNodeDetailCustomization.cpp
--- a/Engine/Plugins/Runtime/Metasound/Source/MetasoundEditor/Private/MetasoundNodeDetailCustomization.cpp
+++ b/Engine/Plugins/Runtime/Metasound/Source/MetasoundEditor/Private/MetasoundNodeDetailCustomization.cpp
@@ -495,18 +495,17 @@ namespace Metasound
 					return;
 				}
 
-				// Hidden input types should be omitted from the drop down.
-				if (HiddenInputTypeNames.Contains(EditorDataType.RegistryInfo.DataTypeName))
-				{
-					return;
-				}
-
 				TSharedPtr<FString> TypeStrPtr = MakeShared<FString>(TypeName);
 				if (TypeName == CurrentTypeName)
 				{
 					CurrentTypeString = TypeStrPtr;
 				}
-				DataTypeNames.Add(TypeStrPtr);
+
+				// Hidden input types should be omitted from the drop down.
+				if (!HiddenInputTypeNames.Contains(EditorDataType.RegistryInfo.DataTypeName))
+				{
+					DataTypeNames.Add(TypeStrPtr);
+				}
 			});
 
 			if (!ensure(CurrentTypeString.IsValid()))
