- Disable Trigger Array type as input type
- Remove Editor Settings for double/int64 that were missed earlier
- Add MetaSound plugin icon & replace execution-style icon for triggers with custom icon

diff --git a/Engine/Plugins/Runtime/Metasound/Source/MetasoundEditor/Private/MetasoundEditorGraphBuilder.cpp b/Engine/Plugins/Runtime/Metasound/Source/MetasoundEditor/Private/MetasoundEditorGraphBuilder.cpp
--- a/Engine/Plugins/Runtime/Metasound/Source/MetasoundEditor/Private/MetasoundEditorGraphBuilder.cpp
+++ b/Engine/Plugins/Runtime/Metasound/Source/MetasoundEditor/Private/MetasoundEditorGraphBuilder.cpp
@@ -35,10 +35,10 @@ namespace Metasound
 	{
 		const FName FGraphBuilder::PinCategoryAudio = "audio";
 		const FName FGraphBuilder::PinCategoryBoolean = "bool";
-		const FName FGraphBuilder::PinCategoryDouble = "double";
+		//const FName FGraphBuilder::PinCategoryDouble = "double";
 		const FName FGraphBuilder::PinCategoryFloat = "float";
 		const FName FGraphBuilder::PinCategoryInt32 = "int";
-		const FName FGraphBuilder::PinCategoryInt64 = "int64";
+		//const FName FGraphBuilder::PinCategoryInt64 = "int64";
 		const FName FGraphBuilder::PinCategoryObject = "object";
 		const FName FGraphBuilder::PinCategoryString = "string";
 		const FName FGraphBuilder::PinCategoryTrigger = "trigger";
@@ -240,20 +240,20 @@ namespace Metasound
 				return Settings->IntPinTypeColor;
 			}
 
-			if (PinType.PinCategory == PinCategoryInt64)
-			{
-				return Settings->Int64PinTypeColor;
-			}
+			//if (PinType.PinCategory == PinCategoryInt64)
+			//{
+			//	return Settings->Int64PinTypeColor;
+			//}
 
 			if (PinType.PinCategory == PinCategoryString)
 			{
 				return Settings->StringPinTypeColor;
 			}
 
-			if (PinType.PinCategory == PinCategoryDouble)
-			{
-				return Settings->DoublePinTypeColor;
-			}
+			//if (PinType.PinCategory == PinCategoryDouble)
+			//{
+			//	return Settings->DoublePinTypeColor;
+			//}
 
 			if (PinType.PinCategory == PinCategoryObject)
 			{
diff --git a/Engine/Plugins/Runtime/Metasound/Source/MetasoundEditor/Private/MetasoundEditorGraphBuilder.h b/Engine/Plugins/Runtime/Metasound/Source/MetasoundEditor/Private/MetasoundEditorGraphBuilder.h
--- a/Engine/Plugins/Runtime/Metasound/Source/MetasoundEditor/Private/MetasoundEditorGraphBuilder.h
+++ b/Engine/Plugins/Runtime/Metasound/Source/MetasoundEditor/Private/MetasoundEditorGraphBuilder.h
@@ -38,10 +38,10 @@ namespace Metasound
 		public:
 			static const FName PinCategoryAudio;
 			static const FName PinCategoryBoolean;
-			static const FName PinCategoryDouble;
+			//static const FName PinCategoryDouble;
 			static const FName PinCategoryFloat;
 			static const FName PinCategoryInt32;
-			static const FName PinCategoryInt64;
+			//static const FName PinCategoryInt64;
 			static const FName PinCategoryObject;
 			static const FName PinCategoryString;
 			static const FName PinCategoryTrigger;
diff --git a/Engine/Plugins/Runtime/Metasound/Source/MetasoundEditor/Private/MetasoundEditorModule.cpp b/Engine/Plugins/Runtime/Metasound/Source/MetasoundEditor/Private/MetasoundEditorModule.cpp
--- a/Engine/Plugins/Runtime/Metasound/Source/MetasoundEditor/Private/MetasoundEditorModule.cpp
+++ b/Engine/Plugins/Runtime/Metasound/Source/MetasoundEditor/Private/MetasoundEditorModule.cpp
@@ -11,7 +11,7 @@
 #include "IDetailCustomization.h"
 #include "ISettingsModule.h"
 #include "Metasound.h"
-#include "MetasoundSource.h"
+#include "MetasoundAudioBuffer.h"
 #include "MetasoundAssetTypeActions.h"
 #include "MetasoundDetailCustomization.h"
 #include "MetasoundEditorGraph.h"
@@ -22,6 +22,9 @@
 #include "MetasoundEditorSettings.h"
 #include "MetasoundFrontendRegistries.h"
 #include "MetasoundNodeDetailCustomization.h"
+#include "MetasoundSource.h"
+#include "MetasoundTime.h"
+#include "MetasoundTrigger.h"
 #include "Modules/ModuleInterface.h"
 #include "Modules/ModuleManager.h"
 #include "PropertyEditorDelegates.h"
@@ -69,6 +72,8 @@ namespace Metasound
 				static const FVector2D Icon16 = FVector2D(16.0f, 16.0f);
 				static const FVector2D Icon64 = FVector2D(64.0f, 64.0f);
 
+				const FVector2D Icon15x11(15.0f, 11.0f);
+
 				// Metasound Editor
 				{
 					// Actions
@@ -88,6 +93,9 @@ namespace Metasound
 					Set("MetasoundEditor.Graph.Node.Body.Input", new FSlateImageBrush(RootToContentDir(TEXT("/Graph/node_input_body_64x.png")), FVector2D(114.0f, 64.0f)));
 					Set("MetasoundEditor.Graph.Node.Body.Default", new FSlateImageBrush(RootToContentDir(TEXT("/Graph/node_default_body_64x.png")), FVector2D(64.0f, 64.0f)));
 
+					Set("MetasoundEditor.Graph.TriggerPin.Connected", new IMAGE_BRUSH(TEXT("Graph/pin_trigger_connected"), Icon15x11));
+					Set("MetasoundEditor.Graph.TriggerPin.Disconnected", new IMAGE_BRUSH(TEXT("Graph/pin_trigger_disconnected"), Icon15x11));
+
 					Set("MetasoundEditor.Graph.Node.Math.Add", new FSlateImageBrush(RootToContentDir(TEXT("/Graph/node_math_add_40x.png")), Icon40x40));
 					Set("MetasoundEditor.Graph.Node.Math.Divide", new FSlateImageBrush(RootToContentDir(TEXT("/Graph/node_math_divide_40x.png")), Icon40x40));
 					Set("MetasoundEditor.Graph.Node.Math.Modulo", new FSlateImageBrush(RootToContentDir(TEXT("/Graph/node_math_modulo_40x.png")), Icon40x40));
@@ -196,7 +204,7 @@ namespace Metasound
 					FName PinSubCategory;
 
 					// Execution path triggers are specialized
-					if (DataTypeName == "Trigger")
+					if (DataTypeName == Frontend::GetDataTypeName<FTrigger>())
 					{
 						PinCategory = FGraphBuilder::PinCategoryTrigger;
 					}
@@ -204,10 +212,10 @@ namespace Metasound
 					// GraphEditor by default designates specialized connection
 					// specification for Int64, so use it even though literal is
 					// boiled down to int32
-					else if (DataTypeName == "Int64")
-					{
-						PinCategory = FGraphBuilder::PinCategoryInt64;
-					}
+					//else if (DataTypeName == Frontend::GetDataTypeName<int64>())
+					//{
+					//	PinCategory = FGraphBuilder::PinCategoryInt64;
+					//}
 
 					// Primitives
 					else
@@ -228,17 +236,13 @@ namespace Metasound
 
 								// Doubles use the same preferred literal
 								// but different colorization
-								if (DataTypeName == "Double")
-								{
-									PinCategory = FGraphBuilder::PinCategoryDouble;
-								}
+								//if (DataTypeName == Frontend::GetDataTypeName<double>())
+								//{
+								//	PinCategory = FGraphBuilder::PinCategoryDouble;
+								//}
 
 								// Differentiate stronger numeric types associated with audio
-								if (DataTypeName == "Frequency"
-									|| DataTypeName == "Time"
-									|| DataTypeName == "Time:HighResolution"
-									|| DataTypeName == "Time:SampleResolution"
-								)
+								if (DataTypeName == Frontend::GetDataTypeName<FTime>())
 								{
 									PinSubCategory = FGraphBuilder::PinSubCategoryTime;
 								}
@@ -272,7 +276,7 @@ namespace Metasound
 							{
 								// Audio types are ubiquitous, so added as subcategory
 								// to be able to stylize connections (i.e. wire color & wire animation)
-								if (DataTypeName == "Audio")
+								if (DataTypeName == Frontend::GetDataTypeName<FAudioBuffer>())
 								{
 									PinCategory = FGraphBuilder::PinCategoryAudio;
 								}
diff --git a/Engine/Plugins/Runtime/Metasound/Source/MetasoundEditor/Private/MetasoundEditorSettings.cpp b/Engine/Plugins/Runtime/Metasound/Source/MetasoundEditor/Private/MetasoundEditorSettings.cpp
--- a/Engine/Plugins/Runtime/Metasound/Source/MetasoundEditor/Private/MetasoundEditorSettings.cpp
+++ b/Engine/Plugins/Runtime/Metasound/Source/MetasoundEditor/Private/MetasoundEditorSettings.cpp
@@ -13,10 +13,10 @@ UMetasoundEditorSettings::UMetasoundEditorSettings(const FObjectInitializer& Obj
 
 	AudioPinTypeColor = FLinearColor(1.0f, 0.3f, 1.0f, 1.0f);					// magenta
 	BooleanPinTypeColor = FLinearColor(0.300000f, 0.0f, 0.0f, 1.0f);			// maroon
-	DoublePinTypeColor = FLinearColor(0.039216f, 0.666667f, 0.0f, 1.0f);		// darker green
+	//DoublePinTypeColor = FLinearColor(0.039216f, 0.666667f, 0.0f, 1.0f);		// darker green
 	FloatPinTypeColor = FLinearColor(0.357667f, 1.0f, 0.060000f, 1.0f);			// bright green
 	IntPinTypeColor = FLinearColor(0.013575f, 0.770000f, 0.429609f, 1.0f);		// green-blue
-	Int64PinTypeColor = FLinearColor(0.413575f, 0.770000f, 0.429609f, 1.0f);	// green
+	//Int64PinTypeColor = FLinearColor(0.413575f, 0.770000f, 0.429609f, 1.0f);	// green
 	ObjectPinTypeColor = FLinearColor(0.0f, 0.4f, 0.910000f, 1.0f);				// sharp blue
 	StringPinTypeColor = FLinearColor(1.0f, 0.0f, 0.660537f, 1.0f);				// bright pink
 	TimePinTypeColor = FLinearColor(0.3f, 1.0f, 1.0f, 1.0f);					// cyan
diff --git a/Engine/Plugins/Runtime/Metasound/Source/MetasoundEditor/Private/MetasoundEditorSettings.h b/Engine/Plugins/Runtime/Metasound/Source/MetasoundEditor/Private/MetasoundEditorSettings.h
--- a/Engine/Plugins/Runtime/Metasound/Source/MetasoundEditor/Private/MetasoundEditorSettings.h
+++ b/Engine/Plugins/Runtime/Metasound/Source/MetasoundEditor/Private/MetasoundEditorSettings.h
@@ -35,8 +35,8 @@ public:
 	FLinearColor BooleanPinTypeColor;
 
 	/** Double pin type color */
-	UPROPERTY(EditAnywhere, config, Category = PinColors)
-	FLinearColor DoublePinTypeColor;
+	//UPROPERTY(EditAnywhere, config, Category = PinColors)
+	//FLinearColor DoublePinTypeColor;
 
 	/** Floating-point pin type color */
 	UPROPERTY(EditAnywhere, config, Category=PinColors)
@@ -47,8 +47,8 @@ public:
 	FLinearColor IntPinTypeColor;
 
 	/** Integer64 pin type color */
-	UPROPERTY(EditAnywhere, config, Category = PinColors)
-	FLinearColor Int64PinTypeColor;
+	//UPROPERTY(EditAnywhere, config, Category = PinColors)
+	//FLinearColor Int64PinTypeColor;
 
 	/** Object pin type color */
 	UPROPERTY(EditAnywhere, config, Category=PinColors)
diff --git a/Engine/Plugins/Runtime/Metasound/Source/MetasoundEditor/Private/MetasoundNodeDetailCustomization.cpp b/Engine/Plugins/Runtime/Metasound/Source/MetasoundEditor/Private/MetasoundNodeDetailCustomization.cpp
--- a/Engine/Plugins/Runtime/Metasound/Source/MetasoundEditor/Private/MetasoundNodeDetailCustomization.cpp
+++ b/Engine/Plugins/Runtime/Metasound/Source/MetasoundEditor/Private/MetasoundNodeDetailCustomization.cpp
@@ -75,7 +75,8 @@ namespace Metasound
 			static const TSet<FName> HiddenInputTypeNames =
 			{
 				"Audio:Mono",
-				"Audio:Stereo"
+				"Audio:Stereo",
+				"Trigger:Array"
 			};
 		}
 
@@ -480,21 +481,22 @@ namespace Metasound
 
 			// Not all types have an equivalent array type. Base types without array
 			// types should have the "Is Array" checkbox disabled. 
-			const bool bIsArrayTypeRegisteredForCurrentType = EditorModule.IsRegisteredDataType(FName(CurrentTypeName + VariableCustomizationPrivate::ArrayIdentifier));
+			const FName ArrayType = *(CurrentTypeName + VariableCustomizationPrivate::ArrayIdentifier);
+			const bool bIsArrayTypeRegistered = EditorModule.IsRegisteredDataType(ArrayType);
+			const bool bIsArrayTypeRegisteredHidden = VariableCustomizationPrivate::HiddenInputTypeNames.Contains(ArrayType);
 
 			DataTypeNames.Reset();
 			EditorModule.IterateDataTypes([&](const FEditorDataType& EditorDataType)
 			{
-				using namespace VariableCustomizationPrivate;
+				const FString TypeName = EditorDataType.RegistryInfo.DataTypeName.ToString(); 
 
-				const FString TypeName = EditorDataType.RegistryInfo.DataTypeName.ToString();
-
-				// Array types are handles separately via checkbox
+				// Array types are handled separately via checkbox
 				if (TypeName.EndsWith(VariableCustomizationPrivate::ArrayIdentifier))
 				{
 					return;
 				}
 
+
 				TSharedPtr<FString> TypeStrPtr = MakeShared<FString>(TypeName);
 				if (TypeName == CurrentTypeName)
 				{
@@ -502,7 +504,7 @@ namespace Metasound
 				}
 
 				// Hidden input types should be omitted from the drop down.
-				if (!HiddenInputTypeNames.Contains(EditorDataType.RegistryInfo.DataTypeName))
+				if (!VariableCustomizationPrivate::HiddenInputTypeNames.Contains(EditorDataType.RegistryInfo.DataTypeName))
 				{
 					DataTypeNames.Add(TypeStrPtr);
 				}
@@ -553,7 +555,7 @@ namespace Metasound
 				.VAlign(VAlign_Center)
 				[
 					SAssignNew(DataTypeArrayCheckbox, SCheckBox)
-					.IsEnabled(bIsArrayTypeRegisteredForCurrentType)
+					.IsEnabled(bIsArrayTypeRegistered && !bIsArrayTypeRegisteredHidden)
 					.IsChecked_Lambda([this, InGraphVariable]()
 					{
 						return OnGetDataTypeArrayCheckState(InGraphVariable);
diff --git a/Engine/Plugins/Runtime/Metasound/Source/MetasoundEditor/Private/SMetasoundGraphNode.cpp b/Engine/Plugins/Runtime/Metasound/Source/MetasoundEditor/Private/SMetasoundGraphNode.cpp
--- a/Engine/Plugins/Runtime/Metasound/Source/MetasoundEditor/Private/SMetasoundGraphNode.cpp
+++ b/Engine/Plugins/Runtime/Metasound/Source/MetasoundEditor/Private/SMetasoundGraphNode.cpp
@@ -44,6 +44,35 @@
 #define LOCTEXT_NAMESPACE "MetasoundGraphNode"
 
 
+// namespace MetaSound
+// {
+// 	namespace Editor
+// 	{
+		class SMetaSoundGraphPinTrigger : public SGraphPin
+		{
+			bool bImagesCached = false;
+
+		public:
+			SLATE_BEGIN_ARGS(SGraphPin)
+			{
+			}
+
+			SLATE_END_ARGS()
+
+			void Construct(const FArguments& InArgs, UEdGraphPin* InPin)
+			{
+				SGraphPin::Construct(SGraphPin::FArguments(), InPin);
+
+				if (const ISlateStyle* MetasoundStyle = FSlateStyleRegistry::FindSlateStyle("MetaSoundStyle"))
+				{
+					CachedImg_Pin_Connected = MetasoundStyle->GetBrush(TEXT("MetasoundEditor.Graph.TriggerPin.Connected"));
+					CachedImg_Pin_Disconnected = MetasoundStyle->GetBrush(TEXT("MetasoundEditor.Graph.TriggerPin.Disconnected"));
+				}
+			}
+		};
+// 	} // namespace Editor
+// } // namespace MetaSound
+
 void SMetasoundGraphNode::Construct(const FArguments& InArgs, class UEdGraphNode* InNode)
 {
 	GraphNode = InNode;
@@ -164,10 +193,10 @@ TSharedPtr<SGraphPin> SMetasoundGraphNode::CreatePinWidget(UEdGraphPin* InPin) c
 			return SNew(SGraphPinBool, InPin);
 		}
 
-		if (InPin->PinType.PinCategory == FGraphBuilder::PinCategoryDouble)
-		{
-			return SNew(SGraphPinNum<double>, InPin);
-		}
+		//if (InPin->PinType.PinCategory == FGraphBuilder::PinCategoryDouble)
+		//{
+		//	return SNew(SGraphPinNum<double>, InPin);
+		//}
 
 		if (InPin->PinType.PinCategory == FGraphBuilder::PinCategoryFloat)
 		{
@@ -183,10 +212,10 @@ TSharedPtr<SGraphPin> SMetasoundGraphNode::CreatePinWidget(UEdGraphPin* InPin) c
 			return SNew(SGraphPinInteger, InPin);
 		}
 
-		if (InPin->PinType.PinCategory == FGraphBuilder::PinCategoryInt64)
-		{
-			return SNew(SGraphPinNum<int64>, InPin);
-		}
+		//if (InPin->PinType.PinCategory == FGraphBuilder::PinCategoryInt64)
+		//{
+		//	return SNew(SGraphPinNum<int64>, InPin);
+		//}
 
 		if (InPin->PinType.PinCategory == FGraphBuilder::PinCategoryObject)
 		{
@@ -200,7 +229,7 @@ TSharedPtr<SGraphPin> SMetasoundGraphNode::CreatePinWidget(UEdGraphPin* InPin) c
 
 		if (InPin->PinType.PinCategory == FGraphBuilder::PinCategoryTrigger)
 		{
-			return SNew(SGraphPinExec, InPin);
+			return SNew(SMetaSoundGraphPinTrigger, InPin);
 		}
 	}
 
