Metasounds: GetWaveDuration node

diff --git a/Engine/Plugins/Runtime/Metasound/Source/MetasoundEngine/Private/MetasoundWaveInfoNode.cpp b/Engine/Plugins/Runtime/Metasound/Source/MetasoundEngine/Private/MetasoundWaveInfoNode.cpp
new file mode 100644
--- /dev/null
+++ b/Engine/Plugins/Runtime/Metasound/Source/MetasoundEngine/Private/MetasoundWaveInfoNode.cpp
@@ -0,0 +1,183 @@
+// Copyright Epic Games, Inc. All Rights Reserved.
+
+#include "CoreMinimal.h"
+#include "Internationalization/Text.h"
+
+#include "MetasoundWave.h"
+#include "MetasoundFacade.h"
+#include "MetasoundPrimitives.h"
+#include "MetasoundAudioBuffer.h"
+#include "MetasoundParamHelper.h"
+#include "MetasoundOperatorSettings.h"
+#include "MetasoundEngineNodesNames.h"
+#include "MetasoundExecutableOperator.h"
+#include "MetasoundNodeRegistrationMacro.h"
+#include "MetasoundStandardNodesCategories.h"
+#include "MetasoundDataTypeRegistrationMacro.h"
+
+#define LOCTEXT_NAMESPACE "MetasoundWaveInfo"
+
+namespace Metasound
+{
+	// forward declarations
+	// ...
+
+	namespace WaveInfoNodeParameterNames
+	{
+		// inputs
+		METASOUND_PARAM(ParamWaveAsset, "Wave", "Input Wave Asset");
+
+		// outputs
+		METASOUND_PARAM(ParamDurationSeconds, "Duration", "Duration of the wave asset in seconds");
+
+	} // namespace WaveInfoNodeParameterNames
+
+	using namespace WaveInfoNodeParameterNames;
+
+
+
+	class FWaveInfoNodeOperator : public TExecutableOperator < FWaveInfoNodeOperator >
+	{
+	public:
+		// ctor
+		FWaveInfoNodeOperator(const FOperatorSettings& InSettings, const FWaveAssetReadRef& InWaveAsset);
+
+		// node interface
+		static const FNodeClassMetadata& GetNodeInfo();
+		static FVertexInterface DeclareVertexInterface();
+		static TUniquePtr<IOperator> CreateOperator(const FCreateOperatorParams& InParams, FBuildErrorArray& OutErrors);
+		virtual FDataReferenceCollection GetInputs() const override;
+		virtual FDataReferenceCollection GetOutputs() const override;
+		void Execute();
+
+	private: // members
+		// input pins
+		FWaveAssetReadRef WaveAsset;
+
+		// output pins
+		FTimeWriteRef DurationSeconds;
+
+		// other
+		FOperatorSettings Settings;
+
+	}; // class FWaveInfoNodeOperator
+
+
+
+
+	// ctor
+	FWaveInfoNodeOperator::FWaveInfoNodeOperator(const FOperatorSettings& InSettings, const FWaveAssetReadRef& InWaveAsset)
+		: WaveAsset(InWaveAsset)
+		, DurationSeconds(FTimeWriteRef::CreateNew(0.0f))
+		, Settings(InSettings)
+	{
+		Execute();
+	}
+
+
+	const FNodeClassMetadata& FWaveInfoNodeOperator::GetNodeInfo()
+	{
+		auto InitNodeInfo = []() -> FNodeClassMetadata
+		{
+			FNodeClassMetadata Info;
+			Info.ClassName = { Metasound::EngineNodes::Namespace, TEXT("Get Wave Duration"), TEXT(" ") };
+			Info.MajorVersion = 1;
+			Info.MinorVersion = 0;
+			Info.DisplayName = LOCTEXT("MetasoundGetWaveDuration_ClassNodeDisplayName", "Get Wave Duration");
+			Info.Description = LOCTEXT("GetWaveDuration_NodeDescription", "Returns the duration of the input Wave asset (in seconds)"),
+			Info.Author = PluginAuthor;
+			Info.PromptIfMissing = PluginNodeMissingPrompt;
+			Info.DefaultInterface = DeclareVertexInterface();
+
+			return Info;
+		};
+
+		static const FNodeClassMetadata Info = InitNodeInfo();
+
+		return Info;
+	}
+
+
+	FVertexInterface FWaveInfoNodeOperator::DeclareVertexInterface()
+	{
+		using namespace WaveInfoNodeParameterNames;
+		static const FVertexInterface Interface(
+			FInputVertexInterface(
+				TInputDataVertexModel<FWaveAsset>(METASOUND_GET_PARAM_NAME(ParamWaveAsset), METASOUND_GET_PARAM_TT(ParamWaveAsset))
+			),
+			FOutputVertexInterface(
+				TOutputDataVertexModel<FTime>(METASOUND_GET_PARAM_NAME(ParamDurationSeconds), METASOUND_GET_PARAM_TT(ParamDurationSeconds))
+			)
+		);
+
+		return Interface;
+	}
+
+
+	TUniquePtr<IOperator> FWaveInfoNodeOperator::CreateOperator(const FCreateOperatorParams& InParams, FBuildErrorArray& OutErrors)
+	{
+		using namespace WaveInfoNodeParameterNames;
+		
+		const FDataReferenceCollection& InputDataRefs = InParams.InputDataReferences;
+
+		// inputs
+		FWaveAssetReadRef WaveAssetIn = InputDataRefs.GetDataReadReferenceOrConstruct<FWaveAsset>(METASOUND_GET_PARAM_NAME(ParamWaveAsset));
+
+		return MakeUnique < FWaveInfoNodeOperator >(InParams.OperatorSettings, WaveAssetIn);
+	}
+
+	FDataReferenceCollection FWaveInfoNodeOperator::GetInputs() const
+	{
+		using namespace WaveInfoNodeParameterNames;
+		FDataReferenceCollection InputDataReferences;
+		InputDataReferences.AddDataReadReference(METASOUND_GET_PARAM_NAME(ParamWaveAsset), FWaveAssetReadRef(WaveAsset));
+
+		return InputDataReferences;
+	}
+
+	FDataReferenceCollection FWaveInfoNodeOperator::GetOutputs() const
+	{
+		// expose read access to our output buffer for other processors in the graph
+		using namespace WaveInfoNodeParameterNames;
+		FDataReferenceCollection OutputDataReferences;
+		OutputDataReferences.AddDataReadReference(METASOUND_GET_PARAM_NAME(ParamDurationSeconds), DurationSeconds);
+
+		return OutputDataReferences;
+	}
+
+	void FWaveInfoNodeOperator::Execute()
+	{
+		if ((*WaveAsset).IsSoundWaveValid())
+		{
+			*DurationSeconds = FTime::FromSeconds((*WaveAsset)->GetDuration());
+		}
+		else
+		{
+			*DurationSeconds = FTime::FromSeconds(0.0f);
+		}
+	}
+
+
+
+	class FWaveInfoNode : public FNodeFacade
+	{
+	public:
+		// public node api needs to define two conversion constructors:
+		// (1: from FString)
+		FWaveInfoNode(const FString& InInstanceName, const FGuid& InInstanceID)
+			: FNodeFacade(InInstanceName, InInstanceID, TFacadeOperatorClass < FWaveInfoNodeOperator >())
+		{ }
+
+		// (2: From an NodeInitData struct)
+		FWaveInfoNode(const FNodeInitData& InInitData)
+			: FWaveInfoNode(InInitData.InstanceName, InInitData.InstanceID)
+		{ }
+
+	};
+
+
+	METASOUND_REGISTER_NODE(FWaveInfoNode);
+
+} // namespace Metasound
+
+#undef LOCTEXT_NAMESPACE //MetasoundWaveInfo
\ No newline at end of file
diff --git a/Engine/Plugins/Runtime/Metasound/Source/MetasoundFrontend/Public/MetasoundTime.h b/Engine/Plugins/Runtime/Metasound/Source/MetasoundFrontend/Public/MetasoundTime.h
--- a/Engine/Plugins/Runtime/Metasound/Source/MetasoundFrontend/Public/MetasoundTime.h
+++ b/Engine/Plugins/Runtime/Metasound/Source/MetasoundFrontend/Public/MetasoundTime.h
@@ -338,6 +338,13 @@ namespace Metasound
 				return static_cast<ArithmeticType>(InTime.GetSeconds() * 1e6);
 			}
 
+			template<typename ArithmeticType>
+			static FTime FromSeconds(const ArithmeticType& InSeconds)
+			{
+				static_assert(TIsArithmetic<ArithmeticType>::Value, "Must be arithmetic type.");
+				return FTime(static_cast<TimeType>(InSeconds));
+			}
+
 			template<typename ArithmeticType>
 			static FTime FromMilliseconds(const ArithmeticType& InMilliseconds)
 			{
