Fix Metasound pop when looping audio

diff --git a/Engine/Plugins/Runtime/Metasound/Source/MetasoundEngine/Private/MetasoundWave.cpp b/Engine/Plugins/Runtime/Metasound/Source/MetasoundEngine/Private/MetasoundWave.cpp
--- a/Engine/Plugins/Runtime/Metasound/Source/MetasoundEngine/Private/MetasoundWave.cpp
+++ b/Engine/Plugins/Runtime/Metasound/Source/MetasoundEngine/Private/MetasoundWave.cpp
@@ -1,18 +1,62 @@
 // Copyright Epic Games, Inc. All Rights Reserved.
 
 #include "MetasoundWave.h"
-#include "MetasoundPrimitives.h"
+
+#include "AudioDevice.h"
+#include "AudioDeviceManager.h"
+#include "ContentStreaming.h"
+#include "DecoderInputFactory.h"
+#include "DSP/ParamInterpolator.h"
 #include "IAudioCodec.h"
 #include "IAudioCodecRegistry.h"
+#include "MetasoundPrimitives.h"
 #include "Sound/SoundWave.h"
-#include "DecoderInputFactory.h"
-#include "AudioDeviceManager.h"
-#include "AudioDevice.h"
-#include "DSP/ParamInterpolator.h"
+
+namespace Metasound
+{
+
+	FWaveAsset::FWaveAsset(const TUniquePtr<Audio::IProxyData>& InInitData)
+	{
+		if (InInitData.IsValid())
+		{
+			if (InInitData->CheckTypeCast<FSoundWaveProxy>())
+			{
+				// should we be getting handed a SharedPtr here?
+				SoundWaveProxy = MakeShared<FSoundWaveProxy, ESPMode::ThreadSafe>(InInitData->GetAs<FSoundWaveProxy>());
+
+				if (SoundWaveProxy.IsValid())
+				{
+					// TODO HACK: Prime the sound for playback.
+					//
+					// Preferably playback latency would be controlled externally.
+					// With the current decoder and waveplayer implementation, the 
+					// wave player does not know whether samples were actually decoded
+					// or if the decoder is still waiting on the stream cache. Generally
+					// this is not an issue except for looping. Looping requires counting
+					// of decoded samples to get exact loop points. When the decoder 
+					// returns zeroed audio (because the stream cache has not loaded
+					// the requested chunk) the sample counting gets off. Currently
+					// there is not route to expose that information to the wave 
+					// player to correct the sample counting logic. 
+					//
+					// In hopes of mitigating the issue, the stream cache
+					// is primed here in the hopes that the chunk is ready by the
+					// time that the decoder attempts to decode audio.
+					IStreamingManager::Get().GetAudioStreamingManager().RequestChunk(SoundWaveProxy, 1, [](EAudioChunkLoadResult) {});
+				}
+			}
+		}
+	}
+
+	bool FWaveAsset::IsSoundWaveValid() const
+	{
+		return SoundWaveProxy.IsValid();
+	}
+}
 
 namespace Audio
 {
-	bool FSimpleDecoderWrapper::Initialize(const InitParams& InInitParams, const FSoundWaveProxyPtr& InWave)
+	bool FSimpleDecoderWrapper::Initialize(const InitParams& InInitParams, const FSoundWaveProxyPtr& InWave, bool bRetainExistingSamples)
 	{
 		// validate data
 		if (!ensure(InWave.IsValid()) || !InWave->IsStreaming())
@@ -42,7 +86,8 @@ namespace Audio
 		DecodeBlockSizeInSamples = DecodeBlockSizeInFrames * NumChannels;
 
 		// set Circular Buffer capacity
-		OutputCircularBuffer.SetCapacity(InInitParams.OutputBlockSizeInFrames * NumChannels * (1.0f + FsOutToInRatio * MaxPitchShiftRatio) * 2);
+		int32 Capacity = FMath::Max(1, static_cast<int32>(InInitParams.OutputBlockSizeInFrames * NumChannels * (1.0f + FsOutToInRatio * MaxPitchShiftRatio) * 2));
+		OutputCircularBuffer.Reserve(Capacity, bRetainExistingSamples);
 
 		TotalNumFramesOutput = 0;
 		TotalNumFramesDecoded = 0;
@@ -135,6 +180,7 @@ namespace Audio
 		return NumOutputSamples; // update once we are aware of partial decode on last buffer
 	}
 
+	/*
 	void FSimpleDecoderWrapper::SeekToTime(const float InSeconds)
 	{
 		if (Input.IsValid() && bIsSeekable)
@@ -146,6 +192,7 @@ namespace Audio
 			ensureMsgf(false, TEXT("Attempting to seek on a sound wave that is not set to Seekable"));
 		}
 	}
+	*/
 
 	bool FSimpleDecoderWrapper::InitializeDecodersInternal(const FSoundWaveProxyPtr& Wave)
 	{
diff --git a/Engine/Plugins/Runtime/Metasound/Source/MetasoundEngine/Private/MetasoundWavePlayerNode.cpp b/Engine/Plugins/Runtime/Metasound/Source/MetasoundEngine/Private/MetasoundWavePlayerNode.cpp
--- a/Engine/Plugins/Runtime/Metasound/Source/MetasoundEngine/Private/MetasoundWavePlayerNode.cpp
+++ b/Engine/Plugins/Runtime/Metasound/Source/MetasoundEngine/Private/MetasoundWavePlayerNode.cpp
@@ -400,7 +400,34 @@ namespace Metasound
 				Params.InitialPitchShiftSemitones = PitchShiftClamped;
 				Params.StartTimeSeconds = SeekTimeSeconds;
 
-				return Decoder.Initialize(Params, WaveProxy);
+				bool bRetainExistingSamples = false;
+
+				return Decoder.Initialize(Params, WaveProxy, bRetainExistingSamples);
+			}
+
+			return false;
+		}
+
+		bool SeekDecoder(float SeekTimeSeconds, bool bLogFailures = false)
+		{
+			if (IsCurrentWaveValid(bLogFailures))
+			{
+				const FSoundWaveProxyPtr& WaveProxy = CurrentWaveAsset.GetSoundWaveProxy();
+
+				CuePoints = WaveProxy->GetCuePoints();
+
+				float PitchShiftClamped = FMath::Clamp(*PitchShift, -12.0f * MaxPitchShiftOctaves, 12.0f * MaxPitchShiftOctaves);
+
+				Audio::FSimpleDecoderWrapper::InitParams Params;
+				Params.OutputBlockSizeInFrames = OutputBlockSizeInFrames;
+				Params.OutputSampleRate = OutputSampleRate;
+				Params.MaxPitchShiftMagnitudeAllowedInOctaves = 6.f;
+				Params.InitialPitchShiftSemitones = PitchShiftClamped;
+				Params.StartTimeSeconds = SeekTimeSeconds;
+
+				bool bRetainExistingSamples = true;
+
+				return Decoder.Initialize(Params, WaveProxy, bRetainExistingSamples);
 			}
 
 			return false;
@@ -537,7 +564,7 @@ namespace Metasound
 						*LoopPercent = 0.0f;
 
 						// Now seek to the loop start frame
-						StartDecoder(LoopSeekTime);
+						SeekDecoder(LoopSeekTime);
 					}
 					else
 					{
diff --git a/Engine/Plugins/Runtime/Metasound/Source/MetasoundEngine/Public/MetasoundWave.h b/Engine/Plugins/Runtime/Metasound/Source/MetasoundEngine/Public/MetasoundWave.h
--- a/Engine/Plugins/Runtime/Metasound/Source/MetasoundEngine/Public/MetasoundWave.h
+++ b/Engine/Plugins/Runtime/Metasound/Source/MetasoundEngine/Public/MetasoundWave.h
@@ -28,24 +28,9 @@ namespace Metasound
 		FWaveAsset(const FWaveAsset&) = default;
 		FWaveAsset& operator=(const FWaveAsset& Other) = default;
 
-		FWaveAsset(const TUniquePtr<Audio::IProxyData>& InInitData)
-		{
-			SoundWaveProxy.Reset();
-
-			if (InInitData.IsValid())
-			{
-				if (InInitData->CheckTypeCast<FSoundWaveProxy>())
-				{
-					// should we be getting handed a SharedPtr here?
-					SoundWaveProxy = MakeShared<FSoundWaveProxy, ESPMode::ThreadSafe>(InInitData->GetAs<FSoundWaveProxy>());
-				}
-			}
-		}
+		FWaveAsset(const TUniquePtr<Audio::IProxyData>& InInitData);
 
-		bool IsSoundWaveValid() const
-		{
-			return SoundWaveProxy.IsValid();
-		}
+		bool IsSoundWaveValid() const;
 
 		const FSoundWaveProxyPtr& GetSoundWaveProxy() const
 		{
@@ -91,12 +76,22 @@ namespace Audio
 			return !bDecoderIsDone && Input.IsValid() && Output.IsValid() && Decoder.IsValid();
 		}
 
-		bool Initialize(const InitParams& InInitParams, const FSoundWaveProxyPtr& InWave);
+		/** Initialize the decoder to produce audio from the given wave.
+		 *
+		 * @param InInitParams - Initialization parameters for the decoder. 
+		 * @param InWave - The wave proxy to decode.
+		 * @param bRetainExistingSamples - If true, any existing decoded samples in the FSimpleDecoderWrapper
+		 *                                 will be output during subsequent calls to GenerateAudio(...). If false,
+		 *                                 any previously existing decoded samples will be discarded.
+		 *
+		 * @return True on success, false on failure. 
+		 */
+		bool Initialize(const InitParams& InInitParams, const FSoundWaveProxyPtr& InWave, bool bRetainExistingSamples=false);
 
 		// returns number of samples written.   
 		uint32 GenerateAudio(float* OutputDest, int32 NumOutputFrames, int32& OutNumFramesConsumed, float PitchShiftInCents = 0.0f, bool bIsLooping = false);
 
-		void SeekToTime(const float InSeconds);
+		//void SeekToTime(const float InSeconds);
 
 	private:
 		// actual decoder objects
diff --git a/Engine/Source/Runtime/SignalProcessing/Public/DSP/Dsp.h b/Engine/Source/Runtime/SignalProcessing/Public/DSP/Dsp.h
--- a/Engine/Source/Runtime/SignalProcessing/Public/DSP/Dsp.h
+++ b/Engine/Source/Runtime/SignalProcessing/Public/DSP/Dsp.h
@@ -787,6 +787,31 @@ namespace Audio
 			InternalBuffer.AddZeroed(Capacity);
 		}
 
+		/** Reserve capacity.
+		 *
+		 * @param InMinimumCapacity - Minimum capacity of circular buffer.
+		 * @param bRetainExistingSamples - If true, existing samples in the buffer will be retained. If false, they are discarded.
+		 */
+		void Reserve(uint32 InMinimumCapacity, bool bRetainExistingSamples)
+		{
+			if (Capacity <= InMinimumCapacity)
+			{
+				uint32 NewCapacity = InMinimumCapacity + 1;
+
+				checkf(NewCapacity < (uint32)TNumericLimits<int32>::Max(), TEXT("Max capacity overflow. Requested %d. Maximum allowed %d"), NewCapacity, TNumericLimits<int32>::Max());
+
+				uint32 NumToAdd = NewCapacity - Capacity;
+				InternalBuffer.AddZeroed(NumToAdd);
+				Capacity = NewCapacity;
+			}
+
+			if (!bRetainExistingSamples)
+			{
+				ReadCounter.Set(0);
+				WriteCounter.Set(0);
+			}
+		}
+
 		// Pushes some amount of samples into this circular buffer.
 		// Returns the amount of samples written.
 		// This can only be used for trivially copyable types.
