Looping metasound wave pop fixes

diff --git a/Engine/Plugins/Runtime/Metasound/Source/MetasoundEngine/Private/MetasoundWave.cpp b/Engine/Plugins/Runtime/Metasound/Source/MetasoundEngine/Private/MetasoundWave.cpp
--- a/Engine/Plugins/Runtime/Metasound/Source/MetasoundEngine/Private/MetasoundWave.cpp
+++ b/Engine/Plugins/Runtime/Metasound/Source/MetasoundEngine/Private/MetasoundWave.cpp
@@ -56,53 +56,74 @@ namespace Metasound
 
 namespace Audio
 {
-	bool FSimpleDecoderWrapper::Initialize(const InitParams& InInitParams, const FSoundWaveProxyPtr& InWave, bool bRetainExistingSamples)
+	FSimpleDecoderWrapper::FSimpleDecoderWrapper(const InitParams& InInitParams)
+	: OutputSampleRate(InInitParams.OutputSampleRate)
+	, MaxPitchShiftCents(InInitParams.MaxPitchShiftMagnitudeAllowedInOctaves * 1200.0f)
+	, MaxPitchShiftRatio(FMath::Pow(2.0f, InInitParams.MaxPitchShiftMagnitudeAllowedInOctaves))
+	, DecodeBlockSizeInFrames(64)
+	, OutputBlockSizeInFrames(InInitParams.OutputBlockSizeInFrames)
 	{
-		// validate data
-		if (!ensure(InWave.IsValid()) || !InWave->IsStreaming())
-		{
-			ensureAlwaysMsgf(InWave->IsStreaming(), TEXT("Metasounds does not support Force Inline (sound must be streaming)"));
+		check(OutputBlockSizeInFrames > 0);
+	}
 
-			Input.Reset();
-			Output.Reset();
-			Decoder.Reset();
-			bDecoderIsDone = false;
-			bIsSeekable = InWave->IsSeekableStreaming();
+	bool FSimpleDecoderWrapper::SetWave(const FSoundWaveProxyPtr& InWave, float InStartTimeSeconds, float InInitialPitchShiftSemitones)
+	{
+		Wave = InWave;
 
-			return false;
-		}
+		bool bSuccessful = true;
 
-		// initialize input/output data
-		InputSampleRate = InWave->GetSampleRate();
-		OutputSampleRate = InInitParams.OutputSampleRate;
-		FsOutToInRatio = (OutputSampleRate / InputSampleRate);
-		MaxPitchShiftRatio = FMath::Pow(2.0f, InInitParams.MaxPitchShiftMagnitudeAllowedInOctaves);
-		MaxPitchShiftCents = InInitParams.MaxPitchShiftMagnitudeAllowedInOctaves * 1200.0f;
-		StartTimeSeconds = InInitParams.StartTimeSeconds;
+		if (Wave.IsValid())
+		{
+			// validate data
+			if (!ensureAlwaysMsgf(Wave->IsStreaming(), TEXT("Metasounds does not support Force Inline (sound must be streaming)")))
+			{
+				Wave.Reset();
+				Input.Reset();
+				Output.Reset();
+				Decoder.Reset();
+				bDecoderIsDone = false;
 
-		NumChannels = InWave->GetNumChannels();
-		NumFrames = InWave->GetNumFrames();
-		DecodeBlockSizeInFrames = 64;
-		DecodeBlockSizeInSamples = DecodeBlockSizeInFrames * NumChannels;
+				return false;
+			}
+
+			// initialize input/output data
+			InputSampleRate = Wave->GetSampleRate();
+			FsOutToInRatio = (OutputSampleRate / InputSampleRate);
+
+			NumChannels = Wave->GetNumChannels();
+			DecodeBlockSizeInSamples = DecodeBlockSizeInFrames * NumChannels;
 
-		// set Circular Buffer capacity
-		int32 Capacity = FMath::Max(1, static_cast<int32>(InInitParams.OutputBlockSizeInFrames * NumChannels * (1.0f + FsOutToInRatio * MaxPitchShiftRatio) * 2));
-		OutputCircularBuffer.Reserve(Capacity, bRetainExistingSamples);
+			// set Circular Buffer capacity
+			int32 Capacity = FMath::Max(1, static_cast<int32>(OutputBlockSizeInFrames * NumChannels * (1.0f + FsOutToInRatio * MaxPitchShiftRatio) * 2));
+			OutputCircularBuffer.Reserve(Capacity, true /* bRetainExistingSamples */);
 
-		TotalNumFramesOutput = 0;
-		TotalNumFramesDecoded = 0;
 
-		// try to initialize decoders
-		bool bSuccessful = InitializeDecodersInternal(InWave);
-		bDecoderIsDone = !bSuccessful;
+			// try to initialize decoders
+			bSuccessful = InitializeDecodersInternal(Wave, InStartTimeSeconds);
+			bDecoderIsDone = !bSuccessful;
 
-		// initialize SRC object (will be re-initialized for pitch shifting)
-		Resampler.Init(Audio::EResamplingMethod::Linear, FsOutToInRatio, NumChannels);
-		PitchShifter.Reset(NumChannels, InInitParams.InitialPitchShiftSemitones);
+			Resampler.Init(Audio::EResamplingMethod::Linear, FsOutToInRatio, NumChannels);
+			PitchShifter.Reset(NumChannels, InInitialPitchShiftSemitones);
+		}
 
 		return bSuccessful;
 	}
 
+	bool FSimpleDecoderWrapper::SeekToTime(const float InSeconds)
+	{
+		if (Wave.IsValid())
+		{
+			// try to initialize decoders
+			bool bSuccessful = InitializeDecodersInternal(Wave, InSeconds);
+			bDecoderIsDone = !bSuccessful;
+
+			return bSuccessful;
+		}
+
+		return false;
+	}
+	
+
 	uint32 FSimpleDecoderWrapper::GenerateAudio(float* OutputDest, int32 NumOutputFrames, int32& OutNumFramesConsumed, float PitchShiftInCents, bool bIsLooping)
 	{
 		const uint32 NumOutputSamples = NumOutputFrames * NumChannels;
@@ -180,31 +201,18 @@ namespace Audio
 		return NumOutputSamples; // update once we are aware of partial decode on last buffer
 	}
 
-	/*
-	void FSimpleDecoderWrapper::SeekToTime(const float InSeconds)
-	{
-		if (Input.IsValid() && bIsSeekable)
-		{
-			Input->SeekToTime(InSeconds);
-		}
-		if (!bIsSeekable)
-		{
-			ensureMsgf(false, TEXT("Attempting to seek on a sound wave that is not set to Seekable"));
-		}
-	}
-	*/
-
-	bool FSimpleDecoderWrapper::InitializeDecodersInternal(const FSoundWaveProxyPtr& Wave)
+	bool FSimpleDecoderWrapper::InitializeDecodersInternal(const FSoundWaveProxyPtr& InWave, float InStartTimeSeconds)
 	{
-		if (!ensure(Wave.IsValid()))
+		if (!ensure(InWave.IsValid()))
 		{
 			return false;
 		}
 
 		// Input:
-		FName OldFormat = Wave->GetRuntimeFormat();
-		Input = MakeShareable(Audio::CreateBackCompatDecoderInput(OldFormat, Wave).Release());
-		Input->SeekToTime(StartTimeSeconds);
+		FName OldFormat = InWave->GetRuntimeFormat();
+		// TODO: Why is this shared? Doesn't look like it needs to be shared. 
+		Input = MakeShareable(Audio::CreateBackCompatDecoderInput(OldFormat, InWave).Release());
+		Input->SeekToTime(InStartTimeSeconds);
 
 		if (!Input)
 		{
diff --git a/Engine/Plugins/Runtime/Metasound/Source/MetasoundEngine/Private/MetasoundWavePlayerNode.cpp b/Engine/Plugins/Runtime/Metasound/Source/MetasoundEngine/Private/MetasoundWavePlayerNode.cpp
--- a/Engine/Plugins/Runtime/Metasound/Source/MetasoundEngine/Private/MetasoundWavePlayerNode.cpp
+++ b/Engine/Plugins/Runtime/Metasound/Source/MetasoundEngine/Private/MetasoundWavePlayerNode.cpp
@@ -130,6 +130,7 @@ namespace Metasound
 			, PlaybackLocation(FFloatWriteRef::CreateNew(0.0f))
 			, AudioBufferL(FAudioBufferWriteRef::CreateNew(InArgs.Settings))
 			, AudioBufferR(FAudioBufferWriteRef::CreateNew(InArgs.Settings))
+			, Decoder(Audio::FSimpleDecoderWrapper::InitParams{static_cast<uint32>(InArgs.Settings.GetNumFramesPerBlock()), InArgs.Settings.GetSampleRate(), 6.f /*MaxPitchShiftMagnitudeAllowedInOctages */})
 			, OutputSampleRate(InArgs.Settings.GetSampleRate())
 			, OutputBlockSizeInFrames(InArgs.Settings.GetNumFramesPerBlock())
 		{
@@ -293,11 +294,13 @@ namespace Metasound
 
 		void StartPlaying()
 		{
+			// Copy the wave asset off on init in case the user changes it while we're playing it.
+			// We'll only check for new wave assets when the current one finishes for sample accurate concantenation
+			CurrentWaveAsset = *WaveAsset;
+
 			if (IsWaveValid())
 			{
-				// Copy the wave asset off on init in case the user changes it while we're playing it.
-				// We'll only check for new wave assets when the current one finishes for sample accurate concantenation
-				CurrentWaveAsset = *WaveAsset;
+				
 				SoundAssetSampleRate = CurrentWaveAsset->GetSampleRate();
 				SoundAssetDurationSeconds = CurrentWaveAsset->GetDuration();
 				SoundAssetNumFrames = CurrentWaveAsset->GetNumFrames();
@@ -321,6 +324,16 @@ namespace Metasound
 				PlaybackStartFrame = (uint32)(StartTimeSeconds * SoundAssetSampleRate);
 				UpdateAndGetLoopStartTime();
 
+
+
+				// Update our current loop frame
+				CurrentConsumedFrameCount = LoopStartFrame;
+
+				// Reset our decoded frame counting for sub loop
+				NumTotalDecodedFramesInLoop = 0;
+
+				*LoopPercent = 0.0f;
+
 				StartDecoder(StartTimeSeconds, true /* bLogFailures */);
 
 				if (!bIsPlaying)
@@ -388,21 +401,12 @@ namespace Metasound
 			if (IsCurrentWaveValid(bLogFailures))
 			{
 				const FSoundWaveProxyPtr& WaveProxy = CurrentWaveAsset.GetSoundWaveProxy();
-
+				
 				CuePoints = WaveProxy->GetCuePoints();
 
 				float PitchShiftClamped = FMath::Clamp(*PitchShift, -12.0f * MaxPitchShiftOctaves, 12.0f * MaxPitchShiftOctaves);
 
-				Audio::FSimpleDecoderWrapper::InitParams Params;
-				Params.OutputBlockSizeInFrames = OutputBlockSizeInFrames;
-				Params.OutputSampleRate = OutputSampleRate;
-				Params.MaxPitchShiftMagnitudeAllowedInOctaves = 6.f;
-				Params.InitialPitchShiftSemitones = PitchShiftClamped;
-				Params.StartTimeSeconds = SeekTimeSeconds;
-
-				bool bRetainExistingSamples = false;
-
-				return Decoder.Initialize(Params, WaveProxy, bRetainExistingSamples);
+				return Decoder.SetWave(WaveProxy, SeekTimeSeconds, PitchShiftClamped);
 			}
 
 			return false;
@@ -410,27 +414,7 @@ namespace Metasound
 
 		bool SeekDecoder(float SeekTimeSeconds, bool bLogFailures = false)
 		{
-			if (IsCurrentWaveValid(bLogFailures))
-			{
-				const FSoundWaveProxyPtr& WaveProxy = CurrentWaveAsset.GetSoundWaveProxy();
-
-				CuePoints = WaveProxy->GetCuePoints();
-
-				float PitchShiftClamped = FMath::Clamp(*PitchShift, -12.0f * MaxPitchShiftOctaves, 12.0f * MaxPitchShiftOctaves);
-
-				Audio::FSimpleDecoderWrapper::InitParams Params;
-				Params.OutputBlockSizeInFrames = OutputBlockSizeInFrames;
-				Params.OutputSampleRate = OutputSampleRate;
-				Params.MaxPitchShiftMagnitudeAllowedInOctaves = 6.f;
-				Params.InitialPitchShiftSemitones = PitchShiftClamped;
-				Params.StartTimeSeconds = SeekTimeSeconds;
-
-				bool bRetainExistingSamples = true;
-
-				return Decoder.Initialize(Params, WaveProxy, bRetainExistingSamples);
-			}
-
-			return false;
+			return Decoder.SeekToTime(SeekTimeSeconds);
 		}
 
 
@@ -619,6 +603,8 @@ namespace Metasound
 			{
 				FMemory::Memzero(FinalOutputLeft, sizeof(float) * NumOutputFramesToGenerate);
 				FMemory::Memzero(FinalOutputRight, sizeof(float) * NumOutputFramesToGenerate);
+
+				NumFramesGenerated = NumOutputFramesToGenerate;
 			}
 
 			return NumFramesGenerated;
diff --git a/Engine/Plugins/Runtime/Metasound/Source/MetasoundEngine/Public/MetasoundWave.h b/Engine/Plugins/Runtime/Metasound/Source/MetasoundEngine/Public/MetasoundWave.h
--- a/Engine/Plugins/Runtime/Metasound/Source/MetasoundEngine/Public/MetasoundWave.h
+++ b/Engine/Plugins/Runtime/Metasound/Source/MetasoundEngine/Public/MetasoundWave.h
@@ -67,40 +67,48 @@ namespace Audio
 			uint32 OutputBlockSizeInFrames{ 512 };
 			float OutputSampleRate { 44100.0f };
 			float MaxPitchShiftMagnitudeAllowedInOctaves { 4.0f };
-			float InitialPitchShiftSemitones { 0.0f };
-			float StartTimeSeconds { 0.0f };
 		};
 
-		bool CanGenerateAudio() const
-		{
-			return !bDecoderIsDone && Input.IsValid() && Output.IsValid() && Decoder.IsValid();
-		}
+		FSimpleDecoderWrapper(const InitParams& InInitparams);
+
 
-		/** Initialize the decoder to produce audio from the given wave.
+		/** Set the wave for decoding.
 		 *
-		 * @param InInitParams - Initialization parameters for the decoder. 
 		 * @param InWave - The wave proxy to decode.
-		 * @param bRetainExistingSamples - If true, any existing decoded samples in the FSimpleDecoderWrapper
-		 *                                 will be output during subsequent calls to GenerateAudio(...). If false,
-		 *                                 any previously existing decoded samples will be discarded.
+		 * @param InStartTimeSeconds - Number of seconds into the wave to start decoding.
+		 * @param InInitialPitchShiftSemitones - Initial pitch shift of wave.
 		 *
 		 * @return True on success, false on failure. 
 		 */
-		bool Initialize(const InitParams& InInitParams, const FSoundWaveProxyPtr& InWave, bool bRetainExistingSamples=false);
+		bool SetWave(const FSoundWaveProxyPtr& InWave, float InStartTimeSeconds=0.f, float InInitialPitchShiftSemitones=0.f);
+
+		/** Seek in the wave.
+		 *
+		 * @param InSeconds - Number of seconds into the wave to start decoding.
+		 *
+		 * @return True on success, false on failure. 
+		 */
+		bool SeekToTime(const float InSeconds);
+
+		bool CanGenerateAudio() const
+		{
+			return !bDecoderIsDone && Input.IsValid() && Output.IsValid() && Decoder.IsValid();
+		}
 
 		// returns number of samples written.   
 		uint32 GenerateAudio(float* OutputDest, int32 NumOutputFrames, int32& OutNumFramesConsumed, float PitchShiftInCents = 0.0f, bool bIsLooping = false);
 
-		//void SeekToTime(const float InSeconds);
 
 	private:
+		FSoundWaveProxyPtr Wave;
+
 		// actual decoder objects
 		TUniquePtr<Audio::IDecoder> Decoder;
 		TUniquePtr<Audio::IDecoderOutput> Output;
 		TSharedPtr<Audio::IDecoderInput, ESPMode::ThreadSafe> Input;
 
 		// init helper for decoders
-		bool InitializeDecodersInternal(const FSoundWaveProxyPtr& Wave);
+		bool InitializeDecodersInternal(const FSoundWaveProxyPtr& InWave, float InStartTimeSeconds);
 
 		// SRC objects
 		Audio::FResampler Resampler;
@@ -119,17 +127,11 @@ namespace Audio
 		float MaxPitchShiftRatio{ 2.0f };
 
 		uint32 NumChannels{ 1 };
-		uint32 NumFrames{ 0 };
 		uint32 DecodeBlockSizeInFrames{ 64 };
 		uint32 DecodeBlockSizeInSamples{ 64 };
+		uint32 OutputBlockSizeInFrames { 512 };
 		
-		// cached values
-		float StartTimeSeconds{ 0.0f };
-		int32 TotalNumFramesOutput{ 0 };
-		int32 TotalNumFramesDecoded{ 0 };
-
 		bool bDecoderIsDone{ true };
-		bool bIsSeekable{ false };
 		bool bDecoderHasLooped{ false };
 
 	}; // class FSimpleDecoderWrapper
diff --git a/Engine/Source/Runtime/AudioPlatformConfiguration/Private/AudioResampler.cpp b/Engine/Source/Runtime/AudioPlatformConfiguration/Private/AudioResampler.cpp
--- a/Engine/Source/Runtime/AudioPlatformConfiguration/Private/AudioResampler.cpp
+++ b/Engine/Source/Runtime/AudioPlatformConfiguration/Private/AudioResampler.cpp
@@ -125,19 +125,29 @@ namespace Audio
 	{
 		int32 ErrorResult = 0;
 
-		// Reset the SRC state if we already have one
+		// Reset the SRC state if we already have one with equal channels
 		if (LibSRCState)
 		{
-			ErrorResult = src_reset(LibSRCState);
-			if (ErrorResult != 0)
+			if (InNumChannels == src_get_channels(LibSRCState))
 			{
-				const char* ErrorString = src_strerror(ErrorResult);
-				UE_LOG(LogAudioResampler, Error, TEXT("Failed to reset sample converter state: %s."), ErrorString);
-				return;
+				ErrorResult = src_reset(LibSRCState);
+				if (ErrorResult != 0)
+				{
+					const char* ErrorString = src_strerror(ErrorResult);
+					UE_LOG(LogAudioResampler, Error, TEXT("Failed to reset sample converter state: %s."), ErrorString);
+					return;
+				}
+			}
+			else
+			{
+				// If the channel counts do not match, then remove one.
+				LibSRCState = src_delete(LibSRCState);
+				check(nullptr == LibSRCState);
 			}
 		}
-		// Otherwise create a new one
-		else
+
+		// Create a new one if one does not exist.
+		if (nullptr == LibSRCState)	
 		{
 			LibSRCState = src_new((int32)ResamplingMethod, InNumChannels, &ErrorResult);
 			if (!LibSRCState)
