Only warm cache if wave is streaming

diff --git a/Engine/Plugins/Runtime/Metasound/Source/MetasoundEngine/Private/MetasoundWave.cpp b/Engine/Plugins/Runtime/Metasound/Source/MetasoundEngine/Private/MetasoundWave.cpp
--- a/Engine/Plugins/Runtime/Metasound/Source/MetasoundEngine/Private/MetasoundWave.cpp
+++ b/Engine/Plugins/Runtime/Metasound/Source/MetasoundEngine/Private/MetasoundWave.cpp
@@ -42,7 +42,10 @@ namespace Metasound
 					// In hopes of mitigating the issue, the stream cache
 					// is primed here in the hopes that the chunk is ready by the
 					// time that the decoder attempts to decode audio.
-					IStreamingManager::Get().GetAudioStreamingManager().RequestChunk(SoundWaveProxy, 1, [](EAudioChunkLoadResult) {});
+					if (SoundWaveProxy->IsStreaming())
+					{
+						IStreamingManager::Get().GetAudioStreamingManager().RequestChunk(SoundWaveProxy, 1, [](EAudioChunkLoadResult) {});
+					}
 				}
 			}
 		}
