Metasound Waveplayer crash fix.
Attempting to access missing chunk

diff --git a/Engine/Plugins/Runtime/Metasound/Source/MetasoundEngine/Private/MetasoundWave.cpp b/Engine/Plugins/Runtime/Metasound/Source/MetasoundEngine/Private/MetasoundWave.cpp
--- a/Engine/Plugins/Runtime/Metasound/Source/MetasoundEngine/Private/MetasoundWave.cpp
+++ b/Engine/Plugins/Runtime/Metasound/Source/MetasoundEngine/Private/MetasoundWave.cpp
@@ -12,6 +12,15 @@
 #include "MetasoundPrimitives.h"
 #include "Sound/SoundWave.h"
 
+
+static int32 DisableMetasoundWaveAssetCachePriming = 0;
+FAutoConsoleVariableRef CVarDisableMetasoundWaveAssetCachePriming(
+	TEXT("au.MetaSound.DisableWaveCachePriming"),
+	DisableMetasoundWaveAssetCachePriming,
+	TEXT("Disables MetaSound Wave Cache Priming.\n")
+	TEXT("0 (default): Enabled, 1: Disabled"),
+	ECVF_Default);
+
 namespace Metasound
 {
 
@@ -42,9 +51,15 @@ namespace Metasound
 					// In hopes of mitigating the issue, the stream cache
 					// is primed here in the hopes that the chunk is ready by the
 					// time that the decoder attempts to decode audio.
-					if (SoundWaveProxy->IsStreaming())
+					if (0 == DisableMetasoundWaveAssetCachePriming)
 					{
-						IStreamingManager::Get().GetAudioStreamingManager().RequestChunk(SoundWaveProxy, 1, [](EAudioChunkLoadResult) {});
+						if (SoundWaveProxy->IsStreaming())
+						{
+							if (SoundWaveProxy->GetNumChunks() > 1)
+							{
+								IStreamingManager::Get().GetAudioStreamingManager().RequestChunk(SoundWaveProxy, 1, [](EAudioChunkLoadResult) {});
+							}
+						}
 					}
 				}
 			}
