Disable MetaSound Trigger:Array type all together

diff --git a/Engine/Plugins/Runtime/Metasound/Source/MetasoundEditor/Private/MetasoundNodeDetailCustomization.cpp b/Engine/Plugins/Runtime/Metasound/Source/MetasoundEditor/Private/MetasoundNodeDetailCustomization.cpp
--- a/Engine/Plugins/Runtime/Metasound/Source/MetasoundEditor/Private/MetasoundNodeDetailCustomization.cpp
+++ b/Engine/Plugins/Runtime/Metasound/Source/MetasoundEditor/Private/MetasoundNodeDetailCustomization.cpp
@@ -75,8 +75,7 @@ namespace Metasound
 			static const TSet<FName> HiddenInputTypeNames =
 			{
 				"Audio:Mono",
-				"Audio:Stereo",
-				"Trigger:Array"
+				"Audio:Stereo"
 			};
 		}
 
diff --git a/Engine/Plugins/Runtime/Metasound/Source/MetasoundFrontend/Private/MetasoundPrimitives.cpp b/Engine/Plugins/Runtime/Metasound/Source/MetasoundFrontend/Private/MetasoundPrimitives.cpp
--- a/Engine/Plugins/Runtime/Metasound/Source/MetasoundFrontend/Private/MetasoundPrimitives.cpp
+++ b/Engine/Plugins/Runtime/Metasound/Source/MetasoundFrontend/Private/MetasoundPrimitives.cpp
@@ -19,6 +19,13 @@ namespace Metasound
 		static constexpr bool Value = false;
 	};
 
+	// Disable arrays of triggers
+	template<>
+	struct TEnableAutoArrayTypeRegistration<FTrigger>
+	{
+		static constexpr bool Value = false;
+	};
+
 	// Disable auto-conversion using the FAudioBuffer(int32) constructor
 	template<typename FromDataType>
 	struct TEnableAutoConverterNodeRegistration<FromDataType, FAudioBuffer>
