Metasound: Fix for Loctext name collision in the mixer node causing vertex mismatches duing metasound graph construction (and no audio to play) in multiple scenarios. (undo/redo operations, silent on certain platforms, etc.)

diff --git a/Engine/Plugins/Runtime/Metasound/Source/MetasoundStandardNodes/Private/MetasoundCrossfadeNode.cpp b/Engine/Plugins/Runtime/Metasound/Source/MetasoundStandardNodes/Private/MetasoundCrossfadeNode.cpp
--- a/Engine/Plugins/Runtime/Metasound/Source/MetasoundStandardNodes/Private/MetasoundCrossfadeNode.cpp
+++ b/Engine/Plugins/Runtime/Metasound/Source/MetasoundStandardNodes/Private/MetasoundCrossfadeNode.cpp
@@ -37,7 +37,7 @@ namespace Metasound
 
 		const FString GetInputName(uint32 InIndex)
 		{
-			return FText::Format(LOCTEXT("CrossfadeInputName", "In {0}"), InIndex).ToString();
+			return FString::Format(TEXT("In {0}"), { InIndex });
 		}
 
 		const FText GetInputDescription(uint32 InIndex)
diff --git a/Engine/Plugins/Runtime/Metasound/Source/MetasoundStandardNodes/Private/MetasoundMixerNode.cpp b/Engine/Plugins/Runtime/Metasound/Source/MetasoundStandardNodes/Private/MetasoundMixerNode.cpp
--- a/Engine/Plugins/Runtime/Metasound/Source/MetasoundStandardNodes/Private/MetasoundMixerNode.cpp
+++ b/Engine/Plugins/Runtime/Metasound/Source/MetasoundStandardNodes/Private/MetasoundMixerNode.cpp
@@ -230,26 +230,24 @@ namespace Metasound
 		{
 			if (NumChannels == 1)
 			{
-				return FText::Format(LOCTEXT("AudioMixerAudioInputName", "In {0}"), InputIndex).ToString();
+				return FString::Printf(TEXT("In %i"), InputIndex);
 			}
 			else if (NumChannels == 2)
 			{
-				return FString::Printf(TEXT("%s %s"), *FText::Format(LOCTEXT("AudioMixerAudioInputName", "In {0}"), InputIndex).ToString(), (ChannelIndex == 0) ? "L" : "R");
+				return FString::Printf(TEXT("In %i %s"), InputIndex, (ChannelIndex == 0) ? "L" : "R");
 			}
 
-			return FText::Format(LOCTEXT("AudioMixerAudioInputName", "In {0}, {1}"), InputIndex, ChannelIndex).ToString();
+			return FString::Printf(TEXT("In %i, %i"), InputIndex, ChannelIndex);
 		}
 
-
 		static const FText GetAudioInputDescription(uint32 InputIndex, uint32 ChannelIndex)
 		{
 			return FText::Format(LOCTEXT("AudioMixerAudioInputDescription", "Audio Input #: {0}, Channel: {1}"), InputIndex, ChannelIndex);
 		}
 
-
 		static const FString GetGainInputName(uint32 InputIndex)
 		{
-			return FText::Format(LOCTEXT("AudioMixerGainInputName", "Gain {0}"), InputIndex).ToString();
+			return FString::Printf(TEXT("Gain %i"), InputIndex);
 		}
 
 		static const FText GetGainInputDescription(uint32 InputIndex)
@@ -261,19 +259,19 @@ namespace Metasound
 		{
 			if (NumChannels == 1)
 			{
-				return LOCTEXT("AudioMixerAudioOUtputName", "Out").ToString();
+				return TEXT("Out");
 			}
 			else if (NumChannels == 2)
 			{
-				return FString::Printf(TEXT("%s %s"), *LOCTEXT("AudioMixerAudioOutputName", "Out").ToString(), (ChannelIndex == 0) ? TEXT("L") : TEXT("R"));
+				return FString::Printf(TEXT("Out %s"), (ChannelIndex == 0) ? TEXT("L") : TEXT("R"));
 			}
 
-			return FText::Format(LOCTEXT("AudioMixerAudioOUtputName", "Out {0}"), ChannelIndex).ToString();
+			return FString::Printf(TEXT("Out %i"), ChannelIndex);
 		}
 
 		static const FText GetAudioOutputDescription(uint32 ChannelIndex)
 		{
-			return FText::Format(LOCTEXT("AudioMixerAudioOUtputName", "Summed output for channel: {0}"), ChannelIndex);
+			return FText::Format(LOCTEXT("AudioMixerAudioOutputDescription", "Summed output for channel: {0}"), ChannelIndex);
 		}
 #pragma endregion
 	}; // class TAudioMixerNodeOperator
diff --git a/Engine/Plugins/Runtime/Metasound/Source/MetasoundStandardNodes/Private/MetasoundTriggerAnyNode.cpp b/Engine/Plugins/Runtime/Metasound/Source/MetasoundStandardNodes/Private/MetasoundTriggerAnyNode.cpp
--- a/Engine/Plugins/Runtime/Metasound/Source/MetasoundStandardNodes/Private/MetasoundTriggerAnyNode.cpp
+++ b/Engine/Plugins/Runtime/Metasound/Source/MetasoundStandardNodes/Private/MetasoundTriggerAnyNode.cpp
@@ -43,7 +43,7 @@ namespace Metasound
 	{
 		METASOUNDSTANDARDNODES_API const FString GetInputTriggerName(uint32 InIndex)
 		{
-			return FText::Format(LOCTEXT("TriggerAnyTriggerInputName", "In {0}"), InIndex).ToString();
+			return FString::Format(TEXT("In {0}"), { InIndex });
 		}
 
 		METASOUNDSTANDARDNODES_API const FText GetInputTriggerDescription(uint32 InIndex)
diff --git a/Engine/Plugins/Runtime/Metasound/Source/MetasoundStandardNodes/Private/MetasoundTriggerRouteNode.cpp b/Engine/Plugins/Runtime/Metasound/Source/MetasoundStandardNodes/Private/MetasoundTriggerRouteNode.cpp
--- a/Engine/Plugins/Runtime/Metasound/Source/MetasoundStandardNodes/Private/MetasoundTriggerRouteNode.cpp
+++ b/Engine/Plugins/Runtime/Metasound/Source/MetasoundStandardNodes/Private/MetasoundTriggerRouteNode.cpp
@@ -27,7 +27,7 @@ namespace Metasound
 	{
 		const FString GetInputTriggerName(uint32 InIndex)
 		{
-			return FText::Format(LOCTEXT("TriggerRouteTriggerInputName", "Set {0}"), InIndex).ToString();
+			return FString::Format(TEXT("Set {0}"), { InIndex });
 		}
 
 		const FText GetInputTriggerDescription(uint32 InIndex)
@@ -41,8 +41,7 @@ namespace Metasound
 
 		const FString GetInputValueName(uint32 InIndex)
 		{
-			return FText::Format(LOCTEXT("TriggerRouteValueInputName", "Value {0}"), InIndex).ToString();
-
+			return FString::Format(TEXT("Value {0}"), { InIndex });
 		}
 
 		const FText GetInputValueDescription(uint32 InIndex)
