Fix for mixer node names being inconsistent on some platforms.

diff --git a/Engine/Plugins/Runtime/Metasound/Source/MetasoundStandardNodes/Private/MetasoundMixerNode.cpp b/Engine/Plugins/Runtime/Metasound/Source/MetasoundStandardNodes/Private/MetasoundMixerNode.cpp
--- a/Engine/Plugins/Runtime/Metasound/Source/MetasoundStandardNodes/Private/MetasoundMixerNode.cpp
+++ b/Engine/Plugins/Runtime/Metasound/Source/MetasoundStandardNodes/Private/MetasoundMixerNode.cpp
@@ -234,7 +234,7 @@ namespace Metasound
 			}
 			else if (NumChannels == 2)
 			{
-				return FString::Printf(TEXT("In %i %s"), InputIndex, (ChannelIndex == 0) ? "L" : "R");
+				return FString::Printf(TEXT("In %i %s"), InputIndex, (ChannelIndex == 0) ? TEXT("L") : TEXT("R"));
 			}
 
 			return FString::Printf(TEXT("In %i, %i"), InputIndex, ChannelIndex);
