Hide 'Advanced' category from MetasSound inspector view
#rnx

diff --git a/Engine/Plugins/Runtime/Metasound/Source/MetasoundEditor/Private/MetasoundDetailCustomization.cpp b/Engine/Plugins/Runtime/Metasound/Source/MetasoundEditor/Private/MetasoundDetailCustomization.cpp
--- a/Engine/Plugins/Runtime/Metasound/Source/MetasoundEditor/Private/MetasoundDetailCustomization.cpp
+++ b/Engine/Plugins/Runtime/Metasound/Source/MetasoundEditor/Private/MetasoundDetailCustomization.cpp
@@ -94,6 +94,7 @@ namespace Metasound
 					GeneralCategoryBuilder.AddProperty(MajorVersionHandle);
 					GeneralCategoryBuilder.AddProperty(MinorVersionHandle);
 
+					DetailLayout.HideCategory("Advanced");
 					DetailLayout.HideCategory("Analysis");
 					DetailLayout.HideCategory("Attenuation");
 					DetailLayout.HideCategory("Effects");
