Use SetCustomPinIcon on SGraphPin to set trigger icon instead of overriding entire class.
#rnx

diff --git a/Engine/Plugins/Runtime/Metasound/Source/MetasoundEditor/Private/SMetasoundGraphNode.cpp b/Engine/Plugins/Runtime/Metasound/Source/MetasoundEditor/Private/SMetasoundGraphNode.cpp
--- a/Engine/Plugins/Runtime/Metasound/Source/MetasoundEditor/Private/SMetasoundGraphNode.cpp
+++ b/Engine/Plugins/Runtime/Metasound/Source/MetasoundEditor/Private/SMetasoundGraphNode.cpp
@@ -44,35 +44,6 @@
 #define LOCTEXT_NAMESPACE "MetasoundGraphNode"
 
 
-// namespace MetaSound
-// {
-// 	namespace Editor
-// 	{
-		class SMetaSoundGraphPinTrigger : public SGraphPin
-		{
-			bool bImagesCached = false;
-
-		public:
-			SLATE_BEGIN_ARGS(SGraphPin)
-			{
-			}
-
-			SLATE_END_ARGS()
-
-			void Construct(const FArguments& InArgs, UEdGraphPin* InPin)
-			{
-				SGraphPin::Construct(SGraphPin::FArguments(), InPin);
-
-				if (const ISlateStyle* MetasoundStyle = FSlateStyleRegistry::FindSlateStyle("MetaSoundStyle"))
-				{
-					CachedImg_Pin_Connected = MetasoundStyle->GetBrush(TEXT("MetasoundEditor.Graph.TriggerPin.Connected"));
-					CachedImg_Pin_Disconnected = MetasoundStyle->GetBrush(TEXT("MetasoundEditor.Graph.TriggerPin.Disconnected"));
-				}
-			}
-		};
-// 	} // namespace Editor
-// } // namespace MetaSound
-
 void SMetasoundGraphNode::Construct(const FArguments& InArgs, class UEdGraphNode* InNode)
 {
 	GraphNode = InNode;
@@ -229,7 +200,15 @@ TSharedPtr<SGraphPin> SMetasoundGraphNode::CreatePinWidget(UEdGraphPin* InPin) c
 
 		if (InPin->PinType.PinCategory == FGraphBuilder::PinCategoryTrigger)
 		{
-			return SNew(SMetaSoundGraphPinTrigger, InPin);
+			TSharedPtr<SGraphPin> TriggerPin = SNew(SGraphPin, InPin);
+
+			if (const ISlateStyle* MetasoundStyle = FSlateStyleRegistry::FindSlateStyle("MetaSoundStyle"))
+			{
+				const FSlateBrush* PinConnectedBrush = MetasoundStyle->GetBrush(TEXT("MetasoundEditor.Graph.TriggerPin.Connected"));
+				const FSlateBrush* PinDisconnectedBrush = MetasoundStyle->GetBrush(TEXT("MetasoundEditor.Graph.TriggerPin.Disconnected"));
+				TriggerPin->SetCustomPinIcon(PinConnectedBrush, PinDisconnectedBrush);
+			}
+			return TriggerPin;
 		}
 	}
 
