Load MetasoundEditor module in all uncooked builds to avoid -game errors for missing types
#rnx

diff --git a/Engine/Plugins/Runtime/Metasound/Metasound.uplugin b/Engine/Plugins/Runtime/Metasound/Metasound.uplugin
--- a/Engine/Plugins/Runtime/Metasound/Metasound.uplugin
+++ b/Engine/Plugins/Runtime/Metasound/Metasound.uplugin
@@ -42,7 +42,7 @@
 		},
 		{
 			"Name": "MetasoundEditor",
-			"Type": "Editor",
+			"Type": "UncookedOnly",
 			"LoadingPhase": "PostEngineInit"
 		}
 	],
