Metasound Class Validation & Versioning First Pass

diff --git a/Engine/Plugins/Runtime/Metasound/Source/MetasoundEditor/Private/MetasoundEditor.cpp b/Engine/Plugins/Runtime/Metasound/Source/MetasoundEditor/Private/MetasoundEditor.cpp
--- a/Engine/Plugins/Runtime/Metasound/Source/MetasoundEditor/Private/MetasoundEditor.cpp
+++ b/Engine/Plugins/Runtime/Metasound/Source/MetasoundEditor/Private/MetasoundEditor.cpp
@@ -21,14 +21,21 @@
 #include "IAudioExtensionPlugin.h"
 #include "IDetailsView.h"
 #include "Kismet2/BlueprintEditorUtils.h"
+#include "Logging/TokenizedMessage.h"
 #include "Metasound.h"
 #include "MetasoundEditorCommands.h"
 #include "MetasoundEditorGraph.h"
 #include "MetasoundEditorGraphBuilder.h"
 #include "MetasoundEditorGraphNode.h"
 #include "MetasoundEditorGraphSchema.h"
+#include "MetasoundEditorGraphValidation.h"
+#include "MetasoundEditorModule.h"
 #include "MetasoundEditorSettings.h"
 #include "MetasoundEditorTabFactory.h"
+#include "MetasoundFrontend.h"
+#include "MetasoundFrontendDocument.h"
+#include "MetasoundFrontendRegistries.h"
+#include "MetasoundFrontendSearchEngine.h"
 #include "MetasoundLog.h"
 #include "MetasoundUObjectRegistry.h"
 #include "Misc/Attribute.h"
@@ -37,13 +44,13 @@
 #include "ScopedTransaction.h"
 #include "SMetasoundPalette.h"
 #include "SNodePanel.h"
+#include "Templates/SharedPointer.h"
 #include "Widgets/Docking/SDockTab.h"
 #include "Widgets/Images/SImage.h"
 #include "Widgets/Input/SButton.h"
-#include "Widgets/SOverlay.h"
 #include "Widgets/Notifications/SNotificationList.h"
+#include "Widgets/SOverlay.h"
 #include "Widgets/SBoxPanel.h"
-#include "Templates/SharedPointer.h"
 
 // This needs to be moved to public directory
 #include "../../GraphEditor/Private/GraphActionNode.h"
@@ -64,6 +71,27 @@ namespace Metasound
 {
 	namespace Editor
 	{
+		namespace Private
+		{
+			void ValidateAndUpdate(UMetasoundEditorGraph& Graph)
+			{
+				using namespace Frontend;
+
+				FGraphValidationResults Results;
+				Graph.Validate(Results);
+
+				// Only auto-update minor upgrades
+				for (const FGraphNodeValidationResult& Result : Results.GetResults())
+				{
+					FMetasoundFrontendVersionNumber MinorUpdateVersion = Result.Node->GetMinorUpdateAvailable();
+					if (MinorUpdateVersion.IsValid())
+					{
+						Result.Node->UpdateToVersion(MinorUpdateVersion, true /* bInPropagateErrorMessages */);
+					}
+				}
+			}
+		}
+
 		static const TArray<FText> NodeSectionNames
 		{
 			LOCTEXT("NodeSectionName_Invalid", "INVALID"),
@@ -145,6 +173,55 @@ namespace Metasound
 			}
 		};
 
+		class SMetaSoundGraphPaletteItem : public SGraphPaletteItem
+		{
+		private:
+			TSharedPtr<FMetasoundGraphNodeSchemaAction> MetasoundAction;
+			bool bIsNameInvalid = false;
+
+		public:
+			SLATE_BEGIN_ARGS(SMetaSoundGraphPaletteItem)
+			{
+			}
+
+			SLATE_END_ARGS()
+
+			void Construct(const FArguments& InArgs, FCreateWidgetForActionData* const InCreateData)
+			{
+				TSharedPtr<FEdGraphSchemaAction> Action = InCreateData->Action;
+				MetasoundAction = StaticCastSharedPtr<FMetasoundGraphNodeSchemaAction>(Action);
+
+				SGraphPaletteItem::Construct(SGraphPaletteItem::FArguments(), InCreateData);
+			}
+
+		protected:
+			virtual void OnNameTextCommitted(const FText& InNewText, ETextCommit::Type InTextCommit) override
+			{
+				using namespace Frontend;
+
+				if (MetasoundAction.IsValid())
+				{
+					if (UMetasoundEditorGraphVariable* GraphVariable = MetasoundAction->GetVariable())
+					{
+						GraphVariable->SetDisplayName(InNewText);
+					}
+				}
+			}
+
+			virtual bool OnNameTextVerifyChanged(const FText& InNewText, FText& OutErrorMessage) override
+			{
+				if (MetasoundAction.IsValid())
+				{
+					if (UMetasoundEditorGraphVariable* GraphVariable = MetasoundAction->GetVariable())
+					{
+						return GraphVariable->CanRename(InNewText, OutErrorMessage);
+					}
+				}
+
+				return false;
+			}
+		};
+
 		void FEditor::RegisterTabSpawners(const TSharedRef<FTabManager>& InTabManager)
 		{
 			WorkspaceMenuCategory = InTabManager->AddLocalWorkspaceMenuCategory(LOCTEXT("WorkspaceMenu_MetasoundEditor", "MetaSound Editor"));
@@ -333,6 +410,8 @@ namespace Metasound
 
 			ExtendToolbar();
 			RegenerateMenusAndToolbars();
+
+			Private::ValidateAndUpdate(GetMetaSoundGraphChecked());
 		}
 
 		UObject* FEditor::GetMetasoundObject() const
@@ -781,17 +860,18 @@ namespace Metasound
 				return;
 			}
 
+			GetMetaSoundGraphChecked().SetPreviewID(InPreviewID);
+		}
+
+		UMetasoundEditorGraph& FEditor::GetMetaSoundGraphChecked()
+		{
 			FMetasoundAssetBase* MetasoundAsset = IMetasoundUObjectRegistry::Get().GetObjectAsAssetBase(Metasound);
 			check(MetasoundAsset);
 
 			UEdGraph* Graph = MetasoundAsset->GetGraph();
-			if (!Graph)
-			{
-				return;
-			}
-			UMetasoundEditorGraph* MetasoundGraph = CastChecked<UMetasoundEditorGraph>(MetasoundAsset->GetGraph());
+			check(Graph);
 
-			MetasoundGraph->SetPreviewID(InPreviewID);
+			return *CastChecked<UMetasoundEditorGraph>(MetasoundAsset->GetGraph());
 		}
 
 		void FEditor::ExecuteNode()
@@ -1026,6 +1106,30 @@ namespace Metasound
 				// Node Commands
 				GraphEditorCommands->MapAction(FGraphEditorCommands::Get().CreateComment,
 					FExecuteAction::CreateSP(this, &FEditor::OnCreateComment));
+
+				GraphEditorCommands->MapAction(FEditorCommands::Get().UpdateNodes,
+					FExecuteAction::CreateLambda([this]()
+					{
+						const FScopedTransaction Transaction(LOCTEXT("NodeVersionUpgrade", "Upgrade MetaSound Node(s)"));
+						check(Metasound);
+						Metasound->Modify();
+
+						UMetasoundEditorGraph& Graph = GetMetaSoundGraphChecked();
+						Graph.Modify();
+
+						const FGraphPanelSelectionSet Selection = MetasoundGraphEditor->GetSelectedNodes();
+						for (UObject* Object : Selection)
+						{
+							if (UMetasoundEditorGraphExternalNode* ExternalNode = Cast<UMetasoundEditorGraphExternalNode>(Object))
+							{
+								FMetasoundFrontendVersionNumber MajorUpdateVersion = ExternalNode->GetMajorUpdateAvailable();
+								if (MajorUpdateVersion.IsValid())
+								{
+									ExternalNode->UpdateToVersion(MajorUpdateVersion, false /* bInPropagateErrorMessages */);
+								}
+							}
+						}
+					}));
 			}
 
 			FGraphAppearanceInfo AppearanceInfo;
@@ -1184,8 +1288,8 @@ namespace Metasound
 					const FScopedTransaction Transaction(LOCTEXT("MetaSoundEditorDeleteSelectedNode", "Delete MetaSound Variable"));
 					Metasound->Modify();
 
-					UMetasoundEditorGraph* Graph = CastChecked<UMetasoundEditorGraph>(MetasoundGraphEditor->GetCurrentGraph());
-					Graph->Modify();
+					UMetasoundEditorGraph& Graph = GetMetaSoundGraphChecked();
+					Graph.Modify();
 
 					TSharedPtr<FMetasoundGraphNodeSchemaAction> ActionToDelete;
 					for (const TSharedPtr<FEdGraphSchemaAction>& Action : Actions)
@@ -1219,7 +1323,7 @@ namespace Metasound
 
 					if (ActionToDelete.IsValid())
 					{
-						DeleteInterfaceItem(ActionToDelete, Graph);
+						DeleteInterfaceItem(ActionToDelete, &Graph);
 					}
 					return;
 				}
@@ -1389,29 +1493,20 @@ namespace Metasound
 				Location = MetasoundGraphEditor->GetPasteLocation();
 			}
 
-			FMetasoundAssetBase* MetasoundAsset = IMetasoundUObjectRegistry::Get().GetObjectAsAssetBase(Metasound);
-			check(MetasoundAsset);
-
-			UEdGraph* Graph = MetasoundAsset->GetGraph();
-			if (!Graph)
-			{
-				return;
-			}
+			UMetasoundEditorGraph& Graph = GetMetaSoundGraphChecked();
 
 			const FScopedTransaction Transaction(InTransactionText);
 			Metasound->Modify();
-			Graph->Modify();
+			Graph.Modify();
 
 			// Clear the selection set (newly pasted stuff will be selected)
 			MetasoundGraphEditor->ClearSelectionSet();
 
 			TSet<UEdGraphNode*> PastedGraphNodes;
-			FEdGraphUtilities::ImportNodesFromText(Graph, NodeTextToPaste, PastedGraphNodes);
+			FEdGraphUtilities::ImportNodesFromText(&Graph, NodeTextToPaste, PastedGraphNodes);
 
 			NodeTextToPaste.Empty();
 
-			UMetasoundEditorGraph* MetasoundGraph = CastChecked<UMetasoundEditorGraph>(MetasoundAsset->GetGraph());
-
 			TArray<UEdGraphNode*> NodesToRemove;
 			for (UEdGraphNode* GraphNode : PastedGraphNodes)
 			{
@@ -1426,14 +1521,14 @@ namespace Metasound
 				}
 				else if (UMetasoundEditorGraphInputNode* InputNode = Cast<UMetasoundEditorGraphInputNode>(GraphNode))
 				{
-					if (!MetasoundGraph->ContainsInput(InputNode->Input))
+					if (!Graph.ContainsInput(InputNode->Input))
 					{
 						NodesToRemove.Add(GraphNode);
 					}
 				}
 				else if (UMetasoundEditorGraphOutputNode* OutputNode = Cast<UMetasoundEditorGraphOutputNode>(GraphNode))
 				{
-					if (!MetasoundGraph->ContainsOutput(OutputNode->Output))
+					if (!Graph.ContainsOutput(OutputNode->Output))
 					{
 						NodesToRemove.Add(GraphNode);
 					}
@@ -1448,7 +1543,7 @@ namespace Metasound
 					};
 
 					// Can only have one output reference node
-					if (MetasoundGraph->Nodes.ContainsByPredicate(NodeMatches))
+					if (Graph.Nodes.ContainsByPredicate(NodeMatches))
 					{
 						NodesToRemove.Add(GraphNode);
 					}
@@ -1463,7 +1558,7 @@ namespace Metasound
 			// in frontend graph.
 			for (UEdGraphNode* Node : NodesToRemove)
 			{
-				Graph->RemoveNode(Node);
+				Graph.RemoveNode(Node);
 				PastedGraphNodes.Remove(Node);
 			}
 
@@ -1549,20 +1644,14 @@ namespace Metasound
 
 		bool FEditor::CanPasteNodes()
 		{
-			FMetasoundAssetBase* MetasoundAsset = IMetasoundUObjectRegistry::Get().GetObjectAsAssetBase(Metasound);
-			check(MetasoundAsset);
-
-			if (UEdGraph* Graph = MetasoundAsset->GetGraph())
+			UMetasoundEditorGraph& Graph = GetMetaSoundGraphChecked();
+			FPlatformApplicationMisc::ClipboardPaste(NodeTextToPaste);
+			if (FEdGraphUtilities::CanImportNodesFromText(&Graph, NodeTextToPaste))
 			{
-				FPlatformApplicationMisc::ClipboardPaste(NodeTextToPaste);
-				if (FEdGraphUtilities::CanImportNodesFromText(Graph, NodeTextToPaste))
-				{
-					return true;
-				}
-
-				NodeTextToPaste.Empty();
+				return true;
 			}
 
+			NodeTextToPaste.Empty();
 			return false;
 		}
 
@@ -1648,9 +1737,7 @@ namespace Metasound
 			FMetasoundAssetBase* MetasoundAsset = IMetasoundUObjectRegistry::Get().GetObjectAsAssetBase(Metasound);
 			check(MetasoundAsset);
 
-			UMetasoundEditorGraph* EditorGraph = Cast<UMetasoundEditorGraph>(MetasoundAsset->GetGraph());
-
-			MetasoundAsset->GetRootGraphHandle()->IterateConstNodes([this, EditorGraph, MetasoundAsset, ActionList = &OutAllActions](const Frontend::FConstNodeHandle& Input)
+			MetasoundAsset->GetRootGraphHandle()->IterateConstNodes([this, ActionList = &OutAllActions](const Frontend::FConstNodeHandle& Input)
 			{
 				const bool bIsLiteralInput = Input->GetNodeStyle().Display.Visibility == EMetasoundFrontendNodeStyleDisplayVisibility::Hidden;
 				if (!ShowLiteralMetasoundInputsInEditorCVar && bIsLiteralInput)
@@ -1660,13 +1747,15 @@ namespace Metasound
 
 				const FText Tooltip = Input->GetDescription();
 				const FText InputName = Input->GetDisplayName();
+				UMetasoundEditorGraph& Graph = GetMetaSoundGraphChecked();
+
 				TSharedPtr<FMetasoundGraphNodeSchemaAction> NewFuncAction = MakeShared<FMetasoundGraphNodeSchemaAction>(FText::GetEmpty(), InputName, Tooltip, 1, ENodeSection::Inputs);
-				NewFuncAction->Graph = MetasoundAsset->GetGraph();
+				NewFuncAction->Graph = &Graph;
 				ActionList->AddAction(NewFuncAction);
 
 				const FGuid NodeID = Input->GetID();
 				NewFuncAction->NodeID = NodeID;
-				UMetasoundEditorGraphInput * EdGraphInput = EditorGraph->FindInput(NodeID);
+				UMetasoundEditorGraphInput * EdGraphInput = Graph.FindInput(NodeID);
 				if (ensure(EdGraphInput))
 				{
 					if (FDelegateHandle* NameChangeDelegate = NameChangeDelegateHandles.Find(NodeID))
@@ -1677,19 +1766,20 @@ namespace Metasound
 				}
 			}, EMetasoundFrontendClassType::Input);
 
-			MetasoundAsset->GetRootGraphHandle()->IterateConstNodes([this, EditorGraph, MetasoundAsset, ActionList = &OutAllActions](const Frontend::FConstNodeHandle& Output)
+			MetasoundAsset->GetRootGraphHandle()->IterateConstNodes([this, ActionList = &OutAllActions](const Frontend::FConstNodeHandle& Output)
 			{
 				const FText Tooltip = Output->GetDescription();
 				const FText OutputName = Output->GetDisplayName();
+				UMetasoundEditorGraph& Graph = GetMetaSoundGraphChecked();
 
 				TSharedPtr<FMetasoundGraphNodeSchemaAction> NewFuncAction = MakeShared<FMetasoundGraphNodeSchemaAction>(FText::GetEmpty(), OutputName, Tooltip, 1, ENodeSection::Outputs);
-				NewFuncAction->Graph = MetasoundAsset->GetGraph();
+				NewFuncAction->Graph = &Graph;
 				NewFuncAction->NodeID = Output->GetID();
 				ActionList->AddAction(NewFuncAction);
 
 				const FGuid NodeID = Output->GetID();
 				NewFuncAction->NodeID = NodeID;
-				UMetasoundEditorGraphOutput* EdGraphOutput = EditorGraph->FindOutput(NodeID);
+				UMetasoundEditorGraphOutput* EdGraphOutput = Graph.FindOutput(NodeID);
 				if (ensure(EdGraphOutput))
 				{
 					if (FDelegateHandle* NameChangeDelegate = NameChangeDelegateHandles.Find(NodeID))
@@ -1758,7 +1848,7 @@ namespace Metasound
 
 		TSharedRef<SWidget> FEditor::OnCreateWidgetForAction(FCreateWidgetForActionData* const InCreateData)
 		{
-			return SNew(SGraphPaletteItem, InCreateData);
+			return SNew(SMetaSoundGraphPaletteItem, InCreateData);
 		}
 
 		FText FEditor::GetSectionTitle(ENodeSection InSection) const
@@ -1828,13 +1918,7 @@ namespace Metasound
 				return FReply::Unhandled();
 			}
 
-			FMetasoundAssetBase* MetasoundAsset = IMetasoundUObjectRegistry::Get().GetObjectAsAssetBase(Metasound);
-			check(MetasoundAsset);
-			UMetasoundEditorGraph* Graph = Cast<UMetasoundEditorGraph>(MetasoundAsset->GetGraph());
-			if (!ensure(Graph))
-			{
-				return FReply::Unhandled();
-			}
+			UMetasoundEditorGraph& Graph = GetMetaSoundGraphChecked();
 
 			FName NameToSelect;
 			switch (static_cast<ENodeSection>(InSectionID))
@@ -1851,7 +1935,7 @@ namespace Metasound
 					Frontend::FNodeHandle NodeHandle = FGraphBuilder::AddInputNodeHandle(*Metasound, NewNodeName, DataTypeName, FText::GetEmpty());
 					if (ensure(NodeHandle->IsValid()))
 					{
-						TObjectPtr<UMetasoundEditorGraphInput> Input = Graph->FindOrAddInput(NodeHandle);
+						TObjectPtr<UMetasoundEditorGraphInput> Input = Graph.FindOrAddInput(NodeHandle);
 						if (ensure(Input))
 						{
 							FGuid NodeID = NodeHandle->GetID();
@@ -1877,7 +1961,7 @@ namespace Metasound
 					Frontend::FNodeHandle NodeHandle = FGraphBuilder::AddOutputNodeHandle(*Metasound, NewNodeName, DataTypeName, FText::GetEmpty());
 					if (ensure(NodeHandle->IsValid()))
 					{
-						TObjectPtr<UMetasoundEditorGraphOutput> Output = Graph->FindOrAddOutput(NodeHandle);
+						TObjectPtr<UMetasoundEditorGraphOutput> Output = Graph.FindOrAddOutput(NodeHandle);
 						if (ensure(Output))
 						{
 							FGuid NodeID = NodeHandle->GetID();
diff --git a/Engine/Plugins/Runtime/Metasound/Source/MetasoundEditor/Private/MetasoundEditor.h b/Engine/Plugins/Runtime/Metasound/Source/MetasoundEditor/Private/MetasoundEditor.h
--- a/Engine/Plugins/Runtime/Metasound/Source/MetasoundEditor/Private/MetasoundEditor.h
+++ b/Engine/Plugins/Runtime/Metasound/Source/MetasoundEditor/Private/MetasoundEditor.h
@@ -140,6 +140,8 @@ namespace Metasound
 			/** Called when the selection changes in the GraphEditor */
 			void OnSelectedNodesChanged(const TSet<UObject*>& NewSelection);
 
+			UMetasoundEditorGraph& GetMetaSoundGraphChecked();
+
 			/**
 			 * Called when a node's title is committed for a rename
 			 *
diff --git a/Engine/Plugins/Runtime/Metasound/Source/MetasoundEditor/Private/MetasoundEditorCommands.cpp b/Engine/Plugins/Runtime/Metasound/Source/MetasoundEditor/Private/MetasoundEditorCommands.cpp
--- a/Engine/Plugins/Runtime/Metasound/Source/MetasoundEditor/Private/MetasoundEditorCommands.cpp
+++ b/Engine/Plugins/Runtime/Metasound/Source/MetasoundEditor/Private/MetasoundEditorCommands.cpp
@@ -22,6 +22,8 @@ namespace Metasound
 
 			UI_COMMAND(EditMetasoundSettings, "MetaSound", "Edit MetaSound settings", EUserInterfaceActionType::Button, FInputChord());
 			UI_COMMAND(EditGeneralSettings, "General", "Edit general sound settings", EUserInterfaceActionType::Button, FInputChord());
+
+			UI_COMMAND(UpdateNodes, "Update Nodes", "Update selected node(s) that have an available update.", EUserInterfaceActionType::Button, FInputChord());
 		}
 	} // namespace Editor
 } // namespace Metasound
diff --git a/Engine/Plugins/Runtime/Metasound/Source/MetasoundEditor/Private/MetasoundEditorGraph.cpp b/Engine/Plugins/Runtime/Metasound/Source/MetasoundEditor/Private/MetasoundEditorGraph.cpp
--- a/Engine/Plugins/Runtime/Metasound/Source/MetasoundEditor/Private/MetasoundEditorGraph.cpp
+++ b/Engine/Plugins/Runtime/Metasound/Source/MetasoundEditor/Private/MetasoundEditorGraph.cpp
@@ -4,9 +4,11 @@
 #include "Audio/AudioParameterInterface.h"
 #include "Components/AudioComponent.h"
 #include "EdGraph/EdGraphNode.h"
+#include "Interfaces/ITargetPlatform.h"
 #include "MetasoundEditorGraphBuilder.h"
-#include "MetasoundEditorGraphNode.h"
 #include "MetasoundEditorGraphInputNodes.h"
+#include "MetasoundEditorGraphNode.h"
+#include "MetasoundEditorGraphValidation.h"
 #include "MetasoundEditorModule.h"
 #include "MetasoundUObjectRegistry.h"
 #include "ScopedTransaction.h"
@@ -35,6 +37,25 @@ TArray<UMetasoundEditorGraphNode*> UMetasoundEditorGraphVariable::GetNodes() con
 	return Nodes;
 }
 
+void UMetasoundEditorGraphVariable::SetDisplayName(const FText& InNewName)
+{
+	using namespace Metasound::Frontend;
+
+	const FText TransactionLabel = FText::Format(LOCTEXT("Rename Variable", "Rename Metasound {0}"), GetVariableLabel());
+	const FScopedTransaction Transaction(TransactionLabel);
+
+	Modify();
+	if (UMetasoundEditorGraph* Graph = Cast<UMetasoundEditorGraph>(GetOuter()))
+	{
+		Graph->GetMetasoundChecked().Modify();
+	}
+
+	FNodeHandle NodeHandle = GetNodeHandle();
+	NodeHandle->SetDisplayName(InNewName);
+
+	NameChanged.Broadcast(NodeID);
+}
+
 void UMetasoundEditorGraphVariable::SetDataType(FName InNewType)
 {
 	using namespace Metasound::Editor;
@@ -77,7 +98,7 @@ void UMetasoundEditorGraphVariable::SetDataType(FName InNewType)
 		return;
 	}
 
-	ClassName = NewNodeHandle->GetClassName();
+	ClassName = NewNodeHandle->GetClassMetadata().ClassName;
 	NodeID = NewNodeHandle->GetID();
 	TypeName = InNewType;
 
@@ -118,6 +139,35 @@ Metasound::Frontend::FConstNodeHandle UMetasoundEditorGraphVariable::GetConstNod
 	return GetNodeHandle();
 }
 
+bool UMetasoundEditorGraphVariable::CanRename(const FText& InNewName, FText& OutError) const
+{
+	using namespace Metasound::Frontend;
+
+	if (InNewName.IsEmpty())
+	{
+		OutError = FText::Format(LOCTEXT("VariableRenameInvalid_NameEmpty", "{0} cannot be empty string."), InNewName);
+		return false;
+	}
+
+	FConstNodeHandle NodeHandle = GetConstNodeHandle();
+	FConstGraphHandle GraphHandle = NodeHandle->GetOwningGraph();
+
+	bool bIsNameValid = true;
+	GraphHandle->IterateConstNodes([&](FConstNodeHandle NodeToCompare)
+	{
+		if (NodeID != NodeToCompare->GetID())
+		{
+			if (InNewName.CompareToCaseIgnored(NodeToCompare->GetDisplayName()) == 0)
+			{
+				bIsNameValid = false;
+				OutError = FText::Format(LOCTEXT("VariableRenameInvalid_NameTaken", "{0} is already in use"), InNewName);
+			}
+		}
+	}, GetClassType());
+
+		return bIsNameValid;
+}
+
 #if WITH_EDITORONLY_DATA
 void UMetasoundEditorGraphInputLiteral::PostEditUndo()
 {
@@ -186,6 +236,12 @@ Metasound::Frontend::FNodeHandle UMetasoundEditorGraphInput::AddNodeHandle(const
 	return NewNodeHandle;
 }
 
+const FText& UMetasoundEditorGraphInput::GetVariableLabel() const
+{
+	static const FText Label = LOCTEXT("VariableLabel_Input", "Input");
+	return Label;
+}
+
 #if WITH_EDITORONLY_DATA
 void UMetasoundEditorGraphInput::PostEditUndo()
 {
@@ -260,6 +316,12 @@ Metasound::Frontend::FNodeHandle UMetasoundEditorGraphOutput::AddNodeHandle(cons
 	return NewNodeHandle;
 }
 
+const FText& UMetasoundEditorGraphOutput::GetVariableLabel() const
+{
+	static const FText Label = LOCTEXT("VariableLabel_Output", "Output");
+	return Label;
+}
+
 UMetasoundEditorGraphInputNode* UMetasoundEditorGraph::CreateInputNode(Metasound::Frontend::FNodeHandle InNodeHandle, bool bInSelectNewNode)
 {
 	using namespace Metasound::Editor;
@@ -291,6 +353,19 @@ UMetasoundEditorGraphInputNode* UMetasoundEditorGraph::CreateInputNode(Metasound
 	return nullptr;	
 }
 
+Metasound::Frontend::FDocumentHandle UMetasoundEditorGraph::GetDocumentHandle() const
+{
+	return GetGraphHandle()->GetOwningDocument();
+}
+
+Metasound::Frontend::FGraphHandle UMetasoundEditorGraph::GetGraphHandle() const
+{
+	FMetasoundAssetBase* MetasoundAsset = Metasound::IMetasoundUObjectRegistry::Get().GetObjectAsAssetBase(&GetMetasoundChecked());
+	check(MetasoundAsset);
+
+	return MetasoundAsset->GetRootGraphHandle();
+}
+
 UObject* UMetasoundEditorGraph::GetMetasound() const
 {
 	return ParentMetasound;
@@ -346,7 +421,7 @@ UMetasoundEditorGraphInput* UMetasoundEditorGraph::FindOrAddInput(Metasound::Fro
 	if (UMetasoundEditorGraphInput* NewInput = NewObject<UMetasoundEditorGraphInput>(this, FName(), RF_Transactional))
 	{
 		NewInput->NodeID = NodeID;
-		NewInput->ClassName = InNodeHandle->GetClassName();
+		NewInput->ClassName = InNodeHandle->GetClassMetadata().ClassName;
 		NewInput->TypeName = TypeName;
 
 		FMetasoundFrontendLiteral DefaultLiteral = Graph->GetDefaultInput(VertexID);
@@ -401,7 +476,7 @@ UMetasoundEditorGraphOutput* UMetasoundEditorGraph::FindOrAddOutput(Metasound::F
 	if (UMetasoundEditorGraphOutput* NewOutput = NewObject<UMetasoundEditorGraphOutput>(this, FName(), RF_Transactional))
 	{
 		NewOutput->NodeID = NodeID;
-		NewOutput->ClassName = InNodeHandle->GetClassName();
+		NewOutput->ClassName = InNodeHandle->GetClassMetadata().ClassName;
 		NewOutput->TypeName = TypeName;
 		Outputs.Add(NewOutput);
 		return NewOutput;
@@ -494,4 +569,32 @@ bool UMetasoundEditorGraph::RemoveVariable(UMetasoundEditorGraphVariable& InVari
 
 	return false;
 }
+
+bool UMetasoundEditorGraph::Validate(Metasound::Editor::FGraphValidationResults& OutResults)
+{
+	using namespace Metasound::Editor;
+	using namespace Metasound::Frontend;
+
+	bool bMarkDirty = false;
+	bool bIsValid = false;
+
+	OutResults = FGraphValidationResults();
+
+	TArray<UMetasoundEditorGraphExternalNode*> ExternalNodes;
+	GetNodesOfClass<UMetasoundEditorGraphExternalNode>(ExternalNodes);
+	for (UMetasoundEditorGraphExternalNode* ExternalNode : ExternalNodes)
+	{
+		FGraphNodeValidationResult NodeResult(*ExternalNode);
+		bIsValid |= ExternalNode->Validate(NodeResult);
+		bMarkDirty |= NodeResult.bIsDirty;
+		OutResults.NodeResults.Add(NodeResult);
+	}
+
+	if (bMarkDirty)
+	{
+		MarkPackageDirty();
+	}
+
+	return bIsValid;
+}
 #undef LOCTEXT_NAMESPACE
diff --git a/Engine/Plugins/Runtime/Metasound/Source/MetasoundEditor/Private/MetasoundEditorGraphBuilder.cpp b/Engine/Plugins/Runtime/Metasound/Source/MetasoundEditor/Private/MetasoundEditorGraphBuilder.cpp
--- a/Engine/Plugins/Runtime/Metasound/Source/MetasoundEditor/Private/MetasoundEditorGraphBuilder.cpp
+++ b/Engine/Plugins/Runtime/Metasound/Source/MetasoundEditor/Private/MetasoundEditorGraphBuilder.cpp
@@ -66,7 +66,7 @@ namespace Metasound
 		UMetasoundEditorGraphExternalNode* FGraphBuilder::AddExternalNode(UObject& InMetasound, Frontend::FNodeHandle& InNodeHandle, FVector2D InLocation, bool bInSelectNewNode)
 		{
 			UMetasoundEditorGraphExternalNode* NewGraphNode = nullptr;
-			if (!ensure(InNodeHandle->GetClassType() == EMetasoundFrontendClassType::External))
+			if (!ensure(InNodeHandle->GetClassMetadata().Type == EMetasoundFrontendClassType::External))
 			{
 				return nullptr;
 			}
@@ -77,7 +77,7 @@ namespace Metasound
 			FGraphNodeCreator<UMetasoundEditorGraphExternalNode> NodeCreator(Graph);
 
 			NewGraphNode = NodeCreator.CreateNode(bInSelectNewNode);
-			NewGraphNode->ClassName = InNodeHandle->GetClassName();
+			NewGraphNode->ClassName = InNodeHandle->GetClassMetadata().ClassName;
 
 			NodeCreator.Finalize();
 			InitGraphNode(InNodeHandle, NewGraphNode, InMetasound);
@@ -87,9 +87,9 @@ namespace Metasound
 			return NewGraphNode;
 		}
 
-		UMetasoundEditorGraphExternalNode* FGraphBuilder::AddExternalNode(UObject& InMetasound, const Frontend::FNodeClassInfo& InClassInfo, FVector2D InLocation, bool bInSelectNewNode)
+		UMetasoundEditorGraphExternalNode* FGraphBuilder::AddExternalNode(UObject& InMetasound, const FMetasoundFrontendClassMetadata& InMetadata, FVector2D InLocation, bool bInSelectNewNode)
 		{
-			Frontend::FNodeHandle NodeHandle = AddExternalNodeHandle(InMetasound, InClassInfo);
+			Frontend::FNodeHandle NodeHandle = AddExternalNodeHandle(InMetasound, InMetadata);
 			return AddExternalNode(InMetasound, NodeHandle, InLocation, bInSelectNewNode);
 		}
 
@@ -98,7 +98,7 @@ namespace Metasound
 			using namespace Frontend;
 
 			UMetasoundEditorGraphOutputNode* NewGraphNode = nullptr;
-			if (!ensure(InNodeHandle->GetClassType() == EMetasoundFrontendClassType::Output))
+			if (!ensure(InNodeHandle->GetClassMetadata().Type == EMetasoundFrontendClassType::Output))
 			{
 				return nullptr;
 			}
@@ -122,20 +122,18 @@ namespace Metasound
 
 		void FGraphBuilder::InitGraphNode(Frontend::FNodeHandle& InNodeHandle, UMetasoundEditorGraphNode* NewGraphNode, UObject& InMetasound)
 		{
-			const FMetasoundFrontendNodeStyle& Style = InNodeHandle->GetNodeStyle();
-
 			NewGraphNode->CreateNewGuid();
 			NewGraphNode->SetNodeID(InNodeHandle->GetID());
 
 			RebuildNodePins(*NewGraphNode);
 		}
 
-		Frontend::FNodeHandle FGraphBuilder::AddExternalNodeHandle(UObject& InMetasound, const Frontend::FNodeClassInfo& InClassInfo)
+		Frontend::FNodeHandle FGraphBuilder::AddExternalNodeHandle(UObject& InMetasound, const FMetasoundFrontendClassMetadata& InMetadata)
 		{
 			FMetasoundAssetBase* MetasoundAsset = IMetasoundUObjectRegistry::Get().GetObjectAsAssetBase(&InMetasound);
 			check(MetasoundAsset);
 
-			return MetasoundAsset->GetRootGraphHandle()->AddNode(InClassInfo);
+			return MetasoundAsset->GetRootGraphHandle()->AddNode(InMetadata);
 		}
 
 		TArray<FString> FGraphBuilder::GetDataTypeNameCategories(const FName& InDataTypeName)
@@ -356,7 +354,7 @@ namespace Metasound
 
 		bool FGraphBuilder::IsLiteralInput(Frontend::FNodeHandle InNodeHandle)
 		{
-			if (InNodeHandle->GetClassType() == EMetasoundFrontendClassType::Input)
+			if (InNodeHandle->GetClassMetadata().Type == EMetasoundFrontendClassType::Input)
 			{
 				return InNodeHandle->GetNodeStyle().Display.Visibility == EMetasoundFrontendNodeStyleDisplayVisibility::Hidden;
 			}
@@ -401,12 +399,12 @@ namespace Metasound
 				FMetasoundFrontendLiteral PinDataTypeDefaultLiteral;
 				if (GetPinDefaultLiteral(InInputPin, PinDataTypeDefaultLiteral))
 				{
-					NewInputNode = AddInputNodeHandle(InMetasound, NewInputName, TypeName, FText::GetEmpty(), EMetasoundFrontendNodeStyleDisplayVisibility::Hidden, &PinDataTypeDefaultLiteral);
+					NewInputNode = AddInputNodeHandle(InMetasound, NewInputName, TypeName, FText::GetEmpty(), true /* bisLiteralInput */, &PinDataTypeDefaultLiteral);
 				}
 			}
 			else
 			{
-				NewInputNode = AddInputNodeHandle(InMetasound, NewInputName, TypeName, FText::GetEmpty(), EMetasoundFrontendNodeStyleDisplayVisibility::Hidden);
+				NewInputNode = AddInputNodeHandle(InMetasound, NewInputName, TypeName, FText::GetEmpty(), true /* bInIsLiteral */);
 
 				Frontend::FGraphHandle GraphHandle = InNodeHandle->GetOwningGraph();
 				const FGuid VertexID = GraphHandle->GetVertexIDForInputVertex(NewInputNode->GetNodeName());
@@ -644,7 +642,7 @@ namespace Metasound
 			return NodeHandle;
 		}
 
-		Frontend::FNodeHandle FGraphBuilder::AddInputNodeHandle(UObject& InMetasound, const FString& InName, const FName InTypeName, const FText& InToolTip, EMetasoundFrontendNodeStyleDisplayVisibility InVisibility, const FMetasoundFrontendLiteral* InDefaultValue)
+		Frontend::FNodeHandle FGraphBuilder::AddInputNodeHandle(UObject& InMetasound, const FString& InName, const FName InTypeName, const FText& InToolTip, bool bInIsLiteral, const FMetasoundFrontendLiteral* InDefaultValue)
 		{
 			FMetasoundAssetBase* MetasoundAsset = IMetasoundUObjectRegistry::Get().GetObjectAsAssetBase(&InMetasound);
 			check(MetasoundAsset);
@@ -661,7 +659,7 @@ namespace Metasound
 			Frontend::FNodeHandle NodeHandle = GraphHandle->AddInputVertex(Description);
 
 			FMetasoundFrontendNodeStyle Style = NodeHandle->GetNodeStyle();
-			Style.Display.Visibility = InVisibility;
+			Style.Display.Visibility = bInIsLiteral ? EMetasoundFrontendNodeStyleDisplayVisibility::Hidden : EMetasoundFrontendNodeStyleDisplayVisibility::Visible;
 			NodeHandle->SetNodeStyle(Style);
 
 			if (ensure(NodeHandle->IsValid()))
@@ -702,7 +700,7 @@ namespace Metasound
 
 		UMetasoundEditorGraphNode* FGraphBuilder::AddNode(UObject& InMetasound, Frontend::FNodeHandle InNodeHandle, FVector2D InLocation, bool bInSelectNewNode)
 		{
-			switch (InNodeHandle->GetClassType())
+			switch (InNodeHandle->GetClassMetadata().Type)
 			{
 				case EMetasoundFrontendClassType::Input:
 				{
@@ -860,7 +858,7 @@ namespace Metasound
 			TArray<FNodeHandle> NodeHandles = GraphHandle->GetNodes();
 			for (FNodeHandle& NodeHandle : NodeHandles)
 			{
-				const EMetasoundFrontendClassType NodeType = NodeHandle->GetClassType();
+				const EMetasoundFrontendClassType NodeType = NodeHandle->GetClassMetadata().Type;
 				FVector2D NewLocation;
 				if (NodeType == EMetasoundFrontendClassType::Input)
 				{
@@ -906,7 +904,7 @@ namespace Metasound
 							if (OutputHandle->IsValid())
 							{
 								const FMetasoundFrontendNodeStyle& Style = OutputHandle->GetOwningNode()->GetNodeStyle();
-								if (Style.Display.Visibility == EMetasoundFrontendNodeStyleDisplayVisibility::Visible)
+								if (!IsLiteralInput(OutputHandle->GetOwningNode()))
 								{
 									UEdGraphNode* OutputGraphNode = NewIdNodeMap.FindChecked(OutputHandle->GetOwningNodeID()).GraphNode;
 									UEdGraphPin* OutputPin = OutputGraphNode->FindPinChecked(OutputHandle->GetName(), EEdGraphPinDirection::EGPD_Output);
@@ -982,6 +980,8 @@ namespace Metasound
 				return false;
 			}
 
+			bool bWasErroredNode = InNode.ErrorType == EMessageSeverity::Error;
+
 			// If node isn't a MetasoundEditorGraphNode, just remove and return (ex. comment nodes)
 			UMetasoundEditorGraphNode* Node = Cast<UMetasoundEditorGraphNode>(&InNode);
 			UMetasoundEditorGraph* Graph = CastChecked<UMetasoundEditorGraph>(InNode.GetGraph());
@@ -991,7 +991,6 @@ namespace Metasound
 				return true;
 			}
 
-
 			FNodeHandle NodeHandle = Node->GetNodeHandle();
 			if (!ensure(NodeHandle->IsValid()))
 			{
@@ -1005,21 +1004,18 @@ namespace Metasound
 			// Remove connects only to pins associated with this EdGraph node
 			// only (Iterate pins and not Frontend representation to preserve
 			// other input/output EditorGraph reference node associations)
-			Node->IteratePins([](UEdGraphPin* Pin, int32 Index)
+			Node->IteratePins([](UEdGraphPin& Pin, int32 Index)
 			{
-				if (Pin)
-				{
-					// Only add literal inputs for output pins as adding when disconnecting
-					// inputs would immediately orphan them on EditorGraph node removal below.
-					const bool bAddLiteralInputs = Pin->Direction == EGPD_Output;
-					FGraphBuilder::DisconnectPin(*Pin, bAddLiteralInputs);
-				}
+				// Only add literal inputs for output pins as adding when disconnecting
+				// inputs would immediately orphan them on EditorGraph node removal below.
+				const bool bAddLiteralInputs = Pin.Direction == EGPD_Output;
+				FGraphBuilder::DisconnectPin(Pin, bAddLiteralInputs);
 			});
 
 			Frontend::FGraphHandle GraphHandle = NodeHandle->GetOwningGraph();
 			if (GraphHandle->IsValid())
 			{
-				switch (NodeHandle->GetClassType())
+				switch (NodeHandle->GetClassMetadata().Type)
 				{
 					case EMetasoundFrontendClassType::Output:
 					case EMetasoundFrontendClassType::Input:
@@ -1039,13 +1035,25 @@ namespace Metasound
 					{
 						static_assert(static_cast<int32>(EMetasoundFrontendClassType::Invalid) == 4, "Possible missing MetasoundFrontendClassType switch case coverage.");
 
-						ensure(GraphHandle->RemoveNode(*NodeHandle));
+						if (ensure(GraphHandle->RemoveNode(*NodeHandle)))
+						{
+							GraphHandle->GetOwningDocument()->SynchronizeDependencies();
+						}
 					}
 					break;
 				}
 			}
 
-			return ensure(Graph->RemoveNode(&InNode));
+			const bool bSuccess = ensure(Graph->RemoveNode(&InNode));
+
+			// Sync the graph after nodes containing errors are deleted to ensure that
+			// the graph is not malformed once all errors are addressed by the user.
+			if (bSuccess && bWasErroredNode)
+			{
+				FGraphBuilder::SynchronizeGraph(Graph->GetMetasoundChecked());
+			}
+
+			return bSuccess;
 		}
 
 		void FGraphBuilder::RebuildNodePins(UMetasoundEditorGraphNode& InGraphNode)
@@ -1063,7 +1071,7 @@ namespace Metasound
 			// Only add input pins of the node is not an input node. Input nodes
 			// have their input pins hidden because they cannot be connected internal
 			// to the graph.
-			if (EMetasoundFrontendClassType::Input != NodeHandle->GetClassType())
+			if (EMetasoundFrontendClassType::Input != NodeHandle->GetClassMetadata().Type)
 			{
 				TArray<FInputHandle> InputHandles = NodeHandle->GetInputs();
 				InputHandles = NodeHandle->GetInputStyle().SortDefaults(InputHandles);
@@ -1076,7 +1084,7 @@ namespace Metasound
 			// Only add output pins of the node is not an output node. Output nodes
 			// have their output pins hidden because they cannot be connected internal
 			// to the graph.
-			if (EMetasoundFrontendClassType::Output != NodeHandle->GetClassType())
+			if (EMetasoundFrontendClassType::Output != NodeHandle->GetClassMetadata().Type)
 			{
 				TArray<FOutputHandle> OutputHandles = NodeHandle->GetOutputs();
 				OutputHandles = NodeHandle->GetOutputStyle().SortDefaults(OutputHandles);
@@ -1087,6 +1095,21 @@ namespace Metasound
 			}
 		}
 
+		void FGraphBuilder::RefreshPinMetadata(UEdGraphPin& InPin, const FMetasoundFrontendVertexMetadata& InMetadata)
+		{
+			InPin.PinToolTip = InMetadata.Description.ToString();
+			InPin.bAdvancedView = InMetadata.bIsAdvancedDisplay;
+			if (InPin.bAdvancedView)
+			{
+				UEdGraphNode* OwningNode = InPin.GetOwningNode();
+				check(OwningNode);
+				if (OwningNode->AdvancedPinDisplay == ENodeAdvancedPins::NoPins)
+				{
+					OwningNode->AdvancedPinDisplay = ENodeAdvancedPins::Hidden;
+				}
+			}
+		}
+
 		bool FGraphBuilder::IsMatchingInputHandleAndPin(const Frontend::FInputHandle& InInputHandle, const UEdGraphPin& InEditorPin)
 		{
 			Frontend::FInputHandle PinInputHandle = GetInputHandleFromPin(&InEditorPin);
@@ -1141,21 +1164,12 @@ namespace Metasound
 			UEdGraphPin* NewPin = InEditorNode.CreatePin(EGPD_Input, PinType, FName(*InInputHandle->GetName()));
 			if (ensure(NewPin))
 			{
-				NewPin->PinToolTip = InInputHandle->GetTooltip().ToString();
-				NewPin->bAdvancedView = InInputHandle->GetMetadata().bIsAdvancedDisplay;
-				if (NewPin->bAdvancedView)
-				{
-					if (InEditorNode.AdvancedPinDisplay == ENodeAdvancedPins::NoPins)
-					{
-						InEditorNode.AdvancedPinDisplay = ENodeAdvancedPins::Hidden;
-					}
-				}
+				RefreshPinMetadata(*NewPin, InInputHandle->GetMetadata());
 
 				FNodeHandle NodeHandle = InInputHandle->GetOwningNode();
 				FGraphBuilder::AddOrUpdateLiteralInput(InEditorNode.GetMetasoundChecked(), NodeHandle, *NewPin);
 			}
 
-
 			return NewPin;
 		}
 
@@ -1211,7 +1225,7 @@ namespace Metasound
 
 				FNodeHandle Node = EditorNode->GetNodeHandle();
 
-				if (EMetasoundFrontendClassType::Input == Node->GetClassType())
+				if (EMetasoundFrontendClassType::Input == Node->GetClassMetadata().Type)
 				{
 					// Skip this node if it is an input node. Input pins on input 
 					// nodes are not connected internal to the graph.
@@ -1249,8 +1263,7 @@ namespace Metasound
 					if (OutputHandle->IsValid())
 					{
 						FNodeHandle InputNodeHandle = OutputHandle->GetOwningNode();
-						const FMetasoundFrontendNodeStyle& InputNodeStyle = InputNodeHandle->GetNodeStyle();
-						bShowConnectionInEditor = InputNodeStyle.Display.Visibility == EMetasoundFrontendNodeStyleDisplayVisibility::Visible;
+						bShowConnectionInEditor = !IsLiteralInput(InputNodeHandle);
 						bIsConnectedToLiteralInput = !bShowConnectionInEditor;
 					}
 					else
@@ -1352,6 +1365,15 @@ namespace Metasound
 			TArray<UMetasoundEditorGraphNode*> EditorNodes;
 			EditorGraph->GetNodesOfClass(EditorNodes);
 
+			// Do not synchronize with errors present as the graph is expected to be malformed.
+			for (UMetasoundEditorGraphNode* Node : EditorNodes)
+			{
+				if (Node->ErrorType == EMessageSeverity::Error)
+				{
+					return true;
+				}
+			}
+
 			TMap<FGuid, UMetasoundEditorGraphNode*> EditorNodesByEdNodeGuid;
 			for (UMetasoundEditorGraphNode* Node : EditorNodes)
 			{
@@ -1540,7 +1562,7 @@ namespace Metasound
 			// Only add input pins of the node is not an input node. Input nodes
 			// have their input pins hidden because they cannot be connected internal
 			// to the graph.
-			if (EMetasoundFrontendClassType::Input != InNode->GetClassType())
+			if (EMetasoundFrontendClassType::Input != InNode->GetClassMetadata().Type)
 			{
 				bIsNodeDirty |= !InputHandles.IsEmpty();
 				InputHandles = InNode->GetOutputStyle().SortDefaults(InputHandles);
@@ -1557,7 +1579,7 @@ namespace Metasound
 			// Only add output pins of the node is not an output node. Output nodes
 			// have their output pins hidden because they cannot be connected internal
 			// to the graph.
-			if (EMetasoundFrontendClassType::Output != InNode->GetClassType())
+			if (EMetasoundFrontendClassType::Output != InNode->GetClassMetadata().Type)
 			{
 				bIsNodeDirty |= !OutputHandles.IsEmpty();
 				OutputHandles = InNode->GetOutputStyle().SortDefaults(OutputHandles);
diff --git a/Engine/Plugins/Runtime/Metasound/Source/MetasoundEditor/Private/MetasoundEditorGraphBuilder.h b/Engine/Plugins/Runtime/Metasound/Source/MetasoundEditor/Private/MetasoundEditorGraphBuilder.h
--- a/Engine/Plugins/Runtime/Metasound/Source/MetasoundEditor/Private/MetasoundEditorGraphBuilder.h
+++ b/Engine/Plugins/Runtime/Metasound/Source/MetasoundEditor/Private/MetasoundEditorGraphBuilder.h
@@ -72,7 +72,7 @@ namespace Metasound
 				const FString& InName,
 				const FName InTypeName,
 				const FText& InToolTip,
-				EMetasoundFrontendNodeStyleDisplayVisibility Visibility = EMetasoundFrontendNodeStyleDisplayVisibility::Visible,
+				bool bInIsLiteral = false,
 				const FMetasoundFrontendLiteral* InDefaultValue = nullptr);
 
 			// Adds a corresponding UMetasoundEditorGraphExternalNode for the provided node handle.
@@ -80,13 +80,13 @@ namespace Metasound
 
 			// Adds an externally-defined node with the given class info to both the editor and document graphs.
 			// Generates analogous FNodeHandle.
-			static UMetasoundEditorGraphExternalNode* AddExternalNode(UObject& InMetasound, const Frontend::FNodeClassInfo& InClassInfo, FVector2D InLocation, bool bInSelectNewNode = true);
+			static UMetasoundEditorGraphExternalNode* AddExternalNode(UObject& InMetasound, const FMetasoundFrontendClassMetadata& InMetadata, FVector2D InLocation, bool bInSelectNewNode = true);
 
 			// Synchronizes node location data
 			static void SynchronizeNodeLocation(FVector2D InLocation, Frontend::FNodeHandle InNodeHandle, UMetasoundEditorGraphNode& InNode);
 
 			// Generates FNodeHandle for the given external node data. Does not bind or create EdGraph representation of given node.
-			static Frontend::FNodeHandle AddExternalNodeHandle(UObject& InMetasound, const Frontend::FNodeClassInfo& InClassInfo);
+			static Frontend::FNodeHandle AddExternalNodeHandle(UObject& InMetasound, const FMetasoundFrontendClassMetadata& InMetadata);
 
 			// Adds an output node to the editor graph that corresponds to the provided node handle.
 			static UMetasoundEditorGraphOutputNode* AddOutputNode(UObject& InMetasound, Frontend::FNodeHandle& InNodeHandle, FVector2D InLocation, bool bInSelectNewNode = true);
@@ -157,6 +157,9 @@ namespace Metasound
 			// Adds an Output UEdGraphPin to a UMetasoundEditorGraphNode
 			static UEdGraphPin* AddPinToNode(UMetasoundEditorGraphNode& InEditorNode, Frontend::FOutputHandle InOutputHandle);
 
+			// Refreshes pin state from class FrontendClassVertexMetadata
+			static void RefreshPinMetadata(UEdGraphPin& InPin, const FMetasoundFrontendVertexMetadata& InMetadata);
+
 			// Adds and removes nodes, pins and connections so that the UEdGraph of the Metasound matches the FMetasoundFrontendDocumentModel
 			//
 			// @return True if the UEdGraph was altered, false otherwise. 
@@ -188,7 +191,7 @@ namespace Metasound
 			//
 			// Functions accept a UEdGraphNode* and return a TSet<UEdGraphNode*> which
 			// represent all the children of the node. 
-			using FDepthFirstVisitFunction = TFunctionRef<TSet<UEdGraphNode*> (UEdGraphNode*)>; 
+			using FDepthFirstVisitFunction = TFunctionRef<TSet<UEdGraphNode*> (UEdGraphNode*)>;
 
 			// Traverse depth first starting at the InInitialNode and calling the InVisitFunction
 			// for each node. 
diff --git a/Engine/Plugins/Runtime/Metasound/Source/MetasoundEditor/Private/MetasoundEditorGraphConnectionDrawingPolicy.cpp b/Engine/Plugins/Runtime/Metasound/Source/MetasoundEditor/Private/MetasoundEditorGraphConnectionDrawingPolicy.cpp
--- a/Engine/Plugins/Runtime/Metasound/Source/MetasoundEditor/Private/MetasoundEditorGraphConnectionDrawingPolicy.cpp
+++ b/Engine/Plugins/Runtime/Metasound/Source/MetasoundEditor/Private/MetasoundEditorGraphConnectionDrawingPolicy.cpp
@@ -48,7 +48,15 @@ namespace Metasound
 			OutParams.AssociatedPin1 = InputPin;
 			OutParams.AssociatedPin2 = OutputPin;
 
-			OutParams.WireColor = FGraphBuilder::GetPinCategoryColor(OutputPin->PinType);
+			if (OutputPin->bOrphanedPin || InputPin->bOrphanedPin)
+			{
+				OutParams.WireColor = FLinearColor::Red;
+			}
+			else
+			{
+				OutParams.WireColor = FGraphBuilder::GetPinCategoryColor(OutputPin->PinType);
+			}
+
 			bool bExecuted = false;
 
 			// Run through the predecessors, and on
diff --git a/Engine/Plugins/Runtime/Metasound/Source/MetasoundEditor/Private/MetasoundEditorGraphNode.cpp b/Engine/Plugins/Runtime/Metasound/Source/MetasoundEditor/Private/MetasoundEditorGraphNode.cpp
--- a/Engine/Plugins/Runtime/Metasound/Source/MetasoundEditor/Private/MetasoundEditorGraphNode.cpp
+++ b/Engine/Plugins/Runtime/Metasound/Source/MetasoundEditor/Private/MetasoundEditorGraphNode.cpp
@@ -7,14 +7,18 @@
 #include "Engine/Font.h"
 #include "Framework/Commands/GenericCommands.h"
 #include "GraphEditorActions.h"
+#include "Logging/TokenizedMessage.h"
 #include "Metasound.h"
+#include "MetasoundEditorCommands.h"
 #include "MetasoundEditorGraph.h"
 #include "MetasoundEditorGraphBuilder.h"
 #include "MetasoundEditorGraphSchema.h"
-#include "MetasoundEditorCommands.h"
+#include "MetasoundEditorGraphValidation.h"
 #include "MetasoundEditorModule.h"
 #include "MetasoundFrontend.h"
 #include "MetasoundFrontendRegistries.h"
+#include "MetasoundFrontendSearchEngine.h"
+#include "MetasoundNodeInterface.h"
 #include "MetasoundUObjectRegistry.h"
 #include "ScopedTransaction.h"
 #include "ToolMenus.h"
@@ -22,6 +26,23 @@
 #define LOCTEXT_NAMESPACE "MetaSoundEditor"
 
 
+namespace Metasound
+{
+	namespace Editor
+	{
+		namespace GraphNodePrivate
+		{
+			void SetGraphNodeMessage(UEdGraphNode& InNode, EMessageSeverity::Type InSeverity, const FString& InMessage)
+			{
+				InNode.bHasCompilerMessage = true;
+				InNode.ErrorMsg = InMessage;
+				InNode.ErrorType = InSeverity;
+			}
+		} // namespace GraphNodePrivate
+	} // namespace Editor
+} // namespace Metasound
+
+
 UMetasoundEditorGraphNode::UMetasoundEditorGraphNode(const FObjectInitializer& ObjectInitializer)
 	: Super(ObjectInitializer)
 {
@@ -124,13 +145,13 @@ Metasound::Frontend::FNodeHandle UMetasoundEditorGraphNode::GetNodeHandle() cons
 	return GetRootGraphHandle()->GetNodeWithID(NodeID);
 }
 
-void UMetasoundEditorGraphNode::IteratePins(TUniqueFunction<void(UEdGraphPin* /* Pin */, int32 /* Index */)> InFunc, EEdGraphPinDirection InPinDirection)
+void UMetasoundEditorGraphNode::IteratePins(TUniqueFunction<void(UEdGraphPin& /* Pin */, int32 /* Index */)> InFunc, EEdGraphPinDirection InPinDirection)
 {
 	for (int32 PinIndex = 0; PinIndex < Pins.Num(); PinIndex++)
 	{
 		if (InPinDirection == EGPD_MAX || Pins[PinIndex]->Direction == InPinDirection)
 		{
-			InFunc(Pins[PinIndex], PinIndex);
+			InFunc(*Pins[PinIndex], PinIndex);
 		}
 	}
 }
@@ -453,4 +474,341 @@ void UMetasoundEditorGraphOutputNode::SetNodeID(FGuid InNodeID)
 		Output->NodeID = InNodeID;
 	}
 }
+
+bool UMetasoundEditorGraphExternalNode::RefreshPinMetadata()
+{
+	using namespace Metasound::Frontend;
+
+	bool bModified = false;
+
+	FConstNodeHandle NodeHandle = GetConstNodeHandle();
+	const FMetasoundFrontendClassInterface& ClassInterface = NodeHandle->GetClassInterface();
+	TArray<FMetasoundFrontendClassInput> Inputs = ClassInterface.Inputs;
+	TArray<FMetasoundFrontendClassOutput> Outputs = ClassInterface.Outputs;
+
+	for (UEdGraphPin* Pin : Pins)
+	{
+		check(Pin);
+
+		if (Pin->Direction == EGPD_Input)
+		{
+			// Remove valid class inputs on iteration to avoid additional
+			// incurred look-up cost on subsequent tooltip update.
+			auto RemoveValidEntry = [&](const FMetasoundFrontendClassInput& Entry)
+			{
+				const FString PinNameString = Pin->PinName.ToString();
+				if (Entry.Name == PinNameString)
+				{
+					Metasound::Editor::FGraphBuilder::RefreshPinMetadata(*Pin, Entry.Metadata);
+					return true;
+				}
+
+				return false;
+			};
+			bModified |= Inputs.RemoveAllSwap(RemoveValidEntry) > 0;
+		}
+
+		if (Pin->Direction == EGPD_Output)
+		{
+			// Remove valid class inputs on iteration to avoid additional
+			// incurred look-up cost on subsequent tooltip update.
+			auto RemoveValidEntry = [&](const FMetasoundFrontendClassOutput& Entry)
+			{
+				const FString PinNameString = Pin->PinName.ToString();
+				if (Entry.Name == PinNameString)
+				{
+					Metasound::Editor::FGraphBuilder::RefreshPinMetadata(*Pin, Entry.Metadata);
+					return true;
+				}
+
+				return false;
+			};
+			bModified |= Outputs.RemoveAllSwap(RemoveValidEntry) > 0;
+		}
+	}
+
+	return bModified;
+}
+
+FMetasoundFrontendVersionNumber UMetasoundEditorGraphExternalNode::GetMajorUpdateAvailable() const
+{
+	using namespace Metasound::Frontend;
+
+	const FMetasoundFrontendClassMetadata& Metadata = GetConstNodeHandle()->GetClassMetadata();
+	const FMetasoundFrontendVersionNumber& CurrentVersion = Metadata.Version;
+	Metasound::FNodeClassName NodeClassName = Metadata.ClassName.ToNodeClassName();
+
+	FMetasoundFrontendClass ClassWithHighestVersion;
+	if (ISearchEngine::Get().FindClassWithHighestVersion(NodeClassName, ClassWithHighestVersion))
+	{
+		if (ClassWithHighestVersion.Metadata.Version.Major > CurrentVersion.Major)
+		{
+			return ClassWithHighestVersion.Metadata.Version;
+		}
+	}
+
+	return FMetasoundFrontendVersionNumber::GetInvalid();
+}
+
+FMetasoundFrontendVersionNumber UMetasoundEditorGraphExternalNode::GetMinorUpdateAvailable() const
+{
+	using namespace Metasound::Frontend;
+
+	const FMetasoundFrontendClassMetadata& Metadata = GetConstNodeHandle()->GetClassMetadata();
+	const FMetasoundFrontendVersionNumber& CurrentVersion = Metadata.Version;
+	Metasound::FNodeClassName NodeClassName = Metadata.ClassName.ToNodeClassName();
+
+	FMetasoundFrontendClass ClassWithHighestVersion;
+	if (ISearchEngine::Get().FindClassWithMajorVersion(NodeClassName, CurrentVersion.Major, ClassWithHighestVersion))
+	{
+		if (ClassWithHighestVersion.Metadata.Version.Minor > CurrentVersion.Minor)
+		{
+			return ClassWithHighestVersion.Metadata.Version;
+		}
+	}
+
+	return FMetasoundFrontendVersionNumber::GetInvalid();
+}
+
+UMetasoundEditorGraphExternalNode* UMetasoundEditorGraphExternalNode::UpdateToVersion(const FMetasoundFrontendVersionNumber& InNewVersion, bool bInPropagateErrorMessages)
+{
+	using namespace Metasound::Editor;
+	using namespace Metasound::Frontend;
+
+	const FMetasoundFrontendClassMetadata& Metadata = GetConstNodeHandle()->GetClassMetadata();
+	const TArray<FMetasoundFrontendClass> SortedVersions = ISearchEngine::Get().FindClassesWithName(Metadata.ClassName.ToNodeClassName(), true /* bInSortByVersion */);
+
+	auto IsClassOfNewVersion = [InNewVersion](const FMetasoundFrontendClass& SortedClass)
+	{
+		return SortedClass.Metadata.Version == InNewVersion;
+	};
+
+	if (!ensure(SortedVersions.ContainsByPredicate(IsClassOfNewVersion)))
+	{
+		return this;
+	}
+
+	struct FConnectionInfo
+	{
+		TArray<UEdGraphPin*> PinsLinkedTo;
+		FName DataType;
+		FString DefaultValue;
+	};
+
+	// Cache pin connections by name to try so they can be
+	// hooked back up after swapping to the new class version.
+	TMap<FName, FConnectionInfo> InputConnections;
+	IteratePins([Connections = &InputConnections](UEdGraphPin& Pin, int32 /* Index */)
+	{
+		FConstInputHandle InputHandle = FGraphBuilder::GetConstInputHandleFromPin(&Pin);
+		Connections->Add(Pin.GetFName(), FConnectionInfo
+		{
+			Pin.LinkedTo,
+			InputHandle->GetDataType(),
+			Pin.DefaultValue
+		});
+	}, EGPD_Input);
+
+	TMap<FName, FConnectionInfo> OutputConnections;
+	IteratePins([Connections = &OutputConnections](UEdGraphPin& Pin, int32 /* Index */)
+	{
+		FConstOutputHandle OutputHandle = FGraphBuilder::GetConstOutputHandleFromPin(&Pin);
+		Connections->Add(Pin.GetFName(), FConnectionInfo
+		{
+			Pin.LinkedTo,
+			OutputHandle->GetDataType()
+		});
+	}, EGPD_Output);
+
+	FNodeHandle NodeHandle = GetNodeHandle();
+	FMetasoundFrontendVersionNumber InitVersion = NodeHandle->GetClassMetadata().Version;
+
+	// Only update the version on the Metadata. Name and version are the only fields that pertain
+	// to building valid class data (via the GetRegistryKey look-up).  All other Metadata is
+	// superfluous for this use case.
+	FMetasoundFrontendClassMetadata NewMetadata = NodeHandle->GetClassMetadata();
+	NewMetadata.Version = InNewVersion;
+
+	const FVector2D NodePos(NodePosX, NodePosY);
+	if (!ensureAlways(FGraphBuilder::DeleteNode(*this)))
+	{
+		return this;
+	}
+
+	Modify();
+	UObject& Metasound = GetMetasoundChecked();
+	UMetasoundEditorGraphExternalNode* ReplacementNode = FGraphBuilder::AddExternalNode(Metasound, NewMetadata, NodePos);
+	if (!ensureAlways(ReplacementNode))
+	{
+		return this;
+	}
+
+	ReplacementNode->Modify();
+	ReplacementNode->IteratePins([&](UEdGraphPin& Pin, int32 /* Index */)
+	{
+		if (FConnectionInfo* ConnectionInfo = InputConnections.Find(Pin.GetFName()))
+		{
+			FInputHandle InputHandle = FGraphBuilder::GetInputHandleFromPin(&Pin);
+			if (ConnectionInfo->DataType == InputHandle->GetDataType())
+			{
+				if (ConnectionInfo->PinsLinkedTo.IsEmpty())
+				{
+					Pin.DefaultValue = ConnectionInfo->DefaultValue;
+					FGraphBuilder::AddOrUpdateLiteralInput(Metasound, InputHandle->GetOwningNode(), Pin, true /* bForcePinValueAsDefault */);
+				}
+				else
+				{
+					for (UEdGraphPin* OutputPin : ConnectionInfo->PinsLinkedTo)
+					{
+						FGraphBuilder::ConnectNodes(Pin, *OutputPin, true /* bInConnectEdPins */);
+					}
+				}
+			}
+		}
+	}, EGPD_Input);
+
+	ReplacementNode->IteratePins([&](UEdGraphPin& Pin, int32 /* Index */)
+	{
+		if (FConnectionInfo* ConnectionInfo = OutputConnections.Find(Pin.GetFName()))
+		{
+			FConstOutputHandle OutputHandle = FGraphBuilder::GetConstOutputHandleFromPin(&Pin);
+			if (ConnectionInfo->DataType == OutputHandle->GetDataType())
+			{
+				for (UEdGraphPin* InputPin : ConnectionInfo->PinsLinkedTo)
+				{
+					FGraphBuilder::ConnectNodes(*InputPin, Pin, true /* bInConnectEdPins */);
+				}
+			}
+		}
+	}, EGPD_Output);
+
+	ReplacementNode->NodeUpgradeMessage = FText::Format(LOCTEXT("NodeVersionUpgradeMessageFormat", "Class upgraded from {0} to {1}"),
+		FText::FromString(*InitVersion.ToString()),
+		FText::FromString(*InNewVersion.ToString())
+	);
+
+	if (ReplacementNode != this)
+	{
+		if (bInPropagateErrorMessages)
+		{
+			ReplacementNode->bHasCompilerMessage = bHasCompilerMessage;
+			ReplacementNode->ErrorMsg = ErrorMsg;
+			ReplacementNode->ErrorType = ErrorType;
+		}
+	}
+
+	return ReplacementNode;
+}
+
+bool UMetasoundEditorGraphExternalNode::Validate(Metasound::Editor::FGraphNodeValidationResult& OutResult)
+{
+	using namespace Metasound::Editor;
+	using namespace Metasound::Frontend;
+
+	OutResult = FGraphNodeValidationResult(*this);
+
+	FNodeHandle NodeHandle = GetNodeHandle();
+
+	if (bHasCompilerMessage)
+	{
+		bHasCompilerMessage = false;
+		ErrorMsg.Reset();
+		OutResult.bIsDirty = true;
+	}
+
+	if (!ensure(NodeHandle->IsValid()))
+	{
+		GraphNodePrivate::SetGraphNodeMessage(*this, EMessageSeverity::Error, TEXT("Node is invalid (Frontend node not found)"));
+		OutResult.bIsInvalid = true;
+		OutResult.bIsDirty = true;
+		return true;
+	}
+
+	const FMetasoundFrontendClassMetadata& Metadata = NodeHandle->GetClassMetadata();
+	if (!ensure(Metadata.Type == EMetasoundFrontendClassType::External))
+	{
+		GraphNodePrivate::SetGraphNodeMessage(*this, EMessageSeverity::Error, TEXT("Node is of type 'External' but Frontend counterpart is not natively-defined."));
+		OutResult.bIsDirty = true;
+		OutResult.bIsInvalid = true;
+		return true;
+	}
+
+	const Metasound::FNodeClassName NodeClassName = Metadata.ClassName.ToNodeClassName();
+	const TArray<FMetasoundFrontendClass> SortedClasses = ISearchEngine::Get().FindClassesWithName(NodeClassName, true /* bSortByVersion */);
+
+	if (SortedClasses.IsEmpty())
+	{
+		OutResult.MissingClass = NodeClassName.GetFullName();
+		const FString NewErrorMsg = FString::Format(TEXT("Missing Class '{0}': {1}"),
+			{
+				*Metadata.ClassName.ToString(),
+				*Metadata.PromptIfMissing.ToString()
+			});
+
+		for (UEdGraphPin* Pin : Pins)
+		{
+			Pin->bOrphanedPin = true;
+		}
+
+		GraphNodePrivate::SetGraphNodeMessage(*this, EMessageSeverity::Error, NewErrorMsg);
+		OutResult.bIsDirty = true;
+		OutResult.bIsInvalid = true;
+		return true;
+	}
+
+	for (UEdGraphPin* Pin : Pins)
+	{
+		if (Pin->bOrphanedPin)
+		{
+			Pin->bOrphanedPin = false;
+			OutResult.bIsDirty = true;
+		}
+	}
+
+	if (SortedClasses[0].Metadata.Version.Major > Metadata.Version.Major)
+	{
+		const FMetasoundFrontendVersionNumber& CurrentVersion = NodeHandle->GetClassMetadata().Version;
+
+		FMetasoundFrontendClass HighestMinorVersionClass;
+		FString NodeMsg;
+		EMessageSeverity::Type Severity;
+
+		// Only compare major version as minor can be incorrect and graph
+		// will still build and divert to using the new upgraded version.
+		if (ISearchEngine::Get().FindClassWithMajorVersion(NodeClassName, CurrentVersion.Major, HighestMinorVersionClass))
+		{
+			NodeMsg = FString::Format(TEXT("Node class '{0}' out-of-date: Eligible for upgrade to {1}"),
+			{
+				*Metadata.ClassName.ToString(),
+				*SortedClasses[0].Metadata.Version.ToString()
+			});
+			Severity = EMessageSeverity::Warning;
+		}
+		else
+		{
+			NodeMsg = FString::Format(TEXT("Node class '{0} ({1})' is missing.  Eligible for upgrade to {2}"),
+			{
+				*Metadata.ClassName.ToString(),
+				*Metadata.Version.ToString(),
+				*SortedClasses[0].Metadata.Version.ToString()
+			});
+			Severity = EMessageSeverity::Error;
+		}
+
+		GraphNodePrivate::SetGraphNodeMessage(*this, Severity, NodeMsg);
+		OutResult.bIsDirty = true;
+	}
+
+	auto IsMinorUpgradeVersion = [&](const FMetasoundFrontendVersionNumber& Version)
+	{
+		if (Version.Major != Metadata.Version.Major)
+		{
+			return false;
+		}
+
+		return Version.Minor > Metadata.Version.Minor;
+	};
+
+	return false;
+}
 #undef LOCTEXT_NAMESPACE
diff --git a/Engine/Plugins/Runtime/Metasound/Source/MetasoundEditor/Private/MetasoundEditorGraphSchema.cpp b/Engine/Plugins/Runtime/Metasound/Source/MetasoundEditor/Private/MetasoundEditorGraphSchema.cpp
--- a/Engine/Plugins/Runtime/Metasound/Source/MetasoundEditor/Private/MetasoundEditorGraphSchema.cpp
+++ b/Engine/Plugins/Runtime/Metasound/Source/MetasoundEditor/Private/MetasoundEditorGraphSchema.cpp
@@ -11,17 +11,20 @@
 #include "GraphEditor.h"
 #include "GraphEditorActions.h"
 #include "GraphEditorSettings.h"
+#include "HAL/IConsoleManager.h"
 #include "Layout/SlateRect.h"
 #include "Metasound.h"
 #include "MetasoundAssetBase.h"
 #include "MetasoundDataReference.h"
 #include "MetasoundEditor.h"
+#include "MetasoundEditorCommands.h"
 #include "MetasoundEditorGraph.h"
 #include "MetasoundEditorGraphBuilder.h"
 #include "MetasoundEditorGraphNode.h"
 #include "MetasoundEditorModule.h"
 #include "MetasoundEditorSettings.h"
 #include "MetasoundFrontend.h"
+#include "MetasoundFrontendSearchEngine.h"
 #include "ScopedTransaction.h"
 #include "Toolkits/ToolkitManager.h"
 #include "ToolMenus.h"
@@ -31,6 +34,13 @@
 
 #define LOCTEXT_NAMESPACE "MetaSoundEditor"
 
+static int32 EnableDeprecatedMetaSoundNodeClassCreationCVar = 0;
+FAutoConsoleVariableRef CVarEnableDeprecatedMetaSoundNodeClassCreation(
+	TEXT("au.Debug.Editor.MetaSounds.EnableDeprecatedNodeClassCreation"),
+	EnableDeprecatedMetaSoundNodeClassCreationCVar,
+	TEXT("Enable creating nodes major versions of deprecated MetaSound classes in the Editor.\n")
+	TEXT("0: Disabled (default), !0: Enabled"),
+	ECVF_Default);
 
 namespace Metasound
 {
@@ -172,7 +182,7 @@ UEdGraphNode* FMetasoundGraphSchemaAction_NewNode::PerformAction(UEdGraph* Paren
 	ParentMetasound.Modify();
 	ParentGraph->Modify();
 
-	if (UMetasoundEditorGraphExternalNode* NewGraphNode = FGraphBuilder::AddExternalNode(ParentMetasound, NodeClassInfo, Location, bSelectNewNode))
+	if (UMetasoundEditorGraphExternalNode* NewGraphNode = FGraphBuilder::AddExternalNode(ParentMetasound, ClassMetadata, Location, bSelectNewNode))
 	{
 		NewGraphNode->Modify();
 		SchemaPrivate::TryConnectNewNodeToPin(*NewGraphNode, FromPin);
@@ -254,7 +264,7 @@ UEdGraphNode* FMetasoundGraphSchemaAction_PromoteToInput::PerformAction(UEdGraph
 	FMetasoundFrontendLiteral DefaultValue;
 	FGraphBuilder::GetPinDefaultLiteral(*FromPin, DefaultValue);
 
-	FNodeHandle NodeHandle = FGraphBuilder::AddInputNodeHandle(ParentMetasound, NewNodeName, InputHandle->GetDataType(), FText::GetEmpty(), EMetasoundFrontendNodeStyleDisplayVisibility::Visible, &DefaultValue);
+	FNodeHandle NodeHandle = FGraphBuilder::AddInputNodeHandle(ParentMetasound, NewNodeName, InputHandle->GetDataType(), FText::GetEmpty(), false /* bIsLiteralInput */, &DefaultValue);
 	if (ensure(NodeHandle->IsValid()))
 	{
 		UMetasoundEditorGraphInput* Input = MetasoundGraph->FindOrAddInput(NodeHandle);
@@ -504,6 +514,8 @@ void UMetasoundEditorGraphSchema::GetGraphContextActions(FGraphContextMenuBuilde
 
 void UMetasoundEditorGraphSchema::GetContextMenuActions(class UToolMenu* Menu, class UGraphNodeContextMenuContext* Context) const
 {
+	using namespace Metasound::Editor;
+
 	if (Context->Pin)
 	{
 		FToolMenuSection& Section = Menu->AddSection("MetasoundGraphSchemaPinActions", LOCTEXT("PinActionsMenuHeader", "Pin Actions"));
@@ -522,6 +534,17 @@ void UMetasoundEditorGraphSchema::GetContextMenuActions(class UToolMenu* Menu, c
 		Section.AddMenuEntry(FGenericCommands::Get().Copy);
 		Section.AddMenuEntry(FGenericCommands::Get().Duplicate);
 		Section.AddMenuEntry(FGraphEditorCommands::Get().BreakNodeLinks);
+
+		// Only display update ability if node is of type external
+		// and node registry is reporting a major update is available.
+		if (const UMetasoundEditorGraphExternalNode* ExternalNode = Cast<UMetasoundEditorGraphExternalNode>(Context->Node))
+		{
+			FMetasoundFrontendVersionNumber MajorUpdateVersion = ExternalNode->GetMajorUpdateAvailable();
+			if (MajorUpdateVersion.IsValid())
+			{
+				Section.AddMenuEntry(FEditorCommands::Get().UpdateNodes);
+			}
+		}
 	}
 
 	Super::GetContextMenuActions(Menu, Context);
@@ -568,9 +591,32 @@ const FPinConnectionResponse UMetasoundEditorGraphSchema::CanCreateConnection(co
 		return FPinConnectionResponse(CONNECT_RESPONSE_DISALLOW, LOCTEXT("ConnectionTypeIncorrect", "Connection pin types do not match"));
 	}
 
+	// Check if nodes being connected are in error/out-of-date version state before checking if can be connected as 
+	// nodes with errors can contain invalid pins/handles).
+	bool bConnectingNodesWithErrors = false;
+	UEdGraphNode* InputNode = InputPin->GetOwningNode();
+	if (ensure(InputNode))
+	{
+		if (InputNode->ErrorType == EMessageSeverity::Error)
+		{
+			bConnectingNodesWithErrors = true;
+		}
+	}
+	UEdGraphNode* OutputNode = InputPin->GetOwningNode();
+	if (ensure(OutputNode))
+	{
+		if (OutputNode->ErrorType == EMessageSeverity::Error)
+		{
+			bConnectingNodesWithErrors = true;
+		}
+	}
+	if (bConnectingNodesWithErrors)
+	{
+		return FPinConnectionResponse(CONNECT_RESPONSE_DISALLOW, LOCTEXT("ConnectionCannotContainErrorNode", "Cannot create new connections with node containing errors."));
+	}
+
 	Frontend::FInputHandle InputHandle = Editor::FGraphBuilder::GetInputHandleFromPin(InputPin);
 	Frontend::FOutputHandle OutputHandle = Editor::FGraphBuilder::GetOutputHandleFromPin(OutputPin);
-
 	if (InputHandle->IsValid() && OutputHandle->IsValid())
 	{
 		// TODO: Implement YesWithConverterNode to provide conversion options
@@ -665,7 +711,7 @@ FText UMetasoundEditorGraphSchema::GetPinDisplayName(const UEdGraphPin* Pin) con
 
 	UMetasoundEditorGraphNode* Node = CastChecked<UMetasoundEditorGraphNode>(Pin->GetOwningNode());
 	FNodeHandle NodeHandle = Node->GetNodeHandle();
-	const EMetasoundFrontendClassType ClassType = NodeHandle->GetClassType();
+	const EMetasoundFrontendClassType ClassType = NodeHandle->GetClassMetadata().Type;
 
 	switch (ClassType)
 	{
@@ -713,11 +759,15 @@ FLinearColor UMetasoundEditorGraphSchema::GetPinTypeColor(const FEdGraphPinType&
 }
 
 void UMetasoundEditorGraphSchema::BreakNodeLinks(UEdGraphNode& TargetNode) const
+{
+	BreakNodeLinks(TargetNode, true /* bShouldActuallyTransact */);
+}
+
+void UMetasoundEditorGraphSchema::BreakNodeLinks(UEdGraphNode& TargetNode, bool bShouldActuallyTransact) const
 {
 	using namespace Metasound::Editor;
-	using namespace Metasound::Frontend;
 
-	const FScopedTransaction Transaction(LOCTEXT("BreakNodeLinks", "Break Node Links"));
+	const FScopedTransaction Transaction(LOCTEXT("BreakNodeLinks", "Break Node Links"), bShouldActuallyTransact);
 	UMetasoundEditorGraph* Graph = CastChecked<UMetasoundEditorGraph>(TargetNode.GetGraph());
 	Graph->GetMetasoundChecked().Modify();
 	TargetNode.Modify();
@@ -767,21 +817,21 @@ void UMetasoundEditorGraphSchema::GetConversionActions(FGraphActionMenuBuilder&
 	using namespace Metasound;
 
 	const FText MenuJoinFormat = LOCTEXT("MetasoundActionsFormatSubCategory", "{0}|{1}");
-	const TArray<Frontend::FNodeClassInfo> ClassInfos = Frontend::GetAllAvailableNodeClasses();
-	for (const Frontend::FNodeClassInfo& ClassInfo : ClassInfos)
+	const bool bIncludeDeprecated = static_cast<bool>(EnableDeprecatedMetaSoundNodeClassCreationCVar);
+	const TArray<FMetasoundFrontendClass> FrontendClasses = Frontend::ISearchEngine::Get().FindAllClasses(bIncludeDeprecated);
+	for (const FMetasoundFrontendClass& FrontendClass : FrontendClasses)
 	{
-		const FMetasoundFrontendClass ClassDescription = Frontend::GenerateClassDescription(ClassInfo);
-		if (InFilters.InputFilterFunction && !ClassDescription.Interface.Inputs.ContainsByPredicate(InFilters.InputFilterFunction))
+		if (InFilters.InputFilterFunction && !FrontendClass.Interface.Inputs.ContainsByPredicate(InFilters.InputFilterFunction))
 		{
 			continue;
 		}
 
-		if (InFilters.OutputFilterFunction && !ClassDescription.Interface.Outputs.ContainsByPredicate(InFilters.OutputFilterFunction))
+		if (InFilters.OutputFilterFunction && !FrontendClass.Interface.Outputs.ContainsByPredicate(InFilters.OutputFilterFunction))
 		{
 			continue;
 		}
 
-		const FMetasoundFrontendClassMetadata Metadata = ClassDescription.Metadata;
+		const FMetasoundFrontendClassMetadata& Metadata = FrontendClass.Metadata;
 		const FText Tooltip = Metadata.Author.IsEmpty()
 			? Metadata.Description
 			: FText::Format(LOCTEXT("MetasoundTooltipAuthorFormat", "{0}\nAuthor: {1}"), Metadata.Description, Metadata.Author);
@@ -796,7 +846,7 @@ void UMetasoundEditorGraphSchema::GetConversionActions(FGraphActionMenuBuilder&
 				0
 			);
 
-			NewNodeAction->NodeClassInfo = ClassInfo;
+			NewNodeAction->ClassMetadata = Metadata;
 			ActionMenuBuilder.AddAction(NewNodeAction);
 
 		}
@@ -875,21 +925,21 @@ void UMetasoundEditorGraphSchema::GetFunctionActions(FGraphActionMenuBuilder& Ac
 
 	const FText MenuJoinFormat = LOCTEXT("MetasoundActionsFormatSubCategory", "{0}|{1}");
 
-	const TArray<FNodeClassInfo> ClassInfos = GetAllAvailableNodeClasses();
-	for (const FNodeClassInfo& ClassInfo : ClassInfos)
+	const bool bIncludeDeprecated = static_cast<bool>(EnableDeprecatedMetaSoundNodeClassCreationCVar);
+	const TArray<FMetasoundFrontendClass> FrontendClasses = ISearchEngine::Get().FindAllClasses(bIncludeDeprecated);
+	for (const FMetasoundFrontendClass& FrontendClass : FrontendClasses)
 	{
-		const FMetasoundFrontendClass ClassDescription = GenerateClassDescription(ClassInfo);
-		if (InFilters.InputFilterFunction && !ClassDescription.Interface.Inputs.ContainsByPredicate(InFilters.InputFilterFunction))
+		if (InFilters.InputFilterFunction && !FrontendClass.Interface.Inputs.ContainsByPredicate(InFilters.InputFilterFunction))
 		{
 			continue;
 		}
 
-		if (InFilters.OutputFilterFunction && !ClassDescription.Interface.Outputs.ContainsByPredicate(InFilters.OutputFilterFunction))
+		if (InFilters.OutputFilterFunction && !FrontendClass.Interface.Outputs.ContainsByPredicate(InFilters.OutputFilterFunction))
 		{
 			continue;
 		}
 
-		const FMetasoundFrontendClassMetadata Metadata = ClassDescription.Metadata;
+		const FMetasoundFrontendClassMetadata& Metadata = FrontendClass.Metadata;
 		const FText Tooltip = Metadata.Author.IsEmpty()
 			? Metadata.Description
 			: FText::Format(LOCTEXT("MetasoundTooltipAuthorFormat", "{0}\nAuthor: {1}"), Metadata.Description, Metadata.Author);
@@ -905,7 +955,7 @@ void UMetasoundEditorGraphSchema::GetFunctionActions(FGraphActionMenuBuilder& Ac
 				0
 			);
 
-			NewNodeAction->NodeClassInfo = ClassInfo;
+			NewNodeAction->ClassMetadata = Metadata;
 			ActionMenuBuilder.AddAction(NewNodeAction);
 		}
 	}
diff --git a/Engine/Plugins/Runtime/Metasound/Source/MetasoundEditor/Public/MetasoundEditorCommands.h b/Engine/Plugins/Runtime/Metasound/Source/MetasoundEditor/Public/MetasoundEditorCommands.h
--- a/Engine/Plugins/Runtime/Metasound/Source/MetasoundEditor/Public/MetasoundEditorCommands.h
+++ b/Engine/Plugins/Runtime/Metasound/Source/MetasoundEditor/Public/MetasoundEditorCommands.h
@@ -52,6 +52,9 @@ namespace Metasound
 			/** Shows the Metasound's top-level general object settings in the Inspector */
 			TSharedPtr<FUICommandInfo> EditMetasoundSettings;
 
+			// Run updates on selected nodes if required by any subset of selected nodes
+			TSharedPtr<FUICommandInfo> UpdateNodes;
+
 			/** Initialize commands */
 			virtual void RegisterCommands() override;
 		};
diff --git a/Engine/Plugins/Runtime/Metasound/Source/MetasoundEditor/Public/MetasoundEditorGraph.h b/Engine/Plugins/Runtime/Metasound/Source/MetasoundEditor/Public/MetasoundEditorGraph.h
--- a/Engine/Plugins/Runtime/Metasound/Source/MetasoundEditor/Public/MetasoundEditorGraph.h
+++ b/Engine/Plugins/Runtime/Metasound/Source/MetasoundEditor/Public/MetasoundEditorGraph.h
@@ -4,6 +4,7 @@
 #include "Audio/AudioParameterInterface.h"
 #include "CoreMinimal.h"
 #include "MetasoundEditor.h"
+#include "MetasoundEditorGraphNode.h"
 #include "MetasoundFrontendController.h"
 #include "MetasoundFrontendDocument.h"
 #include "MetasoundFrontendLiteral.h"
@@ -16,9 +17,18 @@
 
 // Forward Declarations
 struct FMetasoundFrontendDocument;
+class ITargetPlatform;
 class UMetasoundEditorGraphInputNode;
 class UMetasoundEditorGraphNode;
 
+namespace Metasound
+{
+	namespace Editor
+	{
+		struct FGraphValidationResults;
+	} // namespace Editor
+} // namespace Metasound
+
 DECLARE_MULTICAST_DELEGATE_OneParam(FOnMetasoundNodeNameChanged, FGuid /* NodeID */);
 
 UCLASS(Abstract)
@@ -48,9 +58,14 @@ public:
 	Metasound::Frontend::FNodeHandle GetNodeHandle() const;
 	Metasound::Frontend::FConstNodeHandle GetConstNodeHandle() const;
 
+	bool CanRename(const FText& InNewName, FText& OutError) const;
 	TArray<UMetasoundEditorGraphNode*> GetNodes() const;
 	void SetDataType(FName InNewType);
+	void SetDisplayName(const FText& InNewName);
+
+	virtual EMetasoundFrontendClassType GetClassType() const { return EMetasoundFrontendClassType::Invalid; }
 	virtual void OnDataTypeChanged() { }
+	virtual const FText& GetVariableLabel() const { return FText::GetEmpty(); }
 };
 
 UCLASS()
@@ -89,7 +104,8 @@ class METASOUNDEDITOR_API UMetasoundEditorGraphInput : public UMetasoundEditorGr
 
 protected:
 	virtual Metasound::Frontend::FNodeHandle AddNodeHandle(const FString& InNodeName, const FText& InNodeDisplayName, FName InDataType) override;
-
+	virtual EMetasoundFrontendClassType GetClassType() const override { return EMetasoundFrontendClassType::Input; }
+	virtual const FText& GetVariableLabel() const override;
 
 public:
 	UPROPERTY(VisibleAnywhere, Category = DefaultValue)
@@ -113,6 +129,8 @@ class METASOUNDEDITOR_API UMetasoundEditorGraphOutput : public UMetasoundEditorG
 
 protected:
 	virtual Metasound::Frontend::FNodeHandle AddNodeHandle(const FString& InNodeName, const FText& InNodeDisplayName, FName InDataType) override;
+	virtual EMetasoundFrontendClassType GetClassType() const override { return EMetasoundFrontendClassType::Output; }
+	virtual const FText& GetVariableLabel() const override;
 };
 
 UCLASS()
@@ -123,6 +141,9 @@ class METASOUNDEDITOR_API UMetasoundEditorGraph : public UMetasoundEditorGraphBa
 public:
 	UMetasoundEditorGraphInputNode* CreateInputNode(Metasound::Frontend::FNodeHandle InNodeHandle, bool bInSelectNewNode);
 
+	Metasound::Frontend::FDocumentHandle GetDocumentHandle() const;
+	Metasound::Frontend::FGraphHandle GetGraphHandle() const;
+
 	UObject* GetMetasound() const;
 	UObject& GetMetasoundChecked() const;
 
@@ -137,6 +158,8 @@ public:
 
 	virtual void Synchronize();
 
+	bool Validate(Metasound::Editor::FGraphValidationResults& OutResults);
+
 private:
 	// Preview ID is the Unique ID provided by the UObject that implements
 	// a sound's ParameterInterface when a sound begins playing.
diff --git a/Engine/Plugins/Runtime/Metasound/Source/MetasoundEditor/Public/MetasoundEditorGraphNode.h b/Engine/Plugins/Runtime/Metasound/Source/MetasoundEditor/Public/MetasoundEditorGraphNode.h
--- a/Engine/Plugins/Runtime/Metasound/Source/MetasoundEditor/Public/MetasoundEditorGraphNode.h
+++ b/Engine/Plugins/Runtime/Metasound/Source/MetasoundEditor/Public/MetasoundEditorGraphNode.h
@@ -4,7 +4,6 @@
 
 #include "CoreMinimal.h"
 #include "EdGraph/EdGraphNode.h"
-#include "UObject/ObjectMacros.h"
 #include "MetasoundFrontend.h"
 #include "MetasoundFrontendController.h"
 #include "MetasoundFrontendDocument.h"
@@ -12,6 +11,7 @@
 #include "MetasoundFrontendRegistries.h"
 #include "Misc/Guid.h"
 #include "Sound/SoundWave.h"
+#include "UObject/ObjectMacros.h"
 
 #include "MetasoundEditorGraphNode.generated.h"
 
@@ -24,10 +24,13 @@ namespace Metasound
 {
 	namespace Editor
 	{
+		struct FGraphNodeValidationResult;
 		class FGraphBuilder;
-	}
-}
 
+		// Map of class names to sorted array of registered version numbers
+		using FSortedClassVersionMap = TMap<FName, TArray<FMetasoundFrontendVersionNumber>>;
+	} // namespace Editor
+} // namespace Metasound
 
 UCLASS(MinimalAPI)
 class UMetasoundEditorGraphNode : public UEdGraphNode
@@ -41,7 +44,7 @@ public:
 	/** Estimate the width of this Node from the length of its title */
 	METASOUNDEDITOR_API int32 EstimateNodeWidth() const;
 
-	METASOUNDEDITOR_API void IteratePins(TUniqueFunction<void(UEdGraphPin* /* Pin */, int32 /* Index */)> Func, EEdGraphPinDirection InPinDirection = EGPD_MAX);
+	METASOUNDEDITOR_API void IteratePins(TUniqueFunction<void(UEdGraphPin& /* Pin */, int32 /* Index */)> Func, EEdGraphPinDirection InPinDirection = EGPD_MAX);
 
 	// UEdGraphNode interface
 	virtual void AllocateDefaultPins() override;
@@ -135,11 +138,23 @@ public:
 	virtual FMetasoundFrontendClassName GetClassName() const override { return ClassName; }
 	virtual FGuid GetNodeID() const override { return NodeID; }
 
+	FMetasoundFrontendVersionNumber GetMajorUpdateAvailable() const;
+	FMetasoundFrontendVersionNumber GetMinorUpdateAvailable() const;
+
+	// Attempts to replace this node with a new one of the same class and given version number.
+	// If this node is already of the given version, returns itself. If update fails, returns this node.
+	UMetasoundEditorGraphExternalNode* UpdateToVersion(const FMetasoundFrontendVersionNumber& InNewVersion, bool bInPropagateErrorMessages);
+
+	bool Validate(Metasound::Editor::FGraphNodeValidationResult& OutResult);
+
 protected:
+	// Refreshes all Pin Metadata from the associated Frontend node's default class interface.
+	bool RefreshPinMetadata();
+
 	virtual void SetNodeID(FGuid InNodeID) override
 	{
 		NodeID = InNodeID;
 	}
 
 	friend class Metasound::Editor::FGraphBuilder;
-};
\ No newline at end of file
+};
diff --git a/Engine/Plugins/Runtime/Metasound/Source/MetasoundEditor/Public/MetasoundEditorGraphSchema.h b/Engine/Plugins/Runtime/Metasound/Source/MetasoundEditor/Public/MetasoundEditorGraphSchema.h
--- a/Engine/Plugins/Runtime/Metasound/Source/MetasoundEditor/Public/MetasoundEditorGraphSchema.h
+++ b/Engine/Plugins/Runtime/Metasound/Source/MetasoundEditor/Public/MetasoundEditorGraphSchema.h
@@ -101,8 +101,8 @@ struct METASOUNDEDITOR_API FMetasoundGraphSchemaAction_NewNode : public FEdGraph
 {
 	GENERATED_USTRUCT_BODY();
 
-	/** ClassInfo of node to create */
-	Metasound::Frontend::FNodeClassInfo NodeClassInfo;
+	/** Class Metadata of node to create */
+	FMetasoundFrontendClassMetadata ClassMetadata;
 
 	FMetasoundGraphSchemaAction_NewNode() 
 		: FEdGraphSchemaAction()
@@ -203,6 +203,8 @@ class UMetasoundEditorGraphSchema : public UEdGraphSchema
 	virtual TSharedPtr<FEdGraphSchemaAction> GetCreateCommentAction() const override;
 	//~ End EdGraphSchema Interface
 
+	void BreakNodeLinks(UEdGraphNode& TargetNode, bool bShouldActuallyTransact) const;
+
 private:
 	/** Adds actions for creating actions associated with graph DataTypes */
 	void GetConversionActions(FGraphActionMenuBuilder& ActionMenuBuilder, Metasound::Editor::FActionClassFilters InFilters = Metasound::Editor::FActionClassFilters(), bool bShowSelectedActions = true) const;
diff --git a/Engine/Plugins/Runtime/Metasound/Source/MetasoundEditor/Public/MetasoundEditorGraphValidation.h b/Engine/Plugins/Runtime/Metasound/Source/MetasoundEditor/Public/MetasoundEditorGraphValidation.h
new file mode 100644
--- /dev/null
+++ b/Engine/Plugins/Runtime/Metasound/Source/MetasoundEditor/Public/MetasoundEditorGraphValidation.h
@@ -0,0 +1,62 @@
+// Copyright Epic Games, Inc. All Rights Reserved.
+#pragma once
+
+#include "CoreMinimal.h"
+#include "MetasoundFrontendDocument.h"
+#include "MetasoundEditorGraph.h"
+#include "MetasoundEditorGraphNode.h"
+
+
+namespace Metasound
+{
+	namespace Editor
+	{
+		struct METASOUNDEDITOR_API FGraphNodeValidationResult
+		{
+			// Node associated with validation result
+			UMetasoundEditorGraphExternalNode* Node = nullptr;
+
+			// Whether associated node is in invalid state (document is corrupt and no Frontend representation could be
+			// found for the node)
+			bool bIsInvalid = false;
+
+			// Whether validation made changes to the node and is now in a dirty state
+			bool bIsDirty = false;
+
+			// If not none, name of GraphNode's missing (not found in Frontend Registry) class as described in the owning
+			// Frontend document model.
+			FName MissingClass;
+
+			FGraphNodeValidationResult(UMetasoundEditorGraphExternalNode& InNode)
+				: Node(&InNode)
+			{
+			}
+		};
+
+		struct METASOUNDEDITOR_API FGraphValidationResults
+		{
+			TArray<FGraphNodeValidationResult> NodeResults;
+
+			// Results corresponding with node validation
+			const TArray<FGraphNodeValidationResult>& GetResults() const
+			{
+				return NodeResults;
+			}
+
+			// Returns set of missing class names from the validation results
+			TSet<FName> FindMissingClasses() const
+			{
+				TSet<FName> MissingClasses;
+				for (const FGraphNodeValidationResult& Result : NodeResults)
+				{
+					if (!Result.MissingClass.IsNone())
+					{
+						MissingClasses.Add(Result.MissingClass);
+					}
+				}
+
+				return MissingClasses;
+			}
+		};
+	} // namespace Editor
+} // namespace Metasound
diff --git a/Engine/Plugins/Runtime/Metasound/Source/MetasoundEngine/Private/MetasoundSource.cpp b/Engine/Plugins/Runtime/Metasound/Source/MetasoundEngine/Private/MetasoundSource.cpp
--- a/Engine/Plugins/Runtime/Metasound/Source/MetasoundEngine/Private/MetasoundSource.cpp
+++ b/Engine/Plugins/Runtime/Metasound/Source/MetasoundEngine/Private/MetasoundSource.cpp
@@ -186,15 +186,16 @@ bool UMetaSoundSource::GetReceiveNodeMetadataForDataType(const FName& InTypeName
 {
 	using namespace Metasound;
 
-	TArray<FMetasoundFrontendClass> ReceiverNodeClasses = Metasound::Frontend::ISearchEngine::Get().FindClassesWithClassName(FReceiveNodeNames::GetClassNameForDataType(InTypeName));
+	const FNodeClassName ClassName = FReceiveNodeNames::GetClassNameForDataType(InTypeName);
+	TArray<FMetasoundFrontendClass> ReceiverNodeClasses = Frontend::ISearchEngine::Get().FindClassesWithName(ClassName, true /* bInSortByVersion */);
 
-	if (ReceiverNodeClasses.Num() > 0)
+	if (ReceiverNodeClasses.IsEmpty())
 	{
-		OutMetadata = ReceiverNodeClasses[0].Metadata;
-		return true;
+		return false;
 	}
 
-	return false;
+	OutMetadata = ReceiverNodeClasses[0].Metadata;
+	return true;
 }
 
 Metasound::Frontend::FNodeHandle UMetaSoundSource::AddInputPinForSendAddress(const Metasound::FMetasoundInstanceTransmitter::FSendInfo& InSendInfo, Metasound::Frontend::FGraphHandle InGraph) const
diff --git a/Engine/Plugins/Runtime/Metasound/Source/MetasoundFrontend/Private/MetasoundFrontendInvalidController.h b/Engine/Plugins/Runtime/Metasound/Source/MetasoundFrontend/Private/MetasoundFrontendInvalidController.h
--- a/Engine/Plugins/Runtime/Metasound/Source/MetasoundFrontend/Private/MetasoundFrontendInvalidController.h
+++ b/Engine/Plugins/Runtime/Metasound/Source/MetasoundFrontend/Private/MetasoundFrontendInvalidController.h
@@ -155,12 +155,11 @@ namespace Metasound
 			FInputHandle AddOutput(const FString& InVertexName, const FMetasoundFrontendLiteral* InDefault) override { return FInvalidInputController::GetInvalid(); }
 			bool RemoveOutput(FGuid InVertexID) override { return false; }
 
-			EMetasoundFrontendClassType GetClassType() const override { return EMetasoundFrontendClassType::Invalid; }
-			const FMetasoundFrontendClassName& GetClassName() const override { return MetasoundFrontendInvalidControllerPrivate::GetInvalid<FMetasoundFrontendClassName>(); }
-			FMetasoundFrontendVersionNumber GetClassVersionNumber() const override { return MetasoundFrontendInvalidControllerPrivate::GetInvalid<FMetasoundFrontendVersionNumber>(); }
-			const FMetasoundFrontendInterfaceStyle& GetInputStyle() const override { static const FMetasoundFrontendInterfaceStyle Invalid; return Invalid; }
-			const FMetasoundFrontendInterfaceStyle& GetOutputStyle() const override { static const FMetasoundFrontendInterfaceStyle Invalid; return Invalid; }
-			const FMetasoundFrontendClassStyle& GetClassStyle() const override { static const FMetasoundFrontendClassStyle Invalid; return Invalid; }
+			const FMetasoundFrontendClassInterface& GetClassInterface() const override { return MetasoundFrontendInvalidControllerPrivate::GetInvalid<FMetasoundFrontendClassInterface>(); }
+			const FMetasoundFrontendClassMetadata& GetClassMetadata() const override { return MetasoundFrontendInvalidControllerPrivate::GetInvalid<FMetasoundFrontendClassMetadata>(); }
+			const FMetasoundFrontendInterfaceStyle& GetInputStyle() const override { return MetasoundFrontendInvalidControllerPrivate::GetInvalid<FMetasoundFrontendInterfaceStyle>(); }
+			const FMetasoundFrontendInterfaceStyle& GetOutputStyle() const override { return MetasoundFrontendInvalidControllerPrivate::GetInvalid<FMetasoundFrontendInterfaceStyle>(); }
+			const FMetasoundFrontendClassStyle& GetClassStyle() const override { return MetasoundFrontendInvalidControllerPrivate::GetInvalid<FMetasoundFrontendClassStyle>(); }
 
 			const FText& GetDescription() const override { return MetasoundFrontendInvalidControllerPrivate::GetInvalid<FText>(); }
 
@@ -344,7 +343,8 @@ namespace Metasound
 				FConstClassAccessPtr FindOrAddClass(const FNodeClassInfo& InNodeClass) override { return FConstClassAccessPtr(); }
 				FConstClassAccessPtr FindClass(const FMetasoundFrontendClassMetadata& InMetadata) const override{ return FConstClassAccessPtr(); }
 				FConstClassAccessPtr FindOrAddClass(const FMetasoundFrontendClassMetadata& InMetadata) override{ return FConstClassAccessPtr(); }
-				void RemoveUnreferencedDependencies() override {}
+
+				void SynchronizeDependencies() override { }
 
 				TArray<FGraphHandle> GetSubgraphHandles() override { return TArray<FGraphHandle>(); }
 
diff --git a/Engine/Plugins/Runtime/Metasound/Source/MetasoundFrontend/Private/MetasoundFrontendQuery.cpp b/Engine/Plugins/Runtime/Metasound/Source/MetasoundFrontend/Private/MetasoundFrontendQuery.cpp
--- a/Engine/Plugins/Runtime/Metasound/Source/MetasoundFrontend/Private/MetasoundFrontendQuery.cpp
+++ b/Engine/Plugins/Runtime/Metasound/Source/MetasoundFrontend/Private/MetasoundFrontendQuery.cpp
@@ -190,10 +190,8 @@ namespace Metasound
 				{
 					if (InOutResult.GetSelection().Num() > 0)
 					{
-						TArrayView<FFrontendQueryEntry* > SelectionView = InOutResult.GetSelection();
-
+						TArrayView<FFrontendQueryEntry*> SelectionView = InOutResult.GetSelection();
 						TArray<FFrontendQueryEntry*> SortedSelection(SelectionView.GetData(), SelectionView.Num()); 
-
 						InOutResult.ResetSelection();
 
 						Algo::SortBy(SortedSelection, [](const FFrontendQueryEntry* InEntry)
@@ -207,40 +205,36 @@ namespace Metasound
 							}
 						);
 
+						if (!SortedSelection.IsEmpty())
+						{
+							int32 StartIndex = 0;
+							FFrontendQueryEntry::FKey CurrentKey = SortedSelection[0] ? SortedSelection[0]->Key : FFrontendQueryEntry::InvalidKey;
 
-						int32 Num = SortedSelection.Num();
-						int32 StartIndex = 0;
+							auto ReduceFunc = [&](int32 Index)
+							{
+								TArrayView<FFrontendQueryEntry*> ResultsToReduce(&SortedSelection[StartIndex], Index - StartIndex);
+								FFrontendQuerySelection::FReduceOutputView ReduceResult(ResultsToReduce);
 
-						FFrontendQueryEntry::FKey CurrentKey = FFrontendQueryEntry::InvalidKey;
+								Step->Reduce(CurrentKey, ResultsToReduce, ReduceResult);
 
-						auto ReduceFunc = [&](int32 Index)
-						{
-							FFrontendQueryEntry::FKey ThisKey = nullptr == SortedSelection[Index] ? FFrontendQueryEntry::InvalidKey : SortedSelection[Index]->Key;
+								InOutResult.AppendToStorageAndSelection(ReduceResult);
+							};
 
-							if (CurrentKey != ThisKey)
+							const int32 Num = SortedSelection.Num();
+							for (int32 ResultIndex = 1; ResultIndex < Num; ResultIndex++)
 							{
-								if (Index > StartIndex)
+								FFrontendQueryEntry::FKey ThisKey = SortedSelection[ResultIndex] ? SortedSelection[ResultIndex]->Key : FFrontendQueryEntry::InvalidKey;
+								if (CurrentKey != ThisKey)
 								{
-									TArrayView<FFrontendQueryEntry*> ResultsToReduce(&SortedSelection[StartIndex], Index - StartIndex);
+									ReduceFunc(ResultIndex);
 
-									FFrontendQuerySelection::FReduceOutputView ReduceResult(ResultsToReduce);
-
-									Step->Reduce(CurrentKey, ResultsToReduce, ReduceResult);
-
-									InOutResult.AppendToStorageAndSelection(ReduceResult);
+									CurrentKey = ThisKey;
+									StartIndex = ResultIndex;
 								}
-
-								CurrentKey = ThisKey;
-								StartIndex = Index;
 							}
-						};
 
-						for (int32 ResultIndex = 0; ResultIndex < Num; ResultIndex++)
-						{
-							ReduceFunc(ResultIndex);
+							ReduceFunc(Num);
 						}
-
-						ReduceFunc(Num - 1);
 					}
 				}
 			}
diff --git a/Engine/Plugins/Runtime/Metasound/Source/MetasoundFrontend/Private/MetasoundFrontendQuerySteps.cpp b/Engine/Plugins/Runtime/Metasound/Source/MetasoundFrontend/Private/MetasoundFrontendQuerySteps.cpp
--- a/Engine/Plugins/Runtime/Metasound/Source/MetasoundFrontend/Private/MetasoundFrontendQuerySteps.cpp
+++ b/Engine/Plugins/Runtime/Metasound/Source/MetasoundFrontend/Private/MetasoundFrontendQuerySteps.cpp
@@ -85,5 +85,55 @@ namespace Metasound
 	{
 		return InEntry.Value.Get<FMetasoundFrontendClass>().ID == ClassID;
 	}
-}
 
+	FFrontendQueryEntry::FKey FMapClassNameToMajorVersion::Map(const FFrontendQueryEntry& InEntry) const
+	{
+		const FMetasoundFrontendClass& FrontendClass = InEntry.Value.Get<FMetasoundFrontendClass>();
+		const uint32 HashKey = GetTypeHash(FrontendClass.Metadata.ClassName.GetFullName());
+		return static_cast<FFrontendQueryEntry::FKey>(HashKey);
+	}
+
+	void FReduceClassesToHighestVersion::Reduce(FFrontendQueryEntry::FKey InKey, TArrayView<FFrontendQueryEntry*>& InEntries, FReduceOutputView& OutResult) const
+	{
+		FFrontendQueryEntry* HighestVersionEntry = nullptr;
+		int32 HighestMajorVersion = -1;
+
+		for (FFrontendQueryEntry* Entry : InEntries)
+		{
+			const int32 EntryMajorVersion = Entry->Value.Get<FMetasoundFrontendClass>().Metadata.Version.Major;
+
+			if (!HighestVersionEntry || HighestMajorVersion < EntryMajorVersion)
+			{
+				HighestVersionEntry = Entry;
+				HighestMajorVersion = EntryMajorVersion;
+			}
+		}
+
+		OutResult.Add(*HighestVersionEntry);
+	}
+
+	FReduceClassesToMajorVersion::FReduceClassesToMajorVersion(int32 InMajorVersion)
+		: MajorVersion(InMajorVersion)
+	{
+	}
+
+	void FReduceClassesToMajorVersion::Reduce(FFrontendQueryEntry::FKey InKey, TArrayView<FFrontendQueryEntry*>& InEntries, FReduceOutputView& OutResult) const
+	{
+		for (FFrontendQueryEntry* Entry : InEntries)
+		{
+			const int32 EntryMajorVersion = Entry->Value.Get<FMetasoundFrontendClass>().Metadata.Version.Major;
+			if (MajorVersion == EntryMajorVersion)
+			{
+				OutResult.Add(*Entry);
+				return;
+			}
+		}
+	}
+
+	bool FSortClassesByVersion::Sort(const FFrontendQueryEntry& InEntryLHS, const FFrontendQueryEntry& InEntryRHS) const
+	{
+		const FMetasoundFrontendVersionNumber& VersionLHS = InEntryLHS.Value.Get<FMetasoundFrontendClass>().Metadata.Version;
+		const FMetasoundFrontendVersionNumber& VersionRHS = InEntryRHS.Value.Get<FMetasoundFrontendClass>().Metadata.Version;
+		return VersionLHS > VersionRHS;
+	}
+}
diff --git a/Engine/Plugins/Runtime/Metasound/Source/MetasoundFrontend/Private/MetasoundFrontendRegistries.cpp b/Engine/Plugins/Runtime/Metasound/Source/MetasoundFrontend/Private/MetasoundFrontendRegistries.cpp
--- a/Engine/Plugins/Runtime/Metasound/Source/MetasoundFrontend/Private/MetasoundFrontendRegistries.cpp
+++ b/Engine/Plugins/Runtime/Metasound/Source/MetasoundFrontend/Private/MetasoundFrontendRegistries.cpp
@@ -149,7 +149,6 @@ namespace Metasound
 				using FNodeClassMetadata = Metasound::FNodeClassMetadata;
 				using IEnumDataTypeInterface = Metasound::Frontend::IEnumDataTypeInterface;
 
-
 				FRegistryContainerImpl() = default;
 
 				FRegistryContainerImpl(const FRegistryContainerImpl&) = delete;
@@ -181,6 +180,8 @@ namespace Metasound
 				// Return any data types that can be used as a metasound input type or output type.
 				TArray<FName> GetAllValidDataTypes() override;
 
+				void IterateRegistry(Metasound::FIterateMetasoundFrontendClassFunction InIterFunc, EMetasoundFrontendClassType InClassType = EMetasoundFrontendClassType::Invalid) const override;
+
 				// Find Frontend Document data.
 				bool FindFrontendClassFromRegistered(const FNodeClassInfo& InClassInfo, FMetasoundFrontendClass& OutClass) override;
 				bool FindFrontendClassFromRegistered(const FMetasoundFrontendClassMetadata& InMetadata, FMetasoundFrontendClass& OutClass) override;
@@ -563,14 +564,19 @@ namespace Metasound
 				if (ensure(FMetasoundFrontendRegistryContainer::GetRegistryKey(RegistryElement, Key)))
 				{
 					// check to see if an identical node was already registered, and log
-					ensureAlwaysMsgf(!ExternalNodeRegistry.Contains(Key), TEXT("Node with identical name, inputs and outputs to node %s was already registered. The previously registered node will be overwritten. This could also happen because METASOUND_REGISTER_NODE is in a public header."), *Key.NodeClassFullName.ToString());
+					ensureAlwaysMsgf(
+						!ExternalNodeRegistry.Contains(Key),
+						TEXT("Node with identical name '%s' & Major Version already registered. "
+							"The previously registered node will be overwritten. "
+							"This can happen if two classes share the same name or if METASOUND_REGISTER_NODE is defined in a public header."),
+						*Key.NodeClassFullName.ToString());
 
 					// Store update to newly registered node in history so nodes
 					// can be queried by transaction ID
-					FNodeClassInfo ClassInfo = {EMetasoundFrontendClassType::External, Key};
+					FNodeClassInfo ClassInfo = { EMetasoundFrontendClassType::External, Key };
 					RegistryTransactionHistory.Add(MakeNodeRegistrationTransaction(ClassInfo));
 
-					// Store registry elements in map so nodes can be queried using registry key. 
+					// Store registry elements in map so nodes can be queried using registry key.
 					ExternalNodeRegistry.Add(MoveTemp(Key), MoveTemp(RegistryElement));
 
 					return true;
@@ -642,7 +648,6 @@ namespace Metasound
 				return nullptr;
 			}
 
-
 			TSharedPtr<Metasound::IDataChannel, ESPMode::ThreadSafe> FRegistryContainerImpl::CreateDataChannelForDataType(const FName& InDataType, const Metasound::FOperatorSettings& InSettings) const
 			{
 				if (const FDataTypeRegistryElement* Element = DataTypeRegistry.Find(InDataType))
@@ -692,7 +697,6 @@ namespace Metasound
 				Info.LookupKey = GetRegistryKey(InMetadata);
 
 				return FindFrontendClassFromRegistered(Info, OutClass);
-
 			}
 
 			bool FRegistryContainerImpl::FindInputNodeClassMetadataForDataType(const FName& InDataTypeName, FMetasoundFrontendClassMetadata& OutMetadata)
@@ -714,6 +718,71 @@ namespace Metasound
 				}
 				return false;
 			}
+
+			void FRegistryContainerImpl::IterateRegistry(Metasound::FIterateMetasoundFrontendClassFunction InIterFunc, EMetasoundFrontendClassType InClassType) const
+			{
+				auto IterateExternalClasses = [=]
+				{
+					for (const TPair<FNodeRegistryKey, FNodeRegistryElement>& Pair : ExternalNodeRegistry)
+					{
+						InIterFunc(Pair.Value.CreateFrontendClass());
+					}
+				};
+
+				auto IterateInputClasses = [=]
+				{
+					for (const TPair<FNodeRegistryKey, FDataTypeRegistryElement>& Pair : DataTypeNodeRegistry)
+					{
+						InIterFunc(Pair.Value.Callbacks.CreateFrontendInputClass());
+					}
+				};
+
+				auto IterateOutputClasses = [=]
+				{
+					for (const TPair<FNodeRegistryKey, FDataTypeRegistryElement>& Pair : DataTypeNodeRegistry)
+					{
+						InIterFunc(Pair.Value.Callbacks.CreateFrontendOutputClass());
+					}
+				};
+
+				switch (InClassType)
+				{
+					case EMetasoundFrontendClassType::External:
+					{
+						IterateExternalClasses();
+					}
+					break;
+
+					// TODO: Implement
+					case EMetasoundFrontendClassType::Graph:
+					{
+						checkNoEntry();
+					}
+					break;
+
+					case EMetasoundFrontendClassType::Input:
+					{
+						IterateInputClasses();
+					}
+					break;
+
+					case EMetasoundFrontendClassType::Output:
+					{
+						IterateOutputClasses();
+					}
+					break;
+
+					default:
+					case EMetasoundFrontendClassType::Invalid:
+					{
+						IterateExternalClasses();
+						IterateInputClasses();
+						IterateOutputClasses();
+					}
+					break;
+				}
+				static_assert(static_cast<uint32>(EMetasoundFrontendClassType::Invalid) == 4, "Possible missing switch case coverage for EMetasoundFrontendClassType");
+			}
 		}
 	}
 }
diff --git a/Engine/Plugins/Runtime/Metasound/Source/MetasoundFrontend/Private/MetasoundFrontendSearchEngine.cpp b/Engine/Plugins/Runtime/Metasound/Source/MetasoundFrontend/Private/MetasoundFrontendSearchEngine.cpp
--- a/Engine/Plugins/Runtime/Metasound/Source/MetasoundFrontend/Private/MetasoundFrontendSearchEngine.cpp
+++ b/Engine/Plugins/Runtime/Metasound/Source/MetasoundFrontend/Private/MetasoundFrontendSearchEngine.cpp
@@ -2,9 +2,20 @@
 
 #include "MetasoundFrontendSearchEngine.h"
 
+#include "Algo/Transform.h"
+#include "HAL/IConsoleManager.h"
 #include "MetasoundFrontendQuery.h"
 #include "MetasoundFrontendQuerySteps.h"
 
+static int32 ClearMetaSoundFrontendSearchEngineCacheCVar = 0;
+FAutoConsoleVariableRef CVarClearMetaSoundFrontendSearchEngineCache(
+	TEXT("au.Debug.MetaSounds.SearchEngine.ClearCache"),
+	ClearMetaSoundFrontendSearchEngineCacheCVar,
+	TEXT("If true, flags the MetaSound registry SearchEngine cache to be cleared on subsequent query.\n")
+	TEXT("0: Do not clear, !0: Flag to clear"),
+	ECVF_Default);
+
+
 namespace Metasound
 {
 	namespace Frontend
@@ -18,17 +29,19 @@ namespace Metasound
 
 			public:
 				FSearchEngine() = default;
-
 				virtual ~FSearchEngine() = default;
 
-				TArray<FMetasoundFrontendClass> FindClassesWithClassName(const FNodeClassName& InName) override;
+				TArray<FMetasoundFrontendClass> FindAllClasses(bool bInIncludeDeprecated) override;
+				TArray<FMetasoundFrontendClass> FindClassesWithName(const FNodeClassName& InName, bool bInSortByVersion) override;
 
-			private:
+				bool FindClassWithHighestVersion(const FNodeClassName& InName, FMetasoundFrontendClass& OutClass) override;
+				bool FindClassWithMajorVersion(const FNodeClassName& InName, int32 InMajorVersion, FMetasoundFrontendClass& OutClass) override;
 
+			private:
 				FFrontendQuery* FindQuery(const FString& InQueryName);
 				FFrontendQuery* AddQuery(const FString& InQueryName, TUniquePtr<FFrontendQuery>&& InQuery);
 
-				// Store queries. Index them by unique strings. 
+				// Store of queries indexed by unique strings.
 				TMap<FString, TUniquePtr<FFrontendQuery>> QueryCache;
 		};
 
@@ -38,22 +51,67 @@ namespace Metasound
 			return SearchEngine;
 		}
 
-		TArray<FMetasoundFrontendClass> FSearchEngine::FindClassesWithClassName(const FNodeClassName& InName) 
+		TArray<FMetasoundFrontendClass> FSearchEngine::FindAllClasses(bool bInIncludeDeprecated)
 		{
-			// Unique string for query with given parameter name. 
-			FString QueryName = FString(TEXT("FindClassesWithClassName")) + InName.GetFullName().ToString();
+			FString QueryName = FString(TEXT("FindAllClasses"));
+			if (bInIncludeDeprecated)
+			{
+				QueryName += TEXT("IncludingDeprecated");
+			}
 
 			FFrontendQuery* Query = FindQuery(QueryName);
-
-			if (nullptr == Query)
+			if (!Query)
 			{
-				// If the query doesn't exist, it needs to be created. 
 				TUniquePtr<FFrontendQuery> NewQuery = MakeUnique<FFrontendQuery>();
 
-				// Note: Initial query step is FGenerateNewlyAvailableNodeClasses, which does not run query against already examined node classes.
+				NewQuery->AddStep<FGenerateNewlyAvailableNodeClasses>();
+
+				if (!bInIncludeDeprecated)
+				{
+					NewQuery->AddStep<FMapClassNameToMajorVersion>()
+						.AddStep<FReduceClassesToHighestVersion>();
+				}
+
+				Query = AddQuery(QueryName, MoveTemp(NewQuery));
+			}
+
+			TArray<FMetasoundFrontendClass> Result;
+
+			if (ensure(Query))
+			{
+				FFrontendQuerySelectionView Selection = Query->ExecuteQueryAndAppend();
+				TArrayView<const FFrontendQueryEntry* const> Entries = Selection.GetSelection();
+				for (const FFrontendQueryEntry* Entry : Selection.GetSelection())
+				{
+					check(Entry->Value.IsType<FMetasoundFrontendClass>());
+					Result.Add(Entry->Value.Get<FMetasoundFrontendClass>());
+				}
+			}
+
+			return Result;
+		}
+
+		TArray<FMetasoundFrontendClass> FSearchEngine::FindClassesWithName(const FNodeClassName& InName, bool bInSortByVersion)
+		{
+			FString QueryName = FString(TEXT("FindClassesWithName"));
+			if (bInSortByVersion)
+			{
+				QueryName += TEXT("SortedByVersion");
+			}
+			QueryName += InName.GetFullName().ToString();
+			FFrontendQuery* Query = FindQuery(QueryName);
+
+			if (!Query)
+			{
+				TUniquePtr<FFrontendQuery> NewQuery = MakeUnique<FFrontendQuery>();
 				NewQuery->AddStep<FGenerateNewlyAvailableNodeClasses>()
 					.AddStep<FFilterClassesByClassName>(InName);
 
+				if (bInSortByVersion)
+				{
+					NewQuery->AddStep<FSortClassesByVersion>();
+				}
+
 				Query = AddQuery(QueryName, MoveTemp(NewQuery));
 			}
 
@@ -61,14 +119,8 @@ namespace Metasound
 
 			if (ensure(nullptr != Query))
 			{
-				// Execute the query again and append the results. (Note: Initial query 
-				// step is FGenerateNewlyAvailableNodeClasses, which does not run query 
-				// against already examined node classes.)
 				FFrontendQuerySelectionView Selection = Query->ExecuteQueryAndAppend();
-
 				TArrayView<const FFrontendQueryEntry* const> Entries = Selection.GetSelection();
-
-				// Copy result to array
 				for (const FFrontendQueryEntry* Entry : Selection.GetSelection())
 				{
 					check(Entry->Value.IsType<FMetasoundFrontendClass>());
@@ -79,11 +131,87 @@ namespace Metasound
 			return Result;
 		}
 
+		bool FSearchEngine::FindClassWithHighestVersion(const FNodeClassName& InName, FMetasoundFrontendClass& OutClass)
+		{
+			const FString QueryName = FString(TEXT("FindHighestVersion")) + InName.GetFullName().ToString();
+			FFrontendQuery* Query = FindQuery(QueryName);
+			if (!Query)
+			{
+				TUniquePtr<FFrontendQuery> NewQuery = MakeUnique<FFrontendQuery>();
+
+				NewQuery->AddStep<FGenerateNewlyAvailableNodeClasses>()
+					.AddStep<FFilterClassesByClassName>(InName)
+					.AddStep<FMapClassNameToMajorVersion>()
+					.AddStep<FReduceClassesToHighestVersion>();
+
+				Query = AddQuery(QueryName, MoveTemp(NewQuery));
+			}
+
+			if (ensure(Query))
+			{
+				FFrontendQuerySelectionView Selection = Query->ExecuteQueryAndAppend();
+				TArrayView<const FFrontendQueryEntry* const> Entries = Selection.GetSelection();
+				if (Entries.IsEmpty())
+				{
+					return false;
+				}
+
+				if (ensure(Entries.Num() == 1))
+				{
+					OutClass = Entries[0]->Value.Get<FMetasoundFrontendClass>();
+					return true;
+				}
+			}
+
+			return false;
+		}
+
+		bool FSearchEngine::FindClassWithMajorVersion(const FNodeClassName& InName, int32 InMajorVersion, FMetasoundFrontendClass& OutClass)
+		{
+			const FString QueryName = FString(TEXT("FindClassWithMajorVersion")) + InName.GetFullName().ToString() + FString::FromInt(InMajorVersion);
+			FFrontendQuery* Query = FindQuery(QueryName);
+			if (!Query)
+			{
+				TUniquePtr<FFrontendQuery> NewQuery = MakeUnique<FFrontendQuery>();
+
+				NewQuery->AddStep<FGenerateNewlyAvailableNodeClasses>()
+					.AddStep<FFilterClassesByClassName>(InName)
+					.AddStep<FMapClassNameToMajorVersion>()
+					.AddStep<FReduceClassesToMajorVersion>(InMajorVersion);
+
+				Query = AddQuery(QueryName, MoveTemp(NewQuery));
+			}
+
+			if (ensure(Query))
+			{
+				FFrontendQuerySelectionView Selection = Query->ExecuteQueryAndAppend();
+				TArrayView<const FFrontendQueryEntry* const> Entries = Selection.GetSelection();
+				if (Entries.IsEmpty())
+				{
+					return false;
+				}
+
+				if (ensure(Entries.Num() == 1))
+				{
+					OutClass = Entries[0]->Value.Get<FMetasoundFrontendClass>();
+					return true;
+				}
+			}
+
+			return false;
+		}
+
 		FFrontendQuery* FSearchEngine::FindQuery(const FString& InQueryName)
 		{
-			if (TUniquePtr<FFrontendQuery>* QueryUPtr = QueryCache.Find(InQueryName))
+			if (ClearMetaSoundFrontendSearchEngineCacheCVar)
 			{
-				return QueryUPtr->Get();
+				ClearMetaSoundFrontendSearchEngineCacheCVar = 0;
+				QueryCache.Reset();
+			}
+
+			if (TUniquePtr<FFrontendQuery>* Query = QueryCache.Find(InQueryName))
+			{
+				return Query->Get();
 			}
 
 			return nullptr;
@@ -91,11 +219,7 @@ namespace Metasound
 
 		FFrontendQuery* FSearchEngine::AddQuery(const FString& InQueryName, TUniquePtr<FFrontendQuery>&& InQuery)
 		{
-			FFrontendQuery* QueryPtr = InQuery.Get();
-
-			QueryCache.Add(InQueryName, MoveTemp(InQuery));
-
-			return QueryPtr;
+			return QueryCache.Add(InQueryName, MoveTemp(InQuery)).Get();
 		}
 	}
 }
diff --git a/Engine/Plugins/Runtime/Metasound/Source/MetasoundFrontend/Private/MetasoundFrontendStandardController.cpp b/Engine/Plugins/Runtime/Metasound/Source/MetasoundFrontend/Private/MetasoundFrontendStandardController.cpp
--- a/Engine/Plugins/Runtime/Metasound/Source/MetasoundFrontend/Private/MetasoundFrontendStandardController.cpp
+++ b/Engine/Plugins/Runtime/Metasound/Source/MetasoundFrontend/Private/MetasoundFrontendStandardController.cpp
@@ -236,6 +236,24 @@ namespace Metasound
 				};
 			}
 
+			auto FGetClassFromDocumentWithClassInfo(const FNodeClassInfo& InInfo)
+			{
+				return [=](FMetasoundFrontendDocument& InDocument) -> FMetasoundFrontendClass*
+				{
+					auto IsClassWithClassInfo = [=](const FMetasoundFrontendClass& InClass)
+					{
+						return FDocumentController::IsMatchingMetasoundClass(InInfo, InClass.Metadata);
+					};
+
+					FMetasoundFrontendClass* MetasoundClass = InDocument.Dependencies.FindByPredicate(IsClassWithClassInfo);
+					if (nullptr == MetasoundClass)
+					{
+						MetasoundClass = InDocument.Subgraphs.FindByPredicate(IsClassWithClassInfo);
+					}
+					return MetasoundClass;
+				};
+			}
+
 			auto FGetConstClassFromDocumentWithClassInfo(const FNodeClassInfo& InInfo)
 			{
 				return [=](const FMetasoundFrontendDocument& InDocument) -> const FMetasoundFrontendClass*
@@ -254,12 +272,30 @@ namespace Metasound
 				};
 			}
 
+			auto FGetClassFromDocumentWithMetadata(const FMetasoundFrontendClassMetadata& InMetadata)
+			{
+				return [=](FMetasoundFrontendDocument& InDocument) -> FMetasoundFrontendClass*
+				{
+					auto IsClassWithMetadata = [=](const FMetasoundFrontendClass& InClass)
+					{
+						return FDocumentController::IsMatchingMetasoundClass(InMetadata, InClass.Metadata);
+					};
+
+					FMetasoundFrontendClass* MetasoundClass = InDocument.Dependencies.FindByPredicate(IsClassWithMetadata);
+					if (nullptr == MetasoundClass)
+					{
+						MetasoundClass = InDocument.Subgraphs.FindByPredicate(IsClassWithMetadata);
+					}
+					return MetasoundClass;
+				};
+			}
+
 			auto FGetConstClassFromDocumentWithMetadata(const FMetasoundFrontendClassMetadata& InMetadata)
 			{
 				return [=](const FMetasoundFrontendDocument& InDocument) -> const FMetasoundFrontendClass*
 				{
-					auto IsClassWithMetadata = [=](const FMetasoundFrontendClass& InClass) 
-					{ 
+					auto IsClassWithMetadata = [=](const FMetasoundFrontendClass& InClass)
+					{
 						return FDocumentController::IsMatchingMetasoundClass(InMetadata, InClass.Metadata);
 					};
 
@@ -334,6 +370,7 @@ namespace Metasound
 
 		}
 
+
 		FDocumentAccess IDocumentAccessor::GetSharedAccess(IDocumentAccessor& InDocumentAccessor)
 		{
 			return InDocumentAccessor.ShareAccess();
@@ -344,6 +381,7 @@ namespace Metasound
 			return InDocumentAccessor.ShareAccess();
 		}
 
+
 		//
 		// FBaseOutputController
 		//
@@ -492,7 +530,6 @@ namespace Metasound
 			return TArray<FMetasoundFrontendEdge>();
 		}
 
-
 		const FText& FBaseOutputController::GetDisplayName() const
 		{
 			if (ClassOutputPtr.IsValid())
@@ -1080,6 +1117,7 @@ namespace Metasound
 			return Access;
 		}
 
+
 		//
 		// FInputNodeInputController
 		//
@@ -1140,6 +1178,7 @@ namespace Metasound
 			return false;
 		}
 
+
 		//
 		// FBaseNodeController
 		//
@@ -1196,13 +1235,22 @@ namespace Metasound
 			return Metasound::FrontendInvalidID;
 		}
 
-		FMetasoundFrontendVersionNumber FBaseNodeController::GetClassVersionNumber() const
+		const FMetasoundFrontendClassInterface& FBaseNodeController::GetClassInterface() const
+		{
+			if (ClassPtr.IsValid())
+			{
+				return ClassPtr->Interface;
+			}
+			return FrontendControllerIntrinsics::GetInvalidValueConstRef<FMetasoundFrontendClassInterface>();
+		}
+
+		const FMetasoundFrontendClassMetadata& FBaseNodeController::GetClassMetadata() const
 		{
 			if (ClassPtr.IsValid())
 			{
-				return ClassPtr->Metadata.Version;
+				return ClassPtr->Metadata;
 			}
-			return FMetasoundFrontendVersionNumber();
+			return FrontendControllerIntrinsics::GetInvalidValueConstRef<FMetasoundFrontendClassMetadata>();
 		}
 
 		const FMetasoundFrontendInterfaceStyle& FBaseNodeController::GetInputStyle() const
@@ -1275,25 +1323,6 @@ namespace Metasound
 			return FrontendControllerIntrinsics::GetInvalidValueConstRef<FString>();
 		}
 
-		EMetasoundFrontendClassType FBaseNodeController::GetClassType() const
-		{
-			if (ClassPtr.IsValid())
-			{
-				return ClassPtr->Metadata.Type;
-			}
-			return EMetasoundFrontendClassType::Invalid;
-		}
-
-		const FMetasoundFrontendClassName& FBaseNodeController::GetClassName() const
-		{
-			if (ClassPtr.IsValid())
-			{
-				return ClassPtr->Metadata.ClassName;
-			}
-
-			return FrontendControllerIntrinsics::GetInvalidValueConstRef<FMetasoundFrontendClassName>();
-		}
-
 		bool FBaseNodeController::CanAddInput(const FString& InVertexName) const
 		{
 			// TODO: not yet supported
@@ -1641,7 +1670,6 @@ namespace Metasound
 			return FInvalidOutputController::GetInvalid();
 		}
 
-
 		TArray<FBaseNodeController::FInputControllerParams> FBaseNodeController::GetInputControllerParams() const
 		{
 			using namespace FrontendControllerIntrinsics;
@@ -1845,6 +1873,7 @@ namespace Metasound
 			return Access;
 		}
 
+
 		// 
 		// FNodeController
 		//
@@ -1922,6 +1951,7 @@ namespace Metasound
 			return Access;
 		}
 
+
 		//
 		// FOutputNodeController
 		//
@@ -2077,6 +2107,7 @@ namespace Metasound
 			return Access;
 		}
 
+
 		//
 		// FInputNodeController
 		//
@@ -2135,7 +2166,6 @@ namespace Metasound
 			return FInvalidNodeController::GetInvalid();
 		}
 
-
 		bool FInputNodeController::IsValid() const
 		{
 			return OwningGraph->IsValid() && OwningGraphClassInputPtr.IsValid() && GraphPtr.IsValid();
@@ -2815,7 +2845,6 @@ namespace Metasound
 			return ClassMetadata;
 		}
 
-
 		// This can be used to clear the current literal for a given input.
 		// @returns false if the input name couldn't be found.
 		bool FGraphController::ClearLiteralForInput(const FString& InInputName, FGuid InVertexID)
@@ -2853,8 +2882,7 @@ namespace Metasound
 
 		FNodeHandle FGraphController::AddNode(const FMetasoundFrontendClassMetadata& InClassMetadata)
 		{
-			return AddNode(FMetasoundFrontendRegistryContainer
-::GetRegistryKey(InClassMetadata));
+			return AddNode(FMetasoundFrontendRegistryContainer::GetRegistryKey(InClassMetadata));
 		}
 
 		// Remove the node corresponding to this node handle.
@@ -2862,23 +2890,33 @@ namespace Metasound
 		bool FGraphController::RemoveNode(INodeController& InNode)
 		{
 			FGuid NodeID = InNode.GetID();
-			const EMetasoundFrontendClassType NodeType = InNode.GetClassType();
-
 			auto IsNodeWithSameID = [&](const FMetasoundFrontendNode& InDesc) { return InDesc.ID == NodeID; };
-			
 			if (const FMetasoundFrontendNode* Desc = GraphClassPtr->Graph.Nodes.FindByPredicate(IsNodeWithSameID))
 			{
-				if (EMetasoundFrontendClassType::Input == NodeType)
-				{
-					return RemoveInput(*Desc);
-				}
-				else if (EMetasoundFrontendClassType::Output == NodeType)
+				switch(InNode.GetClassMetadata().Type)
 				{
-					return RemoveOutput(*Desc);
-				}
-				else
-				{
-					return RemoveNode(*Desc);
+					case EMetasoundFrontendClassType::Input:
+					{
+						return RemoveInput(*Desc);
+					}
+
+					case EMetasoundFrontendClassType::Output:
+					{
+						return RemoveOutput(*Desc);
+					}
+
+					case EMetasoundFrontendClassType::External:
+					case EMetasoundFrontendClassType::Graph:
+					{
+						return RemoveNode(*Desc);
+					}
+
+					default:
+					case EMetasoundFrontendClassType::Invalid:
+					{
+						static_assert(static_cast<int32>(EMetasoundFrontendClassType::Invalid) == 4, "Possible missing switch case coverage for EMetasoundFrontendClassType.");
+						checkNoEntry();
+					}
 				}
 			}
 
@@ -2999,7 +3037,7 @@ namespace Metasound
 
 				NumRemoved += GraphClassPtr->Graph.Nodes.RemoveAll(IsNodeWithID);
 
-				OwningDocument->RemoveUnreferencedDependencies();
+				OwningDocument->SynchronizeDependencies();
 
 				return (NumRemoved > 0);
 			}
@@ -3383,7 +3421,6 @@ namespace Metasound
 			return FInvalidNodeController::GetInvalid();
 		}
 
-
 		FMetasoundFrontendClassInput* FGraphController::FindInputDescriptionWithName(const FString& InName)
 		{
 			if (GraphClassPtr.IsValid())
@@ -3503,6 +3540,7 @@ namespace Metasound
 			return Access;
 		}
 
+
 		//
 		// FDocumentController
 		//
@@ -3600,20 +3638,38 @@ namespace Metasound
 			
 			FConstClassAccessPtr ClassPtr = FindClass(InNodeClass);
 
-			if ((!ClassPtr.IsValid()) && DocumentPtr.IsValid())
+			auto AddClass = [=](FMetasoundFrontendClass&& NewClassDescription)
 			{
-				// FNodeClassInfo does not contain enough info add a subgraph.
-				check(EMetasoundFrontendClassType::Graph != InNodeClass.NodeType);
+				FConstClassAccessPtr NewClassPtr;
+				if (DocumentPtr.IsValid())
+				{
+					// FNodeClassInfo does not contain enough info add a subgraph.
+					check(EMetasoundFrontendClassType::Graph != InNodeClass.NodeType);
+					NewClassDescription.ID = NewClassID();
 
-				FMetasoundFrontendClass NewClass = GenerateClassDescription(InNodeClass);
-				NewClass.ID = NewClassID();
+					DocumentPtr->Dependencies.Add(MoveTemp(NewClassDescription));
+					NewClassPtr = FindClass(InNodeClass);
+				}
+				return NewClassPtr;
+			};
 
-				DocumentPtr->Dependencies.Add(NewClass);
+			if (ClassPtr.IsValid())
+			{
+				// External node classes must match version to return shared definition.
+				if (InNodeClass.NodeType == EMetasoundFrontendClassType::External)
+				{
+					FMetasoundFrontendClass NewClass = GenerateClassDescription(InNodeClass);
+					if (NewClass.Metadata.Version.Major != ClassPtr->Metadata.Version.Major)
+					{
+						return AddClass(MoveTemp(NewClass));
+					}
+				}
 
-				ClassPtr = FindClass(InNodeClass);
+				return ClassPtr;
 			}
 
-			return ClassPtr;
+			FMetasoundFrontendClass NewClass = GenerateClassDescription(InNodeClass);
+			return AddClass(MoveTemp(NewClass));
 		}
 
 		FConstClassAccessPtr FDocumentController::FindClass(const FMetasoundFrontendClassMetadata& InMetadata) const
@@ -3627,75 +3683,117 @@ namespace Metasound
 		{
 			FConstClassAccessPtr ClassPtr = FindClass(InMetadata);
 
-			if ((!ClassPtr.IsValid()) && DocumentPtr.IsValid())
-			{
-				if ((EMetasoundFrontendClassType::External == InMetadata.Type) || (EMetasoundFrontendClassType::Input == InMetadata.Type) || (EMetasoundFrontendClassType::Output == InMetadata.Type))
-				{
-					FMetasoundFrontendClass NewClass;
-					if (FMetasoundFrontendRegistryContainer::GetFrontendClassFromRegistered(InMetadata, NewClass))
-					{
-						NewClass.ID = NewClassID();
-						DocumentPtr->Dependencies.Add(NewClass);
-					}
-					else
-					{
-						UE_LOG(LogMetaSound, Error, TEXT("Cannot add external dependency. No Metasound class found with matching metadata Name: \"%s\", Version %d.%d. Suggested solution \"%s\" by %s."), *InMetadata.ClassName.GetFullName().ToString(), InMetadata.Version.Major, InMetadata.Version.Minor, *InMetadata.PromptIfMissing.ToString(), *InMetadata.Author.ToString());
-					}
-				} 
-				else if (EMetasoundFrontendClassType::Graph == InMetadata.Type)
-				{
-					FMetasoundFrontendGraphClass NewClass;
-					NewClass.ID = NewClassID();
-					NewClass.Metadata = InMetadata;
+if (!DocumentPtr.IsValid())
+{
+	return ClassPtr;
+}
 
-					DocumentPtr->Subgraphs.Add(NewClass);
-				}
-				else
-				{
-					UE_LOG(LogMetaSound, Error, TEXT("Unsupported metasound class type for node: \"%s\", Version %d.%d."), *InMetadata.ClassName.GetFullName().ToString(), InMetadata.Version.Major, InMetadata.Version.Minor);
-					checkNoEntry();
-				}
+// External node classes must match major version to return shared definition.
+if (EMetasoundFrontendClassType::External == InMetadata.Type)
+{
+	if (ClassPtr.IsValid())
+	{
+		if (InMetadata.Version.Major != ClassPtr->Metadata.Version.Major)
+		{
+			ClassPtr = FConstClassAccessPtr();
+		}
+	}
+}
 
-				ClassPtr = FindClass(InMetadata);
-			}
+if (!ClassPtr.IsValid())
+{
+	if ((EMetasoundFrontendClassType::External == InMetadata.Type) || (EMetasoundFrontendClassType::Input == InMetadata.Type) || (EMetasoundFrontendClassType::Output == InMetadata.Type))
+	{
+		FMetasoundFrontendClass NewClass;
+		if (FMetasoundFrontendRegistryContainer::GetFrontendClassFromRegistered(InMetadata, NewClass))
+		{
+			NewClass.ID = NewClassID();
+			DocumentPtr->Dependencies.Add(NewClass);
+		}
+		else
+		{
+			UE_LOG(LogMetaSound, Error,
+				TEXT("Cannot add external dependency. No Metasound class found with matching metadata Name: "
+					"\"%s\" (%s). Suggested solution \"%s\" by %s."),
+				*InMetadata.ClassName.GetFullName().ToString(),
+				*InMetadata.Version.ToString(),
+				*InMetadata.PromptIfMissing.ToString(),
+				*InMetadata.Author.ToString());
+		}
+	}
+	else if (EMetasoundFrontendClassType::Graph == InMetadata.Type)
+	{
+		FMetasoundFrontendGraphClass NewClass;
+		NewClass.ID = NewClassID();
+		NewClass.Metadata = InMetadata;
 
-			return ClassPtr;
+		DocumentPtr->Subgraphs.Add(NewClass);
+	}
+	else
+	{
+		UE_LOG(LogMetaSound, Error, TEXT(
+			"Unsupported metasound class type for node: \"%s\" (%s)."),
+			*InMetadata.ClassName.GetFullName().ToString(),
+			*InMetadata.Version.ToString());
+		checkNoEntry();
+	}
+
+	ClassPtr = FindClass(InMetadata);
+}
+
+return ClassPtr;
 		}
 
-		void FDocumentController::RemoveUnreferencedDependencies()
+		void FDocumentController::SynchronizeDependencies()
 		{
-			if (DocumentPtr.IsValid())
+			if (!DocumentPtr.IsValid())
 			{
-				int32 NumDependenciesRemovedThisItr = 0;
+				return;
+			}
+
+			int32 NumDependenciesRemovedThisItr = 0;
 
-				// Repeatedly remove unreferenced dependencies until there are 
-				// no unreferenced dependencies left. 
-				do
+			// Repeatedly remove unreferenced dependencies until there are
+			// no unreferenced dependencies left.
+			do
+			{
+				TSet<FGuid> ReferencedDependencyIDs;
+				auto AddNodeClassIDToSet = [&](const FMetasoundFrontendNode& Node)
 				{
-					TSet<FGuid> ReferencedDependencyIDs;
-					auto AddNodeClassIDToSet = [&](const FMetasoundFrontendNode& Node)
-					{
-						ReferencedDependencyIDs.Add(Node.ClassID);
-					};
+					ReferencedDependencyIDs.Add(Node.ClassID);
+				};
 
-					auto AddGraphNodeClassIDsToSet = [&](const FMetasoundFrontendGraphClass& GraphClass)
-					{
-						Algo::ForEach(GraphClass.Graph.Nodes, AddNodeClassIDToSet);
-					};
-					
-					// Referenced dependencies in root class
-					Algo::ForEach(DocumentPtr->RootGraph.Graph.Nodes, AddNodeClassIDToSet);
-					// Referenced dependencies in 
-					Algo::ForEach(DocumentPtr->Subgraphs, AddGraphNodeClassIDsToSet);
+				auto AddGraphNodeClassIDsToSet = [&](const FMetasoundFrontendGraphClass& GraphClass)
+				{
+					Algo::ForEach(GraphClass.Graph.Nodes, AddNodeClassIDToSet);
+				};
 
-					auto IsDependencyUnreferenced = [&](const FMetasoundFrontendClass& ClassDependency)
-					{
-						return !ReferencedDependencyIDs.Contains(ClassDependency.ID);
-					};
+				// Referenced dependencies in root class
+				Algo::ForEach(DocumentPtr->RootGraph.Graph.Nodes, AddNodeClassIDToSet);
 
-					NumDependenciesRemovedThisItr = DocumentPtr->Dependencies.RemoveAllSwap(IsDependencyUnreferenced);
-				} 
-				while(NumDependenciesRemovedThisItr > 0);
+				// Referenced dependencies in subgraphs
+				Algo::ForEach(DocumentPtr->Subgraphs, AddGraphNodeClassIDsToSet);
+
+				auto IsDependencyUnreferenced = [&](const FMetasoundFrontendClass& ClassDependency)
+				{
+					return !ReferencedDependencyIDs.Contains(ClassDependency.ID);
+				};
+
+				NumDependenciesRemovedThisItr = DocumentPtr->Dependencies.RemoveAllSwap(IsDependencyUnreferenced);
+			}
+			while (NumDependenciesRemovedThisItr > 0);
+
+			FMetasoundFrontendRegistryContainer* Registry = FMetasoundFrontendRegistryContainer::Get();
+			if (ensure(Registry))
+			{
+				for (FMetasoundFrontendClass& Class : DocumentPtr->Dependencies)
+				{
+					FMetasoundFrontendClass RegisteredClass;
+					if (Registry->FindFrontendClassFromRegistered(Class.Metadata, RegisteredClass))
+					{
+						Class.Metadata = RegisteredClass.Metadata;
+					}
+				}
 			}
 		}
 
@@ -3854,4 +3952,4 @@ namespace Metasound
 		}
 	}
 }
-#undef LOCTEXT_NAMESPACE 
+#undef LOCTEXT_NAMESPACE
diff --git a/Engine/Plugins/Runtime/Metasound/Source/MetasoundFrontend/Private/MetasoundFrontendStandardController.h b/Engine/Plugins/Runtime/Metasound/Source/MetasoundFrontend/Private/MetasoundFrontendStandardController.h
--- a/Engine/Plugins/Runtime/Metasound/Source/MetasoundFrontend/Private/MetasoundFrontendStandardController.h
+++ b/Engine/Plugins/Runtime/Metasound/Source/MetasoundFrontend/Private/MetasoundFrontendStandardController.h
@@ -19,7 +19,7 @@ namespace Metasound
 		/** FBaseOutputController provides common functionality for multiple derived
 		 * output controllers.
 		 */
-		class METASOUNDFRONTEND_API FBaseOutputController: public IOutputController
+		class METASOUNDFRONTEND_API FBaseOutputController : public IOutputController
 		{
 		public:
 
@@ -352,7 +352,8 @@ namespace Metasound
 			FGuid GetID() const override;
 			FGuid GetClassID() const override;
 
-			FMetasoundFrontendVersionNumber GetClassVersionNumber() const override;
+			const FMetasoundFrontendClassInterface& GetClassInterface() const override;
+			const FMetasoundFrontendClassMetadata& GetClassMetadata() const override;
 			const FMetasoundFrontendInterfaceStyle& GetInputStyle() const override;
 			const FMetasoundFrontendInterfaceStyle& GetOutputStyle() const override;
 			const FMetasoundFrontendClassStyle& GetClassStyle() const override;
@@ -364,7 +365,6 @@ namespace Metasound
 			void SetNodeStyle(const FMetasoundFrontendNodeStyle& InStyle) override;
 
 			const FString& GetNodeName() const override;
-			EMetasoundFrontendClassType GetClassType() const override;
 
 			bool CanAddInput(const FString& InVertexName) const override;
 			FInputHandle AddInput(const FString& InVertexName, const FMetasoundFrontendLiteral* InDefault) override;
@@ -374,8 +374,6 @@ namespace Metasound
 			FInputHandle AddOutput(const FString& InVertexName, const FMetasoundFrontendLiteral* InDefault) override;
 			bool RemoveOutput(FGuid InVertexID) override;
 
-			const FMetasoundFrontendClassName& GetClassName() const override;
-
 			/** Returns all node inputs. */
 			TArray<FInputHandle> GetInputs() override;
 
@@ -484,15 +482,12 @@ namespace Metasound
 			const FMetasoundFrontendVertex* FindNodeInputWithName(const FString& InName) const;
 			const FMetasoundFrontendVertex* FindNodeOutputWithName(const FString& InName) const;
 
-
-
 			virtual FInputHandle CreateInputController(FGuid InVertexID, FConstVertexAccessPtr InNodeVertexPtr, FConstClassInputAccessPtr InClassInputPtr, FNodeHandle InOwningNode) const = 0;
-
 			virtual FOutputHandle CreateOutputController(FGuid InVertexID, FConstVertexAccessPtr InNodeVertexPtr, FConstClassOutputAccessPtr InClassOutputPtr, FNodeHandle InOwningNode) const = 0;
 		};
 
 		/** FNodeController represents a external or subgraph node. */
-		class METASOUNDFRONTEND_API FNodeController: public FBaseNodeController
+		class METASOUNDFRONTEND_API FNodeController : public FBaseNodeController
 		{
 			// Private token only allows members or friends to call constructor.
 			enum EPrivateToken { Token };
@@ -530,19 +525,15 @@ namespace Metasound
 			bool IsValid() const override;
 
 		protected:
-
 			FDocumentAccess ShareAccess() override;
 			FConstDocumentAccess ShareAccess() const override;
 
 
 		private:
-
 			FInputHandle CreateInputController(FGuid InVertexID, FConstVertexAccessPtr InNodeVertexPtr, FConstClassInputAccessPtr InClassInputPtr, FNodeHandle InOwningNode) const override;
-
 			FOutputHandle CreateOutputController(FGuid InVertexID, FConstVertexAccessPtr InNodeVertexPtr, FConstClassOutputAccessPtr InClassOutputPtr, FNodeHandle InOwningNode) const override;
 
 			FGraphAccessPtr GraphPtr;
-
 		};
 
 		/** FOutputNodeController represents an output node. */
@@ -925,7 +916,7 @@ namespace Metasound
 			FConstClassAccessPtr FindOrAddClass(const FNodeClassInfo& InNodeClass) override;
 			FConstClassAccessPtr FindOrAddClass(const FMetasoundFrontendClassMetadata& InMetadata) override;
 
-			void RemoveUnreferencedDependencies() override;
+			void SynchronizeDependencies() override;
 
 			FGraphHandle GetRootGraph() override;
 			FConstGraphHandle GetRootGraph() const override;
diff --git a/Engine/Plugins/Runtime/Metasound/Source/MetasoundFrontend/Public/MetasoundFrontendController.h b/Engine/Plugins/Runtime/Metasound/Source/MetasoundFrontend/Public/MetasoundFrontendController.h
--- a/Engine/Plugins/Runtime/Metasound/Source/MetasoundFrontend/Public/MetasoundFrontendController.h
+++ b/Engine/Plugins/Runtime/Metasound/Source/MetasoundFrontend/Public/MetasoundFrontendController.h
@@ -383,12 +383,9 @@ namespace Metasound
 			virtual FInputHandle AddOutput(const FString& InVertexName, const FMetasoundFrontendLiteral* InDefault) = 0;
 			virtual bool RemoveOutput(FGuid InVertexID) = 0;
 
-			/** Returns the class type for this node. */
-			virtual EMetasoundFrontendClassType GetClassType() const = 0;
-
-			/** Returns the name of the node class. */
-			virtual const FMetasoundFrontendClassName& GetClassName() const = 0;
-			virtual FMetasoundFrontendVersionNumber GetClassVersionNumber() const = 0;
+			/** Returns associated node class data */
+			virtual const FMetasoundFrontendClassInterface& GetClassInterface() const = 0;
+			virtual const FMetasoundFrontendClassMetadata& GetClassMetadata() const = 0;
 			virtual const FMetasoundFrontendInterfaceStyle& GetOutputStyle() const = 0;
 			virtual const FMetasoundFrontendInterfaceStyle& GetInputStyle() const = 0;
 			virtual const FMetasoundFrontendClassStyle& GetClassStyle() const = 0;
@@ -524,7 +521,6 @@ namespace Metasound
 			virtual FConstClassInputAccessPtr FindClassInputWithName(const FString& InName) const = 0;
 			virtual FConstClassOutputAccessPtr FindClassOutputWithName(const FString& InName) const = 0;
 
-
 			/** Add a new input node using the input description. 
 			 *
 			 * @param InDescription - Description for input of graph.
@@ -748,8 +744,9 @@ namespace Metasound
 			 */
 			virtual FConstClassAccessPtr FindOrAddClass(const FMetasoundFrontendClassMetadata& InMetadata) = 0;
 
-			/** Removes all dependencies which are no longer referenced by any graphs within this document.  */
-			virtual void RemoveUnreferencedDependencies() = 0;
+			/** Removes all dependencies which are no longer referenced by any graphs within this document
+			  * and updates dependency Metadata where necessary with that found in the registry.  */
+			virtual void SynchronizeDependencies() = 0;
 
 			/** Returns an array of all subgraphs for this document. */
 			virtual TArray<FGraphHandle> GetSubgraphHandles() = 0;
diff --git a/Engine/Plugins/Runtime/Metasound/Source/MetasoundFrontend/Public/MetasoundFrontendDocument.h b/Engine/Plugins/Runtime/Metasound/Source/MetasoundFrontend/Public/MetasoundFrontendDocument.h
--- a/Engine/Plugins/Runtime/Metasound/Source/MetasoundFrontend/Public/MetasoundFrontendDocument.h
+++ b/Engine/Plugins/Runtime/Metasound/Source/MetasoundFrontend/Public/MetasoundFrontendDocument.h
@@ -43,7 +43,7 @@ enum class EMetasoundFrontendClassType : uint8
 
 // General purpose version number for Metasound Frontend objects.
 USTRUCT()
-struct FMetasoundFrontendVersionNumber
+struct METASOUNDFRONTEND_API FMetasoundFrontendVersionNumber
 {
 	GENERATED_BODY()
 
@@ -54,6 +54,72 @@ struct FMetasoundFrontendVersionNumber
 	// Minor version number.
 	UPROPERTY(VisibleAnywhere, Category = General)
 	int32 Minor = 0;
+
+	static const FMetasoundFrontendVersionNumber& GetInvalid()
+	{
+		static const FMetasoundFrontendVersionNumber Invalid { 0, 0 };
+		return Invalid;
+	}
+
+	bool IsValid() const
+	{
+		return *this != GetInvalid();
+	}
+
+	friend bool operator==(const FMetasoundFrontendVersionNumber& InLHS, const FMetasoundFrontendVersionNumber& InRHS)
+	{
+		return InLHS.Major == InRHS.Major && InLHS.Minor == InRHS.Minor;
+	}
+
+	friend bool operator!=(const FMetasoundFrontendVersionNumber& InLHS, const FMetasoundFrontendVersionNumber& InRHS)
+	{
+		return InLHS.Major != InRHS.Major || InLHS.Minor != InRHS.Minor;
+	}
+
+	friend bool operator>(const FMetasoundFrontendVersionNumber& InLHS, const FMetasoundFrontendVersionNumber& InRHS)
+	{
+		if (InLHS.Major > InRHS.Major)
+		{
+			return true;
+		}
+
+		if (InLHS.Major == InRHS.Major)
+		{
+			return InLHS.Minor > InRHS.Minor;
+		}
+
+		return false;
+	}
+
+	friend bool operator>=(const FMetasoundFrontendVersionNumber& InLHS, const FMetasoundFrontendVersionNumber& InRHS)
+	{
+		return InLHS == InRHS || InLHS > InRHS;
+	}
+
+	friend bool operator<(const FMetasoundFrontendVersionNumber& InLHS, const FMetasoundFrontendVersionNumber& InRHS)
+	{
+		if (InLHS.Major < InRHS.Major)
+		{
+			return true;
+		}
+
+		if (InLHS.Major == InRHS.Major)
+		{
+			return InLHS.Minor < InRHS.Minor;
+		}
+
+		return false;
+	}
+
+	friend bool operator<=(const FMetasoundFrontendVersionNumber& InLHS, const FMetasoundFrontendVersionNumber& InRHS)
+	{
+		return InLHS == InRHS || InLHS < InRHS;
+	}
+
+	FString ToString() const
+	{
+		return FString::Format(TEXT("v{0}.{1}"), { Major, Minor });
+	}
 };
 
 // General purpose version info for Metasound Frontend objects.
@@ -534,7 +600,6 @@ struct METASOUNDFRONTEND_API FMetasoundFrontendClassName
 	FMetasoundFrontendClassName(const FName& InNamespace, const FName& InName, const FName& InVariant);
 
 	FMetasoundFrontendClassName(const Metasound::FNodeClassName& InName);
-	
 
 	// Namespace of class.
 	UPROPERTY(EditAnywhere, Category = General)
@@ -554,6 +619,13 @@ struct METASOUNDFRONTEND_API FMetasoundFrontendClassName
 	// Returns scoped name representing namespace and name. 
 	FName GetScopedName() const;
 
+	// Returns NodeClassName version of full name
+	Metasound::FNodeClassName ToNodeClassName() const
+	{
+		return { Namespace, Name, Variant };
+	}
+
+
 	// Return string version of full name.
 	FString ToString() const;
 
diff --git a/Engine/Plugins/Runtime/Metasound/Source/MetasoundFrontend/Public/MetasoundFrontendQuery.h b/Engine/Plugins/Runtime/Metasound/Source/MetasoundFrontend/Public/MetasoundFrontendQuery.h
--- a/Engine/Plugins/Runtime/Metasound/Source/MetasoundFrontend/Public/MetasoundFrontendQuery.h
+++ b/Engine/Plugins/Runtime/Metasound/Source/MetasoundFrontend/Public/MetasoundFrontendQuery.h
@@ -3,6 +3,7 @@
 #pragma once
 
 #include "CoreMinimal.h"
+#include "Math/NumericLimits.h"
 #include "MetasoundFrontendDocument.h"
 #include "Misc/TVariant.h"
 
@@ -16,7 +17,7 @@ namespace Metasound
 		using FValue = TVariant<FMetasoundFrontendDocument, FMetasoundFrontendClass, FMetasoundFrontendNode>;
 		using FKey = int32;
 
-		static constexpr FKey InvalidKey = 0xFFFFFFFF;
+		static constexpr FKey InvalidKey = TNumericLimits<int32>::Max();
 
 		FKey Key = InvalidKey;
 		FValue Value;
diff --git a/Engine/Plugins/Runtime/Metasound/Source/MetasoundFrontend/Public/MetasoundFrontendQuerySteps.h b/Engine/Plugins/Runtime/Metasound/Source/MetasoundFrontend/Public/MetasoundFrontendQuerySteps.h
--- a/Engine/Plugins/Runtime/Metasound/Source/MetasoundFrontend/Public/MetasoundFrontendQuerySteps.h
+++ b/Engine/Plugins/Runtime/Metasound/Source/MetasoundFrontend/Public/MetasoundFrontendQuerySteps.h
@@ -85,5 +85,33 @@ namespace Metasound
 	private:
 		FGuid ClassID;
 	};
-}
 
+	class METASOUNDFRONTEND_API FMapClassNameToMajorVersion : public IFrontendQueryMapStep
+	{
+	public:
+		FFrontendQueryEntry::FKey Map(const FFrontendQueryEntry& InEntry) const override;
+	};
+
+	class METASOUNDFRONTEND_API FReduceClassesToHighestVersion : public IFrontendQueryReduceStep
+	{
+	public:
+		void Reduce(FFrontendQueryEntry::FKey InKey, TArrayView<FFrontendQueryEntry*>& InEntries, FReduceOutputView& OutResult) const override;
+	};
+
+	class METASOUNDFRONTEND_API FReduceClassesToMajorVersion : public IFrontendQueryReduceStep
+	{
+	public:
+		FReduceClassesToMajorVersion(int32 InMajorVersion);
+
+		void Reduce(FFrontendQueryEntry::FKey InKey, TArrayView<FFrontendQueryEntry*>& InEntries, FReduceOutputView& OutResult) const override;
+
+	private:
+		int32 MajorVersion = -1;
+	};
+
+	class METASOUNDFRONTEND_API FSortClassesByVersion : public IFrontendQuerySortStep
+	{
+	public:
+		bool Sort(const FFrontendQueryEntry& InEntryLHS, const FFrontendQueryEntry& InEntryRHS) const override;
+	};
+}
diff --git a/Engine/Plugins/Runtime/Metasound/Source/MetasoundFrontend/Public/MetasoundFrontendRegistries.h b/Engine/Plugins/Runtime/Metasound/Source/MetasoundFrontend/Public/MetasoundFrontendRegistries.h
--- a/Engine/Plugins/Runtime/Metasound/Source/MetasoundFrontend/Public/MetasoundFrontendRegistries.h
+++ b/Engine/Plugins/Runtime/Metasound/Source/MetasoundFrontend/Public/MetasoundFrontendRegistries.h
@@ -61,18 +61,18 @@ namespace Metasound
 	};
 	*/
 
-
-	typedef TFunction<TUniquePtr<Metasound::INode>(::Metasound::FInputNodeConstructorParams&&)> FCreateInputNodeFunction;
-	typedef TFunction<TUniquePtr<Metasound::INode>(const ::Metasound::FOutputNodeConstructorParams&)> FCreateOutputNodeFunction;
+	using FCreateInputNodeFunction = TFunction<TUniquePtr<Metasound::INode>(::Metasound::FInputNodeConstructorParams&&)>;
+	using FCreateOutputNodeFunction = TFunction<TUniquePtr<Metasound::INode>(const ::Metasound::FOutputNodeConstructorParams&)>;
 
 	// This function is used to create a proxy from a datatype's base uclass.
-	typedef TFunction<Audio::IProxyDataPtr(UObject*)> FCreateAudioProxyFunction;
+	using FCreateAudioProxyFunction = TFunction<Audio::IProxyDataPtr(UObject*)> ;
 
 	// Creates a data channel for senders and receivers of this data type.
-	typedef TFunction<TSharedPtr<IDataChannel, ESPMode::ThreadSafe>(const FOperatorSettings&)> FCreateDataChannelFunction;
+	using FCreateDataChannelFunction = TFunction<TSharedPtr<IDataChannel, ESPMode::ThreadSafe>(const FOperatorSettings&)>;
 
-	typedef TFunction<TUniquePtr<Metasound::INode>(const Metasound::FNodeInitData&)> FCreateMetasoundNodeFunction;
-	typedef TFunction<FMetasoundFrontendClass()> FCreateMetasoundFrontendClassFunction;
+	using FCreateMetasoundNodeFunction = TFunction<TUniquePtr<Metasound::INode>(const Metasound::FNodeInitData&)>;
+	using FCreateMetasoundFrontendClassFunction = TFunction<FMetasoundFrontendClass()>;
+	using FIterateMetasoundFrontendClassFunction = TFunction<void(const FMetasoundFrontendClass&)>;
 
 	// Various elements that we pass to the frontend registry based on templated type traits.
 	struct FDataTypeRegistryInfo
@@ -273,7 +273,7 @@ namespace Metasound
 
 		FCreateDataChannelFunction CreateDataChannel;
 	};
-}
+} // namespace Frontend
 
 /**
  * Singleton registry for all types and nodes.
@@ -330,7 +330,6 @@ public:
 	virtual bool RegisterConversionNode(const FConverterNodeRegistryKey& InNodeKey, const FConverterNodeInfo& InNodeInfo) = 0;
 	virtual bool IsNodeRegistered(const FNodeRegistryKey& InKey) const = 0;
 
-
 	/** Return all node classes registered. */
 	virtual TArray<Metasound::Frontend::FNodeClassInfo> GetAllAvailableNodeClasses(Metasound::Frontend::FRegistryTransactionID* OutCurrentRegistryTransactionID) const = 0;
 
@@ -340,6 +339,9 @@ public:
 	// Return any data types that can be used as a metasound input type or output type.
 	virtual TArray<FName> GetAllValidDataTypes() = 0;
 
+	// Iterates class types in registry.  If InClassType is set to a valid class type (optional), only iterates classes of the given type
+	virtual void IterateRegistry(Metasound::FIterateMetasoundFrontendClassFunction InIterFunc, EMetasoundFrontendClassType InClassType = EMetasoundFrontendClassType::Invalid) const = 0;
+
 	// Query for MetaSound Frontend document objects.
 	virtual bool FindFrontendClassFromRegistered(const Metasound::Frontend::FNodeClassInfo& InClassInfo, FMetasoundFrontendClass& OutClass) = 0;
 	virtual bool FindFrontendClassFromRegistered(const FMetasoundFrontendClassMetadata& InMetadata, FMetasoundFrontendClass& OutClass) = 0;
@@ -355,8 +357,6 @@ public:
 	// Returns an empty array if none are available.
 	virtual TArray<FConverterNodeInfo> GetPossibleConverterNodes(const FName& FromDataType, const FName& ToDataType) = 0;
 
-
-
 	virtual bool RegisterDataType(const FDataTypeRegistryInfo& InDataInfo, const FDataTypeConstructorCallbacks& InCallbacks) = 0;
 
 	virtual bool RegisterEnumDataInterface(FName InDataType, TSharedPtr<IEnumDataTypeInterface>&& InInterface) = 0;
diff --git a/Engine/Plugins/Runtime/Metasound/Source/MetasoundFrontend/Public/MetasoundFrontendSearchEngine.h b/Engine/Plugins/Runtime/Metasound/Source/MetasoundFrontend/Public/MetasoundFrontendSearchEngine.h
--- a/Engine/Plugins/Runtime/Metasound/Source/MetasoundFrontend/Public/MetasoundFrontendSearchEngine.h
+++ b/Engine/Plugins/Runtime/Metasound/Source/MetasoundFrontend/Public/MetasoundFrontendSearchEngine.h
@@ -5,6 +5,7 @@
 #include "MetasoundFrontendDocument.h"
 #include "MetasoundNodeInterface.h"
 
+
 namespace Metasound
 {
 	namespace Frontend
@@ -19,13 +20,25 @@ namespace Metasound
 				/** Return an instance of a search engine. */
 				static ISearchEngine& Get();
 
-				/** Find all FMetasoundFrontendClasses with a matching FNodeClassName. */
-				virtual TArray<FMetasoundFrontendClass> FindClassesWithClassName(const FNodeClassName& InName) = 0;
+				/** Find all FMetasoundFrontendClasses.
+				  * (Optional) Include deprecated classes (i.e. versions of classes that are not the highest major version).
+				  */
+				virtual TArray<FMetasoundFrontendClass> FindAllClasses(bool bInIncludeDeprecated) = 0;
+
+				/** Find all classes with the given ClassName.
+				  * (Optional) Sort matches based on version.
+				  */
+				virtual TArray<FMetasoundFrontendClass> FindClassesWithName(const FNodeClassName& InName, bool bInSortByVersion) = 0;
+
+				/** Find the highest version of a class with the given ClassName. Returns false if not found, true if found. */
+				virtual bool FindClassWithHighestVersion(const FNodeClassName& InName, FMetasoundFrontendClass& OutClass) = 0;
+
+				/** Find the class with the given ClassName & Major Version. Returns false if not found, true if found. */
+				virtual bool FindClassWithMajorVersion(const FNodeClassName& InName, int32 InMajorVersion, FMetasoundFrontendClass& OutClass) = 0;
 
 				virtual ~ISearchEngine() = default;
 
 			protected:
-
 				ISearchEngine() = default;
 		};
 	}
diff --git a/Engine/Plugins/Runtime/Metasound/Source/MetasoundGraphCore/Public/MetasoundNodeInterface.h b/Engine/Plugins/Runtime/Metasound/Source/MetasoundGraphCore/Public/MetasoundNodeInterface.h
--- a/Engine/Plugins/Runtime/Metasound/Source/MetasoundGraphCore/Public/MetasoundNodeInterface.h
+++ b/Engine/Plugins/Runtime/Metasound/Source/MetasoundGraphCore/Public/MetasoundNodeInterface.h
@@ -14,7 +14,7 @@
 namespace Metasound
 {
 	static const FText PluginAuthor = NSLOCTEXT("MetasoundGraphCore", "Metasound_DefaultAuthor", "Epic Games, Inc.");
-	static const FText PluginNodeMissingPrompt = NSLOCTEXT("MetasoundGraphCore", "Metasound_DefaultMissingPrompt", "Make sure that the Metasound plugin is loaded.");
+	static const FText PluginNodeMissingPrompt = NSLOCTEXT("MetasoundGraphCore", "Metasound_DefaultMissingNodePrompt", "The node was likely removed, renamed, or the Metasound plugin is not loaded.");
 
 	/**
 	 * This struct is used to pass in any arguments required for constructing a single node instance.
