Fix null deref
#rnx

diff --git a/Engine/Plugins/Runtime/Metasound/Source/MetasoundFrontend/Private/MetasoundFrontendQuerySteps.cpp b/Engine/Plugins/Runtime/Metasound/Source/MetasoundFrontend/Private/MetasoundFrontendQuerySteps.cpp
--- a/Engine/Plugins/Runtime/Metasound/Source/MetasoundFrontend/Private/MetasoundFrontendQuerySteps.cpp
+++ b/Engine/Plugins/Runtime/Metasound/Source/MetasoundFrontend/Private/MetasoundFrontendQuerySteps.cpp
@@ -109,7 +109,10 @@ namespace Metasound
 			}
 		}
 
-		OutResult.Add(*HighestVersionEntry);
+		if (HighestVersionEntry)
+		{
+			OutResult.Add(*HighestVersionEntry);
+		}
 	}
 
 	FReduceClassesToMajorVersion::FReduceClassesToMajorVersion(int32 InMajorVersion)
