Fix for trigger counter

diff --git a/Engine/Plugins/Runtime/Metasound/Source/MetasoundStandardNodes/Private/MetasoundTriggerCounterNode.cpp b/Engine/Plugins/Runtime/Metasound/Source/MetasoundStandardNodes/Private/MetasoundTriggerCounterNode.cpp
--- a/Engine/Plugins/Runtime/Metasound/Source/MetasoundStandardNodes/Private/MetasoundTriggerCounterNode.cpp
+++ b/Engine/Plugins/Runtime/Metasound/Source/MetasoundStandardNodes/Private/MetasoundTriggerCounterNode.cpp
@@ -144,9 +144,9 @@ namespace Metasound
 			{
 				++CurrentTriggerCount;
 
-				if (CurrentAutoResetCount > 0 && CurrentTriggerCount >= CurrentAutoResetCount)
+				if (CurrentAutoResetCount > 0 && CurrentTriggerCount > CurrentAutoResetCount)
 				{
-					CurrentTriggerCount = 0;
+					CurrentTriggerCount = 1;
 					*OutValue = *StartValue;
 					TriggerOnReset->TriggerFrame(StartFrame);
 				}
