- Fix errors from cooking Editor-only content in MetaSounds by marking UMetasoundEditorGraphBase as IsEditorOnly() & !NeedsLoadForEditorGame()
- Make MetaSoundEdGraph transient (never serialized)
- Move FText::FromString calls to concrete FText definitions to avoid warnings while packaging

diff --git a/Engine/Plugins/Runtime/Metasound/Source/MetasoundEditor/Private/MetasoundEditor.cpp b/Engine/Plugins/Runtime/Metasound/Source/MetasoundEditor/Private/MetasoundEditor.cpp
--- a/Engine/Plugins/Runtime/Metasound/Source/MetasoundEditor/Private/MetasoundEditor.cpp
+++ b/Engine/Plugins/Runtime/Metasound/Source/MetasoundEditor/Private/MetasoundEditor.cpp
@@ -349,6 +349,15 @@ namespace Metasound
 
 			Metasound = ObjectToEdit;
 
+			FMetasoundAssetBase* MetasoundAsset = IMetasoundUObjectRegistry::Get().GetObjectAsAssetBase(Metasound);
+			if (!MetasoundAsset->GetGraph())
+			{
+				UMetasoundEditorGraph* Graph = NewObject<UMetasoundEditorGraph>(Metasound, FName(), RF_Transactional);
+				Graph->Schema = UMetasoundEditorGraphSchema::StaticClass();
+				MetasoundAsset->SetGraph(Graph);
+				FGraphBuilder::SynchronizeGraph(*Metasound);
+			}
+
 			GEditor->RegisterForUndo(this);
 
 			FGraphEditorCommands::Register();
diff --git a/Engine/Plugins/Runtime/Metasound/Source/MetasoundEditor/Private/MetasoundEditorGraph.cpp b/Engine/Plugins/Runtime/Metasound/Source/MetasoundEditor/Private/MetasoundEditorGraph.cpp
--- a/Engine/Plugins/Runtime/Metasound/Source/MetasoundEditor/Private/MetasoundEditorGraph.cpp
+++ b/Engine/Plugins/Runtime/Metasound/Source/MetasoundEditor/Private/MetasoundEditorGraph.cpp
@@ -368,18 +368,19 @@ Metasound::Frontend::FGraphHandle UMetasoundEditorGraph::GetGraphHandle() const
 
 UObject* UMetasoundEditorGraph::GetMetasound() const
 {
-	return ParentMetasound;
+	return GetOuter();
 }
 
 UObject& UMetasoundEditorGraph::GetMetasoundChecked() const
 {
+	UObject* ParentMetasound = GetOuter();
 	check(ParentMetasound);
 	return *ParentMetasound;
 }
 
 void UMetasoundEditorGraph::Synchronize()
 {
-	if (ParentMetasound)
+	if (UObject* ParentMetasound = GetOuter())
 	{
 		Metasound::Editor::FGraphBuilder::SynchronizeGraph(*ParentMetasound);
 	}
diff --git a/Engine/Plugins/Runtime/Metasound/Source/MetasoundEditor/Private/MetasoundFactory.cpp b/Engine/Plugins/Runtime/Metasound/Source/MetasoundEditor/Private/MetasoundFactory.cpp
--- a/Engine/Plugins/Runtime/Metasound/Source/MetasoundEditor/Private/MetasoundFactory.cpp
+++ b/Engine/Plugins/Runtime/Metasound/Source/MetasoundEditor/Private/MetasoundFactory.cpp
@@ -86,8 +86,6 @@ UObject* UMetasoundSourceFactory::FactoryCreateNew(UClass* InClass, UObject* InP
 	UMetaSoundSource* MetasoundSource = NewObject<UMetaSoundSource>(InParent, Name, Flags);
 	Metasound::Editor::FactoryPrivate::InitMetasound(MetasoundSource, Flags);
 
-	CastChecked<UMetasoundEditorGraph>(&MetasoundSource->GetGraphChecked())->ParentMetasound = MetasoundSource;
-
 	FGraphBuilder::ConstructGraph(*MetasoundSource);
 	FGraphBuilder::SynchronizeGraph(*MetasoundSource);
 
diff --git a/Engine/Plugins/Runtime/Metasound/Source/MetasoundEditor/Public/MetasoundEditorGraph.h b/Engine/Plugins/Runtime/Metasound/Source/MetasoundEditor/Public/MetasoundEditorGraph.h
--- a/Engine/Plugins/Runtime/Metasound/Source/MetasoundEditor/Public/MetasoundEditorGraph.h
+++ b/Engine/Plugins/Runtime/Metasound/Source/MetasoundEditor/Public/MetasoundEditorGraph.h
@@ -165,9 +165,6 @@ private:
 	// a sound's ParameterInterface when a sound begins playing.
 	uint32 PreviewID = INDEX_NONE;
 
-	UPROPERTY()
-	TObjectPtr<UObject> ParentMetasound;
-
 	UPROPERTY()
 	TArray<TObjectPtr<UMetasoundEditorGraphInput>> Inputs;
 
diff --git a/Engine/Plugins/Runtime/Metasound/Source/MetasoundEngine/Private/MetasoundSource.cpp b/Engine/Plugins/Runtime/Metasound/Source/MetasoundEngine/Private/MetasoundSource.cpp
--- a/Engine/Plugins/Runtime/Metasound/Source/MetasoundEngine/Private/MetasoundSource.cpp
+++ b/Engine/Plugins/Runtime/Metasound/Source/MetasoundEngine/Private/MetasoundSource.cpp
@@ -450,8 +450,9 @@ const FMetasoundFrontendArchetype& UMetaSoundSource::GetBaseArchetype()
 		
 		FMetasoundFrontendClassVertex OnPlayTrigger;
 		OnPlayTrigger.Name = UMetaSoundSource::GetOnPlayInputName();
-		OnPlayTrigger.Metadata.DisplayName = FText::FromString(OnPlayTrigger.Name);
+
 		OnPlayTrigger.TypeName = Metasound::Frontend::GetDataTypeName<Metasound::FTrigger>();
+		OnPlayTrigger.Metadata.DisplayName = LOCTEXT("OnPlay", "On Play");
 		OnPlayTrigger.Metadata.Description = LOCTEXT("OnPlayTriggerToolTip", "Trigger executed when this source is played.");
 		OnPlayTrigger.VertexID = FGuid::NewGuid();
 
@@ -459,8 +460,9 @@ const FMetasoundFrontendArchetype& UMetaSoundSource::GetBaseArchetype()
 
 		FMetasoundFrontendClassVertex OnFinished;
 		OnFinished.Name = UMetaSoundSource::GetIsFinishedOutputName();
-		OnFinished.Metadata.DisplayName = FText::FromString(OnFinished.Name);
+
 		OnFinished.TypeName = Metasound::Frontend::GetDataTypeName<Metasound::FTrigger>();
+		OnFinished.Metadata.DisplayName = LOCTEXT("OnFinished", "On Finished");
 		OnFinished.Metadata.Description = LOCTEXT("OnFinishedToolTip", "Trigger executed to initiate stopping the source.");
 		OnFinished.VertexID = FGuid::NewGuid();
 
diff --git a/Engine/Plugins/Runtime/Metasound/Source/MetasoundEngine/Public/MetasoundSource.h b/Engine/Plugins/Runtime/Metasound/Source/MetasoundEngine/Public/MetasoundSource.h
--- a/Engine/Plugins/Runtime/Metasound/Source/MetasoundEngine/Public/MetasoundSource.h
+++ b/Engine/Plugins/Runtime/Metasound/Source/MetasoundEngine/Public/MetasoundSource.h
@@ -27,6 +27,9 @@ class METASOUNDENGINE_API UMetasoundEditorGraphBase : public UEdGraph
 	GENERATED_BODY()
 
 public:
+	virtual bool IsEditorOnly() const override { return true; }
+	virtual bool NeedsLoadForEditorGame() const override { return false; }
+
 	virtual void Synchronize() { }
 };
 
@@ -43,7 +46,7 @@ protected:
 	FMetasoundFrontendDocument RootMetasoundDocument;
 
 #if WITH_EDITORONLY_DATA
-	UPROPERTY(EditAnywhere, Category = CustomView)
+	UPROPERTY(Transient)
 	UMetasoundEditorGraphBase* Graph;
 #endif // WITH_EDITORONLY_DATA
 
diff --git a/Engine/Plugins/Runtime/Metasound/Source/MetasoundFrontend/Public/MetasoundArrayNodes.h b/Engine/Plugins/Runtime/Metasound/Source/MetasoundFrontend/Public/MetasoundArrayNodes.h
--- a/Engine/Plugins/Runtime/Metasound/Source/MetasoundFrontend/Public/MetasoundArrayNodes.h
+++ b/Engine/Plugins/Runtime/Metasound/Source/MetasoundFrontend/Public/MetasoundArrayNodes.h
@@ -88,11 +88,11 @@ namespace Metasound
 		{
 			auto CreateNodeClassMetadata = []() -> FNodeClassMetadata
 			{
-				FName DataTypeName = GetMetasoundDataTypeName<ArrayType>();
-				FName OperatorName = TEXT("Num"); 
-				FText NodeDisplayName = FText::Format(LOCTEXT("ArrayOpArrayNumDisplayNamePattern", "Num ({0})"), FText::FromString(GetMetasoundDataTypeString<ArrayType>()));
-				FText NodeDescription = LOCTEXT("ArrayOpArrayNumDescription", "Number of elements in the array");
-				FVertexInterface NodeInterface = GetDefaultInterface();
+				const FName DataTypeName = GetMetasoundDataTypeName<ArrayType>();
+				const FName OperatorName = TEXT("Num");
+				const FText NodeDisplayName = FText::Format(LOCTEXT("ArrayOpArrayNumDisplayNamePattern", "Num ({0})"), GetMetasoundDataTypeDisplayText<ArrayType>());
+				const FText NodeDescription = LOCTEXT("ArrayOpArrayNumDescription", "Number of elements in the array");
+				const FVertexInterface NodeInterface = GetDefaultInterface();
 			
 				return MetasoundArrayNodesPrivate::CreateArrayNodeClassMetadata(DataTypeName, OperatorName, NodeDisplayName, NodeDescription, NodeInterface);
 			};
@@ -202,11 +202,11 @@ namespace Metasound
 		{
 			auto CreateNodeClassMetadata = []() -> FNodeClassMetadata
 			{
-				FName DataTypeName = GetMetasoundDataTypeName<ArrayType>();
-				FName OperatorName = TEXT("Get"); 
-				FText NodeDisplayName = FText::Format(LOCTEXT("ArrayOpArrayGetDisplayNamePattern", "Get ({0})"), FText::FromString(GetMetasoundDataTypeString<ArrayType>()));
-				FText NodeDescription = LOCTEXT("ArrayOpArrayGetDescription", "Get element at index in array.");
-				FVertexInterface NodeInterface = GetDefaultInterface();
+				const FName DataTypeName = GetMetasoundDataTypeName<ArrayType>();
+				const FName OperatorName = TEXT("Get"); 
+				const FText NodeDisplayName = FText::Format(LOCTEXT("ArrayOpArrayGetDisplayNamePattern", "Get ({0})"), GetMetasoundDataTypeDisplayText<ArrayType>());
+				const FText NodeDescription = LOCTEXT("ArrayOpArrayGetDescription", "Get element at index in array.");
+				const FVertexInterface NodeInterface = GetDefaultInterface();
 			
 				return MetasoundArrayNodesPrivate::CreateArrayNodeClassMetadata(DataTypeName, OperatorName, NodeDisplayName, NodeDescription, NodeInterface);
 			};
@@ -340,11 +340,11 @@ namespace Metasound
 		{
 			auto CreateNodeClassMetadata = []() -> FNodeClassMetadata
 			{
-				FName DataTypeName = GetMetasoundDataTypeName<ArrayType>();
-				FName OperatorName = TEXT("Set"); 
-				FText NodeDisplayName = FText::Format(LOCTEXT("ArrayOpArraySetDisplayNamePattern", "Set ({0})"), FText::FromString(GetMetasoundDataTypeString<ArrayType>()));
-				FText NodeDescription = LOCTEXT("ArrayOpArraySetDescription", "Set element at index in array.");
-				FVertexInterface NodeInterface = GetDefaultInterface();
+				const FName DataTypeName = GetMetasoundDataTypeName<ArrayType>();
+				const FName OperatorName = TEXT("Set"); 
+				const FText NodeDisplayName = FText::Format(LOCTEXT("ArrayOpArraySetDisplayNamePattern", "Set ({0})"), GetMetasoundDataTypeDisplayText<ArrayType>());
+				const FText NodeDescription = LOCTEXT("ArrayOpArraySetDescription", "Set element at index in array.");
+				const FVertexInterface NodeInterface = GetDefaultInterface();
 			
 				return MetasoundArrayNodesPrivate::CreateArrayNodeClassMetadata(DataTypeName, OperatorName, NodeDisplayName, NodeDescription, NodeInterface);
 			};
@@ -480,11 +480,11 @@ namespace Metasound
 		{
 			auto CreateNodeClassMetadata = []() -> FNodeClassMetadata
 			{
-				FName DataTypeName = GetMetasoundDataTypeName<ArrayType>();
-				FName OperatorName = TEXT("Concat"); 
-				FText NodeDisplayName = FText::Format(LOCTEXT("ArrayOpArrayConcatDisplayNamePattern", "Concatenate ({0})"), FText::FromString(GetMetasoundDataTypeString<ArrayType>()));
-				FText NodeDescription = LOCTEXT("ArrayOpArrayConcatDescription", "Concatenates two arrays on trigger.");
-				FVertexInterface NodeInterface = GetDefaultInterface();
+				const FName DataTypeName = GetMetasoundDataTypeName<ArrayType>();
+				const FName OperatorName = TEXT("Concat"); 
+				const FText NodeDisplayName = FText::Format(LOCTEXT("ArrayOpArrayConcatDisplayNamePattern", "Concatenate ({0})"), GetMetasoundDataTypeDisplayText<ArrayType>());
+				const FText NodeDescription = LOCTEXT("ArrayOpArrayConcatDescription", "Concatenates two arrays on trigger.");
+				const FVertexInterface NodeInterface = GetDefaultInterface();
 			
 				return MetasoundArrayNodesPrivate::CreateArrayNodeClassMetadata(DataTypeName, OperatorName, NodeDisplayName, NodeDescription, NodeInterface);
 			};
@@ -606,11 +606,11 @@ namespace Metasound
 		{
 			auto CreateNodeClassMetadata = []() -> FNodeClassMetadata
 			{
-				FName DataTypeName = GetMetasoundDataTypeName<ArrayType>();
-				FName OperatorName = TEXT("Subset"); 
-				FText NodeDisplayName = FText::Format(LOCTEXT("ArrayOpArraySubsetDisplayNamePattern", "Subset ({0})"), FText::FromString(GetMetasoundDataTypeString<ArrayType>()));
-				FText NodeDescription = LOCTEXT("ArrayOpArraySubsetDescription", "Subset array on trigger.");
-				FVertexInterface NodeInterface = GetDefaultInterface();
+				const FName DataTypeName = GetMetasoundDataTypeName<ArrayType>();
+				const FName OperatorName = TEXT("Subset"); 
+				const FText NodeDisplayName = FText::Format(LOCTEXT("ArrayOpArraySubsetDisplayNamePattern", "Subset ({0})"), GetMetasoundDataTypeDisplayText<ArrayType>());
+				const FText NodeDescription = LOCTEXT("ArrayOpArraySubsetDescription", "Subset array on trigger.");
+				const FVertexInterface NodeInterface = GetDefaultInterface();
 			
 				return MetasoundArrayNodesPrivate::CreateArrayNodeClassMetadata(DataTypeName, OperatorName, NodeDisplayName, NodeDescription, NodeInterface);
 			};
diff --git a/Engine/Plugins/Runtime/Metasound/Source/MetasoundFrontend/Public/MetasoundArrayRandomNode.h b/Engine/Plugins/Runtime/Metasound/Source/MetasoundFrontend/Public/MetasoundArrayRandomNode.h
--- a/Engine/Plugins/Runtime/Metasound/Source/MetasoundFrontend/Public/MetasoundArrayRandomNode.h
+++ b/Engine/Plugins/Runtime/Metasound/Source/MetasoundFrontend/Public/MetasoundArrayRandomNode.h
@@ -140,8 +140,8 @@ namespace Metasound
 			auto CreateNodeClassMetadata = []() -> FNodeClassMetadata
 			{
 				FName DataTypeName = GetMetasoundDataTypeName<ArrayType>();
-				FName OperatorName = TEXT("Random Get");
-				FText NodeDisplayName = FText::Format(LOCTEXT("RandomGetArrayOpDisplayNamePattern", "Random Get ({0})"), FText::FromString(GetMetasoundDataTypeString<ArrayType>()));
+				FName OperatorName = "Random Get";
+				FText NodeDisplayName = FText::Format(LOCTEXT("RandomGetArrayOpDisplayNamePattern", "Random Get ({0})"), GetMetasoundDataTypeDisplayText<ArrayType>());
 				FText NodeDescription = LOCTEXT("RandomGetArrayDescription", "Randomly retrieve data from input array using the supplied weights.");
 				FVertexInterface NodeInterface = GetDefaultInterface();
 
diff --git a/Engine/Plugins/Runtime/Metasound/Source/MetasoundFrontend/Public/MetasoundArrayShuffleNode.h b/Engine/Plugins/Runtime/Metasound/Source/MetasoundFrontend/Public/MetasoundArrayShuffleNode.h
--- a/Engine/Plugins/Runtime/Metasound/Source/MetasoundFrontend/Public/MetasoundArrayShuffleNode.h
+++ b/Engine/Plugins/Runtime/Metasound/Source/MetasoundFrontend/Public/MetasoundArrayShuffleNode.h
@@ -168,11 +168,11 @@ namespace Metasound
 		{
 			auto CreateNodeClassMetadata = []() -> FNodeClassMetadata
 			{
-				FName DataTypeName = GetMetasoundDataTypeName<ArrayType>();
-				FName OperatorName = TEXT("Shuffle");
-				FText NodeDisplayName = FText::Format(LOCTEXT("ArrayOpArrayShuffleDisplayNamePattern", "Shuffle ({0})"), FText::FromString(GetMetasoundDataTypeString<ArrayType>()));
-				FText NodeDescription = LOCTEXT("ArrayOpArrayShuffleDescription", "Output next element of a shuffled array on trigger.");
-				FVertexInterface NodeInterface = GetDefaultInterface();
+				const FName DataTypeName = GetMetasoundDataTypeName<ArrayType>();
+				const FName OperatorName = "Shuffle";
+				const FText NodeDisplayName = FText::Format(LOCTEXT("ArrayOpArrayShuffleDisplayNamePattern", "Shuffle ({0})"), GetMetasoundDataTypeDisplayText<ArrayType>());
+				const FText NodeDescription = LOCTEXT("ArrayOpArrayShuffleDescription", "Output next element of a shuffled array on trigger.");
+				const FVertexInterface NodeInterface = GetDefaultInterface();
 
 				return MetasoundArrayNodesPrivate::CreateArrayNodeClassMetadata(DataTypeName, OperatorName, NodeDisplayName, NodeDescription, NodeInterface);
 			};
diff --git a/Engine/Plugins/Runtime/Metasound/Source/MetasoundGraphCore/Private/MetasoundDataReference.cpp b/Engine/Plugins/Runtime/Metasound/Source/MetasoundGraphCore/Private/MetasoundDataReference.cpp
--- a/Engine/Plugins/Runtime/Metasound/Source/MetasoundGraphCore/Private/MetasoundDataReference.cpp
+++ b/Engine/Plugins/Runtime/Metasound/Source/MetasoundGraphCore/Private/MetasoundDataReference.cpp
@@ -5,6 +5,10 @@
 namespace Metasound
 {
 	const TCHAR* TDataReferenceTypeInfo<void>::TypeName = TEXT("void");
+	const FText& TDataReferenceTypeInfo<void>::GetTypeDisplayText()
+	{
+		return FText::GetEmpty();
+	}
 	const void* const TDataReferenceTypeInfo<void>::TypePtr = nullptr;
 	const void* const TDataReferenceTypeInfo<void>::TypeId = static_cast<const void* const>(TDataReferenceTypeInfo<void>::TypePtr);
 }
diff --git a/Engine/Plugins/Runtime/Metasound/Source/MetasoundGraphCore/Public/MetasoundDataReference.h b/Engine/Plugins/Runtime/Metasound/Source/MetasoundGraphCore/Public/MetasoundDataReference.h
--- a/Engine/Plugins/Runtime/Metasound/Source/MetasoundGraphCore/Public/MetasoundDataReference.h
+++ b/Engine/Plugins/Runtime/Metasound/Source/MetasoundGraphCore/Public/MetasoundDataReference.h
@@ -6,6 +6,7 @@
 #include "MetasoundOperatorSettings.h"
 #include <type_traits>
 
+
 /** Macro to make declaring a metasound parameter simple.  */
 // Declares a metasound parameter type by
 // - Adding typedefs for commonly used template types.
@@ -15,6 +16,7 @@
 	struct ::Metasound::TDataReferenceTypeInfo<DataType> \
 	{ \
 		static ModuleApi const TCHAR* TypeName; \
+		static ModuleApi const FText& GetTypeDisplayText(); \
 		static ModuleApi const FMetasoundDataTypeId TypeId; \
 		\
 		private: \
@@ -30,6 +32,7 @@
 	struct ::Metasound::TDataReferenceTypeInfo<TArray<DataType>> \
 	{ \
 		static ModuleApi const TCHAR* TypeName; \
+		static ModuleApi const FText& GetTypeDisplayText(); \
 		static ModuleApi const FMetasoundDataTypeId TypeId; \
 		\
 		private: \
@@ -39,11 +42,21 @@
 
 // This only needs to be called if you don't plan on calling REGISTER_METASOUND_DATATYPE.
 #define DEFINE_METASOUND_DATA_TYPE(DataType, DataTypeName) \
-	const TCHAR*  				::Metasound::TDataReferenceTypeInfo<DataType>::TypeName = TEXT(DataTypeName); \
+	const TCHAR* ::Metasound::TDataReferenceTypeInfo<DataType>::TypeName = TEXT(DataTypeName); \
+	const FText& ::Metasound::TDataReferenceTypeInfo<DataType>::GetTypeDisplayText() \
+	{ \
+		static const FText DisplayText = NSLOCTEXT("MetaSoundCore_DataReference", DataTypeName, DataTypeName); \
+		return DisplayText; \
+	} \
 	const DataType* const ::Metasound::TDataReferenceTypeInfo<DataType>::TypePtr = nullptr; \
 	const void* const ::Metasound::TDataReferenceTypeInfo<DataType>::TypeId = static_cast<const FMetasoundDataTypeId>(&::Metasound::TDataReferenceTypeInfo<DataType>::TypePtr); \
-	/* Array definition */\
-	const TCHAR*  				::Metasound::TDataReferenceTypeInfo<TArray<DataType>>::TypeName = TEXT(DataTypeName ":Array"); \
+	/* Array definition */ \
+	const TCHAR* ::Metasound::TDataReferenceTypeInfo<TArray<DataType>>::TypeName = TEXT(DataTypeName ":Array"); \
+	const FText& ::Metasound::TDataReferenceTypeInfo<TArray<DataType>>::GetTypeDisplayText() \
+	{ \
+		static const FText DisplayText = NSLOCTEXT("MetaSoundCore_DataReference", DataTypeName "_Array", DataTypeName ":Array"); \
+		return DisplayText; \
+	} \
 	const TArray<DataType>* const ::Metasound::TDataReferenceTypeInfo<TArray<DataType>>::TypePtr = nullptr; \
 	const void* const ::Metasound::TDataReferenceTypeInfo<TArray<DataType>>::TypeId = static_cast<const FMetasoundDataTypeId>(&::Metasound::TDataReferenceTypeInfo<TArray<DataType>>::TypePtr);
 
@@ -94,12 +107,19 @@ namespace Metasound
 		return TypeString;
 	}
 
-	/** Return the data type ID for a registered data type. 
+	/** Return the display text for a registered data type. */
+	template<typename DataType>
+	const FText& GetMetasoundDataTypeDisplayText()
+	{
+		return TDataReferenceTypeInfo<std::decay_t<DataType>>::GetTypeDisplayText();
+	}
+
+	/** Return the data type ID for a registered data type.
 	 *
 	 * This ID is runtime constant but may change between executions and builds.
 	 */
 	template<typename DataType>
-	const void* const GetMetasoundDataTypeId() 
+	const void* const GetMetasoundDataTypeId()
 	{
 		return TDataReferenceTypeInfo<std::decay_t<DataType>>::TypeId;
 	}
@@ -110,6 +130,7 @@ namespace Metasound
 	{
 		static METASOUNDGRAPHCORE_API const TCHAR* TypeName;
 		static METASOUNDGRAPHCORE_API const void* const TypeId;
+		static METASOUNDGRAPHCORE_API const FText& GetTypeDisplayText();
 
 		private:
 
@@ -394,4 +415,3 @@ namespace Metasound
 			}
 	};
 }
-
diff --git a/Engine/Plugins/Runtime/Metasound/Source/MetasoundGraphCore/Public/MetasoundGraph.h b/Engine/Plugins/Runtime/Metasound/Source/MetasoundGraphCore/Public/MetasoundGraph.h
--- a/Engine/Plugins/Runtime/Metasound/Source/MetasoundGraphCore/Public/MetasoundGraph.h
+++ b/Engine/Plugins/Runtime/Metasound/Source/MetasoundGraphCore/Public/MetasoundGraph.h
@@ -5,7 +5,6 @@
 #include "MetasoundNode.h"
 
 // Forward Declarations
-class FText;
 
 
 namespace Metasound
diff --git a/Engine/Plugins/Runtime/Metasound/Source/MetasoundStandardNodes/Private/MetasoundADEnvelopeNode.cpp b/Engine/Plugins/Runtime/Metasound/Source/MetasoundStandardNodes/Private/MetasoundADEnvelopeNode.cpp
--- a/Engine/Plugins/Runtime/Metasound/Source/MetasoundStandardNodes/Private/MetasoundADEnvelopeNode.cpp
+++ b/Engine/Plugins/Runtime/Metasound/Source/MetasoundStandardNodes/Private/MetasoundADEnvelopeNode.cpp
@@ -217,13 +217,13 @@ namespace Metasound
 		{
 			auto CreateNodeClassMetadata = []() -> FNodeClassMetadata
 			{
-				FName DataTypeName = GetMetasoundDataTypeName<ValueType>();
-				FName OperatorName = TEXT("AD Envelope");
-				FText NodeDisplayName = FText::Format(LOCTEXT("ADEnvelopeDisplayNamePattern", "AD Envelope ({0})"), FText::FromString(GetMetasoundDataTypeString<ValueType>()));
-				FText NodeDescription = LOCTEXT("ADEnevelopeDesc", "Generates an attack-decay envelope value output when triggered.");
-				FVertexInterface NodeInterface = GetDefaultInterface();
+				const FName DataTypeName = GetMetasoundDataTypeName<ValueType>();
+				const FName OperatorName = "AD Envelope";
+				const FText NodeDisplayName = FText::Format(LOCTEXT("ADEnvelopeDisplayNamePattern", "AD Envelope ({0})"), GetMetasoundDataTypeDisplayText<ValueType>());
+				const FText NodeDescription = LOCTEXT("ADEnevelopeDesc", "Generates an attack-decay envelope value output when triggered.");
+				const FVertexInterface NodeInterface = GetDefaultInterface();
 
-				FNodeClassMetadata Metadata
+				const FNodeClassMetadata Metadata
 				{
 					FNodeClassName{FName("AD Envelope"), OperatorName, DataTypeName},
 					1, // Major Version
diff --git a/Engine/Plugins/Runtime/Metasound/Source/MetasoundStandardNodes/Private/MetasoundADSREnvelopeNode.cpp b/Engine/Plugins/Runtime/Metasound/Source/MetasoundStandardNodes/Private/MetasoundADSREnvelopeNode.cpp
--- a/Engine/Plugins/Runtime/Metasound/Source/MetasoundStandardNodes/Private/MetasoundADSREnvelopeNode.cpp
+++ b/Engine/Plugins/Runtime/Metasound/Source/MetasoundStandardNodes/Private/MetasoundADSREnvelopeNode.cpp
@@ -356,13 +356,13 @@ namespace Metasound
 		{
 			auto CreateNodeClassMetadata = []() -> FNodeClassMetadata
 			{
-				FName DataTypeName = GetMetasoundDataTypeName<ValueType>();
-				FName OperatorName = TEXT("ADSR Envelope");
-				FText NodeDisplayName = FText::Format(LOCTEXT("ADSREnvelopeDisplayNamePattern", "ADSR Envelope ({0})"), FText::FromString(GetMetasoundDataTypeString<ValueType>()));
-				FText NodeDescription = LOCTEXT("ADSREnevelopeDesc", "Generates an attack-decay-sustain-release envelope value output when triggered.");
-				FVertexInterface NodeInterface = GetDefaultInterface();
+				const FName DataTypeName = GetMetasoundDataTypeName<ValueType>();
+				const FName OperatorName = "ADSR Envelope";
+				const FText NodeDisplayName = FText::Format(LOCTEXT("ADSREnvelopeDisplayNamePattern", "ADSR Envelope ({0})"), GetMetasoundDataTypeDisplayText<ValueType>());
+				const FText NodeDescription = LOCTEXT("ADSREnevelopeDesc", "Generates an attack-decay-sustain-release envelope value output when triggered.");
+				const FVertexInterface NodeInterface = GetDefaultInterface();
 
-				FNodeClassMetadata Metadata
+				const FNodeClassMetadata Metadata
 				{
 					FNodeClassName{FName("ADSR Envelope"), OperatorName, DataTypeName},
 					1, // Major Version
diff --git a/Engine/Plugins/Runtime/Metasound/Source/MetasoundStandardNodes/Private/MetasoundAudioFormats.cpp b/Engine/Plugins/Runtime/Metasound/Source/MetasoundStandardNodes/Private/MetasoundAudioFormats.cpp
--- a/Engine/Plugins/Runtime/Metasound/Source/MetasoundStandardNodes/Private/MetasoundAudioFormats.cpp
+++ b/Engine/Plugins/Runtime/Metasound/Source/MetasoundStandardNodes/Private/MetasoundAudioFormats.cpp
@@ -182,10 +182,10 @@ namespace Metasound
 		{
 			FNodeClassMetadata Info;
 
-			Info.ClassName = {TEXT("Output"), GetMetasoundDataTypeName<FStereoAudioFormat>(), TEXT("")};
+			Info.ClassName = { "Output", GetMetasoundDataTypeName<FStereoAudioFormat>(), "" };
 			Info.MajorVersion = 1;
 			Info.MinorVersion = 0;
-			Info.DisplayName = FText::Format(LOCTEXT("Metasound_OutputNodeDisplayNameFormat", "Output {0}"), FText::FromString(GetMetasoundDataTypeString<FStereoAudioFormat>()));
+			Info.DisplayName = FText::Format(LOCTEXT("Metasound_OutputNodeDisplayNameFormat", "Output {0}"), GetMetasoundDataTypeDisplayText<FStereoAudioFormat>());
 			Info.Description = LOCTEXT("Metasound_OutputNodeDescription", "Output from the parent Metasound graph.");
 			Info.Author = PluginAuthor;
 			Info.PromptIfMissing = PluginNodeMissingPrompt;
@@ -321,10 +321,10 @@ namespace Metasound
 		{
 			FNodeClassMetadata Info;
 
-			Info.ClassName = {TEXT("Input"), GetMetasoundDataTypeName<FStereoAudioFormat>(), TEXT("")};
+			Info.ClassName = { "Input", GetMetasoundDataTypeName<FStereoAudioFormat>(), "" };
 			Info.MajorVersion = 1;
 			Info.MinorVersion = 0;
-			Info.DisplayName = FText::Format(LOCTEXT("Metasound_InputNodeDisplayNameFormat", "Input {0}"), FText::FromString(GetMetasoundDataTypeString<FStereoAudioFormat>()));
+			Info.DisplayName = FText::Format(LOCTEXT("Metasound_InputNodeDisplayNameFormat", "Input {0}"), GetMetasoundDataTypeDisplayText<FStereoAudioFormat>());
 			Info.Description = LOCTEXT("Metasound_InputNodeDescription", "Input from the parent Metasound graph.");
 			Info.Author = PluginAuthor;
 			Info.PromptIfMissing = PluginNodeMissingPrompt;
@@ -455,10 +455,10 @@ namespace Metasound
 		{
 			FNodeClassMetadata Info;
 
-			Info.ClassName = {TEXT("Output"), GetMetasoundDataTypeName<FMonoAudioFormat>(), TEXT("")};
+			Info.ClassName = { "Output", GetMetasoundDataTypeName<FMonoAudioFormat>(), "" };
 			Info.MajorVersion = 1;
 			Info.MinorVersion = 0;
-			Info.DisplayName = FText::Format(LOCTEXT("Metasound_OutputNodeDisplayNameFormat", "Output {0}"), FText::FromString(GetMetasoundDataTypeString<FMonoAudioFormat>()));
+			Info.DisplayName = FText::Format(LOCTEXT("Metasound_OutputNodeDisplayNameFormat", "Output {0}"), GetMetasoundDataTypeDisplayText<FMonoAudioFormat>());
 			Info.Description = LOCTEXT("Metasound_OutputNodeDescription", "Output from the parent Metasound graph.");
 			Info.Author = PluginAuthor;
 			Info.PromptIfMissing = PluginNodeMissingPrompt;
@@ -588,10 +588,10 @@ namespace Metasound
 		{
 			FNodeClassMetadata Info;
 
-			Info.ClassName = {TEXT("Input"), GetMetasoundDataTypeName<FMonoAudioFormat>(), TEXT("")};
+			Info.ClassName = { "Input", GetMetasoundDataTypeName<FMonoAudioFormat>(), "" };
 			Info.MajorVersion = 1;
 			Info.MinorVersion = 0;
-			Info.DisplayName = FText::Format(LOCTEXT("Metasound_InputNodeDisplayNameFormat", "Input {0}"), FText::FromString(GetMetasoundDataTypeString<FMonoAudioFormat>()));
+			Info.DisplayName = FText::Format(LOCTEXT("Metasound_InputNodeDisplayNameFormat", "Input {0}"), GetMetasoundDataTypeDisplayText<FMonoAudioFormat>());
 			Info.Description = LOCTEXT("Metasound_InputNodeDescription", "Input from the parent Metasound graph.");
 			Info.Author = PluginAuthor;
 			Info.PromptIfMissing = PluginNodeMissingPrompt;
diff --git a/Engine/Plugins/Runtime/Metasound/Source/MetasoundStandardNodes/Private/MetasoundClampNode.cpp b/Engine/Plugins/Runtime/Metasound/Source/MetasoundStandardNodes/Private/MetasoundClampNode.cpp
--- a/Engine/Plugins/Runtime/Metasound/Source/MetasoundStandardNodes/Private/MetasoundClampNode.cpp
+++ b/Engine/Plugins/Runtime/Metasound/Source/MetasoundStandardNodes/Private/MetasoundClampNode.cpp
@@ -173,11 +173,11 @@ namespace Metasound
 		{
 			auto CreateNodeClassMetadata = []() -> FNodeClassMetadata
 			{
-				FName DataTypeName = GetMetasoundDataTypeName<ValueType>();
-				FName OperatorName = TEXT("Clamp");
-				FText NodeDisplayName = FText::Format(LOCTEXT("ClampDisplayNamePattern", "Clamp ({0})"), FText::FromString(GetMetasoundDataTypeString<ValueType>()));
-				FText NodeDescription = LOCTEXT("ClampDesc", "Returns the clamped value of the input within the given value range.");
-				FVertexInterface NodeInterface = GetDefaultInterface();
+				const FName DataTypeName = GetMetasoundDataTypeName<ValueType>();
+				const FName OperatorName = "Clamp";
+				const FText NodeDisplayName = FText::Format(LOCTEXT("ClampDisplayNamePattern", "Clamp ({0})"), GetMetasoundDataTypeDisplayText<ValueType>());
+				const FText NodeDescription = LOCTEXT("ClampDesc", "Returns the clamped value of the input within the given value range.");
+				const FVertexInterface NodeInterface = GetDefaultInterface();
 
 				return MetasoundClampNodePrivate::CreateNodeClassMetadata(DataTypeName, OperatorName, NodeDisplayName, NodeDescription, NodeInterface);
 			};
diff --git a/Engine/Plugins/Runtime/Metasound/Source/MetasoundStandardNodes/Private/MetasoundConversionNodes.cpp b/Engine/Plugins/Runtime/Metasound/Source/MetasoundStandardNodes/Private/MetasoundConversionNodes.cpp
--- a/Engine/Plugins/Runtime/Metasound/Source/MetasoundStandardNodes/Private/MetasoundConversionNodes.cpp
+++ b/Engine/Plugins/Runtime/Metasound/Source/MetasoundStandardNodes/Private/MetasoundConversionNodes.cpp
@@ -152,8 +152,8 @@ namespace Metasound
 		{
 			using namespace ConversionNodeVertexNames;
 
-			static FText InputDesc = FText::Format(LOCTEXT("ConvDisplayNamePatternFrom", "Input {0} value."), FText::FromString(GetMetasoundDataTypeString<FromType>()));
-			static FText OutputDesc = FText::Format(LOCTEXT("ConvDisplayNamePatternTo", "Ouput {0} value."), FText::FromString(GetMetasoundDataTypeString<ToType>()));
+			static FText InputDesc = FText::Format(LOCTEXT("ConvDisplayNamePatternFrom", "Input {0} value."), GetMetasoundDataTypeDisplayText<FromType>());
+			static FText OutputDesc = FText::Format(LOCTEXT("ConvDisplayNamePatternTo", "Ouput {0} value."), GetMetasoundDataTypeDisplayText<ToType>());
 
 			static const FVertexInterface DefaultInterface(
 				FInputVertexInterface(
@@ -171,14 +171,18 @@ namespace Metasound
 		{
 			auto InitNodeInfo = []() -> FNodeClassMetadata
 			{
-				FString FromTypeString = GetMetasoundDataTypeString<FromType>();
-				FString ToTypeString = GetMetasoundDataTypeString<ToType>();
-				FName ClassName = *FString::Format(TEXT("Conversion{0}To{1}"), { FromTypeString, ToTypeString });
-				FText NodeDisplayName = FText::Format(LOCTEXT("ConverterNodeDisplayName", "{0} To {1}"), FText::FromString(FromTypeString), FText::FromString(ToTypeString));
-				FText NodeDescription = FText::Format(LOCTEXT("ConverterNodeDesc", "Converts from {0} to {1}."), FText::FromString(FromTypeString), FText::FromString(ToTypeString));
+				const FText& FromTypeText = GetMetasoundDataTypeDisplayText<FromType>();
+				const FText& ToTypeText = GetMetasoundDataTypeDisplayText<ToType>();
+
+				const FString& FromTypeString = GetMetasoundDataTypeString<FromType>();
+				const FString& ToTypeString = GetMetasoundDataTypeString<ToType>();
+
+				const FName ClassName = *FString::Format(TEXT("Conversion{0}To{1}"), { FromTypeString, ToTypeString });
+				const FText NodeDisplayName = FText::Format(LOCTEXT("ConverterNodeDisplayName", "{0} To {1}"), FromTypeText, ToTypeText);
+				const FText NodeDescription = FText::Format(LOCTEXT("ConverterNodeDesc", "Converts from {0} to {1}."), FromTypeText, ToTypeText);
 
 				FNodeClassMetadata Info;
-				Info.ClassName = { Metasound::StandardNodes::Namespace, ClassName, TEXT("") };
+				Info.ClassName = { Metasound::StandardNodes::Namespace, ClassName, "" };
 				Info.MajorVersion = 1;
 				Info.MinorVersion = 0;
 				Info.DisplayName = NodeDisplayName;
diff --git a/Engine/Plugins/Runtime/Metasound/Source/MetasoundStandardNodes/Private/MetasoundCrossfadeNode.cpp b/Engine/Plugins/Runtime/Metasound/Source/MetasoundStandardNodes/Private/MetasoundCrossfadeNode.cpp
--- a/Engine/Plugins/Runtime/Metasound/Source/MetasoundStandardNodes/Private/MetasoundCrossfadeNode.cpp
+++ b/Engine/Plugins/Runtime/Metasound/Source/MetasoundStandardNodes/Private/MetasoundCrossfadeNode.cpp
@@ -183,7 +183,7 @@ namespace Metasound
 			{
 				FName DataTypeName = GetMetasoundDataTypeName<ValueType>();
 				FName OperatorName = *FString::Printf(TEXT("Trigger Route (%s, %d)"), *DataTypeName.ToString(), NumInputs);
-				FText NodeDisplayName = FText::Format(LOCTEXT("CrossfadeDisplayNamePattern", "Crossfade ({0}, {1})"), FText::FromString(GetMetasoundDataTypeString<ValueType>()), NumInputs);
+				FText NodeDisplayName = FText::Format(LOCTEXT("CrossfadeDisplayNamePattern", "Crossfade ({0}, {1})"), GetMetasoundDataTypeDisplayText<ValueType>(), NumInputs);
 				FText NodeDescription = LOCTEXT("CrossfadeDescription", "Crossfades inputs to outputs.");
 				FVertexInterface NodeInterface = GetVertexInterface();
 
diff --git a/Engine/Plugins/Runtime/Metasound/Source/MetasoundStandardNodes/Private/MetasoundDebugLogNode.cpp b/Engine/Plugins/Runtime/Metasound/Source/MetasoundStandardNodes/Private/MetasoundDebugLogNode.cpp
--- a/Engine/Plugins/Runtime/Metasound/Source/MetasoundStandardNodes/Private/MetasoundDebugLogNode.cpp
+++ b/Engine/Plugins/Runtime/Metasound/Source/MetasoundStandardNodes/Private/MetasoundDebugLogNode.cpp
@@ -92,7 +92,7 @@ namespace Metasound
 				{
 					FName DataTypeName = GetMetasoundDataTypeName<PrintLogType>();
 					FName OperatorName = TEXT("Print Log");
-					FText NodeDisplayName = FText::Format(LOCTEXT("PrintLogDisplayNamePattern", "Print Log ({0})"), FText::FromString(GetMetasoundDataTypeString<PrintLogType>()));
+					FText NodeDisplayName = FText::Format(LOCTEXT("PrintLogDisplayNamePattern", "Print Log ({0})"), GetMetasoundDataTypeDisplayText<PrintLogType>());
 					FText NodeDescription = LOCTEXT("PrintLogOpDescription", "Used to record values to the log, on trigger");
 					FVertexInterface NodeInterface = GetDefaultInterface();
 
diff --git a/Engine/Plugins/Runtime/Metasound/Source/MetasoundStandardNodes/Private/MetasoundMIDIToFreqNode.cpp b/Engine/Plugins/Runtime/Metasound/Source/MetasoundStandardNodes/Private/MetasoundMIDIToFreqNode.cpp
--- a/Engine/Plugins/Runtime/Metasound/Source/MetasoundStandardNodes/Private/MetasoundMIDIToFreqNode.cpp
+++ b/Engine/Plugins/Runtime/Metasound/Source/MetasoundStandardNodes/Private/MetasoundMIDIToFreqNode.cpp
@@ -148,10 +148,10 @@ namespace Metasound
 	{
 		auto InitNodeInfo = []() -> FNodeClassMetadata
 		{
-			FName DataTypeName = GetMetasoundDataTypeName<ValueType>();
-			FName OperatorName = TEXT("MIDI To Frequency");
-			FText NodeDisplayName = FText::Format(LOCTEXT("RandomGetArrayOpDisplayNamePattern", "MIDI To Frequency ({0})"), FText::FromString(GetMetasoundDataTypeString<ValueType>()));
-			FText NodeDescription = LOCTEXT("Metasound_MidiToFreqNodeDescription", "Converts a Midi note value to a frequency (hz) value.");
+			const FName DataTypeName = GetMetasoundDataTypeName<ValueType>();
+			const FName OperatorName = TEXT("MIDI To Frequency");
+			const FText NodeDisplayName = FText::Format(LOCTEXT("RandomGetArrayOpDisplayNamePattern", "MIDI To Frequency ({0})"), GetMetasoundDataTypeDisplayText<ValueType>());
+			const FText NodeDescription = LOCTEXT("Metasound_MidiToFreqNodeDescription", "Converts a Midi note value to a frequency (hz) value.");
 
 			FNodeClassMetadata Info;
 			Info.ClassName = { Metasound::StandardNodes::Namespace, OperatorName, DataTypeName };
diff --git a/Engine/Plugins/Runtime/Metasound/Source/MetasoundStandardNodes/Private/MetasoundMathNodes.cpp b/Engine/Plugins/Runtime/Metasound/Source/MetasoundStandardNodes/Private/MetasoundMathNodes.cpp
--- a/Engine/Plugins/Runtime/Metasound/Source/MetasoundStandardNodes/Private/MetasoundMathNodes.cpp
+++ b/Engine/Plugins/Runtime/Metasound/Source/MetasoundStandardNodes/Private/MetasoundMathNodes.cpp
@@ -54,7 +54,7 @@ namespace Metasound
 		template<typename DataType>
 		const FText GetAddDisplayName()
 		{
-			static const FText DisplayName = FText::Format(LOCTEXT("AddNodeDisplayNamePattern", "Add ({0})"), FText::FromString(GetMetasoundDataTypeString<DataType>()));
+			static const FText DisplayName = FText::Format(LOCTEXT("AddNodeDisplayNamePattern", "Add ({0})"), GetMetasoundDataTypeDisplayText<DataType>());
 			return DisplayName;
 		}
 
@@ -63,15 +63,15 @@ namespace Metasound
 		{
 			static const FText DisplayName = FText::Format(
 				LOCTEXT("AddNodeDisplayNamePattern", "Add ({0} to {1})"),
-				FText::FromString(GetMetasoundDataTypeString<OperandDataType>()),
-				FText::FromString(GetMetasoundDataTypeString<DataType>()));
+				GetMetasoundDataTypeDisplayText<OperandDataType>(),
+				GetMetasoundDataTypeDisplayText<DataType>());
 			return DisplayName;
 		}
 
 		template<typename DataType>
 		const FText GetSubtractDisplayName()
 		{
-			static const FText DisplayName = FText::Format(LOCTEXT("SubtractNodeDisplayNamePattern", "Subtract ({0})"), FText::FromString(GetMetasoundDataTypeString<DataType>()));
+			static const FText DisplayName = FText::Format(LOCTEXT("SubtractNodeDisplayNamePattern", "Subtract ({0})"), GetMetasoundDataTypeDisplayText<DataType>());
 			return DisplayName;
 		}
 
@@ -80,15 +80,15 @@ namespace Metasound
 		{
 			static const FText DisplayName = FText::Format(
 				LOCTEXT("SubtractNodeOperandTypedDisplayNamePattern", "Subtract ({0} from {1})"),
-				FText::FromString(GetMetasoundDataTypeString<OperandDataType>()),
-				FText::FromString(GetMetasoundDataTypeString<DataType>()));
+				GetMetasoundDataTypeDisplayText<OperandDataType>(),
+				GetMetasoundDataTypeDisplayText<DataType>());
 			return DisplayName;
 		}
 
 		template<typename DataType>
 		const FText GetMultiplyDisplayName()
 		{
-			static const FText DisplayName = FText::Format(LOCTEXT("MultiplyNodeDisplayNamePattern", "Multiply ({0})"), FText::FromString(GetMetasoundDataTypeString<DataType>()));
+			static const FText DisplayName = FText::Format(LOCTEXT("MultiplyNodeDisplayNamePattern", "Multiply ({0})"), GetMetasoundDataTypeDisplayText<DataType>());
 			return DisplayName;
 		}
 
@@ -97,15 +97,15 @@ namespace Metasound
 		{
 			static const FText DisplayName = FText::Format(
 				LOCTEXT("MultiplyNodeOperandTypedDisplayNamePattern", "Multiply ({0} by {1})"),
-				FText::FromString(GetMetasoundDataTypeString<DataType>()),
-				FText::FromString(GetMetasoundDataTypeString<OperandDataType>()));
+				GetMetasoundDataTypeDisplayText<DataType>(),
+				GetMetasoundDataTypeDisplayText<OperandDataType>());
 			return DisplayName;
 		}
 
 		template<typename DataType>
 		const FText GetDivideDisplayName()
 		{
-			static const FText DisplayName = FText::Format(LOCTEXT("DivideNodeDisplayNamePattern", "Divide ({0})"), FText::FromString(GetMetasoundDataTypeString<DataType>()));
+			static const FText DisplayName = FText::Format(LOCTEXT("DivideNodeDisplayNamePattern", "Divide ({0})"), GetMetasoundDataTypeDisplayText<DataType>());
 			return DisplayName;
 		}
 
@@ -114,15 +114,15 @@ namespace Metasound
 		{
 			static const FText DisplayName = FText::Format(
 				LOCTEXT("DivideNodeOperandTypedDisplayNamePattern", "Divide ({0} by {1})"),
-				FText::FromString(GetMetasoundDataTypeString<DataType>()),
-				FText::FromString(GetMetasoundDataTypeString<OperandDataType>()));
+				GetMetasoundDataTypeDisplayText<DataType>(),
+				GetMetasoundDataTypeDisplayText<OperandDataType>());
 			return DisplayName;
 		}
 
 		template<typename DataType>
 		const FText GetModuloDisplayName()
 		{
-			static const FText DisplayName = FText::Format(LOCTEXT("ModuloNodeDisplayNamePattern", "Modulo ({0})"), FText::FromString(GetMetasoundDataTypeString<DataType>()));
+			static const FText DisplayName = FText::Format(LOCTEXT("ModuloNodeDisplayNamePattern", "Modulo ({0})"), GetMetasoundDataTypeDisplayText<DataType>());
 			return DisplayName;
 		}
 
@@ -131,15 +131,15 @@ namespace Metasound
 		{
 			static const FText DisplayName = FText::Format(
 				LOCTEXT("ModuloNodeOperandTypedDisplayNamePattern", "Modulo ({0} by {1})"),
-				FText::FromString(GetMetasoundDataTypeString<DataType>()),
-				FText::FromString(GetMetasoundDataTypeString<OperandDataType>()));
+				GetMetasoundDataTypeDisplayText<DataType>(),
+				GetMetasoundDataTypeDisplayText<OperandDataType>());
 			return DisplayName;
 		}
 
 		template<typename DataType>
 		const FText GetPowerDisplayName()
 		{
-			static const FText DisplayName = FText::Format(LOCTEXT("PowerNodeDisplayNamePattern", "Power ({0})"), FText::FromString(GetMetasoundDataTypeString<DataType>()));
+			static const FText DisplayName = FText::Format(LOCTEXT("PowerNodeDisplayNamePattern", "Power ({0})"), GetMetasoundDataTypeDisplayText<DataType>());
 			return DisplayName;
 		}
 
@@ -148,15 +148,15 @@ namespace Metasound
 		{
 			static const FText DisplayName = FText::Format(
 				LOCTEXT("PowerNodeOperandTypedDisplayNamePattern", "Power ({0} to the power of {1})"),
-				FText::FromString(GetMetasoundDataTypeString<DataType>()),
-				FText::FromString(GetMetasoundDataTypeString<OperandDataType>()));
+				GetMetasoundDataTypeDisplayText<DataType>(),
+				GetMetasoundDataTypeDisplayText<OperandDataType>());
 			return DisplayName;
 		}
 
 		template<typename DataType>
 		const FText GetLogarithmDisplayName()
 		{
-			static const FText DisplayName = FText::Format(LOCTEXT("LogNodeDisplayNamePattern", "Log ({0})"), FText::FromString(GetMetasoundDataTypeString<DataType>()));
+			static const FText DisplayName = FText::Format(LOCTEXT("LogNodeDisplayNamePattern", "Log ({0})"), GetMetasoundDataTypeDisplayText<DataType>());
 			return DisplayName;
 		}
 
@@ -165,8 +165,8 @@ namespace Metasound
 		{
 			static const FText DisplayName = FText::Format(
 				LOCTEXT("LogarithmNodeOperandTypedDisplayNamePattern", "Log ({0}-Base logarithm of {1})"),
-				FText::FromString(GetMetasoundDataTypeString<DataType>()),
-				FText::FromString(GetMetasoundDataTypeString<OperandDataType>()));
+				GetMetasoundDataTypeDisplayText<DataType>(),
+				GetMetasoundDataTypeDisplayText<OperandDataType>());
 			return DisplayName;
 		}
 	}
diff --git a/Engine/Plugins/Runtime/Metasound/Source/MetasoundStandardNodes/Private/MetasoundMaxNode.cpp b/Engine/Plugins/Runtime/Metasound/Source/MetasoundStandardNodes/Private/MetasoundMaxNode.cpp
--- a/Engine/Plugins/Runtime/Metasound/Source/MetasoundStandardNodes/Private/MetasoundMaxNode.cpp
+++ b/Engine/Plugins/Runtime/Metasound/Source/MetasoundStandardNodes/Private/MetasoundMaxNode.cpp
@@ -153,11 +153,11 @@ namespace Metasound
 		{
 			auto CreateNodeClassMetadata = []() -> FNodeClassMetadata
 			{
-				FName DataTypeName = GetMetasoundDataTypeName<ValueType>();
-				FName OperatorName = TEXT("Max");
-				FText NodeDisplayName = FText::Format(LOCTEXT("MaxDisplayNamePattern", "Max ({0})"), FText::FromString(GetMetasoundDataTypeString<ValueType>()));
-				FText NodeDescription = LOCTEXT("MaxDesc", "Returns the max of A and B.");
-				FVertexInterface NodeInterface = GetDefaultInterface();
+				const FName DataTypeName = GetMetasoundDataTypeName<ValueType>();
+				const FName OperatorName = TEXT("Max");
+				const FText NodeDisplayName = FText::Format(LOCTEXT("MaxDisplayNamePattern", "Max ({0})"), GetMetasoundDataTypeDisplayText<ValueType>());
+				const FText NodeDescription = LOCTEXT("MaxDesc", "Returns the max of A and B.");
+				const FVertexInterface NodeInterface = GetDefaultInterface();
 
 				return MetasoundMaxNodePrivate::CreateNodeClassMetadata(DataTypeName, OperatorName, NodeDisplayName, NodeDescription, NodeInterface);
 			};
diff --git a/Engine/Plugins/Runtime/Metasound/Source/MetasoundStandardNodes/Private/MetasoundMinNode.cpp b/Engine/Plugins/Runtime/Metasound/Source/MetasoundStandardNodes/Private/MetasoundMinNode.cpp
--- a/Engine/Plugins/Runtime/Metasound/Source/MetasoundStandardNodes/Private/MetasoundMinNode.cpp
+++ b/Engine/Plugins/Runtime/Metasound/Source/MetasoundStandardNodes/Private/MetasoundMinNode.cpp
@@ -152,11 +152,11 @@ namespace Metasound
 		{
 			auto CreateNodeClassMetadata = []() -> FNodeClassMetadata
 			{
-				FName DataTypeName = GetMetasoundDataTypeName<ValueType>();
-				FName OperatorName = TEXT("Min");
-				FText NodeDisplayName = FText::Format(LOCTEXT("MinDisplayNamePattern", "Min ({0})"), FText::FromString(GetMetasoundDataTypeString<ValueType>()));
-				FText NodeDescription = LOCTEXT("MinDesc", "Returns the min of A and B.");
-				FVertexInterface NodeInterface = GetDefaultInterface();
+				const FName DataTypeName = GetMetasoundDataTypeName<ValueType>();
+				const FName OperatorName = TEXT("Min");
+				const FText NodeDisplayName = FText::Format(LOCTEXT("MinDisplayNamePattern", "Min ({0})"), GetMetasoundDataTypeDisplayText<ValueType>());
+				const FText NodeDescription = LOCTEXT("MinDesc", "Returns the min of A and B.");
+				const FVertexInterface NodeInterface = GetDefaultInterface();
 
 				return MetasoundMinNodePrivate::CreateNodeClassMetadata(DataTypeName, OperatorName, NodeDisplayName, NodeDescription, NodeInterface);
 			};
diff --git a/Engine/Plugins/Runtime/Metasound/Source/MetasoundStandardNodes/Private/MetasoundTriggerRouteNode.cpp b/Engine/Plugins/Runtime/Metasound/Source/MetasoundStandardNodes/Private/MetasoundTriggerRouteNode.cpp
--- a/Engine/Plugins/Runtime/Metasound/Source/MetasoundStandardNodes/Private/MetasoundTriggerRouteNode.cpp
+++ b/Engine/Plugins/Runtime/Metasound/Source/MetasoundStandardNodes/Private/MetasoundTriggerRouteNode.cpp
@@ -131,7 +131,7 @@ namespace Metasound
 			{
 				FName DataTypeName = GetMetasoundDataTypeName<ValueType>();
 				FName OperatorName = *FString::Printf(TEXT("Trigger Route (%s, %d)"), *DataTypeName.ToString(), NumInputs);
-				FText NodeDisplayName = FText::Format(LOCTEXT("TriggerRouteDisplayNamePattern", "Trigger Route ({0}, {1})"), FText::FromString(GetMetasoundDataTypeString<ValueType>()), NumInputs);
+				FText NodeDisplayName = FText::Format(LOCTEXT("TriggerRouteDisplayNamePattern", "Trigger Route ({0}, {1})"), GetMetasoundDataTypeDisplayText<ValueType>(), NumInputs);
 				FText NodeDescription = LOCTEXT("TriggerRouteDescription", "Allows routing different values to the same output pin depending on trigger inputs.");
 				FVertexInterface NodeInterface = GetDefaultInterface();
 
diff --git a/Engine/Plugins/Runtime/Metasound/Source/MetasoundStandardNodes/Public/MetasoundTriggerCompareNode.h b/Engine/Plugins/Runtime/Metasound/Source/MetasoundStandardNodes/Public/MetasoundTriggerCompareNode.h
--- a/Engine/Plugins/Runtime/Metasound/Source/MetasoundStandardNodes/Public/MetasoundTriggerCompareNode.h
+++ b/Engine/Plugins/Runtime/Metasound/Source/MetasoundStandardNodes/Public/MetasoundTriggerCompareNode.h
@@ -89,11 +89,11 @@ namespace Metasound
 		{
 			auto CreateNodeClassMetadata = []() -> FNodeClassMetadata
 			{
-				FName DataTypeName = GetMetasoundDataTypeName<ValueType>();
-				FName OperatorName = TEXT("Clamp");
-				FText NodeDisplayName = FText::Format(LOCTEXT("TriggerCompareDisplayPattern", "Trigger Compare ({0})"), FText::FromString(GetMetasoundDataTypeString<ValueType>()));
-				FText NodeDescription = LOCTEXT("TriggerCompareDisc", "Output triggers (True or False) based on comparing inputs, A and B.");
-				FVertexInterface NodeInterface = GetDefaultInterface();
+				const FName DataTypeName = GetMetasoundDataTypeName<ValueType>();
+				const FName OperatorName = TEXT("Clamp");
+				const FText NodeDisplayName = FText::Format(LOCTEXT("TriggerCompareDisplayPattern", "Trigger Compare ({0})"), GetMetasoundDataTypeDisplayText<ValueType>());
+				const FText NodeDescription = LOCTEXT("TriggerCompareDisc", "Output triggers (True or False) based on comparing inputs, A and B.");
+				const FVertexInterface NodeInterface = GetDefaultInterface();
 
 				return MetasoundTriggerCompareNodePrivate::CreateNodeClassMetadata(DataTypeName, OperatorName, NodeDisplayName, NodeDescription, NodeInterface);
 			};
diff --git a/Engine/Plugins/Runtime/Metasound/Source/MetasoundStandardNodes/Public/MetasoundValueNode.h b/Engine/Plugins/Runtime/Metasound/Source/MetasoundStandardNodes/Public/MetasoundValueNode.h
--- a/Engine/Plugins/Runtime/Metasound/Source/MetasoundStandardNodes/Public/MetasoundValueNode.h
+++ b/Engine/Plugins/Runtime/Metasound/Source/MetasoundStandardNodes/Public/MetasoundValueNode.h
@@ -60,11 +60,11 @@ namespace Metasound
 		{
 			auto CreateNodeClassMetadata = []() -> FNodeClassMetadata
 			{
-				FName DataTypeName = GetMetasoundDataTypeName<ValueType>();
-				FName OperatorName = TEXT("Value");
-				FText NodeDisplayName = FText::Format(LOCTEXT("ValueDisplayNamePattern", "Value ({0})"), FText::FromString(GetMetasoundDataTypeString<ValueType>()));
-				FText NodeDescription = LOCTEXT("ValueDescription", "Allows setting a value to output on trigger.");
-				FVertexInterface NodeInterface = GetDefaultInterface();
+				const FName DataTypeName = GetMetasoundDataTypeName<ValueType>();
+				const FName OperatorName = "Value";
+				const FText NodeDisplayName = FText::Format(LOCTEXT("ValueDisplayNamePattern", "Value ({0})"), GetMetasoundDataTypeDisplayText<ValueType>());
+				const FText NodeDescription = LOCTEXT("ValueDescription", "Allows setting a value to output on trigger.");
+				const FVertexInterface NodeInterface = GetDefaultInterface();
 
 				return MetasoundValueNodePrivate::CreateNodeClassMetadata(DataTypeName, OperatorName, NodeDisplayName, NodeDescription, NodeInterface);
 			};
