Cleanup TAccessPtr
- Moved TAccessPtr<> member accessors to their own implementation file.
- Implemented specific derived classes for TAccessPtr<> to metasound frontend document sturctures.
- Fixed broken unit test

diff --git a/Engine/Plugins/Runtime/Metasound/Source/MetasoundEngine/Private/Metasound.cpp b/Engine/Plugins/Runtime/Metasound/Source/MetasoundEngine/Private/Metasound.cpp
--- a/Engine/Plugins/Runtime/Metasound/Source/MetasoundEngine/Private/Metasound.cpp
+++ b/Engine/Plugins/Runtime/Metasound/Source/MetasoundEngine/Private/Metasound.cpp
@@ -42,15 +42,17 @@ void UMetaSound::SetGraph(UEdGraph* InGraph)
 #endif // WITH_EDITORONLY_DATA
 
 // Returns document object responsible for serializing asset
-Metasound::Frontend::TAccessPtr<FMetasoundFrontendDocument> UMetaSound::GetDocument()
+Metasound::Frontend::FDocumentAccessPtr UMetaSound::GetDocument()
 {
-	return Metasound::Frontend::MakeAccessPtr(MetasoundDocument.AccessPoint, MetasoundDocument);
+	using namespace Metasound::Frontend;
+	return MakeAccessPtr<FDocumentAccessPtr>(MetasoundDocument.AccessPoint, MetasoundDocument);
 }
 
 // Returns document object responsible for serializing asset
-Metasound::Frontend::TAccessPtr<const FMetasoundFrontendDocument> UMetaSound::GetDocument() const
+Metasound::Frontend::FConstDocumentAccessPtr UMetaSound::GetDocument() const
 {
-	return Metasound::Frontend::MakeAccessPtr<const FMetasoundFrontendDocument>(MetasoundDocument.AccessPoint, MetasoundDocument);
+	using namespace Metasound::Frontend;
+	return Metasound::Frontend::MakeAccessPtr<FConstDocumentAccessPtr>(MetasoundDocument.AccessPoint, MetasoundDocument);
 }
 
 const TArray<FMetasoundFrontendArchetype>& UMetaSound::GetPreferredMetasoundArchetypes() const
diff --git a/Engine/Plugins/Runtime/Metasound/Source/MetasoundEngine/Private/MetasoundSource.cpp b/Engine/Plugins/Runtime/Metasound/Source/MetasoundEngine/Private/MetasoundSource.cpp
--- a/Engine/Plugins/Runtime/Metasound/Source/MetasoundEngine/Private/MetasoundSource.cpp
+++ b/Engine/Plugins/Runtime/Metasound/Source/MetasoundEngine/Private/MetasoundSource.cpp
@@ -111,6 +111,7 @@ float UMetaSoundSource::GetDuration()
 ISoundGeneratorPtr UMetaSoundSource::CreateSoundGenerator(const FSoundGeneratorInitParams& InParams)
 {
 	using namespace Metasound;
+	using namespace Metasound::Frontend;
 
 	Duration = INDEFINITELY_LOOPING_DURATION;
 	bLooping = true;
@@ -153,7 +154,7 @@ ISoundGeneratorPtr UMetaSoundSource::CreateSoundGenerator(const FSoundGeneratorI
 	ensure(CopyDocumentAndInjectReceiveNodes(InParams.InstanceID, *OriginalDoc, DocumentWithInjectedReceives));
 
 	// Create handles for new root graph
-	Frontend::FConstDocumentHandle NewDocumentHandle = Frontend::IDocumentController::CreateDocumentHandle(Frontend::MakeAccessPtr<const FMetasoundFrontendDocument>(DocumentWithInjectedReceives.AccessPoint, DocumentWithInjectedReceives));
+	FConstDocumentHandle NewDocumentHandle = IDocumentController::CreateDocumentHandle(MakeAccessPtr<FConstDocumentAccessPtr>(DocumentWithInjectedReceives.AccessPoint, DocumentWithInjectedReceives));
 
 	FMetasoundGeneratorInitParams InitParams = 
 	{
@@ -215,11 +216,12 @@ Metasound::Frontend::FNodeHandle UMetaSoundSource::AddInputPinForSendAddress(con
 bool UMetaSoundSource::CopyDocumentAndInjectReceiveNodes(uint64 InInstanceID, const FMetasoundFrontendDocument& InSourceDoc, FMetasoundFrontendDocument& OutDestDoc) const
 {
 	using namespace Metasound;
+	using namespace Metasound::Frontend;
 
 	OutDestDoc = InSourceDoc;
 
-	Frontend::FDocumentHandle Document = Frontend::IDocumentController::CreateDocumentHandle(Frontend::MakeAccessPtr(OutDestDoc.AccessPoint, OutDestDoc));
-	Frontend::FGraphHandle RootGraph = Document->GetRootGraph();
+	FDocumentHandle Document = IDocumentController::CreateDocumentHandle(MakeAccessPtr<FDocumentAccessPtr>(OutDestDoc.AccessPoint, OutDestDoc));
+	FGraphHandle RootGraph = Document->GetRootGraph();
 
 	TArray<FSendInfoAndVertexName> SendInfoAndVertexes = GetSendInfos(InInstanceID);
 
@@ -235,19 +237,19 @@ bool UMetaSoundSource::CopyDocumentAndInjectReceiveNodes(uint64 InInstanceID, co
 			// TODO: log warning
 			continue;
 		}
-		Frontend::FNodeHandle ReceiveNode = RootGraph->AddNode(ReceiveNodeMetadata);
+		FNodeHandle ReceiveNode = RootGraph->AddNode(ReceiveNodeMetadata);
 
 		// Add receive node address to graph
-		Frontend::FNodeHandle AddressNode = UMetaSoundSource::AddInputPinForSendAddress(InfoAndVertexName.SendInfo, RootGraph);
-		TArray<Frontend::FOutputHandle> AddressNodeOutputs = AddressNode->GetOutputs();
+		FNodeHandle AddressNode = UMetaSoundSource::AddInputPinForSendAddress(InfoAndVertexName.SendInfo, RootGraph);
+		TArray<FOutputHandle> AddressNodeOutputs = AddressNode->GetOutputs();
 		if (AddressNodeOutputs.Num() != 1)
 		{
 			// TODO: log warning
 			continue;
 		}
 
-		Frontend::FOutputHandle AddressOutput = AddressNodeOutputs[0];
-		TArray<Frontend::FInputHandle> ReceiveAddressInput = ReceiveNode->GetInputsWithVertexName(Metasound::FReceiveNodeNames::GetAddressInputName());
+		FOutputHandle AddressOutput = AddressNodeOutputs[0];
+		TArray<FInputHandle> ReceiveAddressInput = ReceiveNode->GetInputsWithVertexName(Metasound::FReceiveNodeNames::GetAddressInputName());
 		if (ReceiveAddressInput.Num() != 1)
 		{
 			// TODO: log error
@@ -259,29 +261,29 @@ bool UMetaSoundSource::CopyDocumentAndInjectReceiveNodes(uint64 InInstanceID, co
 
 
 		// Swap input node connections with receive node connections
-		Frontend::FNodeHandle InputNode = RootGraph->GetInputNodeWithName(InfoAndVertexName.VertexName);
+		FNodeHandle InputNode = RootGraph->GetInputNodeWithName(InfoAndVertexName.VertexName);
 		if (!ensure(InputNode->GetOutputs().Num() == 1))
 		{
 			// TODO: handle input node with varying number of outputs or varying output types.
 			continue;
 		}
 
-		Frontend::FOutputHandle InputNodeOutput = InputNode->GetOutputs()[0];
+		FOutputHandle InputNodeOutput = InputNode->GetOutputs()[0];
 
 		if (ensure(ReceiveNode->IsValid()))
 		{
-			TArray<Frontend::FOutputHandle> ReceiveNodeOutputs = ReceiveNode->GetOutputs();
+			TArray<FOutputHandle> ReceiveNodeOutputs = ReceiveNode->GetOutputs();
 			if (!ensure(ReceiveNodeOutputs.Num() == 1))
 			{
 				// TODO: handle array outputs and receive nodes of varying formats.
 				continue;
 			}
 
-			TArray<Frontend::FInputHandle> ReceiveDefaultInputs = ReceiveNode->GetInputsWithVertexName(Metasound::FReceiveNodeNames::GetDefaultDataInputName());
+			TArray<FInputHandle> ReceiveDefaultInputs = ReceiveNode->GetInputsWithVertexName(Metasound::FReceiveNodeNames::GetDefaultDataInputName());
 			if (ensure(ReceiveDefaultInputs.Num() == 1))
 			{
-				Frontend::FOutputHandle ReceiverNodeOutput = ReceiveNodeOutputs[0];
-				for (Frontend::FInputHandle NodeInput : InputNodeOutput->GetCurrentlyConnectedInputs())
+				FOutputHandle ReceiverNodeOutput = ReceiveNodeOutputs[0];
+				for (FInputHandle NodeInput : InputNodeOutput->GetCurrentlyConnectedInputs())
 				{
 					// Swap connections to receiver node
 					ensure(InputNodeOutput->Disconnect(*NodeInput));
diff --git a/Engine/Plugins/Runtime/Metasound/Source/MetasoundEngine/Public/Metasound.h b/Engine/Plugins/Runtime/Metasound/Source/MetasoundEngine/Public/Metasound.h
--- a/Engine/Plugins/Runtime/Metasound/Source/MetasoundEngine/Public/Metasound.h
+++ b/Engine/Plugins/Runtime/Metasound/Source/MetasoundEngine/Public/Metasound.h
@@ -2,9 +2,10 @@
 #pragma once
 
 #include "CoreMinimal.h"
+#include "MetasoundAssetBase.h"
 #include "MetasoundFrontend.h"
 #include "MetasoundFrontendDocument.h"
-#include "MetasoundAssetBase.h"
+#include "MetasoundFrontendDocumentAccessPtr.h"
 
 
 #if WITH_EDITORONLY_DATA
@@ -33,10 +34,10 @@ protected:
 	FMetasoundFrontendDocument MetasoundDocument;
 
 	// Returns document object responsible for serializing asset
-	Metasound::Frontend::TAccessPtr<FMetasoundFrontendDocument> GetDocument() override;
+	Metasound::Frontend::FDocumentAccessPtr GetDocument() override;
 
 	// Returns document object responsible for serializing asset
-	Metasound::Frontend::TAccessPtr<const FMetasoundFrontendDocument> GetDocument() const override;
+	Metasound::Frontend::FConstDocumentAccessPtr GetDocument() const override;
 
 #if WITH_EDITORONLY_DATA
 	UPROPERTY()
diff --git a/Engine/Plugins/Runtime/Metasound/Source/MetasoundEngine/Public/MetasoundSource.h b/Engine/Plugins/Runtime/Metasound/Source/MetasoundEngine/Public/MetasoundSource.h
--- a/Engine/Plugins/Runtime/Metasound/Source/MetasoundEngine/Public/MetasoundSource.h
+++ b/Engine/Plugins/Runtime/Metasound/Source/MetasoundEngine/Public/MetasoundSource.h
@@ -139,18 +139,20 @@ public:
 
 protected:
 
-	Metasound::Frontend::TAccessPtr<FMetasoundFrontendDocument> GetDocument() override
+	Metasound::Frontend::FDocumentAccessPtr GetDocument() override
 	{
+		using namespace Metasound::Frontend;
 		// Return document using FAccessPoint to inform the TAccessPtr when the 
 		// object is no longer valid.
-		return Metasound::Frontend::MakeAccessPtr<FMetasoundFrontendDocument>(RootMetasoundDocument.AccessPoint, RootMetasoundDocument);
+		return MakeAccessPtr<FDocumentAccessPtr>(RootMetasoundDocument.AccessPoint, RootMetasoundDocument);
 	}
 
-	Metasound::Frontend::TAccessPtr<const FMetasoundFrontendDocument> GetDocument() const override
+	Metasound::Frontend::FConstDocumentAccessPtr GetDocument() const override
 	{
+		using namespace Metasound::Frontend;
 		// Return document using FAccessPoint to inform the TAccessPtr when the 
 		// object is no longer valid.
-		return Metasound::Frontend::MakeAccessPtr<const FMetasoundFrontendDocument>(RootMetasoundDocument.AccessPoint, RootMetasoundDocument);
+		return MakeAccessPtr<FConstDocumentAccessPtr>(RootMetasoundDocument.AccessPoint, RootMetasoundDocument);
 	}
 
 private:
diff --git a/Engine/Plugins/Runtime/Metasound/Source/MetasoundFrontend/Private/MetasoundAssetBase.cpp b/Engine/Plugins/Runtime/Metasound/Source/MetasoundFrontend/Private/MetasoundAssetBase.cpp
--- a/Engine/Plugins/Runtime/Metasound/Source/MetasoundFrontend/Private/MetasoundAssetBase.cpp
+++ b/Engine/Plugins/Runtime/Metasound/Source/MetasoundFrontend/Private/MetasoundAssetBase.cpp
@@ -148,7 +148,7 @@ Metasound::Frontend::FConstGraphHandle FMetasoundAssetBase::GetRootGraphHandle()
 
 bool FMetasoundAssetBase::ImportFromJSON(const FString& InJSON)
 {
-	Metasound::Frontend::TAccessPtr<FMetasoundFrontendDocument> Document = GetDocument();
+	Metasound::Frontend::FDocumentAccessPtr Document = GetDocument();
 	if (ensure(Document.IsValid()))
 	{
 		bool bSuccess = Metasound::Frontend::ImportJSONToMetasound(InJSON, *Document);
@@ -165,7 +165,7 @@ bool FMetasoundAssetBase::ImportFromJSON(const FString& InJSON)
 
 bool FMetasoundAssetBase::ImportFromJSONAsset(const FString& InAbsolutePath)
 {
-	Metasound::Frontend::TAccessPtr<FMetasoundFrontendDocument> Document = GetDocument();
+	Metasound::Frontend::FDocumentAccessPtr Document = GetDocument();
 	if (ensure(Document.IsValid()))
 	{
 		bool bSuccess = Metasound::Frontend::ImportJSONAssetToMetasound(InAbsolutePath, *Document);
@@ -182,7 +182,7 @@ bool FMetasoundAssetBase::ImportFromJSONAsset(const FString& InAbsolutePath)
 
 FMetasoundFrontendDocument& FMetasoundAssetBase::GetDocumentChecked()
 {
-	Metasound::Frontend::TAccessPtr<FMetasoundFrontendDocument> DocAccessPtr = GetDocument();
+	Metasound::Frontend::FDocumentAccessPtr DocAccessPtr = GetDocument();
 
 	check(DocAccessPtr.IsValid());
 	return *DocAccessPtr;
@@ -190,7 +190,7 @@ FMetasoundFrontendDocument& FMetasoundAssetBase::GetDocumentChecked()
 
 const FMetasoundFrontendDocument& FMetasoundAssetBase::GetDocumentChecked() const
 {
-	Metasound::Frontend::TAccessPtr<const FMetasoundFrontendDocument> DocAccessPtr = GetDocument();
+	Metasound::Frontend::FConstDocumentAccessPtr DocAccessPtr = GetDocument();
 
 	check(DocAccessPtr.IsValid());
 	return *DocAccessPtr;
diff --git a/Engine/Plugins/Runtime/Metasound/Source/MetasoundFrontend/Private/MetasoundFrontendController.cpp b/Engine/Plugins/Runtime/Metasound/Source/MetasoundFrontend/Private/MetasoundFrontendController.cpp
--- a/Engine/Plugins/Runtime/Metasound/Source/MetasoundFrontend/Private/MetasoundFrontendController.cpp
+++ b/Engine/Plugins/Runtime/Metasound/Source/MetasoundFrontend/Private/MetasoundFrontendController.cpp
@@ -35,16 +35,16 @@ namespace Metasound
 			return FInvalidDocumentController::GetInvalid();
 		}
 			
-		FDocumentHandle IDocumentController::CreateDocumentHandle(TAccessPtr<FMetasoundFrontendDocument> InDocument)
+		FDocumentHandle IDocumentController::CreateDocumentHandle(FDocumentAccessPtr InDocument)
 		{
 			// Create using standard document controller.
 			return FDocumentController::CreateDocumentHandle(InDocument);
 		}
 
-		FConstDocumentHandle IDocumentController::CreateDocumentHandle(TAccessPtr<const FMetasoundFrontendDocument> InDocument)
+		FConstDocumentHandle IDocumentController::CreateDocumentHandle(FConstDocumentAccessPtr InDocument)
 		{
 			// Create using standard document controller. 
-			return FDocumentController::CreateDocumentHandle(ConstCastAccessPtr<FMetasoundFrontendDocument>(InDocument));
+			return FDocumentController::CreateDocumentHandle(ConstCastAccessPtr<FDocumentAccessPtr>(InDocument));
 		}
 	}
 }
diff --git a/Engine/Plugins/Runtime/Metasound/Source/MetasoundFrontend/Private/MetasoundFrontendDocumentAccessPtr.cpp b/Engine/Plugins/Runtime/Metasound/Source/MetasoundFrontend/Private/MetasoundFrontendDocumentAccessPtr.cpp
new file mode 100644
--- /dev/null
+++ b/Engine/Plugins/Runtime/Metasound/Source/MetasoundFrontend/Private/MetasoundFrontendDocumentAccessPtr.cpp
@@ -0,0 +1,680 @@
+// Copyright Epic Games, Inc. All Rights Reserved.
+
+#include "MetasoundFrontendDocumentAccessPtr.h"
+#include "MetasoundFrontendRegistries.h"
+
+namespace Metasound
+{
+	namespace Frontend
+	{
+		namespace MetasoundFrontendDocumentAccessPtrPrivate
+		{
+			using FRegistry = FMetasoundFrontendRegistryContainer;
+
+			bool IsMatchingMetasoundClass(const FMetasoundFrontendClassMetadata& InMetadataA, const FMetasoundFrontendClassMetadata& InMetadataB) 
+			{
+				if (InMetadataA.Type == InMetadataB.Type)
+				{
+					if (InMetadataA.ClassName == InMetadataB.ClassName)
+					{
+						return FRegistry::GetRegistryKey(InMetadataA) == FRegistry::GetRegistryKey(InMetadataB);
+					}
+				}
+				return false;
+			}
+
+			bool IsMatchingMetasoundClass(const FNodeClassInfo& InNodeClass, const FMetasoundFrontendClassMetadata& InMetadata) 
+			{
+				if (InNodeClass.NodeType == InMetadata.Type)
+				{
+					return InNodeClass.LookupKey == FRegistry::GetRegistryKey(InMetadata);
+				}
+				return false;
+			}
+		}
+
+		FVertexAccessPtr FNodeAccessPtr::GetInputWithName(const FString& InName)
+		{
+			auto FindInputWithName = [Name=InName](FMetasoundFrontendNode& InNode) -> FMetasoundFrontendVertex*
+			{
+				auto IsVertexWithName = [&](const FMetasoundFrontendVertex& Vertex) { return Vertex.Name == Name; };
+				return InNode.Interface.Inputs.FindByPredicate(IsVertexWithName);
+			};
+
+			return GetMemberAccessPtr<FVertexAccessPtr>(FindInputWithName);
+		}
+
+		FVertexAccessPtr FNodeAccessPtr::GetInputWithVertexID(const FGuid& InID)
+		{
+			auto FindInputWithID = [ID=InID](FMetasoundFrontendNode& InNode) -> FMetasoundFrontendVertex*
+			{
+				auto IsVertexWithVertexID = [=](const FMetasoundFrontendVertex& Vertex) { return Vertex.VertexID == ID; };
+				return InNode.Interface.Inputs.FindByPredicate(IsVertexWithVertexID);
+			};
+			return GetMemberAccessPtr<FVertexAccessPtr>(FindInputWithID);
+		}
+
+		FVertexAccessPtr FNodeAccessPtr::GetOutputWithName(const FString& InName)
+		{
+			auto FindOutputWithName = [Name=InName](FMetasoundFrontendNode& InNode) -> FMetasoundFrontendVertex*
+			{
+				auto IsVertexWithName = [&](const FMetasoundFrontendVertex& Vertex) { return Vertex.Name == Name; };
+				return InNode.Interface.Outputs.FindByPredicate(IsVertexWithName);
+			};
+
+			return GetMemberAccessPtr<FVertexAccessPtr>(FindOutputWithName);
+		}
+
+		FVertexAccessPtr FNodeAccessPtr::GetOutputWithVertexID(const FGuid& InID)
+		{
+			auto FindOutputWithID = [ID=InID](FMetasoundFrontendNode& InNode) -> FMetasoundFrontendVertex*
+			{
+				auto IsVertexWithVertexID = [=](const FMetasoundFrontendVertex& Vertex) { return Vertex.VertexID == ID; };
+				return InNode.Interface.Outputs.FindByPredicate(IsVertexWithVertexID);
+			};
+			return GetMemberAccessPtr<FVertexAccessPtr>(FindOutputWithID);
+		}
+
+		FConstVertexAccessPtr FNodeAccessPtr::GetInputWithName(const FString& InName) const
+		{
+			auto FindInputWithName = [Name=InName](const FMetasoundFrontendNode& InNode) -> const FMetasoundFrontendVertex*
+			{
+				auto IsVertexWithName = [&](const FMetasoundFrontendVertex& Vertex) { return Vertex.Name == Name; };
+				return InNode.Interface.Inputs.FindByPredicate(IsVertexWithName);
+			};
+
+			return GetMemberAccessPtr<FConstVertexAccessPtr>(FindInputWithName);
+		}
+
+		FConstVertexAccessPtr FNodeAccessPtr::GetInputWithVertexID(const FGuid& InID) const
+		{
+			auto FindInputWithID = [ID=InID](const FMetasoundFrontendNode& InNode) -> const FMetasoundFrontendVertex*
+			{
+				auto IsVertexWithVertexID = [=](const FMetasoundFrontendVertex& Vertex) { return Vertex.VertexID == ID; };
+				return InNode.Interface.Inputs.FindByPredicate(IsVertexWithVertexID);
+			};
+			return GetMemberAccessPtr<FConstVertexAccessPtr>(FindInputWithID);
+		}
+
+		FConstVertexAccessPtr FNodeAccessPtr::GetOutputWithName(const FString& InName) const
+		{
+			auto FindOutputWithName = [Name=InName](const FMetasoundFrontendNode& InNode) -> const FMetasoundFrontendVertex*
+			{
+				auto IsVertexWithName = [&](const FMetasoundFrontendVertex& Vertex) { return Vertex.Name == Name; };
+				return InNode.Interface.Outputs.FindByPredicate(IsVertexWithName);
+			};
+
+			return GetMemberAccessPtr<FConstVertexAccessPtr>(FindOutputWithName);
+		}
+
+		FConstVertexAccessPtr FNodeAccessPtr::GetOutputWithVertexID(const FGuid& InID) const
+		{
+			auto FindOutputWithID = [ID=InID](const FMetasoundFrontendNode& InNode) -> const FMetasoundFrontendVertex*
+			{
+				auto IsVertexWithVertexID = [=](const FMetasoundFrontendVertex& Vertex) { return Vertex.VertexID == ID; };
+				return InNode.Interface.Outputs.FindByPredicate(IsVertexWithVertexID);
+			};
+			return GetMemberAccessPtr<FConstVertexAccessPtr>(FindOutputWithID);
+		}
+
+		FConstVertexAccessPtr FConstNodeAccessPtr::GetInputWithName(const FString& InName) const
+		{
+			auto FindInputWithName = [Name=InName](const FMetasoundFrontendNode& InNode) -> const FMetasoundFrontendVertex*
+			{
+				auto IsVertexWithName = [&](const FMetasoundFrontendVertex& Vertex) { return Vertex.Name == Name; };
+				return InNode.Interface.Inputs.FindByPredicate(IsVertexWithName);
+			};
+
+			return GetMemberAccessPtr<FConstVertexAccessPtr>(FindInputWithName);
+		}
+
+		FConstVertexAccessPtr FConstNodeAccessPtr::GetInputWithVertexID(const FGuid& InID) const
+		{
+			auto FindInputWithID = [ID=InID](const FMetasoundFrontendNode& InNode) -> const FMetasoundFrontendVertex*
+			{
+				auto IsVertexWithVertexID = [=](const FMetasoundFrontendVertex& Vertex) { return Vertex.VertexID == ID; };
+				return InNode.Interface.Inputs.FindByPredicate(IsVertexWithVertexID);
+			};
+			return GetMemberAccessPtr<FConstVertexAccessPtr>(FindInputWithID);
+		}
+
+		FConstVertexAccessPtr FConstNodeAccessPtr::GetOutputWithName(const FString& InName) const
+		{
+			auto FindOutputWithName = [Name=InName](const FMetasoundFrontendNode& InNode) -> const FMetasoundFrontendVertex*
+			{
+				auto IsVertexWithName = [&](const FMetasoundFrontendVertex& Vertex) { return Vertex.Name == Name; };
+				return InNode.Interface.Outputs.FindByPredicate(IsVertexWithName);
+			};
+
+			return GetMemberAccessPtr<FConstVertexAccessPtr>(FindOutputWithName);
+		}
+
+		FConstVertexAccessPtr FConstNodeAccessPtr::GetOutputWithVertexID(const FGuid& InID) const
+		{
+			auto FindOutputWithID = [ID=InID](const FMetasoundFrontendNode& InNode) -> const FMetasoundFrontendVertex*
+			{
+				auto IsVertexWithVertexID = [=](const FMetasoundFrontendVertex& Vertex) { return Vertex.VertexID == ID; };
+				return InNode.Interface.Outputs.FindByPredicate(IsVertexWithVertexID);
+			};
+			return GetMemberAccessPtr<FConstVertexAccessPtr>(FindOutputWithID);
+		}
+
+
+		// FConstClassAccessPtr
+		FClassInputAccessPtr FClassAccessPtr::GetInputWithName(const FString& InName)
+		{
+			auto FindClassInputWithName = [Name=InName](FMetasoundFrontendClass& InClass)
+			{
+				auto IsClassInputWithName = [&](const FMetasoundFrontendClassInput& ClassInput) { return ClassInput.Name == Name; };
+				return InClass.Interface.Inputs.FindByPredicate(IsClassInputWithName);
+
+			};
+
+			return GetMemberAccessPtr<FClassInputAccessPtr>(FindClassInputWithName);
+		}
+
+		FClassOutputAccessPtr FClassAccessPtr::GetOutputWithName(const FString& InName)
+		{
+			auto FindClassOutputWithName = [Name=InName](FMetasoundFrontendClass& InClass)
+			{
+				auto IsClassOutputWithName = [&](const FMetasoundFrontendClassOutput& ClassOutput) { return ClassOutput.Name == Name; };
+				return InClass.Interface.Outputs.FindByPredicate(IsClassOutputWithName);
+
+			};
+
+			return GetMemberAccessPtr<FClassOutputAccessPtr>(FindClassOutputWithName);
+		}
+
+		FConstClassInputAccessPtr FClassAccessPtr::GetInputWithName(const FString& InName) const
+		{
+			auto FindClassInputWithName = [Name=InName](const FMetasoundFrontendClass& InClass)
+			{
+				auto IsClassInputWithName = [&](const FMetasoundFrontendClassInput& ClassInput) { return ClassInput.Name == Name; };
+				return InClass.Interface.Inputs.FindByPredicate(IsClassInputWithName);
+
+			};
+
+			return GetMemberAccessPtr<FConstClassInputAccessPtr>(FindClassInputWithName);
+		}
+
+		FConstClassOutputAccessPtr FClassAccessPtr::GetOutputWithName(const FString& InName) const
+		{
+			auto FindClassOutputWithName = [Name=InName](const FMetasoundFrontendClass& InClass)
+			{
+				auto IsClassOutputWithName = [&](const FMetasoundFrontendClassOutput& ClassOutput) { return ClassOutput.Name == Name; };
+				return InClass.Interface.Outputs.FindByPredicate(IsClassOutputWithName);
+
+			};
+
+			return GetMemberAccessPtr<FConstClassOutputAccessPtr>(FindClassOutputWithName);
+		}
+
+		// FConstClassAccessPtr
+		FConstClassInputAccessPtr FConstClassAccessPtr::GetInputWithName(const FString& InName) const
+		{
+			auto FindClassInputWithName = [Name=InName](const FMetasoundFrontendClass& InClass)
+			{
+				auto IsClassInputWithName = [&](const FMetasoundFrontendClassInput& ClassInput) { return ClassInput.Name == Name; };
+				return InClass.Interface.Inputs.FindByPredicate(IsClassInputWithName);
+
+			};
+
+			return GetMemberAccessPtr<FConstClassInputAccessPtr>(FindClassInputWithName);
+		}
+
+		FConstClassOutputAccessPtr FConstClassAccessPtr::GetOutputWithName(const FString& InName) const
+		{
+			auto FindClassOutputWithName = [Name=InName](const FMetasoundFrontendClass& InClass)
+			{
+				auto IsClassOutputWithName = [&](const FMetasoundFrontendClassOutput& ClassOutput) { return ClassOutput.Name == Name; };
+				return InClass.Interface.Outputs.FindByPredicate(IsClassOutputWithName);
+
+			};
+
+			return GetMemberAccessPtr<FConstClassOutputAccessPtr>(FindClassOutputWithName);
+		}
+
+		// FGraphClassAccessPtr
+		FClassInputAccessPtr FGraphClassAccessPtr::GetInputWithName(const FString& InName)
+		{
+			auto FindClassInputWithName = [Name=InName](FMetasoundFrontendGraphClass& InClass)
+			{
+				auto IsClassInputWithName = [&](const FMetasoundFrontendClassInput& ClassInput) { return ClassInput.Name == Name; };
+				return InClass.Interface.Inputs.FindByPredicate(IsClassInputWithName);
+
+			};
+
+			return GetMemberAccessPtr<FClassInputAccessPtr>(FindClassInputWithName);
+		}
+
+		FClassInputAccessPtr FGraphClassAccessPtr::GetInputWithNodeID(const FGuid& InNodeID)
+		{
+			auto FindClassInputWithNodeID = [NodeID=InNodeID](FMetasoundFrontendGraphClass& InGraphClass) -> FMetasoundFrontendClassInput*
+			{
+				return InGraphClass.Interface.Inputs.FindByPredicate([&](const FMetasoundFrontendClassInput& Desc) { return Desc.NodeID == NodeID; });
+			};
+			return GetMemberAccessPtr<FClassInputAccessPtr>(FindClassInputWithNodeID);
+		}
+
+		FClassOutputAccessPtr FGraphClassAccessPtr::GetOutputWithName(const FString& InName)
+		{
+			auto FindClassOutputWithName = [Name=InName](FMetasoundFrontendGraphClass& InClass)
+			{
+				auto IsClassOutputWithName = [&](const FMetasoundFrontendClassOutput& ClassOutput) { return ClassOutput.Name == Name; };
+				return InClass.Interface.Outputs.FindByPredicate(IsClassOutputWithName);
+
+			};
+
+			return GetMemberAccessPtr<FClassOutputAccessPtr>(FindClassOutputWithName);
+		}
+
+		FClassOutputAccessPtr FGraphClassAccessPtr::GetOutputWithNodeID(const FGuid& InNodeID)
+		{
+			auto FindClassOutputWithNodeID = [NodeID=InNodeID](FMetasoundFrontendGraphClass& InGraphClass) -> FMetasoundFrontendClassOutput*
+			{
+				return InGraphClass.Interface.Outputs.FindByPredicate([&](const FMetasoundFrontendClassOutput& Desc) { return Desc.NodeID == NodeID; });
+			};
+			return GetMemberAccessPtr<FClassOutputAccessPtr>(FindClassOutputWithNodeID);
+		}
+
+		FNodeAccessPtr FGraphClassAccessPtr::GetNodeWithNodeID(const FGuid& InNodeID)
+		{
+			auto FindNodeWithNodeID = [ID=InNodeID](FMetasoundFrontendGraphClass& InGraphClass) -> FMetasoundFrontendNode*
+			{
+				auto IsNodeWithID = [&](const FMetasoundFrontendNode& InNode) { return InNode.ID == ID; };
+				return InGraphClass.Graph.Nodes.FindByPredicate(IsNodeWithID);
+			};
+
+			return GetMemberAccessPtr<FNodeAccessPtr>(FindNodeWithNodeID);
+		}
+
+		FGraphAccessPtr FGraphClassAccessPtr::GetGraph()
+		{
+			auto FindGraph = [](FMetasoundFrontendGraphClass& InGraphClass) -> FMetasoundFrontendGraph* 
+			{ 
+				return &InGraphClass.Graph; 
+			};
+
+			return GetMemberAccessPtr<FGraphAccessPtr>(FindGraph);
+		}
+
+
+		FConstClassInputAccessPtr FGraphClassAccessPtr::GetInputWithName(const FString& InName) const
+		{
+			auto FindClassInputWithName = [Name=InName](const FMetasoundFrontendGraphClass& InClass)
+			{
+				auto IsClassInputWithName = [&](const FMetasoundFrontendClassInput& ClassInput) { return ClassInput.Name == Name; };
+				return InClass.Interface.Inputs.FindByPredicate(IsClassInputWithName);
+
+			};
+
+			return GetMemberAccessPtr<FConstClassInputAccessPtr>(FindClassInputWithName);
+		}
+
+
+		FConstClassInputAccessPtr FGraphClassAccessPtr::GetInputWithNodeID(const FGuid& InNodeID) const
+		{
+			auto FindClassInputWithNodeID = [NodeID=InNodeID](const FMetasoundFrontendGraphClass& InGraphClass) -> const FMetasoundFrontendClassInput*
+			{
+				return InGraphClass.Interface.Inputs.FindByPredicate([&](const FMetasoundFrontendClassInput& Desc) { return Desc.NodeID == NodeID; });
+			};
+			return GetMemberAccessPtr<FConstClassInputAccessPtr>(FindClassInputWithNodeID);
+		}
+
+		FConstClassOutputAccessPtr FGraphClassAccessPtr::GetOutputWithName(const FString& InName) const
+		{
+			auto FindClassOutputWithName = [Name=InName](const FMetasoundFrontendGraphClass& InClass)
+			{
+				auto IsClassOutputWithName = [&](const FMetasoundFrontendClassOutput& ClassOutput) { return ClassOutput.Name == Name; };
+				return InClass.Interface.Outputs.FindByPredicate(IsClassOutputWithName);
+
+			};
+
+			return GetMemberAccessPtr<FConstClassOutputAccessPtr>(FindClassOutputWithName);
+		}
+
+		FConstClassOutputAccessPtr FGraphClassAccessPtr::GetOutputWithNodeID(const FGuid& InNodeID) const
+		{
+			auto FindClassOutputWithNodeID = [NodeID=InNodeID](const FMetasoundFrontendGraphClass& InGraphClass) -> const FMetasoundFrontendClassOutput*
+			{
+				return InGraphClass.Interface.Outputs.FindByPredicate([&](const FMetasoundFrontendClassOutput& Desc) { return Desc.NodeID == NodeID; });
+			};
+			return GetMemberAccessPtr<FConstClassOutputAccessPtr>(FindClassOutputWithNodeID);
+		}
+
+		FConstNodeAccessPtr FGraphClassAccessPtr::GetNodeWithNodeID(const FGuid& InNodeID) const
+		{
+			auto FindNodeWithNodeID = [ID=InNodeID](const FMetasoundFrontendGraphClass& InGraphClass) -> const FMetasoundFrontendNode*
+			{
+				auto IsNodeWithID = [&](const FMetasoundFrontendNode& InNode) { return InNode.ID == ID; };
+				return InGraphClass.Graph.Nodes.FindByPredicate(IsNodeWithID);
+			};
+
+			return GetMemberAccessPtr<FConstNodeAccessPtr>(FindNodeWithNodeID);
+		}
+
+		FConstGraphAccessPtr FGraphClassAccessPtr::GetGraph() const
+		{
+			auto FindGraph = [](const FMetasoundFrontendGraphClass& InGraphClass) -> const FMetasoundFrontendGraph* 
+			{ 
+				return &InGraphClass.Graph; 
+			};
+
+			return GetMemberAccessPtr<FConstGraphAccessPtr>(FindGraph);
+		}
+
+		// FConstGraphClassAccessPtr
+		FConstClassInputAccessPtr FConstGraphClassAccessPtr::GetInputWithName(const FString& InName) const
+		{
+			auto FindClassInputWithName = [Name=InName](const FMetasoundFrontendGraphClass& InClass)
+			{
+				auto IsClassInputWithName = [&](const FMetasoundFrontendClassInput& ClassInput) { return ClassInput.Name == Name; };
+				return InClass.Interface.Inputs.FindByPredicate(IsClassInputWithName);
+
+			};
+
+			return GetMemberAccessPtr<FConstClassInputAccessPtr>(FindClassInputWithName);
+		}
+
+		FConstClassInputAccessPtr FConstGraphClassAccessPtr::GetInputWithNodeID(const FGuid& InNodeID) const
+		{
+			auto FindClassInputWithNodeID = [NodeID=InNodeID](const FMetasoundFrontendGraphClass& InGraphClass) -> const FMetasoundFrontendClassInput*
+			{
+				return InGraphClass.Interface.Inputs.FindByPredicate([&](const FMetasoundFrontendClassInput& Desc) { return Desc.NodeID == NodeID; });
+			};
+			return GetMemberAccessPtr<FConstClassInputAccessPtr>(FindClassInputWithNodeID);
+		}
+
+		FConstClassOutputAccessPtr FConstGraphClassAccessPtr::GetOutputWithName(const FString& InName) const
+		{
+			auto FindClassOutputWithName = [Name=InName](const FMetasoundFrontendGraphClass& InClass)
+			{
+				auto IsClassOutputWithName = [&](const FMetasoundFrontendClassOutput& ClassOutput) { return ClassOutput.Name == Name; };
+				return InClass.Interface.Outputs.FindByPredicate(IsClassOutputWithName);
+
+			};
+
+			return GetMemberAccessPtr<FConstClassOutputAccessPtr>(FindClassOutputWithName);
+		}
+
+		FConstClassOutputAccessPtr FConstGraphClassAccessPtr::GetOutputWithNodeID(const FGuid& InNodeID) const
+		{
+			auto FindClassOutputWithNodeID = [NodeID=InNodeID](const FMetasoundFrontendGraphClass& InGraphClass) -> const FMetasoundFrontendClassOutput*
+			{
+				return InGraphClass.Interface.Outputs.FindByPredicate([&](const FMetasoundFrontendClassOutput& Desc) { return Desc.NodeID == NodeID; });
+			};
+			return GetMemberAccessPtr<FConstClassOutputAccessPtr>(FindClassOutputWithNodeID);
+		}
+
+		FConstNodeAccessPtr FConstGraphClassAccessPtr::GetNodeWithNodeID(const FGuid& InNodeID) const
+		{
+			auto FindNodeWithNodeID = [ID=InNodeID](const FMetasoundFrontendGraphClass& InGraphClass) -> const FMetasoundFrontendNode*
+			{
+				auto IsNodeWithID = [&](const FMetasoundFrontendNode& InNode) { return InNode.ID == ID; };
+				return InGraphClass.Graph.Nodes.FindByPredicate(IsNodeWithID);
+			};
+
+			return GetMemberAccessPtr<FConstNodeAccessPtr>(FindNodeWithNodeID);
+		}
+
+		FConstGraphAccessPtr FConstGraphClassAccessPtr::GetGraph() const
+		{
+			auto FindGraph = [](const FMetasoundFrontendGraphClass& InGraphClass) -> const FMetasoundFrontendGraph* 
+			{ 
+				return &InGraphClass.Graph; 
+			};
+
+			return GetMemberAccessPtr<FConstGraphAccessPtr>(FindGraph);
+		}
+
+		// FDocumentAccessPtr
+		
+		FGraphClassAccessPtr FDocumentAccessPtr::GetRootGraph()
+		{
+			return GetMemberAccessPtr<FGraphClassAccessPtr>([](FMetasoundFrontendDocument& InDoc) { return &InDoc.RootGraph; });
+		}
+
+		FGraphClassAccessPtr FDocumentAccessPtr::GetSubgraphWithID(const FGuid& InID)
+		{
+			auto FindSubgraphWithID = [ID=InID](FMetasoundFrontendDocument& InDoc) -> FMetasoundFrontendGraphClass*
+			{
+				return InDoc.Subgraphs.FindByPredicate([&](const FMetasoundFrontendGraphClass& InSubgraph) { return InSubgraph.ID == ID; });
+			};
+
+			return GetMemberAccessPtr<FGraphClassAccessPtr>(FindSubgraphWithID);
+		}
+
+		FClassAccessPtr FDocumentAccessPtr::GetDependencyWithID(const FGuid& InID)
+		{
+			auto FindDependencyWithID = [ID=InID](FMetasoundFrontendDocument& InDoc) -> FMetasoundFrontendClass*
+			{
+				return InDoc.Dependencies.FindByPredicate([&](const FMetasoundFrontendClass& InDependency) { return InDependency.ID == ID; });
+			};
+			return GetMemberAccessPtr<FClassAccessPtr>(FindDependencyWithID);
+		}
+
+
+		FClassAccessPtr FDocumentAccessPtr::GetClassWithID(const FGuid& InID) 
+		{
+			auto FindClassWithID = [ID=InID](FMetasoundFrontendDocument& InDoc) -> FMetasoundFrontendClass*
+			{
+				FMetasoundFrontendClass* MetasoundClass = InDoc.Dependencies.FindByPredicate([&](const FMetasoundFrontendClass& InDependency) { return InDependency.ID == ID; });
+
+				if (nullptr == MetasoundClass)
+				{
+					MetasoundClass = InDoc.Subgraphs.FindByPredicate([&](const FMetasoundFrontendGraphClass& InSubgraph) { return InSubgraph.ID == ID; });
+				}
+
+				return MetasoundClass;
+			};
+			return GetMemberAccessPtr<FClassAccessPtr>(FindClassWithID);
+		}
+
+		FClassAccessPtr FDocumentAccessPtr::GetClassWithMetadata(const FMetasoundFrontendClassMetadata& InMetadata)
+		{
+			auto FindClassWithMetadata = [Metadata=InMetadata](FMetasoundFrontendDocument& InDoc)
+			{
+				auto IsClassWithMetadata = [=](const FMetasoundFrontendClass& InClass) 
+				{ 
+					using namespace MetasoundFrontendDocumentAccessPtrPrivate;
+
+					return IsMatchingMetasoundClass(Metadata, InClass.Metadata);
+				};
+
+				FMetasoundFrontendClass* MetasoundClass = InDoc.Dependencies.FindByPredicate(IsClassWithMetadata);
+				if (nullptr == MetasoundClass)
+				{
+					MetasoundClass = InDoc.Subgraphs.FindByPredicate(IsClassWithMetadata);
+				}
+				return MetasoundClass;
+			};
+
+			return GetMemberAccessPtr<FClassAccessPtr>(FindClassWithMetadata);
+		}
+
+		FClassAccessPtr FDocumentAccessPtr::GetClassWithInfo(const FNodeClassInfo& InInfo)
+		{
+			auto FindClassWithInfo = [Info=InInfo](FMetasoundFrontendDocument& InDoc)
+			{
+				auto IsClassWithInfo = [=](const FMetasoundFrontendClass& InClass) 
+				{ 
+					using namespace MetasoundFrontendDocumentAccessPtrPrivate;
+					return IsMatchingMetasoundClass(Info, InClass.Metadata);
+				};
+
+				FMetasoundFrontendClass* MetasoundClass = InDoc.Dependencies.FindByPredicate(IsClassWithInfo);
+				if (nullptr == MetasoundClass)
+				{
+					MetasoundClass = InDoc.Subgraphs.FindByPredicate(IsClassWithInfo);
+				}
+				return MetasoundClass;
+			};
+
+			return GetMemberAccessPtr<FClassAccessPtr>(FindClassWithInfo);
+		}
+
+		FConstGraphClassAccessPtr FDocumentAccessPtr::GetRootGraph() const
+		{
+			return GetMemberAccessPtr<FConstGraphClassAccessPtr>([](const FMetasoundFrontendDocument& InDoc) { return &InDoc.RootGraph; });
+		}
+
+		FConstGraphClassAccessPtr FDocumentAccessPtr::GetSubgraphWithID(const FGuid& InID) const
+		{
+			auto FindSubgraphWithID = [ID=InID](const FMetasoundFrontendDocument& InDoc)
+			{
+				return InDoc.Subgraphs.FindByPredicate([&](const FMetasoundFrontendGraphClass& InSubgraph) { return InSubgraph.ID == ID; });
+			};
+			return GetMemberAccessPtr<FConstGraphClassAccessPtr>(FindSubgraphWithID);
+		}
+
+		FConstClassAccessPtr FDocumentAccessPtr::GetDependencyWithID(const FGuid& InID) const
+		{
+			auto FindDependencyWithID = [ID=InID](const FMetasoundFrontendDocument& InDoc)
+			{
+				return InDoc.Dependencies.FindByPredicate([&](const FMetasoundFrontendClass& InDependency) { return InDependency.ID == ID; });
+			};
+			return GetMemberAccessPtr<FConstClassAccessPtr>(FindDependencyWithID);
+		}
+
+
+		FConstClassAccessPtr FDocumentAccessPtr::GetClassWithID(const FGuid& InID) const
+		{
+			auto FindClassWithID = [ID=InID](const FMetasoundFrontendDocument& InDoc) -> const FMetasoundFrontendClass*
+			{
+				const FMetasoundFrontendClass* MetasoundClass = InDoc.Dependencies.FindByPredicate([&](const FMetasoundFrontendClass& InDependency) { return InDependency.ID == ID; });
+
+				if (nullptr == MetasoundClass)
+				{
+					MetasoundClass = InDoc.Subgraphs.FindByPredicate([&](const FMetasoundFrontendGraphClass& InSubgraph) { return InSubgraph.ID == ID; });
+				}
+
+				return MetasoundClass;
+			};
+			return GetMemberAccessPtr<FConstClassAccessPtr>(FindClassWithID);
+		}
+
+		FConstClassAccessPtr FDocumentAccessPtr::GetClassWithMetadata(const FMetasoundFrontendClassMetadata& InMetadata) const
+		{
+			auto FindClassWithMetadata = [Metadata=InMetadata](const FMetasoundFrontendDocument& InDoc)
+			{
+				auto IsClassWithMetadata = [=](const FMetasoundFrontendClass& InClass) 
+				{ 
+					using namespace MetasoundFrontendDocumentAccessPtrPrivate;
+					return IsMatchingMetasoundClass(Metadata, InClass.Metadata);
+				};
+
+				const FMetasoundFrontendClass* MetasoundClass = InDoc.Dependencies.FindByPredicate(IsClassWithMetadata);
+				if (nullptr == MetasoundClass)
+				{
+					MetasoundClass = InDoc.Subgraphs.FindByPredicate(IsClassWithMetadata);
+				}
+				return MetasoundClass;
+			};
+
+			return GetMemberAccessPtr<FConstClassAccessPtr>(FindClassWithMetadata);
+		}
+
+		FConstClassAccessPtr FDocumentAccessPtr::GetClassWithInfo(const FNodeClassInfo& InInfo) const
+		{
+			auto FindClassWithInfo = [Info=InInfo](const FMetasoundFrontendDocument& InDoc)
+			{
+				auto IsClassWithInfo = [=](const FMetasoundFrontendClass& InClass) 
+				{ 
+					using namespace MetasoundFrontendDocumentAccessPtrPrivate;
+					return IsMatchingMetasoundClass(Info, InClass.Metadata);
+				};
+
+				const FMetasoundFrontendClass* MetasoundClass = InDoc.Dependencies.FindByPredicate(IsClassWithInfo);
+				if (nullptr == MetasoundClass)
+				{
+					MetasoundClass = InDoc.Subgraphs.FindByPredicate(IsClassWithInfo);
+				}
+				return MetasoundClass;
+			};
+
+			return GetMemberAccessPtr<FConstClassAccessPtr>(FindClassWithInfo);
+		}
+
+		FConstGraphClassAccessPtr FConstDocumentAccessPtr::GetRootGraph() const
+		{
+			return GetMemberAccessPtr<FConstGraphClassAccessPtr>([](const FMetasoundFrontendDocument& InDoc) { return &InDoc.RootGraph; });
+		}
+
+		FConstGraphClassAccessPtr FConstDocumentAccessPtr::GetSubgraphWithID(const FGuid& InID) const
+		{
+			auto FindSubgraphWithID = [ID=InID](const FMetasoundFrontendDocument& InDoc)
+			{
+				return InDoc.Subgraphs.FindByPredicate([&](const FMetasoundFrontendGraphClass& InSubgraph) { return InSubgraph.ID == ID; });
+			};
+			return GetMemberAccessPtr<FConstGraphClassAccessPtr>(FindSubgraphWithID);
+		}
+
+		FConstClassAccessPtr FConstDocumentAccessPtr::GetDependencyWithID(const FGuid& InID) const
+		{
+			auto FindDependencyWithID = [ID=InID](const FMetasoundFrontendDocument& InDoc)
+			{
+				return InDoc.Dependencies.FindByPredicate([&](const FMetasoundFrontendClass& InDependency) { return InDependency.ID == ID; });
+			};
+			return GetMemberAccessPtr<FConstClassAccessPtr>(FindDependencyWithID);
+		}
+
+
+		FConstClassAccessPtr FConstDocumentAccessPtr::GetClassWithID(const FGuid& InID) const
+		{
+			auto FindClassWithID = [ID=InID](const FMetasoundFrontendDocument& InDoc)
+			{
+				const FMetasoundFrontendClass* MetasoundClass = InDoc.Dependencies.FindByPredicate([&](const FMetasoundFrontendClass& InDependency) { return InDependency.ID == ID; });
+
+				if (nullptr == MetasoundClass)
+				{
+					MetasoundClass = InDoc.Subgraphs.FindByPredicate([&](const FMetasoundFrontendGraphClass& InSubgraph) { return InSubgraph.ID == ID; });
+				}
+
+				return MetasoundClass;
+			};
+			return GetMemberAccessPtr<FConstClassAccessPtr>(FindClassWithID);
+		}
+
+		FConstClassAccessPtr FConstDocumentAccessPtr::GetClassWithMetadata(const FMetasoundFrontendClassMetadata& InMetadata) const
+		{
+			auto FindClassWithMetadata = [Metadata=InMetadata](const FMetasoundFrontendDocument& InDoc)
+			{
+				auto IsClassWithMetadata = [=](const FMetasoundFrontendClass& InClass) 
+				{ 
+					using namespace MetasoundFrontendDocumentAccessPtrPrivate;
+					return IsMatchingMetasoundClass(Metadata, InClass.Metadata);
+				};
+
+				const FMetasoundFrontendClass* MetasoundClass = InDoc.Dependencies.FindByPredicate(IsClassWithMetadata);
+				if (nullptr == MetasoundClass)
+				{
+					MetasoundClass = InDoc.Subgraphs.FindByPredicate(IsClassWithMetadata);
+				}
+				return MetasoundClass;
+			};
+
+			return GetMemberAccessPtr<FConstClassAccessPtr>(FindClassWithMetadata);
+		}
+
+		FConstClassAccessPtr FConstDocumentAccessPtr::GetClassWithInfo(const FNodeClassInfo& InInfo) const
+		{
+			auto FindClassWithInfo = [Info=InInfo](const FMetasoundFrontendDocument& InDoc)
+			{
+				auto IsClassWithInfo = [=](const FMetasoundFrontendClass& InClass) 
+				{ 
+					using namespace MetasoundFrontendDocumentAccessPtrPrivate;
+					return IsMatchingMetasoundClass(Info, InClass.Metadata);
+				};
+
+				const FMetasoundFrontendClass* MetasoundClass = InDoc.Dependencies.FindByPredicate(IsClassWithInfo);
+				if (nullptr == MetasoundClass)
+				{
+					MetasoundClass = InDoc.Subgraphs.FindByPredicate(IsClassWithInfo);
+				}
+				return MetasoundClass;
+			};
+
+			return GetMemberAccessPtr<FConstClassAccessPtr>(FindClassWithInfo);
+		}
+	}
+}
diff --git a/Engine/Plugins/Runtime/Metasound/Source/MetasoundFrontend/Private/MetasoundFrontendStandardController.cpp b/Engine/Plugins/Runtime/Metasound/Source/MetasoundFrontend/Private/MetasoundFrontendStandardController.cpp
--- a/Engine/Plugins/Runtime/Metasound/Source/MetasoundFrontend/Private/MetasoundFrontendStandardController.cpp
+++ b/Engine/Plugins/Runtime/Metasound/Source/MetasoundFrontend/Private/MetasoundFrontendStandardController.cpp
@@ -9,6 +9,7 @@
 #include "MetasoundAccessPtr.h"
 #include "MetasoundFrontendController.h"
 #include "MetasoundFrontendDocument.h"
+#include "MetasoundFrontendDocumentAccessPtr.h"
 #include "MetasoundFrontendGraph.h"
 #include "MetasoundFrontendRegistries.h"
 #include "MetasoundJsonBackend.h"
@@ -34,21 +35,6 @@ namespace Metasound
 				return InvalidValue;
 			}
 
-			// Utility function to get a value for an existing description pointer
-			// or return an invalid value const ref if the TAccessPtr is invalid.
-			template<typename ValueType, typename PtrType>
-			const ValueType& GetValueOrReturnInvalid(const TAccessPtr<PtrType>& InPointer, ValueType (PtrType::*InMember))
-			{
-				static const ValueType InvalidValue = GetInvalidValue<ValueType>();
-
-				if (InPointer.IsValid())
-				{
-					return (*InPointer).*InMember;
-				}
-
-				return InvalidValue;
-			}
-
 			// Invalid value specialization for int32
 			template<>
 			int32 GetInvalidValue<int32>() { return INDEX_NONE; }
@@ -67,307 +53,6 @@ namespace Metasound
 				static const ValueType Value = GetInvalidValue<ValueType>();
 				return Value;
 			}
-
-			template<typename Type>
-			TAccessPtr<Type> MakeAutoAccessPtr(Type& InObj)
-			{
-				return MakeAccessPtr(InObj.AccessPoint, InObj);
-				//return MakeAccessPtr(InObj);
-			}
-
-
-			// TODO: All these access ptr helper routines should be merged into
-			// the TAccessPtr<> interface. The TAccessPtr<> class should be specialized 
-			// for the types used here, with the specialized TAccessPtr<> interface
-			// adding support for getting specific member access ptrs and doing
-			// specific queries. 
-			struct FGetSubgraphFromDocumentByID
-			{
-				FGetSubgraphFromDocumentByID(FGuid InID)
-				: ID(InID)
-				{
-				}
-
-				FMetasoundFrontendGraphClass* operator() (FMetasoundFrontendDocument& InDoc) const
-				{
-					return InDoc.Subgraphs.FindByPredicate([&](const FMetasoundFrontendGraphClass& GraphClass) { return GraphClass.ID == ID; });
-				}
-
-				const FMetasoundFrontendGraphClass* operator() (const FMetasoundFrontendDocument& InDoc) const
-				{
-					return InDoc.Subgraphs.FindByPredicate([&](const FMetasoundFrontendGraphClass& GraphClass) { return GraphClass.ID == ID; });
-				}
-
-			private:
-				FGuid ID;
-			};
-
-			struct FGetClassInputFromClassWithName
-			{
-				FGetClassInputFromClassWithName(const FString& InName)
-				: Name(InName)
-				{
-				}
-
-				FMetasoundFrontendClassInput* operator() (FMetasoundFrontendClass& InClass)
-				{
-					auto IsClassInputWithName = [&](const FMetasoundFrontendClassInput& ClassInput) { return ClassInput.Name == Name; };
-					return InClass.Interface.Inputs.FindByPredicate(IsClassInputWithName);
-				}
-
-				const FMetasoundFrontendClassInput* operator() (const FMetasoundFrontendClass& InClass)
-				{
-					auto IsClassInputWithName = [&](const FMetasoundFrontendClassInput& ClassInput) { return ClassInput.Name == Name; };
-					return InClass.Interface.Inputs.FindByPredicate(IsClassInputWithName);
-				}
-				private:
-				FString Name;
-			};
-
-			struct FGetClassOutputFromClassWithName
-			{
-				FGetClassOutputFromClassWithName(const FString& InName)
-				: Name(InName)
-				{
-				}
-
-				FMetasoundFrontendClassOutput* operator() (FMetasoundFrontendClass& InClass)
-				{
-					auto IsClassOutputWithName = [&](const FMetasoundFrontendClassOutput& ClassOutput) { return ClassOutput.Name == Name; };
-					return InClass.Interface.Outputs.FindByPredicate(IsClassOutputWithName);
-				}
-
-				const FMetasoundFrontendClassOutput* operator() (const FMetasoundFrontendClass& InClass)
-				{
-					auto IsClassOutputWithName = [&](const FMetasoundFrontendClassOutput& ClassOutput) { return ClassOutput.Name == Name; };
-					return InClass.Interface.Outputs.FindByPredicate(IsClassOutputWithName);
-				}
-				private:
-				FString Name;
-			};
-
-			auto FGetNodeInputFromNodeWithName(const FString& InName)
-			{
-				return [=](FMetasoundFrontendNode& InNode) -> FMetasoundFrontendVertex*
-				{
-					auto IsVertexWithName = [=](const FMetasoundFrontendVertex& Vertex) { return Vertex.Name == InName; };
-					return InNode.Interface.Inputs.FindByPredicate(IsVertexWithName);
-				};
-			}
-
-			auto FGetConstNodeInputFromNodeWithName(const FString& InName)
-			{
-				return [=](const FMetasoundFrontendNode& InNode) -> const FMetasoundFrontendVertex*
-				{
-					auto IsVertexWithName = [=](const FMetasoundFrontendVertex& Vertex) { return Vertex.Name == InName; };
-					return InNode.Interface.Inputs.FindByPredicate(IsVertexWithName);
-				};
-			}
-
-			auto FGetNodeOutputFromNodeWithName(const FString& InName)
-			{
-				return [=](FMetasoundFrontendNode& InNode) -> FMetasoundFrontendVertex*
-				{
-					auto IsVertexWithName = [=](const FMetasoundFrontendVertex& Vertex) { return Vertex.Name == InName; };
-					return InNode.Interface.Outputs.FindByPredicate(IsVertexWithName);
-				};
-			}
-
-			auto FGetConstNodeOutputFromNodeWithName(const FString& InName)
-			{
-				return [=](const FMetasoundFrontendNode& InNode) -> const FMetasoundFrontendVertex*
-				{
-					auto IsVertexWithName = [=](const FMetasoundFrontendVertex& Vertex) { return Vertex.Name == InName; };
-					return InNode.Interface.Outputs.FindByPredicate(IsVertexWithName);
-				};
-			}
-
-			auto FGetNodeInputFromNodeWithVertexID(FGuid InVertexID)
-			{
-				return [=](FMetasoundFrontendNode& InNode) -> FMetasoundFrontendVertex*
-				{
-					auto IsVertexWithVertexID = [=](const FMetasoundFrontendVertex& Vertex) { return Vertex.VertexID == InVertexID; };
-					return InNode.Interface.Inputs.FindByPredicate(IsVertexWithVertexID);
-				};
-			}
-
-			auto FGetConstNodeInputFromNodeWithVertexID(FGuid InVertexID)
-			{
-				return [=](const FMetasoundFrontendNode& InNode) -> const FMetasoundFrontendVertex*
-				{
-					auto IsVertexWithVertexID = [=](const FMetasoundFrontendVertex& Vertex) { return Vertex.VertexID == InVertexID; };
-					return InNode.Interface.Inputs.FindByPredicate(IsVertexWithVertexID);
-				};
-			}
-
-			auto FGetNodeOutputFromNodeWithVertexID(FGuid InVertexID)
-			{
-				return [=](FMetasoundFrontendNode& InNode) -> FMetasoundFrontendVertex*
-				{
-					auto IsVertexWithVertexID = [=](const FMetasoundFrontendVertex& Vertex) { return Vertex.VertexID == InVertexID; };
-					return InNode.Interface.Outputs.FindByPredicate(IsVertexWithVertexID);
-				};
-			}
-
-			auto FGetConstNodeOutputFromNodeWithVertexID(FGuid InVertexID)
-			{
-				return [=](const FMetasoundFrontendNode& InNode) -> const FMetasoundFrontendVertex*
-				{
-					auto IsVertexWithVertexID = [=](const FMetasoundFrontendVertex& Vertex) { return Vertex.VertexID == InVertexID; };
-					return InNode.Interface.Outputs.FindByPredicate(IsVertexWithVertexID);
-				};
-			}
-
-			auto FGetConstClassFromDocumentDependenciesWithID(FGuid InClassID)
-			{
-				return [=](const FMetasoundFrontendDocument& InDocument) -> const FMetasoundFrontendClass*
-				{
-					auto IsClassWithID = [=](const FMetasoundFrontendClass& InClass) { return InClass.ID == InClassID; };
-					return InDocument.Dependencies.FindByPredicate(IsClassWithID);
-				};
-			}
-
-			auto FGetConstGraphClassFromDocumentWithID(FGuid InClassID)
-			{
-				return [=](const FMetasoundFrontendDocument& InDocument) -> const FMetasoundFrontendGraphClass*
-				{
-					auto IsClassWithID = [=](const FMetasoundFrontendGraphClass& InClass) { return InClass.ID == InClassID; };
-					return InDocument.Subgraphs.FindByPredicate(IsClassWithID);
-				};
-			}
-
-			auto FGetClassFromDocumentWithClassInfo(const FNodeClassInfo& InInfo)
-			{
-				return [=](FMetasoundFrontendDocument& InDocument) -> FMetasoundFrontendClass*
-				{
-					auto IsClassWithClassInfo = [=](const FMetasoundFrontendClass& InClass)
-					{
-						return FDocumentController::IsMatchingMetasoundClass(InInfo, InClass.Metadata);
-					};
-
-					FMetasoundFrontendClass* MetasoundClass = InDocument.Dependencies.FindByPredicate(IsClassWithClassInfo);
-					if (nullptr == MetasoundClass)
-					{
-						MetasoundClass = InDocument.Subgraphs.FindByPredicate(IsClassWithClassInfo);
-					}
-					return MetasoundClass;
-				};
-			}
-
-			auto FGetConstClassFromDocumentWithClassInfo(const FNodeClassInfo& InInfo)
-			{
-				return [=](const FMetasoundFrontendDocument& InDocument) -> const FMetasoundFrontendClass*
-				{
-					auto IsClassWithClassInfo = [=](const FMetasoundFrontendClass& InClass) 
-					{ 
-						return FDocumentController::IsMatchingMetasoundClass(InInfo, InClass.Metadata);
-					};
-
-					const FMetasoundFrontendClass* MetasoundClass = InDocument.Dependencies.FindByPredicate(IsClassWithClassInfo);
-					if (nullptr == MetasoundClass)
-					{
-						MetasoundClass = InDocument.Subgraphs.FindByPredicate(IsClassWithClassInfo);
-					}
-					return MetasoundClass;
-				};
-			}
-
-			auto FGetClassFromDocumentWithMetadata(const FMetasoundFrontendClassMetadata& InMetadata)
-			{
-				return [=](FMetasoundFrontendDocument& InDocument) -> FMetasoundFrontendClass*
-				{
-					auto IsClassWithMetadata = [=](const FMetasoundFrontendClass& InClass)
-					{
-						return FDocumentController::IsMatchingMetasoundClass(InMetadata, InClass.Metadata);
-					};
-
-					FMetasoundFrontendClass* MetasoundClass = InDocument.Dependencies.FindByPredicate(IsClassWithMetadata);
-					if (nullptr == MetasoundClass)
-					{
-						MetasoundClass = InDocument.Subgraphs.FindByPredicate(IsClassWithMetadata);
-					}
-					return MetasoundClass;
-				};
-			}
-
-			auto FGetConstClassFromDocumentWithMetadata(const FMetasoundFrontendClassMetadata& InMetadata)
-			{
-				return [=](const FMetasoundFrontendDocument& InDocument) -> const FMetasoundFrontendClass*
-				{
-					auto IsClassWithMetadata = [=](const FMetasoundFrontendClass& InClass)
-					{
-						return FDocumentController::IsMatchingMetasoundClass(InMetadata, InClass.Metadata);
-					};
-
-					const FMetasoundFrontendClass* MetasoundClass = InDocument.Dependencies.FindByPredicate(IsClassWithMetadata);
-					if (nullptr == MetasoundClass)
-					{
-						MetasoundClass = InDocument.Subgraphs.FindByPredicate(IsClassWithMetadata);
-					}
-					return MetasoundClass;
-				};
-			}
-
-			auto FGetNodeFromGraphClassByNodeID(FGuid InNodeID)
-			{
-				return [=](FMetasoundFrontendGraphClass& InGraphClass) -> FMetasoundFrontendNode*
-				{
-					auto IsNodeWithID = [&](const FMetasoundFrontendNode& InNode) { return InNode.ID == InNodeID; };
-					return InGraphClass.Graph.Nodes.FindByPredicate(IsNodeWithID);
-				};
-			}
-
-			auto FGetConstGraphFromGraphClass()
-			{
-				return [](const FMetasoundFrontendGraphClass& InGraphClass) -> const FMetasoundFrontendGraph* { return &InGraphClass.Graph; };
-			}
-			auto FGetGraphFromGraphClass()
-			{
-				return [](FMetasoundFrontendGraphClass& InGraphClass) -> FMetasoundFrontendGraph* { return &InGraphClass.Graph; };
-			}
-
-			auto FGetConstNodeFromGraphClassByNodeID(FGuid InNodeID)
-			{
-				return [=](const FMetasoundFrontendGraphClass& InGraphClass) -> const FMetasoundFrontendNode*
-				{
-					auto IsNodeWithID = [&](const FMetasoundFrontendNode& InNode) { return InNode.ID == InNodeID; };
-					return InGraphClass.Graph.Nodes.FindByPredicate(IsNodeWithID);
-				};
-			}
-
-			auto FGetConstClassInputFromGraphClassWithNodeID(FGuid InNodeID)
-			{
-				return [=](const FMetasoundFrontendGraphClass& InGraphClass) -> const FMetasoundFrontendClassInput*
-				{
-					return InGraphClass.Interface.Inputs.FindByPredicate([&](const FMetasoundFrontendClassInput& Desc) { return Desc.NodeID == InNodeID; });
-				};
-			}
-
-			auto FGetClassInputFromGraphClassWithNodeID(FGuid InNodeID)
-			{
-				return [=](FMetasoundFrontendGraphClass& InGraphClass) -> FMetasoundFrontendClassInput*
-				{
-					return InGraphClass.Interface.Inputs.FindByPredicate([&](const FMetasoundFrontendClassInput& Desc) { return Desc.NodeID == InNodeID; });
-				};
-			}
-
-			auto FGetConstClassOutputFromGraphClassWithNodeID(FGuid InNodeID)
-			{
-				return [=](const FMetasoundFrontendGraphClass& InGraphClass) -> const FMetasoundFrontendClassOutput*
-				{
-					return InGraphClass.Interface.Outputs.FindByPredicate([&](const FMetasoundFrontendClassOutput& Desc) { return Desc.NodeID == InNodeID; });
-				};
-			}
-
-			auto FGetClassOutputFromGraphClassWithNodeID(FGuid InNodeID)
-			{
-				return [=](FMetasoundFrontendGraphClass& InGraphClass) -> FMetasoundFrontendClassOutput*
-				{
-					return InGraphClass.Interface.Outputs.FindByPredicate([&](const FMetasoundFrontendClassOutput& Desc) { return Desc.NodeID == InNodeID; });
-				};
-			}
-
-
 		}
 
 
@@ -1672,16 +1357,14 @@ namespace Metasound
 
 		TArray<FBaseNodeController::FInputControllerParams> FBaseNodeController::GetInputControllerParams() const
 		{
-			using namespace FrontendControllerIntrinsics;
-
 			TArray<FBaseNodeController::FInputControllerParams> Inputs;
 
 			if (NodePtr.IsValid())
 			{
 				for (const FMetasoundFrontendVertex& NodeInputVertex : NodePtr->Interface.Inputs)
 				{
-					FConstVertexAccessPtr NodeVertexPtr = NodePtr.GetMemberAccessPtr<const FMetasoundFrontendVertex>(FGetNodeInputFromNodeWithName(NodeInputVertex.Name));
-					FConstClassInputAccessPtr ClassInputPtr = ClassPtr.GetMemberAccessPtr<const FMetasoundFrontendClassInput>(FGetClassInputFromClassWithName(NodeInputVertex.Name));
+					FConstVertexAccessPtr NodeVertexPtr = NodePtr.GetInputWithName(NodeInputVertex.Name);
+					FConstClassInputAccessPtr ClassInputPtr = ClassPtr.GetInputWithName(NodeInputVertex.Name);
 
 					Inputs.Add({NodeInputVertex.VertexID, NodeVertexPtr, ClassInputPtr});
 				}
@@ -1692,16 +1375,14 @@ namespace Metasound
 
 		TArray<FBaseNodeController::FOutputControllerParams> FBaseNodeController::GetOutputControllerParams() const
 		{
-			using namespace FrontendControllerIntrinsics;
-
 			TArray<FBaseNodeController::FOutputControllerParams> Outputs;
 
 			for (const FMetasoundFrontendVertex& NodeOutputVertex : NodePtr->Interface.Outputs)
 			{
 				const FString& VertexName = NodeOutputVertex.Name;
 
-				FConstVertexAccessPtr NodeVertexPtr = NodePtr.GetMemberAccessPtr<const FMetasoundFrontendVertex>(FGetConstNodeOutputFromNodeWithName(VertexName));
-				FConstClassOutputAccessPtr ClassOutputPtr = ClassPtr.GetMemberAccessPtr<const FMetasoundFrontendClassOutput>(FGetClassOutputFromClassWithName(VertexName));
+				FConstVertexAccessPtr NodeVertexPtr = NodePtr.GetOutputWithName(VertexName);
+				FConstClassOutputAccessPtr ClassOutputPtr = ClassPtr.GetOutputWithName(VertexName);
 
 				Outputs.Add({NodeOutputVertex.VertexID, NodeVertexPtr, ClassOutputPtr});
 			}
@@ -1711,15 +1392,13 @@ namespace Metasound
 
 		TArray<FBaseNodeController::FInputControllerParams> FBaseNodeController::GetInputControllerParamsWithVertexName(const FString& InName) const
 		{
-			using namespace FrontendControllerIntrinsics;
-
 			TArray<FBaseNodeController::FInputControllerParams> Inputs;
 
-			FConstVertexAccessPtr NodeVertexPtr = NodePtr.GetMemberAccessPtr<const FMetasoundFrontendVertex>(FGetConstNodeInputFromNodeWithName(InName));
+			FConstVertexAccessPtr NodeVertexPtr = NodePtr.GetInputWithName(InName);
 
 			if (const FMetasoundFrontendVertex* Vertex = NodeVertexPtr.Get())
 			{
-				FConstClassInputAccessPtr ClassInputPtr = ClassPtr.GetMemberAccessPtr<const FMetasoundFrontendClassInput>(FGetClassInputFromClassWithName(InName));
+				FConstClassInputAccessPtr ClassInputPtr = ClassPtr.GetInputWithName(InName);
 
 				Inputs.Add({Vertex->VertexID, NodeVertexPtr, ClassInputPtr});
 			}
@@ -1729,15 +1408,13 @@ namespace Metasound
 
 		TArray<FBaseNodeController::FOutputControllerParams> FBaseNodeController::GetOutputControllerParamsWithVertexName(const FString& InName) const
 		{
-			using namespace FrontendControllerIntrinsics;
-
 			TArray<FBaseNodeController::FOutputControllerParams> Outputs;
 
-			FConstVertexAccessPtr NodeVertexPtr = NodePtr.GetMemberAccessPtr<const FMetasoundFrontendVertex>(FGetConstNodeOutputFromNodeWithName(InName));
+			FConstVertexAccessPtr NodeVertexPtr = NodePtr.GetOutputWithName(InName);
 
 			if (const FMetasoundFrontendVertex* Vertex = NodeVertexPtr.Get())
 			{
-				FConstClassOutputAccessPtr ClassOutputPtr = ClassPtr.GetMemberAccessPtr<const FMetasoundFrontendClassOutput>(FGetClassOutputFromClassWithName(InName));
+				FConstClassOutputAccessPtr ClassOutputPtr = ClassPtr.GetOutputWithName(InName);
 
 				Outputs.Add({Vertex->VertexID, NodeVertexPtr, ClassOutputPtr});
 			}
@@ -1747,13 +1424,11 @@ namespace Metasound
 
 		bool FBaseNodeController::FindInputControllerParamsWithID(FGuid InVertexID, FInputControllerParams& OutParams) const
 		{
-			using namespace FrontendControllerIntrinsics;
-
-			FConstVertexAccessPtr NodeVertexPtr = NodePtr.GetMemberAccessPtr<const FMetasoundFrontendVertex>(FGetConstNodeInputFromNodeWithVertexID(InVertexID));
+			FConstVertexAccessPtr NodeVertexPtr = NodePtr.GetInputWithVertexID(InVertexID);
 
 			if (const FMetasoundFrontendVertex* Vertex = NodeVertexPtr.Get())
 			{
-				FConstClassInputAccessPtr ClassInputPtr = ClassPtr.GetMemberAccessPtr<const FMetasoundFrontendClassInput>(FGetClassInputFromClassWithName(Vertex->Name));
+				FConstClassInputAccessPtr ClassInputPtr = ClassPtr.GetInputWithName(Vertex->Name);
 
 				OutParams = FInputControllerParams{InVertexID, NodeVertexPtr, ClassInputPtr};
 				return true;
@@ -1764,13 +1439,11 @@ namespace Metasound
 
 		bool FBaseNodeController::FindOutputControllerParamsWithID(FGuid InVertexID, FOutputControllerParams& OutParams) const
 		{
-			using namespace FrontendControllerIntrinsics;
-
-			FConstVertexAccessPtr NodeVertexPtr = NodePtr.GetMemberAccessPtr<const FMetasoundFrontendVertex>(FGetConstNodeOutputFromNodeWithVertexID(InVertexID));
+			FConstVertexAccessPtr NodeVertexPtr = NodePtr.GetOutputWithVertexID(InVertexID);
 
 			if (const FMetasoundFrontendVertex* Vertex = NodeVertexPtr.Get())
 			{
-				FConstClassOutputAccessPtr ClassOutputPtr = ClassPtr.GetMemberAccessPtr<const FMetasoundFrontendClassOutput>(FGetClassOutputFromClassWithName(Vertex->Name));
+				FConstClassOutputAccessPtr ClassOutputPtr = ClassPtr.GetOutputWithName(Vertex->Name);
 
 				OutParams = FOutputControllerParams{InVertexID, NodeVertexPtr, ClassOutputPtr};
 				return true;
@@ -2008,18 +1681,19 @@ namespace Metasound
 
 		void FOutputNodeController::SetDescription(const FText& InDescription)
 		{
-			if (OwningGraphClassOutputPtr.IsValid())
+			// TODO: can we remove the const cast by constructing output nodes with a non-const access to class outputs?
+			if (FMetasoundFrontendClassOutput* ClassOutput = ConstCastAccessPtr<FClassOutputAccessPtr>(OwningGraphClassOutputPtr).Get())
 			{
-				FClassOutputAccessPtr ClassOutput = ConstCastAccessPtr<FMetasoundFrontendClassOutput>(OwningGraphClassOutputPtr);
 				ClassOutput->Metadata.Description = InDescription;
 			}
 		}
 
 		void FOutputNodeController::SetDisplayName(const FText& InDisplayName)
 		{
-			if (OwningGraphClassOutputPtr.IsValid())
+			// TODO: can we remove the const cast by constructing output nodes with a non-const access to class outputs?
+			if (FMetasoundFrontendClassOutput* ClassOutput = ConstCastAccessPtr<FClassOutputAccessPtr>(OwningGraphClassOutputPtr).Get())
 			{
-				ConstCastAccessPtr<FMetasoundFrontendClassOutput>(OwningGraphClassOutputPtr)->Metadata.DisplayName = InDisplayName;
+				ClassOutput->Metadata.DisplayName = InDisplayName;
 			}
 		}
 
@@ -2227,18 +1901,19 @@ namespace Metasound
 
 		void FInputNodeController::SetDescription(const FText& InDescription)
 		{
-			if (OwningGraphClassInputPtr.IsValid())
+			// TODO: can we remove these const casts by constructing FINputNodeController with non-const access to the class input?
+			if (FMetasoundFrontendClassInput* ClassInput = ConstCastAccessPtr<FClassInputAccessPtr>(OwningGraphClassInputPtr).Get())
 			{
-				FClassInputAccessPtr ClassInput = ConstCastAccessPtr<FMetasoundFrontendClassInput>(OwningGraphClassInputPtr);
 				ClassInput->Metadata.Description = InDescription;
 			}
 		}
 
 		void FInputNodeController::SetDisplayName(const FText& InDisplayName)
 		{
-			if (OwningGraphClassInputPtr.IsValid())
+			// TODO: can we remove these const casts by constructing FINputNodeController with non-const access to the class input?
+			if (FMetasoundFrontendClassInput* ClassInput = ConstCastAccessPtr<FClassInputAccessPtr>(OwningGraphClassInputPtr).Get())
 			{
-				ConstCastAccessPtr<FMetasoundFrontendClassInput>(OwningGraphClassInputPtr)->Metadata.DisplayName = InDisplayName;
+				ClassInput->Metadata.DisplayName = InDisplayName;
 			}
 		}
 
@@ -2357,14 +2032,12 @@ namespace Metasound
 
 		FConstClassInputAccessPtr FGraphController::FindClassInputWithName(const FString& InName) const
 		{
-			using namespace FrontendControllerIntrinsics;
-			return GraphClassPtr.GetMemberAccessPtr<FMetasoundFrontendClassInput>(FGetClassInputFromClassWithName(InName));
+			return GraphClassPtr.GetInputWithName(InName);
 		}
 
 		FConstClassOutputAccessPtr FGraphController::FindClassOutputWithName(const FString& InName) const
 		{
-			using namespace FrontendControllerIntrinsics;
-			return GraphClassPtr.GetMemberAccessPtr<FMetasoundFrontendClassOutput>(FGetClassOutputFromClassWithName(InName));
+			return GraphClassPtr.GetOutputWithName(InName);
 		}
 
 		FGuid FGraphController::GetVertexIDForInputVertex(const FString& InInputName) const
@@ -2433,8 +2106,6 @@ namespace Metasound
 
 		void FGraphController::IterateNodes(TUniqueFunction<void(FNodeHandle)> InFunction, EMetasoundFrontendClassType InClassType)
 		{
-			using namespace FrontendControllerIntrinsics;
-
 			if (GraphClassPtr.IsValid())
 			{
 				for (const FMetasoundFrontendNode& Node : GraphClassPtr->Graph.Nodes)
@@ -2444,7 +2115,7 @@ namespace Metasound
 					{
 						if (InClassType == EMetasoundFrontendClassType::Invalid || NodeClass->Metadata.Type == InClassType)
 						{
-							FNodeAccessPtr NodePtr = GraphClassPtr.GetMemberAccessPtr<FMetasoundFrontendNode>(FGetNodeFromGraphClassByNodeID(Node.ID));
+							FNodeAccessPtr NodePtr = GraphClassPtr.GetNodeWithNodeID(Node.ID);
 							FNodeHandle NodeHandle = GetNodeHandle(FGraphController::FNodeAndClass{ NodePtr, NodeClassPtr });
 							InFunction(NodeHandle);
 						}
@@ -2459,8 +2130,6 @@ namespace Metasound
 
 		void FGraphController::IterateConstNodes(TUniqueFunction<void(FConstNodeHandle)> InFunction, EMetasoundFrontendClassType InClassType) const
 		{
-			using namespace FrontendControllerIntrinsics;
-
 			if (GraphClassPtr.IsValid())
 			{
 				for (const FMetasoundFrontendNode& Node : GraphClassPtr->Graph.Nodes)
@@ -2470,7 +2139,7 @@ namespace Metasound
 					{
 						if (InClassType == EMetasoundFrontendClassType::Invalid || NodeClass->Metadata.Type == InClassType)
 						{
-							FConstNodeAccessPtr NodePtr = GraphClassPtr.GetMemberAccessPtr<const FMetasoundFrontendNode>(FGetConstNodeFromGraphClassByNodeID(Node.ID));
+							FConstNodeAccessPtr NodePtr = GraphClassPtr.GetNodeWithNodeID(Node.ID);
 							FConstNodeHandle NodeHandle = GetNodeHandle(FGraphController::FConstNodeAndClass{ NodePtr, NodeClassPtr });
 							InFunction(NodeHandle);
 						}
@@ -2565,8 +2234,6 @@ namespace Metasound
 
 		FNodeHandle FGraphController::AddInputVertex(const FMetasoundFrontendClassInput& InClassInput)
 		{
-			using namespace FrontendControllerIntrinsics;
-
 			if (FMetasoundFrontendGraphClass* GraphClass = GraphClassPtr.Get())
 			{
 				auto IsInputWithSameName = [&](const FMetasoundFrontendClassInput& ExistingDesc) { return ExistingDesc.Name == InClassInput.Name; };
@@ -2604,7 +2271,7 @@ namespace Metasound
 							FMetasoundFrontendClassInput& NewInput = GraphClassPtr->Interface.Inputs.Add_GetRef(InClassInput);
 							NewInput.NodeID = Node.ID;
 
-							FNodeAccessPtr NodePtr = GraphClassPtr.GetMemberAccessPtr<FMetasoundFrontendNode>(FGetNodeFromGraphClassByNodeID(Node.ID));
+							FNodeAccessPtr NodePtr = GraphClassPtr.GetNodeWithNodeID(Node.ID);
 							return GetNodeHandle(FGraphController::FNodeAndClass{NodePtr, InputClass});
 						}
 					}
@@ -2641,8 +2308,6 @@ namespace Metasound
 
 		FNodeHandle FGraphController::AddOutputVertex(const FMetasoundFrontendClassOutput& InClassOutput)
 		{
-			using namespace FrontendControllerIntrinsics;
-
 			if (FMetasoundFrontendGraphClass* GraphClass = GraphClassPtr.Get())
 			{
 				auto IsOutputWithSameName = [&](const FMetasoundFrontendClassOutput& ExistingDesc) { return ExistingDesc.Name == InClassOutput.Name; };
@@ -2681,7 +2346,7 @@ namespace Metasound
 							FMetasoundFrontendClassOutput& NewOutput = GraphClassPtr->Interface.Outputs.Add_GetRef(InClassOutput);
 							NewOutput.NodeID = Node.ID;
 
-							FNodeAccessPtr NodePtr = GraphClassPtr.GetMemberAccessPtr<FMetasoundFrontendNode>(FGetNodeFromGraphClassByNodeID(Node.ID));
+							FNodeAccessPtr NodePtr = GraphClassPtr.GetNodeWithNodeID(Node.ID);
 							return GetNodeHandle(FGraphController::FNodeAndClass{NodePtr, OutputClass});
 						}
 					}
@@ -2951,7 +2616,6 @@ namespace Metasound
 					{
 						UE_LOG(LogMetaSound, Error, TEXT("Cannot add new subgraph. Metasound class already exists with matching metadata Name: \"%s\", Version %d.%d"), *(ExistingDependency->Metadata.ClassName.GetFullName().ToString()), ExistingDependency->Metadata.Version.Major, ExistingDependency->Metadata.Version.Minor);
 					}
-					//else if (const FMetasoundFrontendClass* DependencyDescription = OwningDocument->FindOrAddClass(InInfo))
 					else if (FConstClassAccessPtr DependencyDescription = OwningDocument->FindOrAddClass(InInfo))
 					{
 						return AddNode(DependencyDescription);
@@ -3005,8 +2669,6 @@ namespace Metasound
 
 		FNodeHandle FGraphController::AddNode(FConstClassAccessPtr InExistingDependency)
 		{
-			using namespace FrontendControllerIntrinsics;
-
 			if (FMetasoundFrontendGraphClass* GraphClass = GraphClassPtr.Get())
 			{
 				if (const FMetasoundFrontendClass* NodeClass = InExistingDependency.Get())
@@ -3015,7 +2677,7 @@ namespace Metasound
 
 					Node.ID = NewNodeID();
 
-					FNodeAccessPtr NodePtr = GraphClassPtr.GetMemberAccessPtr<FMetasoundFrontendNode>(FGetNodeFromGraphClassByNodeID(Node.ID));
+					FNodeAccessPtr NodePtr = GraphClassPtr.GetNodeWithNodeID(Node.ID);
 
 					return GetNodeHandle(FGraphController::FNodeAndClass{NodePtr, InExistingDependency});
 				}
@@ -3107,14 +2769,13 @@ namespace Metasound
 
 		TArray<FGraphController::FNodeAndClass> FGraphController::GetNodesAndClasses()
 		{
-			using namespace FrontendControllerIntrinsics;
 			TArray<FNodeAndClass> NodesAndClasses;
 
 			if (GraphClassPtr.IsValid())
 			{
 				for (FMetasoundFrontendNode& Node : GraphClassPtr->Graph.Nodes)
 				{
-					FNodeAccessPtr NodePtr = GraphClassPtr.GetMemberAccessPtr<FMetasoundFrontendNode>(FGetNodeFromGraphClassByNodeID(Node.ID));
+					FNodeAccessPtr NodePtr = GraphClassPtr.GetNodeWithNodeID(Node.ID);
 					FConstClassAccessPtr NodeClassPtr = OwningDocument->FindClassWithID(Node.ClassID);
 
 					if (NodeClassPtr && NodePtr)
@@ -3133,14 +2794,13 @@ namespace Metasound
 
 		TArray<FGraphController::FConstNodeAndClass> FGraphController::GetNodesAndClasses() const
 		{
-			using namespace FrontendControllerIntrinsics;
 			TArray<FConstNodeAndClass> NodesAndClasses;
 
 			if (GraphClassPtr.IsValid())
 			{
 				for (const FMetasoundFrontendNode& Node : GraphClassPtr->Graph.Nodes)
 				{
-					FConstNodeAccessPtr NodePtr = GraphClassPtr.GetMemberAccessPtr<const FMetasoundFrontendNode>(FGetConstNodeFromGraphClassByNodeID(Node.ID));
+					FConstNodeAccessPtr NodePtr = GraphClassPtr.GetNodeWithNodeID(Node.ID);
 					FConstClassAccessPtr NodeClassPtr = OwningDocument->FindClassWithID(Node.ClassID);
 
 					if (NodeClassPtr && NodePtr)
@@ -3159,7 +2819,6 @@ namespace Metasound
 
 		TArray<FGraphController::FNodeAndClass> FGraphController::GetNodesAndClassesByPredicate(TFunctionRef<bool (const FMetasoundFrontendClass&, const FMetasoundFrontendNode&)> InPredicate)
 		{
-			using namespace FrontendControllerIntrinsics;
 
 			TArray<FNodeAndClass> NodesAndClasses;
 
@@ -3172,7 +2831,7 @@ namespace Metasound
 					{
 						if (InPredicate(*NodeClass, Node))
 						{
-							FNodeAccessPtr NodePtr = GraphClassPtr.GetMemberAccessPtr<FMetasoundFrontendNode>(FGetNodeFromGraphClassByNodeID(Node.ID));
+							FNodeAccessPtr NodePtr = GraphClassPtr.GetNodeWithNodeID(Node.ID);
 							NodesAndClasses.Add(FGraphController::FNodeAndClass{ NodePtr, NodeClassPtr });
 						}
 					}
@@ -3188,8 +2847,6 @@ namespace Metasound
 
 		TArray<FGraphController::FConstNodeAndClass> FGraphController::GetNodesAndClassesByPredicate(TFunctionRef<bool (const FMetasoundFrontendClass&, const FMetasoundFrontendNode&)> InPredicate) const
 		{
-			using namespace FrontendControllerIntrinsics;
-
 			TArray<FConstNodeAndClass> NodesAndClasses;
 
 			if (GraphClassPtr.IsValid())
@@ -3201,7 +2858,7 @@ namespace Metasound
 					{
 						if (InPredicate(*NodeClass, Node))
 						{
-							FConstNodeAccessPtr NodePtr = GraphClassPtr.GetMemberAccessPtr<const FMetasoundFrontendNode>(FGetConstNodeFromGraphClassByNodeID(Node.ID));
+							FConstNodeAccessPtr NodePtr = GraphClassPtr.GetNodeWithNodeID(Node.ID);
 							NodesAndClasses.Add(FGraphController::FConstNodeAndClass{ NodePtr, NodeClassPtr });
 						}
 					}
@@ -3297,12 +2954,10 @@ namespace Metasound
 
 		FNodeHandle FGraphController::GetNodeHandle(const FGraphController::FNodeAndClass& InNodeAndClass)
 		{
-			using namespace FrontendControllerIntrinsics;
-
 			if (InNodeAndClass.IsValid() && GraphClassPtr.IsValid())
 			{
 				FGraphHandle OwningGraph = this->AsShared();
-				FGraphAccessPtr GraphPtr = GraphClassPtr.GetMemberAccessPtr<FMetasoundFrontendGraph>(FGetGraphFromGraphClass());
+				FGraphAccessPtr GraphPtr = GraphClassPtr.GetGraph();
 
 				switch (InNodeAndClass.Class->Metadata.Type)
 				{
@@ -3360,12 +3015,10 @@ namespace Metasound
 
 		FConstNodeHandle FGraphController::GetNodeHandle(const FGraphController::FConstNodeAndClass& InNodeAndClass) const
 		{
-			using namespace FrontendControllerIntrinsics;
-
 			if (InNodeAndClass.IsValid() && GraphClassPtr.IsValid())
 			{
 				FConstGraphHandle OwningGraph = this->AsShared();
-				FConstGraphAccessPtr GraphPtr = GraphClassPtr.GetMemberAccessPtr<const FMetasoundFrontendGraph>(FGetConstGraphFromGraphClass());
+				FConstGraphAccessPtr GraphPtr = GraphClassPtr.GetGraph();
 
 				switch (InNodeAndClass.Class->Metadata.Type)
 				{
@@ -3374,10 +3027,10 @@ namespace Metasound
 						{
 							FInputNodeController::FInitParams InitParams
 							{
-								ConstCastAccessPtr<FMetasoundFrontendNode>(InNodeAndClass.Node),
+								ConstCastAccessPtr<FNodeAccessPtr>(InNodeAndClass.Node),
 								InNodeAndClass.Class,
-								ConstCastAccessPtr<FMetasoundFrontendClassInput>(OwningGraphClassInputPtr),
-								ConstCastAccessPtr<FMetasoundFrontendGraph>(GraphPtr),
+								ConstCastAccessPtr<FClassInputAccessPtr>(OwningGraphClassInputPtr),
+								ConstCastAccessPtr<FGraphAccessPtr>(GraphPtr),
 								ConstCastSharedRef<IGraphController>(OwningGraph)
 							};
 							return FInputNodeController::CreateConstInputNodeHandle(InitParams);
@@ -3389,10 +3042,10 @@ namespace Metasound
 						{
 							FOutputNodeController::FInitParams InitParams
 							{
-								ConstCastAccessPtr<FMetasoundFrontendNode>(InNodeAndClass.Node),
+								ConstCastAccessPtr<FNodeAccessPtr>(InNodeAndClass.Node),
 								InNodeAndClass.Class,
-								ConstCastAccessPtr<FMetasoundFrontendClassOutput>(OwningGraphClassOutputPtr),
-								ConstCastAccessPtr<FMetasoundFrontendGraph>(GraphPtr),
+								ConstCastAccessPtr<FClassOutputAccessPtr>(OwningGraphClassOutputPtr),
+								ConstCastAccessPtr<FGraphAccessPtr>(GraphPtr),
 								ConstCastSharedRef<IGraphController>(OwningGraph)
 							};
 							return FOutputNodeController::CreateConstOutputNodeHandle(InitParams);
@@ -3404,9 +3057,9 @@ namespace Metasound
 						{
 							FNodeController::FInitParams InitParams
 							{
-								ConstCastAccessPtr<FMetasoundFrontendNode>(InNodeAndClass.Node),
+								ConstCastAccessPtr<FNodeAccessPtr>(InNodeAndClass.Node),
 								InNodeAndClass.Class,
-								ConstCastAccessPtr<FMetasoundFrontendGraph>(GraphPtr),
+								ConstCastAccessPtr<FGraphAccessPtr>(GraphPtr),
 								ConstCastSharedRef<IGraphController>(OwningGraph)
 							};
 							return FNodeController::CreateConstNodeHandle(InitParams);
@@ -3495,30 +3148,22 @@ namespace Metasound
 
 		FClassInputAccessPtr FGraphController::FindInputDescriptionWithNodeID(FGuid InNodeID)
 		{
-			using namespace FrontendControllerIntrinsics;
-
-			return GraphClassPtr.GetMemberAccessPtr<FMetasoundFrontendClassInput>(FGetClassInputFromGraphClassWithNodeID(InNodeID));
+			return GraphClassPtr.GetInputWithNodeID(InNodeID);
 		}
 
 		FConstClassInputAccessPtr FGraphController::FindInputDescriptionWithNodeID(FGuid InNodeID) const
 		{
-			using namespace FrontendControllerIntrinsics;
-
-			return GraphClassPtr.GetMemberAccessPtr<const FMetasoundFrontendClassInput>(FGetConstClassInputFromGraphClassWithNodeID(InNodeID));
+			return GraphClassPtr.GetInputWithNodeID(InNodeID);
 		}
 
 		FClassOutputAccessPtr FGraphController::FindOutputDescriptionWithNodeID(FGuid InNodeID)
 		{
-			using namespace FrontendControllerIntrinsics;
-
-			return GraphClassPtr.GetMemberAccessPtr<FMetasoundFrontendClassOutput>(FGetClassOutputFromGraphClassWithNodeID(InNodeID));
+			return GraphClassPtr.GetOutputWithNodeID(InNodeID);
 		}
 
 		FConstClassOutputAccessPtr FGraphController::FindOutputDescriptionWithNodeID(FGuid InNodeID) const
 		{
-			using namespace FrontendControllerIntrinsics;
-
-			return GraphClassPtr.GetMemberAccessPtr<const FMetasoundFrontendClassOutput>(FGetConstClassOutputFromGraphClassWithNodeID(InNodeID));
+			return GraphClassPtr.GetOutputWithNodeID(InNodeID);
 		}
 
 		FDocumentAccess FGraphController::ShareAccess() 
@@ -3603,14 +3248,12 @@ namespace Metasound
 
 		FConstClassAccessPtr FDocumentController::FindDependencyWithID(FGuid InClassID) const 
 		{
-			using namespace FrontendControllerIntrinsics;
-			return DocumentPtr.GetMemberAccessPtr<const FMetasoundFrontendClass>(FGetConstClassFromDocumentDependenciesWithID(InClassID));
+			return DocumentPtr.GetDependencyWithID(InClassID);
 		}
 
 		FConstGraphClassAccessPtr FDocumentController::FindSubgraphWithID(FGuid InClassID) const
 		{
-			using namespace FrontendControllerIntrinsics;
-			return DocumentPtr.GetMemberAccessPtr<const FMetasoundFrontendGraphClass>(FGetConstGraphClassFromDocumentWithID(InClassID));
+			return DocumentPtr.GetSubgraphWithID(InClassID);
 		}
 
 		FConstClassAccessPtr FDocumentController::FindClassWithID(FGuid InClassID) const
@@ -3627,15 +3270,11 @@ namespace Metasound
 
 		FConstClassAccessPtr FDocumentController::FindClass(const FNodeClassInfo& InNodeClass) const
 		{
-			using namespace FrontendControllerIntrinsics;
-
-			return DocumentPtr.GetMemberAccessPtr<const FMetasoundFrontendClass>(FGetConstClassFromDocumentWithClassInfo(InNodeClass));
+			return DocumentPtr.GetClassWithInfo(InNodeClass);
 		}
 
 		FConstClassAccessPtr FDocumentController::FindOrAddClass(const FNodeClassInfo& InNodeClass)
 		{
-			using namespace FrontendControllerIntrinsics;
-			
 			FConstClassAccessPtr ClassPtr = FindClass(InNodeClass);
 
 			auto AddClass = [=](FMetasoundFrontendClass&& NewClassDescription)
@@ -3674,74 +3313,72 @@ namespace Metasound
 
 		FConstClassAccessPtr FDocumentController::FindClass(const FMetasoundFrontendClassMetadata& InMetadata) const
 		{
-			using namespace FrontendControllerIntrinsics;
-
-			return DocumentPtr.GetMemberAccessPtr<const FMetasoundFrontendClass>(FGetConstClassFromDocumentWithMetadata(InMetadata));
+			return DocumentPtr.GetClassWithMetadata(InMetadata);
 		}
 
 		FConstClassAccessPtr FDocumentController::FindOrAddClass(const FMetasoundFrontendClassMetadata& InMetadata)
 		{
 			FConstClassAccessPtr ClassPtr = FindClass(InMetadata);
 
-if (!DocumentPtr.IsValid())
-{
-	return ClassPtr;
-}
+			if (!DocumentPtr.IsValid())
+			{
+				return ClassPtr;
+			}
 
-// External node classes must match major version to return shared definition.
-if (EMetasoundFrontendClassType::External == InMetadata.Type)
-{
-	if (ClassPtr.IsValid())
-	{
-		if (InMetadata.Version.Major != ClassPtr->Metadata.Version.Major)
-		{
-			ClassPtr = FConstClassAccessPtr();
-		}
-	}
-}
+			// External node classes must match major version to return shared definition.
+			if (EMetasoundFrontendClassType::External == InMetadata.Type)
+			{
+				if (ClassPtr.IsValid())
+				{
+					if (InMetadata.Version.Major != ClassPtr->Metadata.Version.Major)
+					{
+						ClassPtr = FConstClassAccessPtr();
+					}
+				}
+			}
 
-if (!ClassPtr.IsValid())
-{
-	if ((EMetasoundFrontendClassType::External == InMetadata.Type) || (EMetasoundFrontendClassType::Input == InMetadata.Type) || (EMetasoundFrontendClassType::Output == InMetadata.Type))
-	{
-		FMetasoundFrontendClass NewClass;
-		if (FMetasoundFrontendRegistryContainer::GetFrontendClassFromRegistered(InMetadata, NewClass))
-		{
-			NewClass.ID = NewClassID();
-			DocumentPtr->Dependencies.Add(NewClass);
-		}
-		else
-		{
-			UE_LOG(LogMetaSound, Error,
-				TEXT("Cannot add external dependency. No Metasound class found with matching metadata Name: "
-					"\"%s\" (%s). Suggested solution \"%s\" by %s."),
-				*InMetadata.ClassName.GetFullName().ToString(),
-				*InMetadata.Version.ToString(),
-				*InMetadata.PromptIfMissing.ToString(),
-				*InMetadata.Author.ToString());
-		}
-	}
-	else if (EMetasoundFrontendClassType::Graph == InMetadata.Type)
-	{
-		FMetasoundFrontendGraphClass NewClass;
-		NewClass.ID = NewClassID();
-		NewClass.Metadata = InMetadata;
+			if (!ClassPtr.IsValid())
+			{
+				if ((EMetasoundFrontendClassType::External == InMetadata.Type) || (EMetasoundFrontendClassType::Input == InMetadata.Type) || (EMetasoundFrontendClassType::Output == InMetadata.Type))
+				{
+					FMetasoundFrontendClass NewClass;
+					if (FMetasoundFrontendRegistryContainer::GetFrontendClassFromRegistered(InMetadata, NewClass))
+					{
+						NewClass.ID = NewClassID();
+						DocumentPtr->Dependencies.Add(NewClass);
+					}
+					else
+					{
+						UE_LOG(LogMetaSound, Error,
+							TEXT("Cannot add external dependency. No Metasound class found with matching metadata Name: "
+								"\"%s\" (%s). Suggested solution \"%s\" by %s."),
+							*InMetadata.ClassName.GetFullName().ToString(),
+							*InMetadata.Version.ToString(),
+							*InMetadata.PromptIfMissing.ToString(),
+							*InMetadata.Author.ToString());
+					}
+				}
+				else if (EMetasoundFrontendClassType::Graph == InMetadata.Type)
+				{
+					FMetasoundFrontendGraphClass NewClass;
+					NewClass.ID = NewClassID();
+					NewClass.Metadata = InMetadata;
 
-		DocumentPtr->Subgraphs.Add(NewClass);
-	}
-	else
-	{
-		UE_LOG(LogMetaSound, Error, TEXT(
-			"Unsupported metasound class type for node: \"%s\" (%s)."),
-			*InMetadata.ClassName.GetFullName().ToString(),
-			*InMetadata.Version.ToString());
-		checkNoEntry();
-	}
+					DocumentPtr->Subgraphs.Add(NewClass);
+				}
+				else
+				{
+					UE_LOG(LogMetaSound, Error, TEXT(
+						"Unsupported metasound class type for node: \"%s\" (%s)."),
+						*InMetadata.ClassName.GetFullName().ToString(),
+						*InMetadata.Version.ToString());
+					checkNoEntry();
+				}
 
-	ClassPtr = FindClass(InMetadata);
-}
+				ClassPtr = FindClass(InMetadata);
+			}
 
-return ClassPtr;
+			return ClassPtr;
 		}
 
 		void FDocumentController::SynchronizeDependencies()
@@ -3801,7 +3438,7 @@ return ClassPtr;
 		{
 			if (DocumentPtr.IsValid())
 			{
-				FGraphClassAccessPtr GraphClass = DocumentPtr.GetMemberAccessPtr<FMetasoundFrontendGraphClass>([](FMetasoundFrontendDocument& Doc) { return &Doc.RootGraph; });
+				FGraphClassAccessPtr GraphClass = DocumentPtr.GetRootGraph();
 				return FGraphController::CreateGraphHandle(FGraphController::FInitParams{GraphClass, this->AsShared()});
 			}
 
@@ -3812,10 +3449,10 @@ return ClassPtr;
 		{
 			if (DocumentPtr.IsValid())
 			{
-				FConstGraphClassAccessPtr GraphClass = DocumentPtr.GetMemberAccessPtr<const FMetasoundFrontendGraphClass>([](const FMetasoundFrontendDocument& Doc) { return &Doc.RootGraph; });
+				FConstGraphClassAccessPtr GraphClass = DocumentPtr.GetRootGraph();
 				return FGraphController::CreateConstGraphHandle(FGraphController::FInitParams
 					{
-						ConstCastAccessPtr<FMetasoundFrontendGraphClass>(GraphClass),
+						ConstCastAccessPtr<FGraphClassAccessPtr>(GraphClass),
 						ConstCastSharedRef<IDocumentController>(this->AsShared())
 					});
 			}
@@ -3854,15 +3491,16 @@ return ClassPtr;
 
 		FGraphHandle FDocumentController::GetSubgraphWithClassID(FGuid InClassID)
 		{
-			FGraphClassAccessPtr GraphClassPtr = DocumentPtr.GetMemberAccessPtr<FMetasoundFrontendGraphClass>(FrontendControllerIntrinsics::FGetSubgraphFromDocumentByID(InClassID));
+			FGraphClassAccessPtr GraphClassPtr = DocumentPtr.GetSubgraphWithID(InClassID);
+
 			return FGraphController::CreateGraphHandle(FGraphController::FInitParams{GraphClassPtr, this->AsShared()});
 		}
 
 		FConstGraphHandle FDocumentController::GetSubgraphWithClassID(FGuid InClassID) const
 		{
-			FConstGraphClassAccessPtr GraphClassPtr = DocumentPtr.GetMemberAccessPtr<const FMetasoundFrontendGraphClass>(FrontendControllerIntrinsics::FGetSubgraphFromDocumentByID(InClassID));
+			FConstGraphClassAccessPtr GraphClassPtr = DocumentPtr.GetSubgraphWithID(InClassID);
 
-			return FGraphController::CreateConstGraphHandle(FGraphController::FInitParams{ConstCastAccessPtr<FMetasoundFrontendGraphClass>(GraphClassPtr), ConstCastSharedRef<IDocumentController>(this->AsShared())});
+			return FGraphController::CreateConstGraphHandle(FGraphController::FInitParams{ConstCastAccessPtr<FGraphClassAccessPtr>(GraphClassPtr), ConstCastSharedRef<IDocumentController>(this->AsShared())});
 		}
 
 		bool FDocumentController::ExportToJSONAsset(const FString& InAbsolutePath) const
diff --git a/Engine/Plugins/Runtime/Metasound/Source/MetasoundFrontend/Public/MetasoundAccessPtr.h b/Engine/Plugins/Runtime/Metasound/Source/MetasoundFrontend/Public/MetasoundAccessPtr.h
--- a/Engine/Plugins/Runtime/Metasound/Source/MetasoundFrontend/Public/MetasoundAccessPtr.h
+++ b/Engine/Plugins/Runtime/Metasound/Source/MetasoundFrontend/Public/MetasoundAccessPtr.h
@@ -10,22 +10,32 @@
 #else
 #define METASOUND_FRONTEND_ACCESSPTR_DEBUG_INFO 0
 #endif
+
 namespace Metasound
 {
 	namespace Frontend
 	{
-
 		class FAccessPoint;
 
+		/** The access token mirrors the lifespan of a object in a non-intrusive
+		 * manner. A TAccessPtr<> can be created to reference an object, but use
+		 * an access token to determine whether the referenced object is accessible.
+		 * When the access token is destroyed, the access pointer becomes invalid.
+		 */
 		struct FAccessToken {};
 
 		/** TAccessPtr
 		 *
 		 * TAccessPtr is used to determine whether an object has been destructed or not.
 		 * It is useful when an object cannot be wrapped in a TSharedPtr. A TAccessPtr
-		 * works similar to a TWeakPtr, but it cannot pin the object. Object pointers 
-		 * held within a TAccessPtr must only be accessed on the thread where the objects
-		 * gets destructed to avoid having the object destructed while in use.
+		 * is functionally similar to a TWeakPtr, but it cannot pin the object. 
+		 * In order to determine whether an object as accessible, the TAccessPtr
+		 * executes a TFunctions<> which returns object. If a nullptr is returned,
+		 * then the object is not accessible. 
+		 *
+		 * Object pointers  held within a TAccessPtr must only be accessed on 
+		 * the thread where the objects gets destructed to avoid having the object 
+		 * destructed while in use.
 		 *
 		 * If the TAccessPtr's underlying object is accessed when the pointer is invalid,
 		 * a fallback object will be returned. 
@@ -35,10 +45,12 @@ namespace Metasound
 		template<typename Type>
 		class TAccessPtr
 		{
-			enum EConstCast { Tag };
+			enum class EConstCast { Tag };
+			enum class EDerivedCopy { Tag };
 
 		public:
 			using FTokenType = FAccessToken;
+			using ObjectType = Type;
 
 			static Type FallbackObject;
 
@@ -51,6 +63,9 @@ namespace Metasound
 				return IsValid();
 			}
 
+			/** Returns a pointer to the accessible object. If the object is 
+			 * inaccessible, a nullptr is returned. 
+			 */
 			Type* Get() const
 			{
 #if METASOUND_FRONTEND_ACCESSPTR_DEBUG_INFO 
@@ -59,6 +74,10 @@ namespace Metasound
 				return GetObject();
 			}
 
+			/** Returns a pointer to the accessible object. If the object is 
+			 * inaccessible, an assert is triggered and the static fallback object
+			 * is returned. 
+			 */
 			Type& operator*() const
 			{
 				if (Type* Object = Get())
@@ -70,6 +89,9 @@ namespace Metasound
 				return FallbackObject;
 			}
 
+			/** Returns a pointer to the accessible object. If the object is 
+			 * inaccessible, a nullptr is returned. 
+			 */
 			Type* operator->() const
 			{
 				if (Type* Object = Get())
@@ -80,9 +102,18 @@ namespace Metasound
 				return &FallbackObject;
 			}
 
-			template<typename MemberType>
-			TAccessPtr<MemberType> GetMemberAccessPtr(TFunction<MemberType*(Type&)> InGetMember) const
+			/** Returns an access pointer to a member of the wrapped object. 
+			 *
+			 * @tparam AccessPtrType - The access pointer type to return.
+			 * @tparam FunctionType - A type which is callable accepts a reference to the wrapped object and returns a pointer to the member.
+			 *
+			 * @param InGetMember - A FunctionType accepts a reference to the wrapped object and returns a pointer to the member.
+			 */
+			template<typename AccessPtrType, typename FunctionType>
+			AccessPtrType GetMemberAccessPtr(FunctionType InGetMember) const
 			{
+				using MemberType = typename AccessPtrType::ObjectType;
+
 				TFunction<MemberType*()> GetMemberFromObject = [=]() -> MemberType*
 				{
 					if (Type* Object = Get())
@@ -92,12 +123,10 @@ namespace Metasound
 					return static_cast<MemberType*>(nullptr);
 				};
 
-				return TAccessPtr<MemberType>(GetMemberFromObject);
+				return AccessPtrType(GetMemberFromObject);
 			}
 
-
 			TAccessPtr()
-			//: Object(nullptr)
 			: GetObject([]() { return static_cast<Type*>(nullptr); })
 			{
 #if METASOUND_FRONTEND_ACCESSPTR_DEBUG_INFO 
@@ -105,11 +134,29 @@ namespace Metasound
 #endif
 			}
 
+			/** Creates a access pointer using an access token. */
+			TAccessPtr(TWeakPtr<FTokenType> AccessToken, Type& InRef)
+			{
+				Type* RefPtr = &InRef;
+
+				GetObject = [=]() -> Type*
+				{
+					Type* Object = nullptr;
+					if (AccessToken.IsValid())
+					{
+						Object = RefPtr;
+					}
+					return Object;
+				};
+#if METASOUND_FRONTEND_ACCESSPTR_DEBUG_INFO 
+				Get();
+#endif
+			}
 
+
+			/** Creates an access pointer from another using a const casts. */
 			template<typename OtherType>
 			TAccessPtr(const TAccessPtr<OtherType>& InOther, EConstCast InTag)
-			//: Token(InOther.Token)
-			//, Object(const_cast<Type*>(InOther.Object))
 			{
 				TFunction<OtherType*()> OtherGetObject = InOther.GetObject;
 				GetObject = [=]() -> Type*
@@ -121,13 +168,12 @@ namespace Metasound
 #endif
 			} 
 
+			/** Creates an access pointer from another using a static cast. */
 			template <
 				typename OtherType,
 				typename = decltype(ImplicitConv<Type*>((OtherType*)nullptr))
 			>
-			TAccessPtr(const TAccessPtr<OtherType>& InOther)
-			//: Token(InOther.Token)
-			//, Object(InOther.Object)
+			TAccessPtr(const TAccessPtr<OtherType>& InOther, EDerivedCopy InTag=EDerivedCopy::Tag)
 			{
 				TFunction<OtherType*()> OtherGetObject = InOther.GetObject;
 
@@ -146,12 +192,13 @@ namespace Metasound
 			TAccessPtr(TAccessPtr<Type>&& InOther) = default;
 			TAccessPtr& operator=(TAccessPtr<Type>&& InOther) = default;
 
-		private:
-			template<typename RelatedType>
-			friend TAccessPtr<RelatedType> MakeAccessPtr(const FAccessPoint& InAccessPoint, RelatedType& InRef);
+		protected:
+
+			template<typename RelatedAccessPtrType, typename RelatedType>
+			friend RelatedAccessPtrType MakeAccessPtr(const FAccessPoint& InAccessPoint, RelatedType& InRef);
 
-			template<typename ToType, typename FromType> 
-			friend TAccessPtr<ToType> ConstCastAccessPtr(const TAccessPtr<FromType>& InAccessPtr);
+			template<typename ToAccessPtrType, typename FromAccessPtrType> 
+			friend ToAccessPtrType ConstCastAccessPtr(const FromAccessPtrType& InAccessPtr);
 
 			template<typename OtherType>
 			friend class TAccessPtr;
@@ -169,28 +216,12 @@ namespace Metasound
 #endif
 			}
 
-			TAccessPtr(TWeakPtr<FTokenType> AccessToken, Type& InRef)
-			{
-				Type* RefPtr = &InRef;
-
-				GetObject = [=]() -> Type*
-				{
-					Type* Object = nullptr;
-					if (AccessToken.IsValid())
-					{
-						Object = RefPtr;
-					}
-					return Object;
-				};
-#if METASOUND_FRONTEND_ACCESSPTR_DEBUG_INFO 
-				Get();
-#endif
-			}
 		};
 
 		template<typename Type>
 		Type TAccessPtr<Type>::FallbackObject = Type();
 
+
 		/** FAccessPoint acts as a lifecycle tracker for the TAccessPtrs it creates. 
 		 * When this object is destructed, all associated TAccessPtrs will become invalid.
 		 */
@@ -217,8 +248,8 @@ namespace Metasound
 			}
 
 		private:
-			template<typename Type>
-			friend TAccessPtr<Type> MakeAccessPtr(const FAccessPoint& InAccessPoint, Type& InRef);
+			template<typename AccessPtrType, typename Type>
+			friend AccessPtrType MakeAccessPtr(const FAccessPoint& InAccessPoint, Type& InRef);
 
 			FAccessPoint(FAccessPoint&&) = delete;
 			FAccessPoint& operator=(FAccessPoint&&) = delete;
@@ -226,16 +257,16 @@ namespace Metasound
 			TSharedPtr<FTokenType> Token;
 		};
 
-		template<typename Type>
-		TAccessPtr<Type> MakeAccessPtr(const FAccessPoint& InAccessPoint, Type& InRef)
+		template<typename AccessPtrType, typename Type>
+		AccessPtrType MakeAccessPtr(const FAccessPoint& InAccessPoint, Type& InRef)
 		{
-			return TAccessPtr<Type>(InAccessPoint.Token, InRef);
+			return AccessPtrType(InAccessPoint.Token, InRef);
 		}
 
-		template<typename ToType, typename FromType> 
-		TAccessPtr<ToType> ConstCastAccessPtr(const TAccessPtr<FromType>& InAccessPtr)
+		template<typename ToAccessPtrType, typename FromAccessPtrType> 
+		ToAccessPtrType ConstCastAccessPtr(const FromAccessPtrType& InAccessPtr)
 		{
-			return TAccessPtr<ToType>(InAccessPtr, TAccessPtr<ToType>::EConstCast::Tag);
+			return ToAccessPtrType(InAccessPtr, ToAccessPtrType::EConstCast::Tag);
 		}
 	}
 }
diff --git a/Engine/Plugins/Runtime/Metasound/Source/MetasoundFrontend/Public/MetasoundAssetBase.h b/Engine/Plugins/Runtime/Metasound/Source/MetasoundFrontend/Public/MetasoundAssetBase.h
--- a/Engine/Plugins/Runtime/Metasound/Source/MetasoundFrontend/Public/MetasoundAssetBase.h
+++ b/Engine/Plugins/Runtime/Metasound/Source/MetasoundFrontend/Public/MetasoundAssetBase.h
@@ -6,6 +6,7 @@
 #include "MetasoundFrontend.h"
 #include "MetasoundFrontendBaseClasses.h"
 #include "MetasoundFrontendController.h"
+#include "MetasoundFrontendDocumentAccessPtr.h"
 #include "MetasoundGraph.h"
 #include "MetasoundLog.h"
 #include "UObject/WeakObjectPtrTemplates.h"
@@ -126,10 +127,10 @@ public:
 protected:
 
 	// Returns an access pointer to the document.
-	virtual Metasound::Frontend::TAccessPtr<FMetasoundFrontendDocument> GetDocument() = 0;
+	virtual Metasound::Frontend::FDocumentAccessPtr GetDocument() = 0;
 
 	// Returns an access pointer to the document.
-	virtual Metasound::Frontend::TAccessPtr<const FMetasoundFrontendDocument> GetDocument() const = 0;
+	virtual Metasound::Frontend::FConstDocumentAccessPtr GetDocument() const = 0;
 
 	// Returns the owning asset responsible for transactions applied to metasound
 	virtual UObject* GetOwningAsset() = 0;
diff --git a/Engine/Plugins/Runtime/Metasound/Source/MetasoundFrontend/Public/MetasoundFrontendController.h b/Engine/Plugins/Runtime/Metasound/Source/MetasoundFrontend/Public/MetasoundFrontendController.h
--- a/Engine/Plugins/Runtime/Metasound/Source/MetasoundFrontend/Public/MetasoundFrontendController.h
+++ b/Engine/Plugins/Runtime/Metasound/Source/MetasoundFrontend/Public/MetasoundFrontendController.h
@@ -8,6 +8,7 @@
 #include "MetasoundFrontend.h"
 #include "MetasoundFrontendBaseClasses.h"
 #include "MetasoundFrontendDocument.h"
+#include "MetasoundFrontendDocumentAccessPtr.h"
 #include "MetasoundFrontendRegistries.h"
 #include "MetasoundGraph.h"
 
@@ -56,22 +57,6 @@ namespace Metasound
 		using FDocumentHandle = TSharedRef<IDocumentController>;
 		using FConstDocumentHandle = TSharedRef<const IDocumentController>;
 
-		using FVertexAccessPtr = TAccessPtr<FMetasoundFrontendVertex>;
-		using FConstVertexAccessPtr = TAccessPtr<const FMetasoundFrontendVertex>;
-		using FClassOutputAccessPtr = TAccessPtr<FMetasoundFrontendClassOutput>;
-		using FConstClassOutputAccessPtr = TAccessPtr<const FMetasoundFrontendClassOutput>;
-		using FClassInputAccessPtr = TAccessPtr<FMetasoundFrontendClassInput>;
-		using FConstClassInputAccessPtr = TAccessPtr<const FMetasoundFrontendClassInput>;
-		using FNodeAccessPtr = TAccessPtr<FMetasoundFrontendNode>;
-		using FConstNodeAccessPtr = TAccessPtr<const FMetasoundFrontendNode>;
-		using FGraphAccessPtr = TAccessPtr<FMetasoundFrontendGraph>;
-		using FConstGraphAccessPtr = TAccessPtr<const FMetasoundFrontendGraph>;
-		using FClassAccessPtr = TAccessPtr<FMetasoundFrontendClass>;
-		using FConstClassAccessPtr = TAccessPtr<const FMetasoundFrontendClass>;
-		using FGraphClassAccessPtr = TAccessPtr<FMetasoundFrontendGraphClass>;
-		using FConstGraphClassAccessPtr = TAccessPtr<const FMetasoundFrontendGraphClass>;
-		using FDocumentAccessPtr = TAccessPtr<FMetasoundFrontendDocument>;
-		using FConstDocumentAccessPtr = TAccessPtr<const FMetasoundFrontendDocument>;
 
 		// Container holding various access pointers to the FMetasoundFrontendDocument
 		struct FConstDocumentAccess
@@ -682,9 +667,9 @@ namespace Metasound
 			static FDocumentHandle GetInvalidHandle();
 
 			/** Create a document from FMetasoundFrontendDocument description pointer. */
-			static FDocumentHandle CreateDocumentHandle(TAccessPtr<FMetasoundFrontendDocument> InDocument);
+			static FDocumentHandle CreateDocumentHandle(FDocumentAccessPtr InDocument);
 			/** Create a document from FMetasoundFrontendDocument description pointer. */
-			static FConstDocumentHandle CreateDocumentHandle(TAccessPtr<const FMetasoundFrontendDocument> InDocument);
+			static FConstDocumentHandle CreateDocumentHandle(FConstDocumentAccessPtr InDocument);
 
 			IDocumentController() = default;
 			virtual ~IDocumentController() = default;
diff --git a/Engine/Plugins/Runtime/Metasound/Source/MetasoundFrontend/Public/MetasoundFrontendDocumentAccessPtr.h b/Engine/Plugins/Runtime/Metasound/Source/MetasoundFrontend/Public/MetasoundFrontendDocumentAccessPtr.h
new file mode 100644
--- /dev/null
+++ b/Engine/Plugins/Runtime/Metasound/Source/MetasoundFrontend/Public/MetasoundFrontendDocumentAccessPtr.h
@@ -0,0 +1,171 @@
+// Copyright Epic Games, Inc. All Rights Reserved.
+
+#pragma once
+
+#include "MetasoundAccessPtr.h"
+#include "MetasoundFrontendDocument.h"
+#include "MetasoundFrontendRegistries.h"
+
+/** Metasound Frontend Document Access Pointers provide convenience methods for
+ * getting access pointers of FMetasoundFrontendDocument sub-elements. For instance,
+ * to get an input vertex subobject on the root graph of the document, one can call:
+ *
+ * FDocumentAccessPtr DocumentAccessPtr = MakeAccessPtr<FDocumentAccessPtr>(AccessPoint, Document);
+ *
+ * FClassAccessPtr ClassAccessPtr = DocumentAccessPtr.GetRootGraph().GetNodeWithNodeID(NodeID).GetInputWithVertexID(VertexID);
+ *
+ */
+
+namespace Metasound
+{
+	namespace Frontend
+	{
+		using FVertexAccessPtr = TAccessPtr<FMetasoundFrontendVertex>;
+		using FConstVertexAccessPtr = TAccessPtr<const FMetasoundFrontendVertex>;
+		using FClassInputAccessPtr = TAccessPtr<FMetasoundFrontendClassInput>;
+		using FConstClassInputAccessPtr = TAccessPtr<const FMetasoundFrontendClassInput>;
+		using FClassOutputAccessPtr = TAccessPtr<FMetasoundFrontendClassOutput>;
+		using FConstClassOutputAccessPtr = TAccessPtr<const FMetasoundFrontendClassOutput>;
+
+		class FNodeAccessPtr : public TAccessPtr<FMetasoundFrontendNode>
+		{
+		public:
+			using Super = TAccessPtr<FMetasoundFrontendNode>;
+
+			// Inherit constructors from base class
+			using Super::Super;
+
+			FVertexAccessPtr GetInputWithName(const FString& InName);
+			FVertexAccessPtr GetInputWithVertexID(const FGuid& InID);
+			FVertexAccessPtr GetOutputWithName(const FString& InName);
+			FVertexAccessPtr GetOutputWithVertexID(const FGuid& InID);
+
+			FConstVertexAccessPtr GetInputWithName(const FString& InName) const;
+			FConstVertexAccessPtr GetInputWithVertexID(const FGuid& InID) const;
+			FConstVertexAccessPtr GetOutputWithName(const FString& InName) const;
+			FConstVertexAccessPtr GetOutputWithVertexID(const FGuid& InID) const;
+		};
+
+		class FConstNodeAccessPtr : public TAccessPtr<const FMetasoundFrontendNode>
+		{
+		public:
+			using Super = TAccessPtr<const FMetasoundFrontendNode>;
+
+			// Inherit constructors from base class
+			using Super::Super;
+
+			FConstVertexAccessPtr GetInputWithName(const FString& InName) const;
+			FConstVertexAccessPtr GetInputWithVertexID(const FGuid& InID) const;
+			FConstVertexAccessPtr GetOutputWithName(const FString& InName) const;
+			FConstVertexAccessPtr GetOutputWithVertexID(const FGuid& InID) const;
+		};
+
+		using FGraphAccessPtr = TAccessPtr<FMetasoundFrontendGraph>;
+		using FConstGraphAccessPtr = TAccessPtr<const FMetasoundFrontendGraph>;
+
+		class FClassAccessPtr : public TAccessPtr<FMetasoundFrontendClass>
+		{
+		public:
+			using Super = TAccessPtr<FMetasoundFrontendClass>;
+
+			// Inherit constructors from base class
+			using Super::Super;
+
+			FClassInputAccessPtr GetInputWithName(const FString& InName);
+			FClassOutputAccessPtr GetOutputWithName(const FString& InName);
+			FConstClassInputAccessPtr GetInputWithName(const FString& InName) const;
+			FConstClassOutputAccessPtr GetOutputWithName(const FString& InName) const;
+		};
+
+		class FConstClassAccessPtr : public TAccessPtr<const FMetasoundFrontendClass>
+		{
+		public:
+			using Super = TAccessPtr<const FMetasoundFrontendClass>;
+
+			// Inherit constructors from base class
+			using Super::Super;
+
+			FConstClassInputAccessPtr GetInputWithName(const FString& InName) const;
+			FConstClassOutputAccessPtr GetOutputWithName(const FString& InName) const;
+		};
+
+		class FGraphClassAccessPtr : public TAccessPtr<FMetasoundFrontendGraphClass>
+		{
+		public:
+			using Super = TAccessPtr<FMetasoundFrontendGraphClass>;
+
+			// Inherit constructors from base class
+			using Super::Super;
+
+			FClassInputAccessPtr GetInputWithName(const FString& InName);
+			FClassInputAccessPtr GetInputWithNodeID(const FGuid& InNodeID);
+			FClassOutputAccessPtr GetOutputWithName(const FString& InName);
+			FClassOutputAccessPtr GetOutputWithNodeID(const FGuid& InNodeID);
+
+			FNodeAccessPtr GetNodeWithNodeID(const FGuid& InNodeID);
+			FGraphAccessPtr GetGraph();
+
+			FConstClassInputAccessPtr GetInputWithName(const FString& InName) const;
+			FConstClassInputAccessPtr GetInputWithNodeID(const FGuid& InNodeID) const;
+			FConstClassOutputAccessPtr GetOutputWithName(const FString& InName) const;
+			FConstClassOutputAccessPtr GetOutputWithNodeID(const FGuid& InNodeID) const;
+
+			FConstNodeAccessPtr GetNodeWithNodeID(const FGuid& InNodeID) const;
+			FConstGraphAccessPtr GetGraph() const;
+		};
+
+		class FConstGraphClassAccessPtr : public TAccessPtr<const FMetasoundFrontendGraphClass>
+		{
+			using Super = TAccessPtr<const FMetasoundFrontendGraphClass>;
+
+			// Inherit constructors from base class
+			using Super::Super;
+
+			FConstClassInputAccessPtr GetInputWithName(const FString& InName) const;
+			FConstClassInputAccessPtr GetInputWithNodeID(const FGuid& InNodeID) const;
+			FConstClassOutputAccessPtr GetOutputWithName(const FString& InName) const;
+			FConstClassOutputAccessPtr GetOutputWithNodeID(const FGuid& InNodeID) const;
+			FConstNodeAccessPtr GetNodeWithNodeID(const FGuid& InNodeID) const;
+			FConstGraphAccessPtr GetGraph() const;
+		};
+
+		class FDocumentAccessPtr : public TAccessPtr<FMetasoundFrontendDocument>
+		{
+		public:
+			using Super = TAccessPtr<FMetasoundFrontendDocument>;
+
+			// Inherit constructors from base class
+			using Super::Super;
+
+			FGraphClassAccessPtr GetRootGraph();
+			FGraphClassAccessPtr GetSubgraphWithID(const FGuid& InID);
+			FClassAccessPtr GetDependencyWithID(const FGuid& InID);
+			FClassAccessPtr GetClassWithID(const FGuid& InID);
+			FClassAccessPtr GetClassWithMetadata(const FMetasoundFrontendClassMetadata& InMetadata);
+			FClassAccessPtr GetClassWithInfo(const FNodeClassInfo& InInfo);
+
+			FConstGraphClassAccessPtr GetRootGraph() const;
+			FConstGraphClassAccessPtr GetSubgraphWithID(const FGuid& InID) const;
+			FConstClassAccessPtr GetDependencyWithID(const FGuid& InID) const;
+			FConstClassAccessPtr GetClassWithID(const FGuid& InID) const;
+			FConstClassAccessPtr GetClassWithMetadata(const FMetasoundFrontendClassMetadata& InMetadata) const;
+			FConstClassAccessPtr GetClassWithInfo(const FNodeClassInfo& InInfo) const;
+		};
+
+		class FConstDocumentAccessPtr : public TAccessPtr<const FMetasoundFrontendDocument>
+		{
+		public:
+			using Super = TAccessPtr<const FMetasoundFrontendDocument>;
+
+			// Inherit constructors from base class
+			using Super::Super;
+
+			FConstGraphClassAccessPtr GetRootGraph() const;
+			FConstGraphClassAccessPtr GetSubgraphWithID(const FGuid& InID) const;
+			FConstClassAccessPtr GetDependencyWithID(const FGuid& InID) const;
+			FConstClassAccessPtr GetClassWithID(const FGuid& InID) const;
+			FConstClassAccessPtr GetClassWithMetadata(const FMetasoundFrontendClassMetadata& InMetadata) const;
+			FConstClassAccessPtr GetClassWithInfo(const FNodeClassInfo& InInfo) const;
+		};
+	}
+}
diff --git a/Engine/Plugins/Runtime/Metasound/Source/MetasoundGenerator/Private/MetasoundGenerator.cpp b/Engine/Plugins/Runtime/Metasound/Source/MetasoundGenerator/Private/MetasoundGenerator.cpp
--- a/Engine/Plugins/Runtime/Metasound/Source/MetasoundGenerator/Private/MetasoundGenerator.cpp
+++ b/Engine/Plugins/Runtime/Metasound/Source/MetasoundGenerator/Private/MetasoundGenerator.cpp
@@ -22,9 +22,10 @@ namespace Metasound
 
 	void FAsyncMetaSoundBuilder::DoWork()
 	{
+		using namespace Frontend;
 		// Create handles for new root graph
-		Frontend::FConstDocumentHandle NewDocumentHandle = Frontend::IDocumentController::CreateDocumentHandle(Frontend::MakeAccessPtr<const FMetasoundFrontendDocument>(InitParams.DocumentCopy.AccessPoint, InitParams.DocumentCopy));
-		Frontend::FConstGraphHandle RootGraph = NewDocumentHandle->GetRootGraph();
+		FConstDocumentHandle NewDocumentHandle = IDocumentController::CreateDocumentHandle(MakeAccessPtr<FConstDocumentAccessPtr>(InitParams.DocumentCopy.AccessPoint, InitParams.DocumentCopy));
+		FConstGraphHandle RootGraph = NewDocumentHandle->GetRootGraph();
 		ensureAlways(RootGraph->IsValid());
 
 		TArray<IOperatorBuilder::FBuildErrorPtr> BuildErrors;
