- Disallow renaming Required MetaSound Inputs/Outputs (ex. 'On Finished') from GraphActionView name selector
- Remove dead customization code

diff --git a/Engine/Plugins/Runtime/Metasound/Source/MetasoundEditor/Private/MetasoundEditorGraph.cpp b/Engine/Plugins/Runtime/Metasound/Source/MetasoundEditor/Private/MetasoundEditorGraph.cpp
--- a/Engine/Plugins/Runtime/Metasound/Source/MetasoundEditor/Private/MetasoundEditorGraph.cpp
+++ b/Engine/Plugins/Runtime/Metasound/Source/MetasoundEditor/Private/MetasoundEditorGraph.cpp
@@ -150,9 +150,15 @@ bool UMetasoundEditorGraphVariable::CanRename(const FText& InNewName, FText& Out
 	}
 
 	FConstNodeHandle NodeHandle = GetConstNodeHandle();
-	FConstGraphHandle GraphHandle = NodeHandle->GetOwningGraph();
+
+	if (NodeHandle->IsRequired())
+	{
+		OutError = FText::Format(LOCTEXT("VariableRenameInvalid_VariableRequired", "{0} is required and cannot be renamed."), InNewName);
+		return false;
+	}
 
 	bool bIsNameValid = true;
+	FConstGraphHandle GraphHandle = NodeHandle->GetOwningGraph();
 	GraphHandle->IterateConstNodes([&](FConstNodeHandle NodeToCompare)
 	{
 		if (NodeID != NodeToCompare->GetID())
@@ -165,7 +171,7 @@ bool UMetasoundEditorGraphVariable::CanRename(const FText& InNewName, FText& Out
 		}
 	}, GetClassType());
 
-		return bIsNameValid;
+	return bIsNameValid;
 }
 
 #if WITH_EDITORONLY_DATA
diff --git a/Engine/Plugins/Runtime/Metasound/Source/MetasoundEditor/Private/MetasoundNodeDetailCustomization.cpp b/Engine/Plugins/Runtime/Metasound/Source/MetasoundEditor/Private/MetasoundNodeDetailCustomization.cpp
--- a/Engine/Plugins/Runtime/Metasound/Source/MetasoundEditor/Private/MetasoundNodeDetailCustomization.cpp
+++ b/Engine/Plugins/Runtime/Metasound/Source/MetasoundEditor/Private/MetasoundNodeDetailCustomization.cpp
@@ -865,44 +865,6 @@ namespace Metasound
 			return TypeName;
 		}
 
-		void FMetasoundInputDetailCustomization::OnDisplayNameChanged(const FText& InNewName)
-		{
-			using namespace Frontend;
-
-			bIsNameInvalid = false;
-			DisplayNameEditableTextBox->SetError(FText::GetEmpty());
-
-			if (!ensure(GraphVariable.IsValid()))
-			{
-				return;
-			}
-
-			if (InNewName.IsEmpty())
-			{
-				bIsNameInvalid = true;
-				DisplayNameEditableTextBox->SetError(FText::Format(LOCTEXT("InputRenameInvalid_NameEmpty", "{0} cannot be empty string."), InNewName));
-				return;
-			}
-
-			FConstNodeHandle NodeHandle = GraphVariable->GetConstNodeHandle();
-			FConstGraphHandle GraphHandle = NodeHandle->GetOwningGraph();
-			const FGuid NodeID = NodeHandle->GetID();
-
-			GraphHandle->IterateConstNodes([this, NodeID, InNewName](FConstNodeHandle NodeToCompare)
-			{
-				// Disregard display name collisions with hidden nodes
-				const bool bIsVisible = NodeToCompare->GetNodeStyle().Display.Visibility == EMetasoundFrontendNodeStyleDisplayVisibility::Visible;
-				if (NodeID != NodeToCompare->GetID() && bIsVisible)
-				{
-					if (InNewName.CompareToCaseIgnored(NodeToCompare->GetDisplayName()) == 0)
-					{
-						bIsNameInvalid = true;
-						DisplayNameEditableTextBox->SetError(FText::Format(LOCTEXT("InputRenameInvalid_NameTaken", "{0} is already in use"), InNewName));
-					}
-				}
-			}, EMetasoundFrontendClassType::Input);
-		}
-
 		void FMetasoundOutputDetailCustomization::CustomizeDetails(IDetailLayoutBuilder& DetailLayout)
 		{
 			using namespace Frontend;
@@ -1048,42 +1010,6 @@ namespace Metasound
 
 			return TypeName;
 		}
-
-		void FMetasoundOutputDetailCustomization::OnDisplayNameChanged(const FText& InNewName)
-		{
-			using namespace Frontend;
-
-			bIsNameInvalid = false;
-			DisplayNameEditableTextBox->SetError(FText::GetEmpty());
-
-			if (!ensure(GraphVariable.IsValid()))
-			{
-				return;
-			}
-
-			if (InNewName.IsEmpty())
-			{
-				bIsNameInvalid = true;
-				DisplayNameEditableTextBox->SetError(FText::Format(LOCTEXT("OutputRenameInvalid_NameEmpty", "{0} cannot be empty string."), InNewName));
-				return;
-			}
-
-			FConstNodeHandle NodeHandle = GraphVariable->GetConstNodeHandle();
-			FConstGraphHandle GraphHandle = NodeHandle->GetOwningGraph();
-			const FGuid NodeID = NodeHandle->GetID();
-
-			GraphHandle->IterateConstNodes([this, NodeID, InNewName](FConstNodeHandle NodeToCompare)
-			{
-				if (NodeID != NodeToCompare->GetID())
-				{
-					if (InNewName.CompareToCaseIgnored(NodeToCompare->GetDisplayName()) == 0)
-					{
-						bIsNameInvalid = true;
-						DisplayNameEditableTextBox->SetError(FText::Format(LOCTEXT("OutputRenameInvalid_NameTaken", "{0} is already in use"), InNewName));
-					}
-				}
-			}, EMetasoundFrontendClassType::Output);
-		}
 	} // namespace Editor
 } // namespace Metasound
 #undef LOCTEXT_NAMESPACE
diff --git a/Engine/Plugins/Runtime/Metasound/Source/MetasoundEditor/Private/MetasoundNodeDetailCustomization.h b/Engine/Plugins/Runtime/Metasound/Source/MetasoundEditor/Private/MetasoundNodeDetailCustomization.h
--- a/Engine/Plugins/Runtime/Metasound/Source/MetasoundEditor/Private/MetasoundNodeDetailCustomization.h
+++ b/Engine/Plugins/Runtime/Metasound/Source/MetasoundEditor/Private/MetasoundNodeDetailCustomization.h
@@ -131,7 +131,25 @@ namespace Metasound
 			}
 			// End of IDetailCustomization interface
 
-			virtual void OnDisplayNameChanged(const FText& InNewText) = 0;
+			void OnDisplayNameChanged(const FText& InNewName)
+			{
+				using namespace Frontend;
+
+				bIsNameInvalid = false;
+				DisplayNameEditableTextBox->SetError(FText::GetEmpty());
+
+				if (!ensure(GraphVariable.IsValid()))
+				{
+					return;
+				}
+
+				FText Error;
+				if (!GraphVariable->CanRename(InNewName, Error))
+				{
+					bIsNameInvalid = true;
+					DisplayNameEditableTextBox->SetError(Error);
+				}
+			}
 
 			FText GetDisplayName() const
 			{
@@ -164,6 +182,13 @@ namespace Metasound
 				{
 					const FText TransactionLabel = FText::Format(LOCTEXT("SetTooltip", "Set the MetaSound {0}'s tooltip"), VariableLabel);
 					const FScopedTransaction Transaction(TransactionLabel);
+
+					GraphVariable->Modify();
+					if (UMetasoundEditorGraph* Graph = Cast<UMetasoundEditorGraph>(GraphVariable->GetOuter()))
+					{
+						Graph->GetMetasoundChecked().Modify();
+					}
+
 					FNodeHandle NodeHandle = GraphVariable->GetNodeHandle();
 					NodeHandle->SetDescription(InNewText);
 				}
@@ -187,14 +212,11 @@ namespace Metasound
 
 				if (!bIsNameInvalid && GraphVariable.IsValid())
 				{
-					const FText TransactionLabel = FText::Format(LOCTEXT("Rename Variable", "Rename Metasound {0}"), VariableLabel);
-					const FScopedTransaction Transaction(TransactionLabel);
-					FNodeHandle NodeHandle = GraphVariable->GetNodeHandle();
-					NodeHandle->SetDisplayName(InNewName);
-
-					GraphVariable->NameChanged.Broadcast(NodeHandle->GetID());
-					DisplayNameEditableTextBox->SetError(FText::GetEmpty());
+					GraphVariable->SetDisplayName(InNewName);
 				}
+
+				DisplayNameEditableTextBox->SetError(FText::GetEmpty());
+				bIsNameInvalid = false;
 			}
 
 			ECheckBoxState OnGetPrivateCheckboxState() const
@@ -242,9 +264,6 @@ namespace Metasound
 			virtual void CustomizeDetails(IDetailLayoutBuilder& DetailLayout) override;
 			// End of IDetailCustomization interface
 
-		protected:
-			virtual void OnDisplayNameChanged(const FText& InNewText) override;
-
 		private:
 			void SetDefaultPropertyMetaData(TSharedRef<IPropertyHandle> InDefaultPropertyHandle) const;
 
@@ -263,9 +282,6 @@ namespace Metasound
 			virtual void CustomizeDetails(IDetailLayoutBuilder& DetailLayout) override;
 			// End of IDetailCustomization interface
 
-		protected:
-			virtual void OnDisplayNameChanged(const FText& InNewText) override;
-
 		private:
 			void SetDefaultPropertyMetaData(TSharedRef<IPropertyHandle> InDefaultPropertyHandle) const;
 
