Unfortunately, not serializing MetaSound EdGraph member on MetaSound assets is not quite ready for prime time. Reasons:
1. Bug in the synchronization location data
2. Need a representation of comment nodes in Frontend (or subset of editor only data that *is* serialized somehow)
3. Need a representation of how visual input nodes are connected (as there are multiple representations of them) so sync code properly hooks up

diff --git a/Engine/Plugins/Runtime/Metasound/Source/MetasoundEngine/Public/MetasoundSource.h b/Engine/Plugins/Runtime/Metasound/Source/MetasoundEngine/Public/MetasoundSource.h
--- a/Engine/Plugins/Runtime/Metasound/Source/MetasoundEngine/Public/MetasoundSource.h
+++ b/Engine/Plugins/Runtime/Metasound/Source/MetasoundEngine/Public/MetasoundSource.h
@@ -46,7 +46,7 @@ protected:
 	FMetasoundFrontendDocument RootMetasoundDocument;
 
 #if WITH_EDITORONLY_DATA
-	UPROPERTY(Transient)
+	UPROPERTY()
 	UMetasoundEditorGraphBase* Graph;
 #endif // WITH_EDITORONLY_DATA
 
