MetaSound Graph Core Subgraph

diff --git a/Engine/Plugins/Runtime/Metasound/Source/MetasoundEditor/Private/MetasoundEditorGraphBuilder.cpp b/Engine/Plugins/Runtime/Metasound/Source/MetasoundEditor/Private/MetasoundEditorGraphBuilder.cpp
--- a/Engine/Plugins/Runtime/Metasound/Source/MetasoundEditor/Private/MetasoundEditorGraphBuilder.cpp
+++ b/Engine/Plugins/Runtime/Metasound/Source/MetasoundEditor/Private/MetasoundEditorGraphBuilder.cpp
@@ -372,7 +372,7 @@ namespace Metasound
 				return;
 			}
 
-			FOutputHandle OutputHandle = InputHandle->GetCurrentlyConnectedOutput();
+			FOutputHandle OutputHandle = InputHandle->GetConnectedOutput();
 			if (OutputHandle->IsValid())
 			{
 				FNodeHandle InputNode = OutputHandle->GetOwningNode();
@@ -648,7 +648,7 @@ namespace Metasound
 			check(MetasoundAsset);
 
 			Frontend::FGraphHandle GraphHandle = MetasoundAsset->GetRootGraphHandle();
-			const FGuid VertexID = GraphHandle->GetNewVertexID();
+			const FGuid VertexID = FGuid::NewGuid();
 
 			FMetasoundFrontendClassInput Description;
 			Description.Name = InName;
@@ -690,7 +690,7 @@ namespace Metasound
 			Description.Name = InName;
 			Description.TypeName = InTypeName;
 			Description.Metadata.Description = InToolTip;
-			Description.VertexID = GraphHandle->GetNewVertexID();
+			Description.VertexID = FGuid::NewGuid();
 
 			Frontend::FNodeHandle NodeHandle = GraphHandle->AddOutputVertex(Description);
 			GraphHandle->SetOutputDisplayName(InName, FText::FromString(InName));
@@ -760,7 +760,7 @@ namespace Metasound
 				return false;
 			}
 
-			FOutputHandle ExistingOutput = InputHandle->GetCurrentlyConnectedOutput();
+			FOutputHandle ExistingOutput = InputHandle->GetConnectedOutput();
 			if (ExistingOutput->IsValid())
 			{
 				FNodeHandle NodeHandle = ExistingOutput->GetOwningNode();
@@ -813,7 +813,7 @@ namespace Metasound
 			for (int32 i = 0; i < InputHandles.Num(); ++i)
 			{
 				FInputHandle InputHandle = InputHandles[i];
-				FConstOutputHandle OutputHandle = InputHandle->GetCurrentlyConnectedOutput();
+				FConstOutputHandle OutputHandle = InputHandle->GetConnectedOutput();
 				const FMetasoundFrontendNodeStyle& Style = OutputHandle->GetOwningNode()->GetNodeStyle();
 
 				// Hidden nodes are not "connected" from the perspective of EdGraph,
@@ -900,7 +900,7 @@ namespace Metasound
 					{
 						case EEdGraphPinDirection::EGPD_Input:
 						{
-							FOutputHandle OutputHandle = NodeInputs[InputIndex]->GetCurrentlyConnectedOutput();
+							FOutputHandle OutputHandle = NodeInputs[InputIndex]->GetConnectedOutput();
 							if (OutputHandle->IsValid())
 							{
 								const FMetasoundFrontendNodeStyle& Style = OutputHandle->GetOwningNode()->GetNodeStyle();
@@ -958,7 +958,7 @@ namespace Metasound
 			TArray<FInputHandle> Inputs = NodeHandle->GetInputs();
 			for (FInputHandle& Input : Inputs)
 			{
-				FOutputHandle Output = Input->GetCurrentlyConnectedOutput();
+				FOutputHandle Output = Input->GetConnectedOutput();
 				if (Output->IsValid())
 				{
 					FNodeHandle InputHandle = Output->GetOwningNode();
@@ -1259,7 +1259,7 @@ namespace Metasound
 					bool bShowConnectionInEditor = true;
 					bool bIsConnectedToLiteralInput = false;
 
-					FOutputHandle OutputHandle = NodeInput->GetCurrentlyConnectedOutput();
+					FOutputHandle OutputHandle = NodeInput->GetConnectedOutput();
 					if (OutputHandle->IsValid())
 					{
 						FNodeHandle InputNodeHandle = OutputHandle->GetOwningNode();
diff --git a/Engine/Plugins/Runtime/Metasound/Source/MetasoundEngine/Private/MetasoundSource.cpp b/Engine/Plugins/Runtime/Metasound/Source/MetasoundEngine/Private/MetasoundSource.cpp
--- a/Engine/Plugins/Runtime/Metasound/Source/MetasoundEngine/Private/MetasoundSource.cpp
+++ b/Engine/Plugins/Runtime/Metasound/Source/MetasoundEngine/Private/MetasoundSource.cpp
@@ -202,7 +202,7 @@ bool UMetaSoundSource::GetReceiveNodeMetadataForDataType(const FName& InTypeName
 Metasound::Frontend::FNodeHandle UMetaSoundSource::AddInputPinForSendAddress(const Metasound::FMetasoundInstanceTransmitter::FSendInfo& InSendInfo, Metasound::Frontend::FGraphHandle InGraph) const
 {
 	FMetasoundFrontendClassInput Description;
-	FGuid VertexID = InGraph->GetNewVertexID();
+	FGuid VertexID = FGuid::NewGuid();
 
 	Description.Name = InSendInfo.Address.ChannelName.ToString();
 	Description.TypeName = Metasound::GetMetasoundDataTypeName<Metasound::FSendAddress>();
@@ -283,7 +283,7 @@ bool UMetaSoundSource::CopyDocumentAndInjectReceiveNodes(uint64 InInstanceID, co
 			if (ensure(ReceiveDefaultInputs.Num() == 1))
 			{
 				FOutputHandle ReceiverNodeOutput = ReceiveNodeOutputs[0];
-				for (FInputHandle NodeInput : InputNodeOutput->GetCurrentlyConnectedInputs())
+				for (FInputHandle NodeInput : InputNodeOutput->GetConnectedInputs())
 				{
 					// Swap connections to receiver node
 					ensure(InputNodeOutput->Disconnect(*NodeInput));
diff --git a/Engine/Plugins/Runtime/Metasound/Source/MetasoundFrontend/Private/MetasoundFrontendController.cpp b/Engine/Plugins/Runtime/Metasound/Source/MetasoundFrontend/Private/MetasoundFrontendController.cpp
--- a/Engine/Plugins/Runtime/Metasound/Source/MetasoundFrontend/Private/MetasoundFrontendController.cpp
+++ b/Engine/Plugins/Runtime/Metasound/Source/MetasoundFrontend/Private/MetasoundFrontendController.cpp
@@ -41,10 +41,20 @@ namespace Metasound
 			return FDocumentController::CreateDocumentHandle(InDocument);
 		}
 
+		FDocumentHandle IDocumentController::CreateDocumentHandle(FMetasoundFrontendDocument& InDocument)
+		{
+			return CreateDocumentHandle(MakeAccessPtr<FDocumentAccessPtr>(InDocument.AccessPoint, InDocument));
+		}
+
 		FConstDocumentHandle IDocumentController::CreateDocumentHandle(FConstDocumentAccessPtr InDocument)
 		{
 			// Create using standard document controller. 
 			return FDocumentController::CreateDocumentHandle(ConstCastAccessPtr<FDocumentAccessPtr>(InDocument));
 		}
+
+		FConstDocumentHandle IDocumentController::CreateDocumentHandle(const FMetasoundFrontendDocument& InDocument)
+		{
+			return CreateDocumentHandle(MakeAccessPtr<FConstDocumentAccessPtr>(InDocument.AccessPoint, InDocument));
+		}
 	}
 }
diff --git a/Engine/Plugins/Runtime/Metasound/Source/MetasoundFrontend/Private/MetasoundFrontendGraph.cpp b/Engine/Plugins/Runtime/Metasound/Source/MetasoundFrontend/Private/MetasoundFrontendGraph.cpp
--- a/Engine/Plugins/Runtime/Metasound/Source/MetasoundFrontend/Private/MetasoundFrontendGraph.cpp
+++ b/Engine/Plugins/Runtime/Metasound/Source/MetasoundFrontend/Private/MetasoundFrontendGraph.cpp
@@ -4,6 +4,7 @@
 
 #include "Algo/AllOf.h"
 #include "Algo/AnyOf.h"
+#include "Algo/TopologicalSort.h"
 #include "Algo/Transform.h"
 #include "CoreMinimal.h"
 #include "MetasoundFrontend.h"
@@ -13,12 +14,13 @@
 
 namespace Metasound
 {
+
 	FFrontendGraph::FFrontendGraph(const FString& InInstanceName, const FGuid& InInstanceID)
 	:	FGraph(InInstanceName, InInstanceID)
 	{
 	}
 
-	void FFrontendGraph::AddInputNode(FGuid InDependencyId, int32 InIndex, const FVertexKey& InVertexKey, TUniquePtr<INode> InNode)
+	void FFrontendGraph::AddInputNode(FGuid InDependencyId, int32 InIndex, const FVertexKey& InVertexKey, TSharedPtr<const INode> InNode)
 	{
 		if (InNode.IsValid())
 		{
@@ -28,11 +30,11 @@ namespace Metasound
 			// Input nodes need an extra Index value to keep track of their position in the graph's inputs.
 			InputNodes.Add(InIndex, InNode.Get());
 			AddInputDataDestination(*InNode, InVertexKey);
-			AddNode(InDependencyId, MoveTemp(InNode));
+			AddNode(InDependencyId, InNode);
 		}
 	}
 
-	void FFrontendGraph::AddOutputNode(FGuid InNodeID, int32 InIndex, const FVertexKey& InVertexKey, TUniquePtr<INode> InNode)
+	void FFrontendGraph::AddOutputNode(FGuid InNodeID, int32 InIndex, const FVertexKey& InVertexKey, TSharedPtr<const INode> InNode)
 	{
 		if (InNode.IsValid())
 		{
@@ -42,12 +44,12 @@ namespace Metasound
 			// Output nodes need an extra Index value to keep track of their position in the graph's inputs.
 			OutputNodes.Add(InIndex, InNode.Get());
 			AddOutputDataSource(*InNode, InVertexKey);
-			AddNode(InNodeID, MoveTemp(InNode));
+			AddNode(InNodeID, InNode);
 		}
 	}
 
 	/** Store a node on this graph. */
-	void FFrontendGraph::AddNode(FGuid InNodeID, TUniquePtr<INode> InNode)
+	void FFrontendGraph::AddNode(FGuid InNodeID, TSharedPtr<const INode> InNode)
 	{
 		if (InNode.IsValid())
 		{
@@ -55,13 +57,13 @@ namespace Metasound
 			check(!NodeMap.Contains(InNodeID));
 
 			NodeMap.Add(InNodeID, InNode.Get());
-			StoreNode(MoveTemp(InNode));
+			StoreNode(InNode);
 		}
 	}
 
 	const INode* FFrontendGraph::FindNode(FGuid InNodeID) const
 	{
-		INode* const* NodePtr = NodeMap.Find(InNodeID);
+		const INode* const* NodePtr = NodeMap.Find(InNodeID);
 
 		if (nullptr != NodePtr)
 		{
@@ -73,7 +75,7 @@ namespace Metasound
 
 	const INode* FFrontendGraph::FindInputNode(int32 InIndex) const
 	{
-		INode* const* NodePtr = InputNodes.Find(InIndex);
+		const INode* const* NodePtr = InputNodes.Find(InIndex);
 
 		if (nullptr != NodePtr)
 		{
@@ -85,7 +87,7 @@ namespace Metasound
 
 	const INode* FFrontendGraph::FindOutputNode(int32 InIndex) const
 	{
-		INode* const* NodePtr = OutputNodes.Find(InIndex);
+		const INode* const* NodePtr = OutputNodes.Find(InIndex);
 
 		if (nullptr != NodePtr)
 		{
@@ -134,15 +136,13 @@ namespace Metasound
 		return true;
 	}
 
-
-	void FFrontendGraph::StoreNode(TUniquePtr<INode> InNode)
+	void FFrontendGraph::StoreNode(TSharedPtr<const INode> InNode)
 	{
 		check(InNode.IsValid());
 		StoredNodes.Add(InNode.Get());
-		Storage.Add(MoveTemp(InNode));
+		NodeStorage.Add(InNode);
 	}
 
-
 	TUniquePtr<INode> FFrontendGraphBuilder::CreateInputNode(const FMetasoundFrontendNode& InNode, const FMetasoundFrontendClass& InClass, const FMetasoundFrontendClassInput& InOwningGraphClassInput) const
 	{
 		const FMetasoundFrontendLiteral* FrontendLiteral = FindInputLiteralForInputNode(InNode, InClass, InOwningGraphClassInput);
@@ -219,20 +219,6 @@ namespace Metasound
 
 		InitData.InstanceID = InNode.ID;
 
-
-		// Copy over our initialization params.
-		/*
-		for (auto& StaticParamTuple : InNode.StaticParameters)
-		{
-			FLiteral LiteralParam = Frontend::GetLiteralParam(StaticParamTuple.Value);
-
-			if (LiteralParam.IsValid())
-			{
-				InitData.ParamMap.Add(StaticParamTuple.Key, MoveTemp(LiteralParam));
-			}
-		}
-		*/
-		
 		// TODO: handle check to see if node interface conforms to class interface here. 
 		// TODO: check to see if external object supports class interface.
 
@@ -246,7 +232,7 @@ namespace Metasound
 	{
 		OutClassInputIndex = INDEX_NONE;
 
-		// TODO: assumes input node has exactly one input 
+		// Input nodes should have exactly one input.
 		if (ensure(InInputNode.Interface.Inputs.Num() == 1))
 		{
 			const FName& TypeName = InInputNode.Interface.Inputs[0].TypeName;
@@ -269,7 +255,7 @@ namespace Metasound
 	{
 		OutClassOutputIndex = INDEX_NONE;
 
-		// TODO: assumes input node has exactly one input 
+		// Output nodes should have exactly one output
 		if (ensure(InOutputNode.Interface.Outputs.Num() == 1))
 		{
 			const FName& TypeName = InOutputNode.Interface.Outputs[0].TypeName;
@@ -342,7 +328,7 @@ namespace Metasound
 	}
 
 	// TODO: add errors here. Most will be a "PromptIfMissing"...
-	void FFrontendGraphBuilder::AddNodesToGraph(const FMetasoundFrontendGraphClass& InGraphClass, const TMap<FGuid, const FMetasoundFrontendClass*>& InClasses, FFrontendGraph& OutGraph) const
+	void FFrontendGraphBuilder::AddNodesToGraph(const FMetasoundFrontendGraphClass& InGraphClass, const FDependencyByIDMap& InClasses, const FSharedNodeByIDMap& InSubgraphs, FFrontendGraph& OutGraph) const
 	{
 		for (const FMetasoundFrontendNode& Node : InGraphClass.Graph.Nodes)
 		{
@@ -359,8 +345,8 @@ namespace Metasound
 
 							if ((nullptr != ClassInput) && (INDEX_NONE != InputIndex))
 							{
-
-								OutGraph.AddInputNode(Node.ID, InputIndex, ClassInput->Name, CreateInputNode(Node, *NodeClass, *ClassInput));
+								TSharedPtr<const INode> InputNode(CreateInputNode(Node, *NodeClass, *ClassInput).Release());
+								OutGraph.AddInputNode(Node.ID, InputIndex, ClassInput->Name, InputNode);
 							}
 							else
 							{
@@ -377,7 +363,8 @@ namespace Metasound
 
 							if ((nullptr != ClassOutput) && (INDEX_NONE != OutputIndex))
 							{
-								OutGraph.AddOutputNode(Node.ID, OutputIndex, ClassOutput->Name, CreateOutputNode(Node, *NodeClass));
+								TSharedPtr<const INode> OutputNode(CreateOutputNode(Node, *NodeClass).Release());
+								OutGraph.AddOutputNode(Node.ID, OutputIndex, ClassOutput->Name, OutputNode);
 							}
 							else
 							{
@@ -387,8 +374,27 @@ namespace Metasound
 
 						break;
 
+					case EMetasoundFrontendClassType::Graph:
+						{
+							const TSharedPtr<const INode>* SubgraphPtr = InSubgraphs.Find(Node.ClassID);
+
+							if (nullptr == SubgraphPtr)
+							{
+								UE_LOG(LogMetaSound, Error, TEXT("Failed to find subgraph for node [NodeID:%s, NodeName:%s, ClassID:%s]"), *Node.ID.ToString(), *Node.Name, *Node.ClassID.ToString());
+							}
+							else
+							{
+								OutGraph.AddNode(Node.ID, *SubgraphPtr);
+							}
+						}
+						break;
+
+					case EMetasoundFrontendClassType::External:
 					default:
-						OutGraph.AddNode(Node.ID, CreateExternalNode(Node, *NodeClass));
+						{
+							TSharedPtr<const INode> ExternalNode(CreateExternalNode(Node, *NodeClass).Release());
+							OutGraph.AddNode(Node.ID, ExternalNode);
+						}
 				}
 			}
 		}
@@ -404,7 +410,6 @@ namespace Metasound
 			const FMetasoundFrontendVertex* Vertex = nullptr;
 		};
 
-		// TODO: add support for array vertices.
 		typedef TTuple<FGuid, FGuid> FNodeIDVertexID;
 
 		TMap<FNodeIDVertexID, FCoreNodeAndFrontendVertex> NodeSourcesByID;
@@ -524,29 +529,62 @@ namespace Metasound
 		return bIsEveryDependencyMet;
 	}
 
-	TUniquePtr<FFrontendGraph> FFrontendGraphBuilder::CreateGraph(const FMetasoundFrontendGraphClass& InGraph, const TArray<FMetasoundFrontendGraphClass>& InSubgraphs, const TArray<FMetasoundFrontendClass>& InDependencies) const
+	bool FFrontendGraphBuilder::SortSubgraphDependencies(TArray<const FMetasoundFrontendGraphClass*>& Subgraphs) const
 	{
-		if (!IsFlat(InGraph, InDependencies))
+		// Helper for caching and querying subgraph dependencies
+		struct FSubgraphDependencyLookup
 		{
-			// Likely this will change in the future and the builder will be able
-			// to build graphs with subgraphs..
-			UE_LOG(LogMetaSound, Error, TEXT("Provided graph not flat. FFrontendGraphBuilder can only build flat graphs"));
-			return TUniquePtr<FFrontendGraph>(nullptr);
-		}
+			FSubgraphDependencyLookup(TArrayView<const FMetasoundFrontendGraphClass*> InGraphs)
+			{
+				// Map ClassID to graph pointer. 
+				TMap<FGuid, const FMetasoundFrontendGraphClass*> ClassIDAndGraph;
+				for (const FMetasoundFrontendGraphClass* Graph: InGraphs)
+				{
+					ClassIDAndGraph.Add(Graph->ID, Graph);
+				}
 
+				// Cache subgraph dependencies.
+				for (const FMetasoundFrontendGraphClass* GraphClass : InGraphs)
+				{
+					for (const FMetasoundFrontendNode& Node : GraphClass->Graph.Nodes)
+					{
+						if (ClassIDAndGraph.Contains(Node.ClassID))
+						{
+							DependencyMap.Add(GraphClass, ClassIDAndGraph[Node.ClassID]);
+						}
+					}
+				}
+			}
 
-		TMap<FGuid, const FMetasoundFrontendClass*> ClassMap;
+			TArray<const FMetasoundFrontendGraphClass*> operator()(const FMetasoundFrontendGraphClass* InParent) const
+			{
+				TArray<const FMetasoundFrontendGraphClass*> Dependencies;
+				DependencyMap.MultiFind(InParent, Dependencies);
+				return Dependencies;
+			}
 
-		for (const FMetasoundFrontendClass& ExtClass : InDependencies)
+		private:
+
+			TMultiMap<const FMetasoundFrontendGraphClass*, const FMetasoundFrontendGraphClass*> DependencyMap;
+		};
+
+		bool bSuccess = Algo::TopologicalSort(Subgraphs, FSubgraphDependencyLookup(Subgraphs));
+		if (!bSuccess)
 		{
-			ClassMap.Add(ExtClass.ID, &ExtClass);
+			UE_LOG(LogMetaSound, Error, TEXT("Failed to topologically sort subgraphs. Possible recursive subgraph dependency"));
 		}
 
+		return bSuccess;
+	}
+
+	TUniquePtr<FFrontendGraph> FFrontendGraphBuilder::CreateGraph(FBuildContext& InContext, const FMetasoundFrontendGraphClass& InGraph) const
+	{
+
 		TUniquePtr<FFrontendGraph> MetasoundGraph = MakeUnique<FFrontendGraph>(InGraph.Metadata.ClassName.GetFullName().ToString(), FGuid::NewGuid());
 
 		// TODO: will likely want to bubble up errors here for case where
 		// a datatype or node is not registered. 
-		AddNodesToGraph(InGraph, ClassMap, *MetasoundGraph);
+		AddNodesToGraph(InGraph, InContext.FrontendClasses, InContext.Graphs, *MetasoundGraph);
 
 		AddEdgesToGraph(InGraph.Graph, *MetasoundGraph);
 
@@ -554,6 +592,51 @@ namespace Metasound
 
 		return MoveTemp(MetasoundGraph);
 	}
+
+
+	TUniquePtr<FFrontendGraph> FFrontendGraphBuilder::CreateGraph(const FMetasoundFrontendGraphClass& InGraph, const TArray<FMetasoundFrontendGraphClass>& InSubgraphs, const TArray<FMetasoundFrontendClass>& InDependencies) const
+	{
+		FBuildContext Context;
+
+		// Gather all references to node classes from external dependencies and subgraphs.
+		for (const FMetasoundFrontendClass& ExtClass : InDependencies)
+		{
+			Context.FrontendClasses.Add(ExtClass.ID, &ExtClass);
+		}
+		for (const FMetasoundFrontendClass& ExtClass : InSubgraphs)
+		{
+			Context.FrontendClasses.Add(ExtClass.ID, &ExtClass);
+		}
+
+		// Sort subgraphs so that dependent subgraphs are created in correct order.
+		TArray<const FMetasoundFrontendGraphClass*> FrontendSubgraphPtrs;
+		Algo::Transform(InSubgraphs, FrontendSubgraphPtrs, [](const FMetasoundFrontendGraphClass& InClass) { return &InClass; });
+
+		bool bSuccess = SortSubgraphDependencies(FrontendSubgraphPtrs);
+		if (!bSuccess)
+		{
+			UE_LOG(LogMetaSound, Error, TEXT("Failed to create graph due to failed subgraph ordering."));
+			return TUniquePtr<FFrontendGraph>(nullptr);
+		}
+
+		// Create each subgraph.
+		for (const FMetasoundFrontendGraphClass* FrontendSubgraphPtr : FrontendSubgraphPtrs)
+		{
+			TSharedPtr<const INode> Subgraph(CreateGraph(Context, *FrontendSubgraphPtr).Release());
+			if (!Subgraph.IsValid())
+			{
+				UE_LOG(LogMetaSound, Warning, TEXT("Failed to create subgraph [SubgraphName: %s]"), *FrontendSubgraphPtr->Metadata.ClassName.ToString());
+			}
+			else
+			{
+				// Add subgraphs to context so they are accessible for subsequent graphs.
+				Context.Graphs.Add(FrontendSubgraphPtr->ID, Subgraph);
+			}
+		}
+
+		// Create parent graph.
+		return CreateGraph(Context, InGraph); 
+	}
 	
 	/* Metasound document should be inflated by now. */
 	TUniquePtr<FFrontendGraph> FFrontendGraphBuilder::CreateGraph(const FMetasoundFrontendDocument& InDocument) const
diff --git a/Engine/Plugins/Runtime/Metasound/Source/MetasoundFrontend/Private/MetasoundFrontendGraph.h b/Engine/Plugins/Runtime/Metasound/Source/MetasoundFrontend/Private/MetasoundFrontendGraph.h
--- a/Engine/Plugins/Runtime/Metasound/Source/MetasoundFrontend/Private/MetasoundFrontendGraph.h
+++ b/Engine/Plugins/Runtime/Metasound/Source/MetasoundFrontend/Private/MetasoundFrontendGraph.h
@@ -16,7 +16,6 @@ namespace Metasound
 	 */
 	class FFrontendGraph : public FGraph
 	{
-		// - TODO: consider adding ability to add in another sub graph
 		public:
 			/** FFrontendGraph constructor.
 			 *
@@ -32,27 +31,27 @@ namespace Metasound
 			 * @param InNodeID - The NodeID related to the parent FMetasoundFrontendClass.
 			 * @param InIndex - The positional index for the input.
 			 * @param InVertexKey - The key for the graph input vertex.
-			 * @param InNode - A unique pointer to an input node. 
+			 * @param InNode - A shared pointer to an input node. 
 			 */
-			void AddInputNode(FGuid InNodeID, int32 InIndex, const FVertexKey& InVertexKey, TUniquePtr<INode> InNode);
+			void AddInputNode(FGuid InNodeID, int32 InIndex, const FVertexKey& InVertexKey, TSharedPtr<const INode> InNode);
 
 			/** Add an output node to this graph.
 			 *
 			 * @param InNodeID - The NodeID related to the parent FMetasoundFrontendClass.
 			 * @param InIndex - The positional index for the output.
 			 * @param InVertexKey - The key for the graph output vertex.
-			 * @param InNode - A unique pointer to an output node. 
+			 * @param InNode - A shared pointer to an output node. 
 			 */
-			void AddOutputNode(FGuid InNodeID, int32 InIndex, const FVertexKey& InVertexKey, TUniquePtr<INode> InNode);
+			void AddOutputNode(FGuid InNodeID, int32 InIndex, const FVertexKey& InVertexKey, TSharedPtr<const INode> InNode);
 
 			/** Store a node on this graph. 
 			 *
 			 * @param InNodeID - The NodeID related to the parent FMetasoundFrontendClass.
-			 * @param InNode - A unique pointer to a node. 
+			 * @param InNode - A shared pointer to a node. 
 			 */
-			void AddNode(FGuid InNodeID, TUniquePtr<INode> InNode);
+			void AddNode(FGuid InNodeID, TSharedPtr<const INode> InNode);
 
-			/** Retrieve node by dependency ID.
+			/** Retrieve node by node ID.
 			 *
 			 * @param InNodeID - The NodeID of the requested Node.
 			 *
@@ -82,14 +81,14 @@ namespace Metasound
 
 		private:
 
-			void StoreNode(TUniquePtr<INode> InNode);
+			void StoreNode(TSharedPtr<const INode> InNode);
 
-			TMap<int32, INode*> InputNodes;
-			TMap<int32, INode*> OutputNodes;
+			TMap<int32, const INode*> InputNodes;
+			TMap<int32, const INode*> OutputNodes;
 
-			TMap<FGuid, INode*> NodeMap;
+			TMap<FGuid, const INode*> NodeMap;
 			TSet<const INode*> StoredNodes;
-			TArray<TUniquePtr<INode>> Storage;
+			TArray<TSharedPtr<const INode>> NodeStorage;
 	};
 
 	/** FFrontendGraphBuilder builds a FFrontendGraph from a FMetasoundDoucment
@@ -117,6 +116,18 @@ namespace Metasound
 
 
 	private:
+		using FDependencyByIDMap = TMap<FGuid, const FMetasoundFrontendClass*>;
+		using FSharedNodeByIDMap = TMap<FGuid, TSharedPtr<const INode>>;
+
+		struct FBuildContext
+		{
+			FDependencyByIDMap FrontendClasses;
+			FSharedNodeByIDMap Graphs;
+		};
+
+		bool SortSubgraphDependencies(TArray<const FMetasoundFrontendGraphClass*>& Subgraphs) const;
+
+		TUniquePtr<FFrontendGraph> CreateGraph(FBuildContext& InContext, const FMetasoundFrontendGraphClass& InSubgraph) const;
 
 		const FMetasoundFrontendClassInput* FindClassInputForInputNode(const FMetasoundFrontendGraphClass& InOwningGraph, const FMetasoundFrontendNode& InInputNode, int32& OutClassInputIndex) const;
 
@@ -131,7 +142,7 @@ namespace Metasound
 		TUniquePtr<INode> CreateExternalNode(const FMetasoundFrontendNode& InNode, const FMetasoundFrontendClass& InClass) const;
 
 		// TODO: add errors here. Most will be a "PromptIfMissing"...
-		void AddNodesToGraph(const FMetasoundFrontendGraphClass& InGraphClass, const TMap<FGuid, const FMetasoundFrontendClass*>& InClasses, FFrontendGraph& OutGraph) const;
+		void AddNodesToGraph(const FMetasoundFrontendGraphClass& InGraphClass, const FDependencyByIDMap& InClasses, const FSharedNodeByIDMap& InSubgraphs, FFrontendGraph& OutGraph) const;
 
 		void AddEdgesToGraph(const FMetasoundFrontendGraph& InGraphDescription, FFrontendGraph& OutGraph) const;
 	};
diff --git a/Engine/Plugins/Runtime/Metasound/Source/MetasoundFrontend/Private/MetasoundFrontendInvalidController.h b/Engine/Plugins/Runtime/Metasound/Source/MetasoundFrontend/Private/MetasoundFrontendInvalidController.h
--- a/Engine/Plugins/Runtime/Metasound/Source/MetasoundFrontend/Private/MetasoundFrontendInvalidController.h
+++ b/Engine/Plugins/Runtime/Metasound/Source/MetasoundFrontend/Private/MetasoundFrontendInvalidController.h
@@ -54,8 +54,9 @@ namespace Metasound
 			TSharedRef<INodeController> GetOwningNode() override;
 			TSharedRef<const INodeController> GetOwningNode() const override;
 
-			TArray<FInputHandle> GetCurrentlyConnectedInputs() override { return TArray<FInputHandle>(); }
-			TArray<FConstInputHandle> GetCurrentlyConnectedInputs() const override { return TArray<FConstInputHandle>(); }
+			bool IsConnected() const override { return false; }
+			TArray<FInputHandle> GetConnectedInputs() override { return TArray<FInputHandle>(); }
+			TArray<FConstInputHandle> GetConstConnectedInputs() const override { return TArray<FConstInputHandle>(); }
 			bool Disconnect() override { return false; }
 
 			FConnectability CanConnectTo(const IInputController& InController) const override { return FConnectability(); }
@@ -98,8 +99,8 @@ namespace Metasound
 			TSharedRef<INodeController> GetOwningNode() override;
 			TSharedRef<const INodeController> GetOwningNode() const override;
 
-			virtual TSharedRef<IOutputController> GetCurrentlyConnectedOutput() override { return FInvalidOutputController::GetInvalid(); }
-			virtual TSharedRef<const IOutputController> GetCurrentlyConnectedOutput() const override { return FInvalidOutputController::GetInvalid(); }
+			virtual TSharedRef<IOutputController> GetConnectedOutput() override { return FInvalidOutputController::GetInvalid(); }
+			virtual TSharedRef<const IOutputController> GetConnectedOutput() const override { return FInvalidOutputController::GetInvalid(); }
 			bool Disconnect() override { return false; }
 
 			virtual FConnectability CanConnectTo(const IOutputController& InController) const override { return FConnectability(); }
@@ -136,9 +137,9 @@ namespace Metasound
 			TArray<TSharedRef<const IOutputController>> GetConstOutputs() const override { return TArray<TSharedRef<const IOutputController>>(); }
 
 			TArray<FInputHandle> GetInputsWithVertexName(const FString& InName) override { return TArray<FInputHandle>(); }
-			TArray<FConstInputHandle> GetInputsWithVertexName(const FString& InName) const override { return TArray<FConstInputHandle>(); }
+			TArray<FConstInputHandle> GetConstInputsWithVertexName(const FString& InName) const override { return TArray<FConstInputHandle>(); }
 			TArray<FOutputHandle> GetOutputsWithVertexName(const FString& InName) override { return TArray<FOutputHandle>(); }
-			TArray<FConstOutputHandle> GetOutputsWithVertexName(const FString& InName) const override { return TArray<FConstOutputHandle>(); }
+			TArray<FConstOutputHandle> GetConstOutputsWithVertexName(const FString& InName) const override { return TArray<FConstOutputHandle>(); }
 			TSharedRef<IInputController> GetInputWithID(FGuid InVertexID) override { return FInvalidInputController::GetInvalid(); }
 			TSharedRef<IOutputController> GetOutputWithID(FGuid InVertexID) override { return FInvalidOutputController::GetInvalid(); }
 			TSharedRef<const IInputController> GetInputWithID(FGuid InVertexID) const override { return FInvalidInputController::GetInvalid(); }
@@ -212,15 +213,15 @@ namespace Metasound
 			}
 
 			bool IsValid() const override { return false; }
+			FGuid GetClassID() const override { return Metasound::FrontendInvalidID; }
+			const FText& GetDisplayName() const override { return MetasoundFrontendInvalidControllerPrivate::GetInvalid<FText>(); }
 
-			FGuid GetNewVertexID() const override { return Metasound::FrontendInvalidID; }
 			TArray<FString> GetInputVertexNames() const override { return TArray<FString>(); }
 			TArray<FString> GetOutputVertexNames() const override { return TArray<FString>(); }
 
 			TArray<TSharedRef<INodeController>> GetNodes() override { return TArray<TSharedRef<INodeController>>(); }
 			TArray<TSharedRef<const INodeController>> GetConstNodes() const override { return TArray<TSharedRef<const INodeController>>(); }
 
-			FGuid GetClassID() const override { return Metasound::FrontendInvalidID; }
 			TSharedRef<const INodeController> GetNodeWithID(FGuid InNodeID) const override { return FInvalidNodeController::GetInvalid(); }
 			TSharedRef<INodeController> GetNodeWithID(FGuid InNodeID) override { return FInvalidNodeController::GetInvalid(); }
 
@@ -280,6 +281,7 @@ namespace Metasound
 			TSharedRef<INodeController> AddNode(const FNodeClassInfo& InNodeClass) override { return FInvalidNodeController::GetInvalid(); }
 			TSharedRef<INodeController> AddNode(const FNodeRegistryKey& InNodeClass) override { return FInvalidNodeController::GetInvalid(); }
 			TSharedRef<INodeController> AddNode(const FMetasoundFrontendClassMetadata& InNodeClass) override { return FInvalidNodeController::GetInvalid(); }
+			TSharedRef<INodeController> AddDuplicateNode(const INodeController& InNode) override { return FInvalidNodeController::GetInvalid(); }
 
 			// Remove the node corresponding to this node handle.
 			// On success, invalidates the received node handle.
@@ -288,14 +290,6 @@ namespace Metasound
 			// Returns the metadata for the current graph, including the name, description and author.
 			const FMetasoundFrontendClassMetadata& GetGraphMetadata() const override { return MetasoundFrontendInvalidControllerPrivate::GetInvalid<FMetasoundFrontendClassMetadata>(); }
 
-			// If the FNodeController given is itself a Metasound graph,
-			// and the FNodeController is a direct member of this FGraphController,
-			// If successful, this will invalidate the FNodeController and paste the graph for this node directly
-			// into the graph.
-			// If not successful, InNode will not be affected.
-			// @returns true on success, false on failure.
-			bool InflateNodeDirectlyIntoGraph(const INodeController& InNode) override { return false; }
-
 			TSharedRef<INodeController> CreateEmptySubgraph(const FMetasoundFrontendClassMetadata& InInfo) override { return FInvalidNodeController::GetInvalid(); }
 
 			TUniquePtr<IOperator> BuildOperator(const FOperatorSettings& InSettings, const FMetasoundEnvironment& InEnvironment, TArray<IOperatorBuilder::FBuildErrorPtr>& OutBuildErrors) const override
@@ -343,6 +337,7 @@ namespace Metasound
 				FConstClassAccessPtr FindOrAddClass(const FNodeClassInfo& InNodeClass) override { return FConstClassAccessPtr(); }
 				FConstClassAccessPtr FindClass(const FMetasoundFrontendClassMetadata& InMetadata) const override{ return FConstClassAccessPtr(); }
 				FConstClassAccessPtr FindOrAddClass(const FMetasoundFrontendClassMetadata& InMetadata) override{ return FConstClassAccessPtr(); }
+				FGraphHandle AddDuplicateSubgraph(const IGraphController& InGraph) override { return FInvalidGraphController::GetInvalid(); }
 
 				void SynchronizeDependencies() override { }
 
diff --git a/Engine/Plugins/Runtime/Metasound/Source/MetasoundFrontend/Private/MetasoundFrontendStandardController.cpp b/Engine/Plugins/Runtime/Metasound/Source/MetasoundFrontend/Private/MetasoundFrontendStandardController.cpp
--- a/Engine/Plugins/Runtime/Metasound/Source/MetasoundFrontend/Private/MetasoundFrontendStandardController.cpp
+++ b/Engine/Plugins/Runtime/Metasound/Source/MetasoundFrontend/Private/MetasoundFrontendStandardController.cpp
@@ -12,6 +12,7 @@
 #include "MetasoundFrontendDocumentAccessPtr.h"
 #include "MetasoundFrontendGraph.h"
 #include "MetasoundFrontendRegistries.h"
+#include "MetasoundFrontendSubgraphNodeController.h"
 #include "MetasoundJsonBackend.h"
 #include "MetasoundOperatorBuilder.h"
 #include "StructDeserializer.h"
@@ -53,6 +54,7 @@ namespace Metasound
 				static const ValueType Value = GetInvalidValue<ValueType>();
 				return Value;
 			}
+
 		}
 
 
@@ -122,7 +124,12 @@ namespace Metasound
 			return OwningNode;
 		}
 
-		TArray<FInputHandle> FBaseOutputController::GetCurrentlyConnectedInputs() 
+		bool FBaseOutputController::IsConnected() const 
+		{
+			return (FindEdges().Num() > 0);
+		}
+
+		TArray<FInputHandle> FBaseOutputController::GetConnectedInputs() 
 		{
 			TArray<FInputHandle> Inputs;
 
@@ -143,7 +150,7 @@ namespace Metasound
 			return Inputs;
 		}
 
-		TArray<FConstInputHandle> FBaseOutputController::GetCurrentlyConnectedInputs() const 
+		TArray<FConstInputHandle> FBaseOutputController::GetConstConnectedInputs() const 
 		{
 			TArray<FConstInputHandle> Inputs;
 
@@ -167,7 +174,7 @@ namespace Metasound
 		bool FBaseOutputController::Disconnect() 
 		{
 			bool bSuccess = true;
-			for (FInputHandle Input : GetCurrentlyConnectedInputs())
+			for (FInputHandle Input : GetConnectedInputs())
 			{
 				if (Input->IsValid())
 				{
@@ -504,7 +511,7 @@ namespace Metasound
 			return OwningNode;
 		}
 
-		FOutputHandle FBaseInputController::GetCurrentlyConnectedOutput()
+		FOutputHandle FBaseInputController::GetConnectedOutput()
 		{
 			if (const FMetasoundFrontendEdge* Edge = FindEdge())
 			{
@@ -517,7 +524,7 @@ namespace Metasound
 			return FInvalidOutputController::GetInvalid();
 		}
 
-		FConstOutputHandle FBaseInputController::GetCurrentlyConnectedOutput() const
+		FConstOutputHandle FBaseInputController::GetConnectedOutput() const
 		{
 			if (const FMetasoundFrontendEdge* Edge = FindEdge())
 			{
@@ -1240,7 +1247,7 @@ namespace Metasound
 			return Inputs;
 		}
 
-		TArray<FConstInputHandle> FBaseNodeController::GetInputsWithVertexName(const FString& InName) const 
+		TArray<FConstInputHandle> FBaseNodeController::GetConstInputsWithVertexName(const FString& InName) const 
 		{
 			TArray<FConstInputHandle> Inputs;
 
@@ -1277,7 +1284,7 @@ namespace Metasound
 			return Outputs;
 		}
 
-		TArray<FConstOutputHandle> FBaseNodeController::GetOutputsWithVertexName(const FString& InName) const
+		TArray<FConstOutputHandle> FBaseNodeController::GetConstOutputsWithVertexName(const FString& InName) const
 		{
 			TArray<FConstOutputHandle> Outputs;
 
@@ -1359,9 +1366,9 @@ namespace Metasound
 		{
 			TArray<FBaseNodeController::FInputControllerParams> Inputs;
 
-			if (NodePtr.IsValid())
+			if (const FMetasoundFrontendNode* Node = NodePtr.Get())
 			{
-				for (const FMetasoundFrontendVertex& NodeInputVertex : NodePtr->Interface.Inputs)
+				for (const FMetasoundFrontendVertex& NodeInputVertex : Node->Interface.Inputs)
 				{
 					FConstVertexAccessPtr NodeVertexPtr = NodePtr.GetInputWithName(NodeInputVertex.Name);
 					FConstClassInputAccessPtr ClassInputPtr = ClassPtr.GetInputWithName(NodeInputVertex.Name);
@@ -1377,14 +1384,17 @@ namespace Metasound
 		{
 			TArray<FBaseNodeController::FOutputControllerParams> Outputs;
 
-			for (const FMetasoundFrontendVertex& NodeOutputVertex : NodePtr->Interface.Outputs)
+			if (const FMetasoundFrontendNode* Node = NodePtr.Get())
 			{
-				const FString& VertexName = NodeOutputVertex.Name;
+				for (const FMetasoundFrontendVertex& NodeOutputVertex : Node->Interface.Outputs)
+				{
+					const FString& VertexName = NodeOutputVertex.Name;
 
-				FConstVertexAccessPtr NodeVertexPtr = NodePtr.GetOutputWithName(VertexName);
-				FConstClassOutputAccessPtr ClassOutputPtr = ClassPtr.GetOutputWithName(VertexName);
+					FConstVertexAccessPtr NodeVertexPtr = NodePtr.GetOutputWithName(VertexName);
+					FConstClassOutputAccessPtr ClassOutputPtr = ClassPtr.GetOutputWithName(VertexName);
 
-				Outputs.Add({NodeOutputVertex.VertexID, NodeVertexPtr, ClassOutputPtr});
+					Outputs.Add({NodeOutputVertex.VertexID, NodeVertexPtr, ClassOutputPtr});
+				}
 			}
 
 			return Outputs;
@@ -1452,54 +1462,6 @@ namespace Metasound
 			return false;
 		}
 
-		const FMetasoundFrontendClassInput* FBaseNodeController::FindClassInputWithName(const FString& InName) const
-		{
-			if (ClassPtr.IsValid())
-			{
-				auto IsClassInputWithName = [&](const FMetasoundFrontendClassInput& ClassInput) { return ClassInput.Name == InName; };
-
-				return ClassPtr->Interface.Inputs.FindByPredicate(IsClassInputWithName);
-			}
-
-			return nullptr;
-		}
-
-		const FMetasoundFrontendClassOutput* FBaseNodeController::FindClassOutputWithName(const FString& InName) const
-		{
-			if (ClassPtr.IsValid())
-			{
-				auto IsClassOutputWithName = [&](const FMetasoundFrontendClassOutput& ClassOutput) { return ClassOutput.Name == InName; };
-
-				return ClassPtr->Interface.Outputs.FindByPredicate(IsClassOutputWithName);
-			}
-
-			return nullptr;
-		}
-
-		const FMetasoundFrontendVertex* FBaseNodeController::FindNodeInputWithName(const FString& InName) const
-		{
-			if (NodePtr.IsValid())
-			{
-				auto IsNodeInputWithName = [&](const FMetasoundFrontendVertex& NodeInput) { return NodeInput.Name == InName; };
-
-				return NodePtr->Interface.Inputs.FindByPredicate(IsNodeInputWithName);
-			}
-
-			return nullptr;
-		}
-
-		const FMetasoundFrontendVertex* FBaseNodeController::FindNodeOutputWithName(const FString& InName) const
-		{
-			if (NodePtr.IsValid())
-			{
-				auto IsNodeOutputWithName = [&](const FMetasoundFrontendVertex& NodeOutput) { return NodeOutput.Name == InName; };
-
-				return NodePtr->Interface.Outputs.FindByPredicate(IsNodeOutputWithName);
-			}
-
-			return nullptr;
-		}
-
 		FGraphHandle FBaseNodeController::AsGraph()
 		{
 			// TODO: consider adding support for external graph owned in another document.
@@ -1525,6 +1487,7 @@ namespace Metasound
 			return FInvalidGraphController::GetInvalid();
 		}
 
+
 		FDocumentAccess FBaseNodeController::ShareAccess() 
 		{
 			FDocumentAccess Access;
@@ -1995,9 +1958,14 @@ namespace Metasound
 			return Metasound::FrontendInvalidID;
 		}
 
-		FGuid FGraphController::GetNewVertexID() const
+		const FText& FGraphController::GetDisplayName() const 
 		{
-			return FGuid::NewGuid();
+			if (const FMetasoundFrontendGraphClass* GraphClass = GraphClassPtr.Get())
+			{
+				return GraphClass->Metadata.DisplayName;
+			}
+
+			return FrontendControllerIntrinsics::GetInvalidValueConstRef<FText>();
 		}
 
 		TArray<FString> FGraphController::GetInputVertexNames() const 
@@ -2247,7 +2215,7 @@ namespace Metasound
 							FMetasoundFrontendNode& Node = GraphClass->Graph.Nodes.Emplace_GetRef(*InputClass);
 
 							Node.Name = InClassInput.Name;
-							Node.ID = NewNodeID();
+							Node.ID = FGuid::NewGuid();
 
 							auto IsVertexWithTypeName = [&](FMetasoundFrontendVertex& Vertex) { return Vertex.TypeName == InClassInput.TypeName; };
 							if (FMetasoundFrontendVertex* InputVertex = Node.Interface.Inputs.FindByPredicate(IsVertexWithTypeName))
@@ -2321,7 +2289,7 @@ namespace Metasound
 							FMetasoundFrontendNode& Node = GraphClass->Graph.Nodes.Emplace_GetRef(*OutputClass);
 
 							Node.Name = InClassOutput.Name;
-							Node.ID = NewNodeID();
+							Node.ID = FGuid::NewGuid();
 
 							// TODO: have something that checks if input node has valid interface.
 							auto IsVertexWithTypeName = [&](FMetasoundFrontendVertex& Vertex) { return Vertex.TypeName == InClassOutput.TypeName; };
@@ -2550,6 +2518,29 @@ namespace Metasound
 			return AddNode(FMetasoundFrontendRegistryContainer::GetRegistryKey(InClassMetadata));
 		}
 
+		FNodeHandle FGraphController::AddDuplicateNode(const INodeController& InNode)
+		{
+			const FMetasoundFrontendClassMetadata& ClassMetadata = InNode.GetClassMetadata();
+
+			if (EMetasoundFrontendClassType::Graph == ClassMetadata.Type)
+			{
+				// Add subgraph and dependencies if needed
+				if (!OwningDocument->FindClass(ClassMetadata).IsValid())
+				{
+					// Class does not exist, need to add the subgraph
+					OwningDocument->AddDuplicateSubgraph(*(InNode.AsGraph()));
+				}
+			}
+
+			// TODO: will need to copy node interface when dynamic pins exist.
+			if (FConstClassAccessPtr DependencyDescription = OwningDocument->FindOrAddClass(ClassMetadata))
+			{
+				return AddNode(DependencyDescription);
+			}
+
+			return INodeController::GetInvalidHandle();
+		}
+
 		// Remove the node corresponding to this node handle.
 		// On success, invalidates the received node handle.
 		bool FGraphController::RemoveNode(INodeController& InNode)
@@ -2598,14 +2589,6 @@ namespace Metasound
 			return FrontendControllerIntrinsics::GetInvalidValueConstRef<FMetasoundFrontendClassMetadata>();
 		}
 
-		bool FGraphController::InflateNodeDirectlyIntoGraph(const INodeController& InNode)
-		{
-			// TODO: implement
-			checkNoEntry();
-
-			return false;
-		}
-
 		FNodeHandle FGraphController::CreateEmptySubgraph(const FMetasoundFrontendClassMetadata& InInfo)
 		{
 			if (IsValid())
@@ -2653,8 +2636,8 @@ namespace Metasound
 
 			// Step 5: Invoke Operator Builder
 			FOperatorBuilder OperatorBuilder(FOperatorBuilderSettings::GetDefaultSettings());
-
-			return OperatorBuilder.BuildGraphOperator(*Graph, InSettings, InEnvironment, OutBuildErrors);
+			FBuildGraphParams BuildParams{*Graph, InSettings, FDataReferenceCollection{}, InEnvironment};
+			return OperatorBuilder.BuildGraphOperator(BuildParams, OutBuildErrors);
 		}
 
 		FDocumentHandle FGraphController::GetOwningDocument()
@@ -2675,7 +2658,7 @@ namespace Metasound
 				{
 					FMetasoundFrontendNode& Node = GraphClass->Graph.Nodes.Emplace_GetRef(*InExistingDependency);
 
-					Node.ID = NewNodeID();
+					Node.ID = FGuid::NewGuid();
 
 					FNodeAccessPtr NodePtr = GraphClassPtr.GetNodeWithNodeID(Node.ID);
 
@@ -2738,11 +2721,6 @@ namespace Metasound
 			return false;
 		}
 
-		FGuid FGraphController::NewNodeID() const
-		{
-			return FGuid::NewGuid();
-		}
-
 		bool FGraphController::ContainsNodesAndClassesByPredicate(TFunctionRef<bool (const FMetasoundFrontendClass&, const FMetasoundFrontendNode&)> InPredicate) const
 		{
 			if (GraphClassPtr.IsValid())
@@ -2974,6 +2952,19 @@ namespace Metasound
 							};
 							return FInputNodeController::CreateInputNodeHandle(InitParams);
 						}
+						else
+						{
+							// TODO: This supports input nodes introduced during subgraph inflation. Input nodes
+							// should be replaced with value nodes once they are implemented. 
+							FNodeController::FInitParams InitParams
+							{
+								InNodeAndClass.Node,
+								InNodeAndClass.Class,
+								GraphPtr,
+								OwningGraph
+							};
+							return FNodeController::CreateNodeHandle(InitParams);
+						}
 						break;
 
 					case EMetasoundFrontendClassType::Output:
@@ -2989,10 +2980,22 @@ namespace Metasound
 							};
 							return FOutputNodeController::CreateOutputNodeHandle(InitParams);
 						}
+						else
+						{
+							// TODO: This supports output nodes introduced during subgraph inflation. Output nodes
+							// should be replaced with value nodes once they are implemented. 
+							FNodeController::FInitParams InitParams
+							{
+								InNodeAndClass.Node,
+								InNodeAndClass.Class,
+								GraphPtr,
+								OwningGraph
+							};
+							return FNodeController::CreateNodeHandle(InitParams);
+						}
 						break;
 
 					case EMetasoundFrontendClassType::External:
-					case EMetasoundFrontendClassType::Graph:
 						{
 							FNodeController::FInitParams InitParams
 							{
@@ -3005,6 +3008,19 @@ namespace Metasound
 						}
 						break;
 
+					case EMetasoundFrontendClassType::Graph:
+						{
+							FSubgraphNodeController::FInitParams InitParams
+							{
+								InNodeAndClass.Node,
+								InNodeAndClass.Class,
+								GraphPtr,
+								OwningGraph
+							};
+							return FSubgraphNodeController::CreateNodeHandle(InitParams);
+						}
+						break;
+
 					default:
 						checkNoEntry();
 				}
@@ -3246,6 +3262,104 @@ namespace Metasound
 			return Classes;
 		}
 
+		bool FDocumentController::AddDuplicateSubgraph(const FMetasoundFrontendGraphClass& InGraphToCopy, const FMetasoundFrontendDocument& InOtherDocument)
+		{
+			FMetasoundFrontendDocument* ThisDocument = DocumentPtr.Get();
+			if (nullptr == ThisDocument)
+			{
+				UE_LOG(LogMetaSound, Error, TEXT("Cannot add subgraph to invalid document"));
+				return false;
+			}
+
+			// Direct copy of subgraph
+			bool bSuccess = true;
+			FMetasoundFrontendGraphClass SubgraphCopy(InGraphToCopy);
+
+			for (FMetasoundFrontendNode& Node : SubgraphCopy.Graph.Nodes)
+			{
+				const FGuid OriginalClassID = Node.ClassID;
+
+				auto IsClassWithClassID = [&](const FMetasoundFrontendClass& InClass) -> bool
+				{
+					return InClass.ID == OriginalClassID;
+				};
+
+				if (const FMetasoundFrontendClass* OriginalNodeClass = InOtherDocument.Dependencies.FindByPredicate(IsClassWithClassID))
+				{
+					// Should not be a graph class since it's in the dependencies list
+					check(EMetasoundFrontendClassType::Graph != OriginalNodeClass->Metadata.Type);
+
+					if (const FMetasoundFrontendClass* NodeClass = FindOrAddClass(OriginalNodeClass->Metadata).Get())
+					{
+						// All this just to update this ID. Maybe having globally 
+						// consistent class IDs would help. Or using the classname & version as
+						// a class ID. 
+						Node.ClassID = NodeClass->ID;
+					}
+					else
+					{
+						UE_LOG(LogMetaSound, Error, TEXT("Failed to add subgraph dependency [Class:%s]"), *OriginalNodeClass->Metadata.ClassName.ToString());
+						bSuccess = false;
+					}
+				}
+				else if (const FMetasoundFrontendGraphClass* OriginalNodeGraphClass = InOtherDocument.Subgraphs.FindByPredicate(IsClassWithClassID))
+				{
+					bSuccess = bSuccess && AddDuplicateSubgraph(*OriginalNodeGraphClass, InOtherDocument);
+					if (!bSuccess)
+					{
+						break;
+					}
+				}
+				else
+				{
+					bSuccess = false;
+					UE_LOG(LogMetaSound, Error, TEXT("Failed to copy subgraph. Subgraph document is missing dependency info for node [Node:%s, NodeID:%s]"), *Node.Name, *Node.ID.ToString());
+				}
+			}
+
+			if (bSuccess)
+			{
+				ThisDocument->Subgraphs.Add(SubgraphCopy);
+			}
+
+			return bSuccess;
+		}
+
+		FGraphHandle FDocumentController::AddDuplicateSubgraph(const IGraphController& InGraph)
+		{
+			// TODO: class IDs are fucked. 
+			// Maybe class IDs should just be the linking info and not a GUID.
+			// ClassIDs are just internal references, so they need to be fixed up
+			// here if swapping documents. 
+
+			FConstDocumentAccess GraphDocumentAccess = GetSharedAccess(*InGraph.GetOwningDocument());
+			const FMetasoundFrontendDocument* OtherDocument = GraphDocumentAccess.ConstDocument.Get();
+			if (nullptr == OtherDocument)
+			{
+				UE_LOG(LogMetaSound, Error, TEXT("Cannot add subgraph from invalid document"));
+				return IGraphController::GetInvalidHandle();
+			}
+
+			FConstDocumentAccess GraphAccess = GetSharedAccess(InGraph);
+			const FMetasoundFrontendGraphClass* OtherGraph = GraphAccess.ConstGraphClass.Get();
+			if (nullptr == OtherGraph)
+			{
+				UE_LOG(LogMetaSound, Error, TEXT("Cannot add invalid subgraph to document"));
+				return IGraphController::GetInvalidHandle();
+			}
+
+			if (AddDuplicateSubgraph(*OtherGraph, *OtherDocument))
+			{
+				if (const FMetasoundFrontendClass* SubgraphClass = FindClass(OtherGraph->Metadata).Get())
+				{
+					return GetSubgraphWithClassID(SubgraphClass->ID);
+				}
+			}
+
+			return IGraphController::GetInvalidHandle();
+		}
+ 
+
 		FConstClassAccessPtr FDocumentController::FindDependencyWithID(FGuid InClassID) const 
 		{
 			return DocumentPtr.GetDependencyWithID(InClassID);
@@ -3282,9 +3396,9 @@ namespace Metasound
 				FConstClassAccessPtr NewClassPtr;
 				if (DocumentPtr.IsValid())
 				{
-					// FNodeClassInfo does not contain enough info add a subgraph.
+					// FNodeClassInfo does not contain enough info to add a subgraph.
 					check(EMetasoundFrontendClassType::Graph != InNodeClass.NodeType);
-					NewClassDescription.ID = NewClassID();
+					NewClassDescription.ID = FGuid::NewGuid();
 
 					DocumentPtr->Dependencies.Add(MoveTemp(NewClassDescription));
 					NewClassPtr = FindClass(InNodeClass);
@@ -3344,7 +3458,7 @@ namespace Metasound
 					FMetasoundFrontendClass NewClass;
 					if (FMetasoundFrontendRegistryContainer::GetFrontendClassFromRegistered(InMetadata, NewClass))
 					{
-						NewClass.ID = NewClassID();
+						NewClass.ID = FGuid::NewGuid();
 						DocumentPtr->Dependencies.Add(NewClass);
 					}
 					else
@@ -3361,7 +3475,7 @@ namespace Metasound
 				else if (EMetasoundFrontendClassType::Graph == InMetadata.Type)
 				{
 					FMetasoundFrontendGraphClass NewClass;
-					NewClass.ID = NewClassID();
+					NewClass.ID = FGuid::NewGuid();
 					NewClass.Metadata = InMetadata;
 
 					DocumentPtr->Subgraphs.Add(NewClass);
@@ -3584,10 +3698,6 @@ namespace Metasound
 			return Access;
 		}
 
-		FGuid FDocumentController::NewClassID() const 
-		{
-			return FGuid::NewGuid();
-		}
 	}
 }
 #undef LOCTEXT_NAMESPACE
diff --git a/Engine/Plugins/Runtime/Metasound/Source/MetasoundFrontend/Private/MetasoundFrontendStandardController.h b/Engine/Plugins/Runtime/Metasound/Source/MetasoundFrontend/Private/MetasoundFrontendStandardController.h
--- a/Engine/Plugins/Runtime/Metasound/Source/MetasoundFrontend/Private/MetasoundFrontendStandardController.h
+++ b/Engine/Plugins/Runtime/Metasound/Source/MetasoundFrontend/Private/MetasoundFrontendStandardController.h
@@ -19,7 +19,7 @@ namespace Metasound
 		/** FBaseOutputController provides common functionality for multiple derived
 		 * output controllers.
 		 */
-		class METASOUNDFRONTEND_API FBaseOutputController : public IOutputController
+		class FBaseOutputController : public IOutputController
 		{
 		public:
 
@@ -56,8 +56,9 @@ namespace Metasound
 			FNodeHandle GetOwningNode() override;
 			FConstNodeHandle GetOwningNode() const override;
 
-			TArray<FInputHandle> GetCurrentlyConnectedInputs() override;
-			TArray<FConstInputHandle> GetCurrentlyConnectedInputs() const override;
+			bool IsConnected() const override;
+			TArray<FInputHandle> GetConnectedInputs() override;
+			TArray<FConstInputHandle> GetConstConnectedInputs() const override;
 			bool Disconnect() override;
 
 			// Connection logic.
@@ -89,7 +90,7 @@ namespace Metasound
 		 * FInputNodeOutputController is largely to represent inputs coming into 
 		 * graph. 
 		 */
-		class METASOUNDFRONTEND_API FInputNodeOutputController : public FBaseOutputController
+		class FInputNodeOutputController : public FBaseOutputController
 		{
 		public:
 			struct FInitParams
@@ -133,7 +134,7 @@ namespace Metasound
 		 * FOutputNodeOutputController is largely to represent inputs coming into 
 		 * graph. 
 		 */
-		class METASOUNDFRONTEND_API FOutputNodeOutputController : public FBaseOutputController
+		class FOutputNodeOutputController : public FBaseOutputController
 		{
 		public:
 			struct FInitParams
@@ -173,7 +174,7 @@ namespace Metasound
 		/** FBaseInputController provides common functionality for multiple derived
 		 * input controllers.
 		 */
-		class METASOUNDFRONTEND_API FBaseInputController : public IInputController 
+		class FBaseInputController : public IInputController 
 		{
 		public:
 
@@ -211,8 +212,8 @@ namespace Metasound
 
 			// Connection info
 			bool IsConnected() const override;
-			FOutputHandle GetCurrentlyConnectedOutput() override;
-			FConstOutputHandle GetCurrentlyConnectedOutput() const override;
+			FOutputHandle GetConnectedOutput() override;
+			FConstOutputHandle GetConnectedOutput() const override;
 
 			FConnectability CanConnectTo(const IOutputController& InController) const override;
 			bool Connect(IOutputController& InController) override;
@@ -246,7 +247,7 @@ namespace Metasound
 		 * FOutputNodeInputController is largely to represent outputs exposed from
 		 * a graph. 
 		 */
-		class METASOUNDFRONTEND_API FOutputNodeInputController : public FBaseInputController 
+		class FOutputNodeInputController : public FBaseInputController 
 		{
 		public:
 			struct FInitParams
@@ -289,7 +290,7 @@ namespace Metasound
 		 * FInputNodeInputController is largely to represent outputs exposed from
 		 * a graph. 
 		 */
-		class METASOUNDFRONTEND_API FInputNodeInputController : public FBaseInputController 
+		class FInputNodeInputController : public FBaseInputController 
 		{
 		public:
 			struct FInitParams
@@ -327,7 +328,7 @@ namespace Metasound
 		/** FBaseNodeController provides common functionality for multiple derived
 		 * node controllers.
 		 */
-		class METASOUNDFRONTEND_API FBaseNodeController : public INodeController
+		class FBaseNodeController : public INodeController
 		{
 		public:
 
@@ -366,17 +367,6 @@ namespace Metasound
 
 			const FString& GetNodeName() const override;
 
-			bool CanAddInput(const FString& InVertexName) const override;
-			FInputHandle AddInput(const FString& InVertexName, const FMetasoundFrontendLiteral* InDefault) override;
-			bool RemoveInput(FGuid InVertexID) override;
-
-			bool CanAddOutput(const FString& InVertexName) const override;
-			FInputHandle AddOutput(const FString& InVertexName, const FMetasoundFrontendLiteral* InDefault) override;
-			bool RemoveOutput(FGuid InVertexID) override;
-
-			/** Returns all node inputs. */
-			TArray<FInputHandle> GetInputs() override;
-
 			/** Returns the human-readable display name of the given node. */
 			const FText& GetDisplayName() const override;
 
@@ -389,6 +379,17 @@ namespace Metasound
 			/** Returns the title of the given node (what to label in visual node). */
 			const FText& GetDisplayTitle() const override;
 
+			bool CanAddInput(const FString& InVertexName) const override;
+			FInputHandle AddInput(const FString& InVertexName, const FMetasoundFrontendLiteral* InDefault) override;
+			bool RemoveInput(FGuid InVertexID) override;
+
+			bool CanAddOutput(const FString& InVertexName) const override;
+			FInputHandle AddOutput(const FString& InVertexName, const FMetasoundFrontendLiteral* InDefault) override;
+			bool RemoveOutput(FGuid InVertexID) override;
+
+			/** Returns all node inputs. */
+			TArray<FInputHandle> GetInputs() override;
+
 			/** Returns all node inputs. */
 			TArray<FConstInputHandle> GetConstInputs() const override;
 
@@ -402,7 +403,7 @@ namespace Metasound
 			int32 GetNumOutputs() const override;
 
 			TArray<FInputHandle> GetInputsWithVertexName(const FString& InName) override;
-			TArray<FConstInputHandle> GetInputsWithVertexName(const FString& InName) const override;
+			TArray<FConstInputHandle> GetConstInputsWithVertexName(const FString& InName) const override;
 
 			/** Returns all node outputs. */
 			TArray<FOutputHandle> GetOutputs() override;
@@ -411,7 +412,7 @@ namespace Metasound
 			TArray<FConstOutputHandle> GetConstOutputs() const override;
 
 			TArray<FOutputHandle> GetOutputsWithVertexName(const FString& InName) override;
-			TArray<FConstOutputHandle> GetOutputsWithVertexName(const FString& InName) const override;
+			TArray<FConstOutputHandle> GetConstOutputsWithVertexName(const FString& InName) const override;
 
 			bool IsRequired() const override;
 
@@ -451,8 +452,6 @@ namespace Metasound
 			FConstClassAccessPtr ClassPtr;
 			FGraphHandle OwningGraph;
 
-		private:
-
 			struct FInputControllerParams
 			{
 				FGuid VertexID;
@@ -467,27 +466,23 @@ namespace Metasound
 				FConstClassOutputAccessPtr ClassOutputPtr;
 			};
 
-			TArray<FInputControllerParams> GetInputControllerParams() const;
-			TArray<FOutputControllerParams> GetOutputControllerParams() const;
-
-			TArray<FInputControllerParams> GetInputControllerParamsWithVertexName(const FString& InName) const;
-			TArray<FOutputControllerParams> GetOutputControllerParamsWithVertexName(const FString& InName) const;
+			virtual TArray<FInputControllerParams> GetInputControllerParams() const;
+			virtual TArray<FOutputControllerParams> GetOutputControllerParams() const;
 
-			bool FindInputControllerParamsWithID(FGuid InVertexID, FInputControllerParams& OutParams) const;
-			bool FindOutputControllerParamsWithID(FGuid InVertexID, FOutputControllerParams& OutParams) const;
+			virtual TArray<FInputControllerParams> GetInputControllerParamsWithVertexName(const FString& InName) const;
+			virtual TArray<FOutputControllerParams> GetOutputControllerParamsWithVertexName(const FString& InName) const;
 
-			const FMetasoundFrontendClassInput* FindClassInputWithName(const FString& InName) const;
-			const FMetasoundFrontendClassOutput* FindClassOutputWithName(const FString& InName) const;
+			virtual bool FindInputControllerParamsWithID(FGuid InVertexID, FInputControllerParams& OutParams) const;
+			virtual bool FindOutputControllerParamsWithID(FGuid InVertexID, FOutputControllerParams& OutParams) const;
 
-			const FMetasoundFrontendVertex* FindNodeInputWithName(const FString& InName) const;
-			const FMetasoundFrontendVertex* FindNodeOutputWithName(const FString& InName) const;
+		private:
 
 			virtual FInputHandle CreateInputController(FGuid InVertexID, FConstVertexAccessPtr InNodeVertexPtr, FConstClassInputAccessPtr InClassInputPtr, FNodeHandle InOwningNode) const = 0;
 			virtual FOutputHandle CreateOutputController(FGuid InVertexID, FConstVertexAccessPtr InNodeVertexPtr, FConstClassOutputAccessPtr InClassOutputPtr, FNodeHandle InOwningNode) const = 0;
 		};
 
 		/** FNodeController represents a external or subgraph node. */
-		class METASOUNDFRONTEND_API FNodeController : public FBaseNodeController
+		class FNodeController : public FBaseNodeController
 		{
 			// Private token only allows members or friends to call constructor.
 			enum EPrivateToken { Token };
@@ -537,7 +532,7 @@ namespace Metasound
 		};
 
 		/** FOutputNodeController represents an output node. */
-		class METASOUNDFRONTEND_API FOutputNodeController: public FBaseNodeController
+		class FOutputNodeController: public FBaseNodeController
 		{
 			// Private token only allows members or friends to call constructor.
 			enum EPrivateToken { Token };
@@ -598,7 +593,7 @@ namespace Metasound
 		};
 
 		/** FInputNodeController represents an input node. */
-		class METASOUNDFRONTEND_API FInputNodeController: public FBaseNodeController
+		class FInputNodeController: public FBaseNodeController
 		{
 			// Private token only allows members or friends to call constructor.
 			enum EPrivateToken { Token };
@@ -659,7 +654,7 @@ namespace Metasound
 		};
 
 		/** FGraphController represents a Metasound graph class. */
-		class METASOUNDFRONTEND_API FGraphController : public IGraphController
+		class FGraphController : public IGraphController
 		{
 
 			// Private token only allows members or friends to call constructor.
@@ -697,8 +692,7 @@ namespace Metasound
 			bool IsValid() const override;
 
 			FGuid GetClassID() const override;
-
-			FGuid GetNewVertexID() const override;
+			const FText& GetDisplayName() const override;
 
 			TArray<FString> GetInputVertexNames() const override;
 			TArray<FString> GetOutputVertexNames() const override;
@@ -776,6 +770,7 @@ namespace Metasound
 			FNodeHandle AddNode(const FNodeClassInfo& InNodeClass) override;
 			FNodeHandle AddNode(const FNodeRegistryKey& InNodeClass) override;
 			FNodeHandle AddNode(const FMetasoundFrontendClassMetadata& InClassMetadata) override;
+			FNodeHandle AddDuplicateNode(const INodeController& InNode) override;
 
 			// Remove the node corresponding to this node handle.
 			// On success, invalidates the received node handle.
@@ -784,8 +779,6 @@ namespace Metasound
 			// Returns the metadata for the current graph, including the name, description and author.
 			const FMetasoundFrontendClassMetadata& GetGraphMetadata() const override;
 
-			bool InflateNodeDirectlyIntoGraph(const INodeController& InNode) override;
-
 			FNodeHandle CreateEmptySubgraph(const FMetasoundFrontendClassMetadata& InInfo) override;
 
 			TUniquePtr<IOperator> BuildOperator(const FOperatorSettings& InSettings, const FMetasoundEnvironment& InEnvironment, TArray<IOperatorBuilder::FBuildErrorPtr>& OutBuildErrors) const override;
@@ -808,8 +801,6 @@ namespace Metasound
 			bool RemoveInput(const FMetasoundFrontendNode& InNode);
 			bool RemoveOutput(const FMetasoundFrontendNode& InNode);
 
-			FGuid NewNodeID() const;
-
 			FNodeHandle GetNodeByPredicate(TFunctionRef<bool (const FMetasoundFrontendClass&, const FMetasoundFrontendNode&)> InPredicate);
 			FConstNodeHandle GetNodeByPredicate(TFunctionRef<bool (const FMetasoundFrontendClass&, const FMetasoundFrontendNode&)> InPredicate) const;
 
@@ -873,7 +864,7 @@ namespace Metasound
 		};
 
 		/** FDocumentController represents an entire Metasound document. */
-		class METASOUNDFRONTEND_API FDocumentController : public IDocumentController 
+		class FDocumentController : public IDocumentController 
 		{
 			using FRegistry = FMetasoundFrontendRegistryContainer;
 
@@ -916,6 +907,8 @@ namespace Metasound
 			FConstClassAccessPtr FindOrAddClass(const FNodeClassInfo& InNodeClass) override;
 			FConstClassAccessPtr FindOrAddClass(const FMetasoundFrontendClassMetadata& InMetadata) override;
 
+			virtual FGraphHandle AddDuplicateSubgraph(const IGraphController& InGraph) override;
+
 			void SynchronizeDependencies() override;
 
 			FGraphHandle GetRootGraph() override;
@@ -944,7 +937,7 @@ namespace Metasound
 			FConstDocumentAccess ShareAccess() const override;
 		private:
 
-			FGuid NewClassID() const;
+			bool AddDuplicateSubgraph(const FMetasoundFrontendGraphClass& InGraphToCopy, const FMetasoundFrontendDocument& InOtherDocument);
 
 			FDocumentAccessPtr DocumentPtr;
 		};
diff --git a/Engine/Plugins/Runtime/Metasound/Source/MetasoundFrontend/Private/MetasoundFrontendSubgraphNodeController.cpp b/Engine/Plugins/Runtime/Metasound/Source/MetasoundFrontend/Private/MetasoundFrontendSubgraphNodeController.cpp
new file mode 100644
--- /dev/null
+++ b/Engine/Plugins/Runtime/Metasound/Source/MetasoundFrontend/Private/MetasoundFrontendSubgraphNodeController.cpp
@@ -0,0 +1,234 @@
+// Copyright Epic Games, Inc. All Rights Reserved.
+
+#include "MetasoundFrontendSubgraphNodeController.h"
+
+#include "Algo/Transform.h"
+
+namespace Metasound
+{
+	namespace Frontend
+	{
+		namespace FrontendSubgraphNodeControllerPrivate
+		{
+			template<typename SetAType, typename SetBType>
+			void InplaceBidirectionalVertexSetDifference(TArray<const SetAType*>& SetA, TArray<const SetBType*>& SetB)
+			{
+				// TODO: comparison is done by `Name`. May update to a `ClassVertexID` in the future. 
+				// Reverse iterate to allow removal of elements from SetA within loop.
+				for (int32 SetAIndex = (SetA.Num() - 1); SetAIndex >= 0; SetAIndex--) 
+				{
+					const FString& SetAName = SetA[SetAIndex]->Name;
+					int32 SetBIndex = SetB.IndexOfByPredicate([&](const SetBType* InVertex) { return InVertex->Name == SetAName; });
+					if (INDEX_NONE != SetBIndex)
+					{
+						SetA.RemoveAt(SetAIndex);
+						SetB.RemoveAt(SetBIndex);
+					}
+				}
+			}
+
+			// Manipulates the NodeVertexArray to match the ClassVertexArray.
+			template<typename NodeVertexType, typename ClassVertexType>
+			bool ConformNodeVertexArrayToClassVertexArray(TArray<NodeVertexType>& NodeVertexArray, const TArray<ClassVertexType>& ClassVertexArray)
+			{
+				// TODO: May need to change this to support dynamic pins. 
+				bool bDidAlterNodeVertexArray = false;
+
+				TArray<const NodeVertexType*> NodeVerticesToRemove;
+				Algo::Transform(NodeVertexArray, NodeVerticesToRemove, [](const NodeVertexType& Vertex) { return &Vertex; });
+
+				TArray<const ClassVertexType*> ClassVerticesToAdd;
+				Algo::Transform(ClassVertexArray, ClassVerticesToAdd, [](const ClassVertexType& Vertex) { return &Vertex; });
+			
+				FrontendSubgraphNodeControllerPrivate::InplaceBidirectionalVertexSetDifference(NodeVerticesToRemove, ClassVerticesToAdd);
+
+				bDidAlterNodeVertexArray = (NodeVerticesToRemove.Num() > 0) || (ClassVerticesToAdd.Num() > 0);
+
+				// Remove extra node vertices.
+				for (const NodeVertexType* VertexToRemove : NodeVerticesToRemove)
+				{
+					NodeVertexArray.RemoveAll([&](const NodeVertexType& InVertex) { return InVertex.VertexID == VertexToRemove->VertexID; });
+				}
+
+				// Add missing node vertices. 
+				for (const ClassVertexType* VertexToAdd : ClassVerticesToAdd)
+				{
+					NodeVertexType NewNodeVertex(*VertexToAdd);
+					NewNodeVertex.VertexID = FGuid::NewGuid();
+					NodeVertexArray.Add(NewNodeVertex);
+				}
+
+				return bDidAlterNodeVertexArray;
+			}
+		}
+
+		// 
+		// FSubgraphNodeController
+		//
+		FSubgraphNodeController::FSubgraphNodeController(EPrivateToken InToken, const FSubgraphNodeController::FInitParams& InParams)
+		: FBaseNodeController({InParams.NodePtr, InParams.ClassPtr, InParams.OwningGraph})
+		, GraphPtr(InParams.GraphPtr)
+		{
+		}
+
+		FNodeHandle FSubgraphNodeController::CreateNodeHandle(const FSubgraphNodeController::FInitParams& InParams)
+		{
+			if (InParams.NodePtr.IsValid() && InParams.ClassPtr.IsValid())
+			{
+				// Cannot make a valid node handle if the node description and class description differ
+				if (InParams.NodePtr->ClassID == InParams.ClassPtr->ID)
+				{
+					return MakeShared<FSubgraphNodeController>(EPrivateToken::Token, InParams);
+				}
+				else
+				{
+					UE_LOG(LogMetaSound, Warning, TEXT("Frontend Node [NodeID:%s, ClassID:%s] is not of expected class class [ClassID:%s]"), *InParams.NodePtr->ID.ToString(), *InParams.NodePtr->ClassID.ToString(), *InParams.ClassPtr->ID.ToString());
+				}
+			}
+			return FInvalidNodeController::GetInvalid();
+		}
+
+		FConstNodeHandle FSubgraphNodeController::CreateConstNodeHandle(const FSubgraphNodeController::FInitParams& InParams)
+		{
+			if (InParams.NodePtr.IsValid() && InParams.ClassPtr.IsValid())
+			{
+				// Cannot make a valid node handle if the node description and class description differ
+				if (InParams.NodePtr->ClassID == InParams.ClassPtr->ID)
+				{
+					return MakeShared<const FSubgraphNodeController>(EPrivateToken::Token, InParams);
+				}
+				else
+				{
+					UE_LOG(LogMetaSound, Warning, TEXT("Frontend Node [NodeID:%s, ClassID:%s] is not of expected class class [ClassID:%s]"), *InParams.NodePtr->ID.ToString(), *InParams.NodePtr->ClassID.ToString(), *InParams.ClassPtr->ID.ToString());
+				}
+			}
+			return FInvalidNodeController::GetInvalid();
+		}
+
+		bool FSubgraphNodeController::IsValid() const
+		{
+			return FBaseNodeController::IsValid() && GraphPtr.IsValid();
+		}
+
+		int32 FSubgraphNodeController::GetNumInputs() const
+		{
+			// TODO: Trigger ConformNodeInterfaceToClassInterface() on an
+			// event callback rather than on every call to this function.
+			const_cast<FSubgraphNodeController*>(this)->ConformNodeInterfaceToClassInterface();
+
+			return Super::GetNumInputs();
+		}
+
+		int32 FSubgraphNodeController::GetNumOutputs() const
+		{
+			// TODO: Trigger ConformNodeInterfaceToClassInterface() on an
+			// event callback rather than on every call to this function.
+			const_cast<FSubgraphNodeController*>(this)->ConformNodeInterfaceToClassInterface();
+
+			return Super::GetNumOutputs();
+		}
+
+		TArray<FBaseNodeController::FInputControllerParams> FSubgraphNodeController::GetInputControllerParams() const
+		{
+			// TODO: Trigger ConformNodeInterfaceToClassInterface() on an
+			// event callback rather than on every call to this function.
+			const_cast<FSubgraphNodeController*>(this)->ConformNodeInterfaceToClassInterface();
+
+			return Super::GetInputControllerParams();
+		}
+
+		TArray<FBaseNodeController::FOutputControllerParams> FSubgraphNodeController::GetOutputControllerParams() const
+		{
+			// TODO: Trigger ConformNodeInterfaceToClassInterface() on an
+			// event callback rather than on every call to this function.
+			const_cast<FSubgraphNodeController*>(this)->ConformNodeInterfaceToClassInterface();
+
+			return Super::GetOutputControllerParams();
+		}
+
+		TArray<FBaseNodeController::FInputControllerParams> FSubgraphNodeController::GetInputControllerParamsWithVertexName(const FString& InName) const
+		{
+			// TODO: Trigger ConformNodeInterfaceToClassInterface() on an
+			// event callback rather than on every call to this function.
+			const_cast<FSubgraphNodeController*>(this)->ConformNodeInterfaceToClassInterface();
+
+			return Super::GetInputControllerParamsWithVertexName(InName);
+		}
+
+		TArray<FBaseNodeController::FOutputControllerParams> FSubgraphNodeController::GetOutputControllerParamsWithVertexName(const FString& InName) const
+		{
+			TArray<FBaseNodeController::FOutputControllerParams> Outputs;
+
+			// TODO: Trigger ConformNodeInterfaceToClassInterface() on an
+			// event callback rather than on every call to this function.
+			const_cast<FSubgraphNodeController*>(this)->ConformNodeInterfaceToClassInterface();
+
+			return Super::GetOutputControllerParamsWithVertexName(InName);
+		}
+
+		bool FSubgraphNodeController::FindInputControllerParamsWithID(FGuid InVertexID, FInputControllerParams& OutParams) const
+		{
+			// TODO: Trigger ConformNodeInterfaceToClassInterface() on an
+			// event callback rather than on every call to this function.
+			const_cast<FSubgraphNodeController*>(this)->ConformNodeInterfaceToClassInterface();
+
+			return Super::FindInputControllerParamsWithID(InVertexID, OutParams);
+		}
+
+		bool FSubgraphNodeController::FindOutputControllerParamsWithID(FGuid InVertexID, FOutputControllerParams& OutParams) const
+		{
+			// TODO: Trigger ConformNodeInterfaceToClassInterface() on an
+			// event callback rather than on every call to this function.
+			const_cast<FSubgraphNodeController*>(this)->ConformNodeInterfaceToClassInterface();
+
+			return Super::FindOutputControllerParamsWithID(InVertexID, OutParams);
+		}
+
+		void FSubgraphNodeController::ConformNodeInterfaceToClassInterface() 
+		{
+			using namespace FrontendSubgraphNodeControllerPrivate;
+
+			// TODO: will need to update this logic to handle external classes and 
+			// dynamic interfaces. These updates will get piped down to the INodeClass
+			// interface.
+			if (FMetasoundFrontendNode* Node = NodePtr.Get())
+			{
+				if (const FMetasoundFrontendClass* NodeClass = ClassPtr.Get())
+				{
+					ConformNodeVertexArrayToClassVertexArray(Node->Interface.Inputs, NodeClass->Interface.Inputs);
+					ConformNodeVertexArrayToClassVertexArray(Node->Interface.Outputs, NodeClass->Interface.Outputs);
+					// TODO: Conform environment variables by aggregating all dependent environment variables.
+				}
+			}
+		}
+
+		FDocumentAccess FSubgraphNodeController::ShareAccess() 
+		{
+			FDocumentAccess Access = FBaseNodeController::ShareAccess();
+
+			Access.Graph = GraphPtr;
+			Access.ConstGraph = GraphPtr;
+
+			return Access;
+		}
+
+		FConstDocumentAccess FSubgraphNodeController::ShareAccess() const
+		{
+			FConstDocumentAccess Access = FBaseNodeController::ShareAccess();
+
+			Access.ConstGraph = GraphPtr;
+
+			return Access;
+		}
+
+		FInputHandle FSubgraphNodeController::CreateInputController(FGuid InVertexID, FConstVertexAccessPtr InNodeVertexPtr, FConstClassInputAccessPtr InClassInputPtr, FNodeHandle InOwningNode) const
+		{
+			return MakeShared<FBaseInputController>(FBaseInputController::FInitParams{InVertexID, InNodeVertexPtr, InClassInputPtr, GraphPtr, InOwningNode});
+		}
+
+		FOutputHandle FSubgraphNodeController::CreateOutputController(FGuid InVertexID, FConstVertexAccessPtr InNodeVertexPtr, FConstClassOutputAccessPtr InClassOutputPtr, FNodeHandle InOwningNode) const
+		{
+			return MakeShared<FBaseOutputController>(FBaseOutputController::FInitParams{InVertexID, InNodeVertexPtr, InClassOutputPtr, GraphPtr, InOwningNode});
+		}
+	}
+}
diff --git a/Engine/Plugins/Runtime/Metasound/Source/MetasoundFrontend/Private/MetasoundFrontendSubgraphNodeController.h b/Engine/Plugins/Runtime/Metasound/Source/MetasoundFrontend/Private/MetasoundFrontendSubgraphNodeController.h
new file mode 100644
--- /dev/null
+++ b/Engine/Plugins/Runtime/Metasound/Source/MetasoundFrontend/Private/MetasoundFrontendSubgraphNodeController.h
@@ -0,0 +1,79 @@
+// Copyright Epic Games, Inc. All Rights Reserved.
+
+#pragma once
+
+#include "MetasoundFrontendStandardController.h"
+
+namespace Metasound
+{
+	namespace Frontend
+	{
+		/** Controller for nodes representing subgraphs. Handles synchronization
+		 * of node interface to subgraph interface. 
+		 */
+		class FSubgraphNodeController : public FBaseNodeController
+		{
+			// Private token only allows members or friends to call constructor.
+			enum EPrivateToken { Token };
+
+			using Super = FBaseNodeController;
+
+		public:
+			struct FInitParams
+			{
+				FNodeAccessPtr NodePtr;
+				FConstClassAccessPtr ClassPtr;
+				FGraphAccessPtr GraphPtr;
+				FGraphHandle OwningGraph;
+			};
+
+			// Constructor takes a private token so it can only be instantiated by
+			// using the static creation functions. This protects against some
+			// error conditions which would result in a zombie object. The creation
+			// methods can detect the error conditions and return an invalid controller
+			// on error
+			FSubgraphNodeController(EPrivateToken InToken, const FInitParams& InParams);
+
+			/** Create a node handle for a external or subgraph node. 
+			 *
+			 * @return A Node handle. On error, an invalid node handle is returned. 
+			 */
+			static FNodeHandle CreateNodeHandle(const FInitParams& InParams);
+
+			/** Create a node handle for a external or subgraph node. 
+			 *
+			 * @return A Node handle. On error, an invalid node handle is returned. 
+			 */
+			static FConstNodeHandle CreateConstNodeHandle(const FInitParams& InParams);
+
+			virtual ~FSubgraphNodeController() = default;
+
+			bool IsValid() const override;
+
+			int32 GetNumInputs() const override;
+
+			int32 GetNumOutputs() const override;
+
+		protected:
+			FDocumentAccess ShareAccess() override;
+			FConstDocumentAccess ShareAccess() const override;
+
+			TArray<FBaseNodeController::FInputControllerParams> GetInputControllerParams() const override;
+			TArray<FBaseNodeController::FOutputControllerParams> GetOutputControllerParams() const override;
+
+			TArray<FBaseNodeController::FInputControllerParams> GetInputControllerParamsWithVertexName(const FString& InName) const override;
+			TArray<FBaseNodeController::FOutputControllerParams> GetOutputControllerParamsWithVertexName(const FString& InName) const override;
+
+			bool FindInputControllerParamsWithID(FGuid InVertexID, FInputControllerParams& OutParams) const override;
+			bool FindOutputControllerParamsWithID(FGuid InVertexID, FOutputControllerParams& OutParams) const override;
+
+		private:
+			void ConformNodeInterfaceToClassInterface();
+
+			FInputHandle CreateInputController(FGuid InVertexID, FConstVertexAccessPtr InNodeVertexPtr, FConstClassInputAccessPtr InClassInputPtr, FNodeHandle InOwningNode) const override;
+			FOutputHandle CreateOutputController(FGuid InVertexID, FConstVertexAccessPtr InNodeVertexPtr, FConstClassOutputAccessPtr InClassOutputPtr, FNodeHandle InOwningNode) const override;
+
+			FGraphAccessPtr GraphPtr;
+		};
+	}
+}
diff --git a/Engine/Plugins/Runtime/Metasound/Source/MetasoundFrontend/Public/MetasoundFrontendController.h b/Engine/Plugins/Runtime/Metasound/Source/MetasoundFrontend/Public/MetasoundFrontendController.h
--- a/Engine/Plugins/Runtime/Metasound/Source/MetasoundFrontend/Public/MetasoundFrontendController.h
+++ b/Engine/Plugins/Runtime/Metasound/Source/MetasoundFrontend/Public/MetasoundFrontendController.h
@@ -161,13 +161,16 @@ namespace Metasound
 			/** Returns a FConstNodeHandle to the node which owns this output. */
 			virtual FConstNodeHandle GetOwningNode() const = 0;
 
+			/** Return true if the input is connect to an output. */
+			virtual bool IsConnected() const = 0;
+
 			/** Returns the currently connected output. If this input is not
 			 * connected, the returned handle will be invalid. */
-			virtual TArray<FInputHandle> GetCurrentlyConnectedInputs() = 0;
+			virtual TArray<FInputHandle> GetConnectedInputs() = 0;
 
 			/** Returns the currently connected output. If this input is not
 			 * connected, the returned handle will be invalid. */
-			virtual TArray<FConstInputHandle> GetCurrentlyConnectedInputs() const = 0;
+			virtual TArray<FConstInputHandle> GetConstConnectedInputs() const = 0;
 
 			virtual bool Disconnect() = 0;
 
@@ -237,11 +240,11 @@ namespace Metasound
 			
 			/** Returns the currently connected output. If this input is not
 			 * connected, the returned handle will be invalid. */
-			virtual FOutputHandle GetCurrentlyConnectedOutput() = 0;
+			virtual FOutputHandle GetConnectedOutput() = 0;
 
 			/** Returns the currently connected output. If this input is not
 			 * connected, the returned handle will be invalid. */
-			virtual FConstOutputHandle GetCurrentlyConnectedOutput() const = 0;
+			virtual FConstOutputHandle GetConnectedOutput() const = 0;
 
 			/** Returns information describing connectability between this input and the supplied output. */
 			virtual FConnectability CanConnectTo(const IOutputController& InController) const = 0;
@@ -321,7 +324,7 @@ namespace Metasound
 			virtual int32 GetNumOutputs() const = 0;
 
 			virtual TArray<FInputHandle> GetInputsWithVertexName(const FString& InName) = 0;
-			virtual TArray<FConstInputHandle> GetInputsWithVertexName(const FString& InName) const = 0;
+			virtual TArray<FConstInputHandle> GetConstInputsWithVertexName(const FString& InName) const = 0;
 
 			/** Returns all node outputs. */
 			virtual TArray<FOutputHandle> GetOutputs() = 0;
@@ -330,7 +333,7 @@ namespace Metasound
 			virtual TArray<FConstOutputHandle> GetConstOutputs() const = 0;
 
 			virtual TArray<FOutputHandle> GetOutputsWithVertexName(const FString& InName) = 0;
-			virtual TArray<FConstOutputHandle> GetOutputsWithVertexName(const FString& InName) const = 0;
+			virtual TArray<FConstOutputHandle> GetConstOutputsWithVertexName(const FString& InName) const = 0;
 
 			/** Returns true if node is required to satisfy the document archetype. */
 			virtual bool IsRequired() const = 0;
@@ -422,7 +425,14 @@ namespace Metasound
 			/** Returns true if the controller is in a valid state. */
 			virtual bool IsValid() const = 0;
 
-			virtual FGuid GetNewVertexID() const = 0;
+			/** Returns the ClassID associated with this graph. */
+			virtual FGuid GetClassID() const = 0;
+
+			/** Return the metadata for the current graph. */
+			virtual const FMetasoundFrontendClassMetadata& GetGraphMetadata() const = 0;
+
+			/** Return the display name of the graph. */
+			virtual const FText& GetDisplayName() const = 0;
 
 			virtual TArray<FString> GetInputVertexNames() const = 0;
 			virtual TArray<FString> GetOutputVertexNames() const = 0;
@@ -602,36 +612,27 @@ namespace Metasound
 
 			/** Add a new node to this graph.
 			 *
-			 * @param InClassMetadat - Info for node class.
+			 * @param InClassMetadata - Info for node class.
 			 * 
 			 * @return Node handle for class. On error, an invalid handle is returned. 
 			 */
 			virtual FNodeHandle AddNode(const FMetasoundFrontendClassMetadata& InClassMetadata) = 0;
 
-			/** Remove the node corresponding to this node handle.
+			/** Add a new node to this graph by duplicating the input node.
+			 * The new node has the same interface and node class as the input node.
+			 * This method will not duplicate node connections.
 			 *
-			 * @return True on success, false on failure. 
+			 * @param InNodeController - Node to duplicate.
+			 *
+			 * @return Node handle for new node. On error, an invalid handle is returned. 
 			 */
-			virtual bool RemoveNode(INodeController& InNode) = 0;
-
-			/** Returns the ClassID associated with this graph. */
-			virtual FGuid GetClassID() const = 0;
-
-			/** Return the metadata for the current graph. */
-			virtual const FMetasoundFrontendClassMetadata& GetGraphMetadata() const = 0;
+			virtual FNodeHandle AddDuplicateNode(const INodeController& InNodeController) = 0;
 
-			/** Inflates a subgraph of this graph into the this graph.
-			 *
-			 * If the INodeController given is itself a Metasound graph,
-			 * and the INodeController is a direct member of this IGraphController,
-			 * this will invalidate the FNodeController and paste the graph for this 
-			 * node directly into this graph.
-			 *
-			 * If not successful, InNode will not be affected.
+			/** Remove the node corresponding to this node handle.
 			 *
-			 * @returns True on success, false on failure.
+			 * @return True on success, false on failure. 
 			 */
-			virtual bool InflateNodeDirectlyIntoGraph(const INodeController& InNode) = 0;
+			virtual bool RemoveNode(INodeController& InNode) = 0;
 
 			/** Creates and inserts a new subgraph into this graph using the given metadata.
 			 * By calling AsGraph() on the returned node handle, callers can modify
@@ -668,8 +669,11 @@ namespace Metasound
 
 			/** Create a document from FMetasoundFrontendDocument description pointer. */
 			static FDocumentHandle CreateDocumentHandle(FDocumentAccessPtr InDocument);
+			static FDocumentHandle CreateDocumentHandle(FMetasoundFrontendDocument& InDocument);
+
 			/** Create a document from FMetasoundFrontendDocument description pointer. */
 			static FConstDocumentHandle CreateDocumentHandle(FConstDocumentAccessPtr InDocument);
+			static FConstDocumentHandle CreateDocumentHandle(const FMetasoundFrontendDocument& InDocument);
 
 			IDocumentController() = default;
 			virtual ~IDocumentController() = default;
@@ -729,6 +733,15 @@ namespace Metasound
 			 */
 			virtual FConstClassAccessPtr FindOrAddClass(const FMetasoundFrontendClassMetadata& InMetadata) = 0;
 
+			/** Adds a duplicate subgraph to the document. This method creates a 
+			 * copy of the graph and adds all the graph dependencies to this document.
+			 * 
+			 * @param InGraph - Graph to copy.
+			 *
+			 * @return Handle to new graph. On error, an invalid handle is returned.
+			 */
+			virtual FGraphHandle AddDuplicateSubgraph(const IGraphController& InGraph) = 0;
+
 			/** Removes all dependencies which are no longer referenced by any graphs within this document
 			  * and updates dependency Metadata where necessary with that found in the registry.  */
 			virtual void SynchronizeDependencies() = 0;
diff --git a/Engine/Plugins/Runtime/Metasound/Source/MetasoundFrontend/Public/MetasoundFrontendDocumentAccessPtr.h b/Engine/Plugins/Runtime/Metasound/Source/MetasoundFrontend/Public/MetasoundFrontendDocumentAccessPtr.h
--- a/Engine/Plugins/Runtime/Metasound/Source/MetasoundFrontend/Public/MetasoundFrontendDocumentAccessPtr.h
+++ b/Engine/Plugins/Runtime/Metasound/Source/MetasoundFrontend/Public/MetasoundFrontendDocumentAccessPtr.h
@@ -27,7 +27,7 @@ namespace Metasound
 		using FClassOutputAccessPtr = TAccessPtr<FMetasoundFrontendClassOutput>;
 		using FConstClassOutputAccessPtr = TAccessPtr<const FMetasoundFrontendClassOutput>;
 
-		class FNodeAccessPtr : public TAccessPtr<FMetasoundFrontendNode>
+		class METASOUNDFRONTEND_API FNodeAccessPtr : public TAccessPtr<FMetasoundFrontendNode>
 		{
 		public:
 			using Super = TAccessPtr<FMetasoundFrontendNode>;
@@ -46,7 +46,7 @@ namespace Metasound
 			FConstVertexAccessPtr GetOutputWithVertexID(const FGuid& InID) const;
 		};
 
-		class FConstNodeAccessPtr : public TAccessPtr<const FMetasoundFrontendNode>
+		class METASOUNDFRONTEND_API FConstNodeAccessPtr : public TAccessPtr<const FMetasoundFrontendNode>
 		{
 		public:
 			using Super = TAccessPtr<const FMetasoundFrontendNode>;
@@ -63,7 +63,7 @@ namespace Metasound
 		using FGraphAccessPtr = TAccessPtr<FMetasoundFrontendGraph>;
 		using FConstGraphAccessPtr = TAccessPtr<const FMetasoundFrontendGraph>;
 
-		class FClassAccessPtr : public TAccessPtr<FMetasoundFrontendClass>
+		class METASOUNDFRONTEND_API FClassAccessPtr : public TAccessPtr<FMetasoundFrontendClass>
 		{
 		public:
 			using Super = TAccessPtr<FMetasoundFrontendClass>;
@@ -77,7 +77,7 @@ namespace Metasound
 			FConstClassOutputAccessPtr GetOutputWithName(const FString& InName) const;
 		};
 
-		class FConstClassAccessPtr : public TAccessPtr<const FMetasoundFrontendClass>
+		class METASOUNDFRONTEND_API FConstClassAccessPtr : public TAccessPtr<const FMetasoundFrontendClass>
 		{
 		public:
 			using Super = TAccessPtr<const FMetasoundFrontendClass>;
@@ -89,7 +89,7 @@ namespace Metasound
 			FConstClassOutputAccessPtr GetOutputWithName(const FString& InName) const;
 		};
 
-		class FGraphClassAccessPtr : public TAccessPtr<FMetasoundFrontendGraphClass>
+		class METASOUNDFRONTEND_API FGraphClassAccessPtr : public TAccessPtr<FMetasoundFrontendGraphClass>
 		{
 		public:
 			using Super = TAccessPtr<FMetasoundFrontendGraphClass>;
@@ -114,7 +114,7 @@ namespace Metasound
 			FConstGraphAccessPtr GetGraph() const;
 		};
 
-		class FConstGraphClassAccessPtr : public TAccessPtr<const FMetasoundFrontendGraphClass>
+		class METASOUNDFRONTEND_API FConstGraphClassAccessPtr : public TAccessPtr<const FMetasoundFrontendGraphClass>
 		{
 			using Super = TAccessPtr<const FMetasoundFrontendGraphClass>;
 
@@ -129,7 +129,7 @@ namespace Metasound
 			FConstGraphAccessPtr GetGraph() const;
 		};
 
-		class FDocumentAccessPtr : public TAccessPtr<FMetasoundFrontendDocument>
+		class METASOUNDFRONTEND_API FDocumentAccessPtr : public TAccessPtr<FMetasoundFrontendDocument>
 		{
 		public:
 			using Super = TAccessPtr<FMetasoundFrontendDocument>;
@@ -152,7 +152,7 @@ namespace Metasound
 			FConstClassAccessPtr GetClassWithInfo(const FNodeClassInfo& InInfo) const;
 		};
 
-		class FConstDocumentAccessPtr : public TAccessPtr<const FMetasoundFrontendDocument>
+		class METASOUNDFRONTEND_API FConstDocumentAccessPtr : public TAccessPtr<const FMetasoundFrontendDocument>
 		{
 		public:
 			using Super = TAccessPtr<const FMetasoundFrontendDocument>;
diff --git a/Engine/Plugins/Runtime/Metasound/Source/MetasoundFrontend/Public/MetasoundInputNode.h b/Engine/Plugins/Runtime/Metasound/Source/MetasoundFrontend/Public/MetasoundInputNode.h
--- a/Engine/Plugins/Runtime/Metasound/Source/MetasoundFrontend/Public/MetasoundInputNode.h
+++ b/Engine/Plugins/Runtime/Metasound/Source/MetasoundFrontend/Public/MetasoundInputNode.h
@@ -3,6 +3,7 @@
 
 #include "CoreMinimal.h"
 #include "Internationalization/Text.h"
+#include "MetasoundDataReference.h"
 #include "MetasoundOperatorInterface.h"
 #include "MetasoundNodeInterface.h"
 #include "MetasoundTrigger.h"
@@ -28,10 +29,26 @@ namespace Metasound
 			{
 			}
 
-			virtual ~TInputOperator() {}
+			virtual ~TInputOperator() = default;
 
 			virtual FDataReferenceCollection GetInputs() const override
 			{
+				// TODO: Expose a readable reference instead of a writable reference.
+				//
+				// If data needs to be written to, outside entities should create
+				// it and pass it in as a readable reference. Currently, the workflow
+				// is to have the input node create a writable reference which is then
+				// queried by the outside world. Exposing writable references causes 
+				// code maintainability issues where TInputNode<> specializations need
+				// to handle multiple situations which can happen in an input node.
+				//
+				// The only reason that this code is not removed immediately is because
+				// of the `TExecutableDataType<>` which primarily supports the FTrigger.
+				// The TExecutableDataType<> advances the trigger within the graph. But,
+				// with graph composition, the owner of the data type becomes more 
+				// complicated and hence triggers advancing should be managed by a 
+				// different object. Preferably the graph operator itself, or an
+				// explicit trigger manager tied to the environment.
 				FDataReferenceCollection Inputs;
 				Inputs.AddDataWriteReference<DataType>(DataReferenceName, InputValue);
 				return Inputs;
@@ -49,6 +66,14 @@ namespace Metasound
 				TExecutableDataType<DataType>::Execute(*InputValue, *OutputValue);
 			}
 
+			void ExecutPassThrough()
+			{
+				if (TExecutableDataType<DataType>::bIsExecutable)
+				{
+					*OutputValue = *InputValue;
+				}
+			}
+
 			static void ExecuteFunction(IOperator* InOperator)
 			{
 				static_cast<TInputOperator<DataType>*>(InOperator)->Execute();
@@ -71,53 +96,160 @@ namespace Metasound
 
 	};
 
-	/** TInputOperatorLiteralFactory creates an input by passing it a literal. */
+	/** TPassThroughOperator supplies a readable input and a readable output. 
+	 *
+	 * It does *not* invoke executable data types (see `TExecutableDataType<>`).
+	 */
 	template<typename DataType>
-	class TInputOperatorLiteralFactory : public IOperatorFactory
+	class TPassThroughOperator : public TInputOperator<DataType>
 	{
-		public:
-			// If the data type is parsable from a literal type, then the data type 
-			// can be registered as an input type with the frontend.  To make a 
-			// DataType registrable, either create a constructor for the data type
-			// which accepts the one of the supported literal types with an optional 
-			// FOperatorSettings argument, or create a default constructor, or specialize
-			// this factory with an implementation for that specific data type.
-			static constexpr bool bCanRegister = TLiteralTraits<DataType>::bIsParsableFromAnyLiteralType;
-
-			using FDataWriteReference = TDataWriteReference<DataType>;
-
-			TInputOperatorLiteralFactory(FLiteral&& InInitParam)
-			:	InitParam(MoveTemp(InInitParam))
+	public:
+		using FDataReadReference = TDataReadReference<DataType>;
+		using Super = TInputOperator<DataType>;
+
+		TPassThroughOperator(const FVertexKey& InDataReferenceName, FDataReadReference InDataReference)
+		: TInputOperator<DataType>(InDataReferenceName, WriteCast(InDataReference)) // Write cast is safe because `GetExecuteFunction() and GetInputs() are overridden, ensuring that data is not written.
+		, DataReferenceName(InDataReferenceName)
+		{
+		}
+
+		virtual ~TPassThroughOperator() = default;
+
+		virtual FDataReferenceCollection GetInputs() const override
+		{
+			FDataReferenceCollection Inputs;
+			ensure(Inputs.AddDataReadReferenceFrom(DataReferenceName, Super::GetInputs(), DataReferenceName, GetMetasoundDataTypeName<DataType>()));
+			return Inputs;
+		}
+
+		static void ExecuteFunction(IOperator* InOperator)
+		{
+			static_cast<TPassThroughOperator<DataType>*>(InOperator)->ExecutPassThrough();
+		}
+
+		virtual IOperator::FExecuteFunction GetExecuteFunction() override
+		{
+			// TODO: this is a hack until we can remove TExecutableOperator<>.
+			//
+			// The primary contention is that we would like to allow developers 
+			// to specialize `TInputNode<>` as in `TInputNode<FStereoAudioFormat>`.
+			// `TExecutableOperator<>` adds in a level of complexity that makes it 
+			// difficult to allow specialization of TInputNode and to derive from
+			// TInputNode to create the TPassThroughOperator. Particularly because
+			// TExecutableOperator<> alters which output data reference is used. 
+			// Specializations of TInputNode also tend to alter the output data
+			// references. Supporting both is likely to cause issues. 
+			//
+			// We may need to ensure that input nodes do not provide execution
+			// functions. Or we may need a more explicit way of only allowing
+			// outputs to be modified. Likely a mix of the `final` keyword
+			// and disabling template specialization of a base class. 
+			//
+			// namespace Private
+			// {
+			//     class TInputNodePrivate<>
+			//     {
+			//         GetInputs() final
+			//         GetExecutionFunction() final
+			//         GetOutputs()
+			//     }
+			// }
+			// 
+			// template<DataType>
+			// using TInputNodeBase<DataType> = TInputNodePrivate<DataType>; // Do not allow specialization of TInputNodePrivate<> or TInputNodeBase<> (this works because you can't specialize a template alias)
+			//
+			// // DO ALLOW specialization of TInputNode
+			// template<DataType>
+			// class TInputNode<DataType> : public TInputNodeBase<DataType>
+			// {
+			// };
+			// 
+			// template<>
+			// class TInputNode<MyType> : public TInputNodeBase<MyType>
+			// {
+			// 	 GetOutputs() <-- OK to override
+			// }
+			//
+			if (TExecutableDataType<DataType>::bIsExecutable)
 			{
+				return &TPassThroughOperator<DataType>::ExecuteFunction;
 			}
+			return nullptr;
+		}
 
-			virtual TUniquePtr<IOperator> CreateOperator(const FCreateOperatorParams& InParams, FBuildErrorArray& OutErrors) override;
+	private:
+		FVertexKey DataReferenceName;
+	};
 
-		private:
 
-			FLiteral InitParam;
+	/** Data type creation policy to create by copy construction. */
+	template<typename DataType>
+	struct FCreateDataReferenceWithCopy
+	{
+		template<typename... ArgTypes>
+		FCreateDataReferenceWithCopy(ArgTypes&&... Args)
+		:	Data(Forward<ArgTypes>(Args)...)
+		{
+		}
+
+		TDataWriteReference<DataType> CreateDataReference(const FOperatorSettings& InOperatorSettings) const
+		{
+			return TDataWriteReferenceFactory<DataType>::CreateExplicitArgs(InOperatorSettings, Data);
+		}
+
+	private:
+		DataType Data;
 	};
 
+	/** Data type creation policy to create by literal construction. */
+	template<typename DataType>
+	struct FCreateDataReferenceWithLiteral
+	{
+		// If the data type is parsable from a literal type, then the data type 
+		// can be registered as an input type with the frontend.  To make a 
+		// DataType registrable, either create a constructor for the data type
+		// which accepts the one of the supported literal types with an optional 
+		// FOperatorSettings argument, or create a default constructor, or specialize
+		// this factory with an implementation for that specific data type.
+		static constexpr bool bCanCreateWithLiteral = TLiteralTraits<DataType>::bIsParsableFromAnyLiteralType;
+
+		FCreateDataReferenceWithLiteral(FLiteral&& InLiteral)
+		:	Literal(MoveTemp(InLiteral))
+		{
+		}
+
+		TDataWriteReference<DataType> CreateDataReference(const FOperatorSettings& InOperatorSettings) const
+		{
+			return TDataWriteReferenceLiteralFactory<DataType>::CreateExplicitArgs(InOperatorSettings, Literal);
+		}
+
+	private:
+
+		FLiteral Literal;
+	};
+
+	
+
 	/** TInputOperatorFactory initializes the DataType at construction. It uses
-	 * the DataType's copy operator to create additional version. 
+	 * the ReferenceCreatorType to create a data reference if one is not passed in.
 	 */
-	template<typename DataType>
+	template<typename DataType, typename ReferenceCreatorType>
 	class TInputOperatorFactory : public IOperatorFactory
 	{
 		public:
 
 			using FDataWriteReference = TDataWriteReference<DataType>;
+			using FDataReadReference = TDataReadReference<DataType>;
 
-			template<typename... ArgTypes>
-			TInputOperatorFactory(ArgTypes&&... Args)
-			:	Data(Forward<ArgTypes>(Args)...)
+			TInputOperatorFactory(ReferenceCreatorType&& InReferenceCreator)
+			:	ReferenceCreator(MoveTemp(InReferenceCreator))
 			{
 			}
 
 			virtual TUniquePtr<IOperator> CreateOperator(const FCreateOperatorParams& InParams, FBuildErrorArray& OutErrors) override;
 
 		private:
-			DataType Data;
+			ReferenceCreatorType ReferenceCreator;
 	};
 
 	/** TInputNode represents an input to a metasound graph. */
@@ -126,9 +258,8 @@ namespace Metasound
 	{
 
 		public:
-
 			// If true, this node can be instantiated by the FrontEnd.
-			static constexpr bool bCanRegister = TInputOperatorLiteralFactory<DataType>::bCanRegister;
+			static constexpr bool bCanRegister = FCreateDataReferenceWithLiteral<DataType>::bCanCreateWithLiteral;
 
 			static FVertexInterface DeclareVertexInterface(const FVertexKey& InVertexName)
 			{
@@ -158,6 +289,23 @@ namespace Metasound
 				return Info;
 			}
 
+			template<typename... ArgTypes>
+			static FOperatorFactorySharedRef CreateOperatorFactoryWithArgs(ArgTypes&&... Args)
+			{
+				using FCreatorType = FCreateDataReferenceWithCopy<DataType>;
+				using FFactoryType = TInputOperatorFactory<DataType, FCreatorType>;
+
+				return MakeOperatorFactoryRef<FFactoryType>(FCreatorType(Forward<ArgTypes>(Args)...));
+			}
+
+			static FOperatorFactorySharedRef CreateOperatorFactoryWithLiteral(FLiteral&& InLiteral)
+			{
+				using FCreatorType = FCreateDataReferenceWithLiteral<DataType>;
+				using FFactoryType = TInputOperatorFactory<DataType, FCreatorType>;
+
+				return MakeOperatorFactoryRef<FFactoryType>(FCreatorType(MoveTemp(InLiteral)));
+			}
+
 
 			/* Construct a TInputNode using the TInputOperatorFactory<> and forwarding 
 			 * Args to the TInputOperatorFactory constructor.*/
@@ -166,7 +314,7 @@ namespace Metasound
 			:	FNode(InNodeDescription, InInstanceID, GetNodeInfo(InVertexName))
 			,	VertexName(InVertexName)
 			,	Interface(DeclareVertexInterface(InVertexName))
-			,	Factory(MakeOperatorFactoryRef<TInputOperatorFactory<DataType>>(Forward<ArgTypes>(Args)...))
+			,	Factory(CreateOperatorFactoryWithArgs(Forward<ArgTypes>(Args)...))
 			{
 			}
 
@@ -176,7 +324,7 @@ namespace Metasound
 			:	FNode(InNodeDescription, InInstanceID, GetNodeInfo(InVertexName))
 			,	VertexName(InVertexName)
 			,	Interface(DeclareVertexInterface(InVertexName))
-			, 	Factory(MakeOperatorFactoryRef<TInputOperatorLiteralFactory<DataType>>(MoveTemp(InParam)))
+			, 	Factory(CreateOperatorFactoryWithLiteral(MoveTemp(InParam)))
 			{
 			}
 
@@ -213,31 +361,32 @@ namespace Metasound
 	};
 
 
-	template<typename DataType>
-	TUniquePtr<IOperator> TInputOperatorLiteralFactory<DataType>::CreateOperator(const FCreateOperatorParams& InParams, FBuildErrorArray& OutErrors) 
+	template<typename DataType, typename ReferenceCreatorType>
+	TUniquePtr<IOperator> TInputOperatorFactory<DataType, ReferenceCreatorType>::CreateOperator(const FCreateOperatorParams& InParams, FBuildErrorArray& OutErrors) 
 	{
 		using FInputNodeType = TInputNode<DataType>;
 
-		// Create write reference by calling compatible constructor with literal.
-		FDataWriteReference DataRef = TDataWriteReferenceLiteralFactory<DataType>::CreateExplicitArgs(InParams.OperatorSettings, InitParam);
-
 		const FInputNodeType& InputNode = static_cast<const FInputNodeType&>(InParams.Node);
-
-		return MakeUnique<TInputOperator<DataType>>(InputNode.GetVertexName(), DataRef);
+		const FVertexKey& VertexKey = InputNode.GetVertexName();
+
+		if (InParams.InputDataReferences.ContainsDataWriteReference<DataType>(VertexKey))
+		{
+			// Data is externally owned. Use pass through operator
+			FDataWriteReference DataRef = InParams.InputDataReferences.GetDataWriteReference<DataType>(VertexKey);
+			return MakeUnique<TPassThroughOperator<DataType>>(InputNode.GetVertexName(), DataRef);
+		}
+		else if (InParams.InputDataReferences.ContainsDataReadReference<DataType>(VertexKey))
+		{
+			// Data is externally owned. Use pass through operator
+			FDataReadReference DataRef = InParams.InputDataReferences.GetDataReadReference<DataType>(VertexKey);
+			return MakeUnique<TPassThroughOperator<DataType>>(InputNode.GetVertexName(), DataRef);
+		}
+		else
+		{
+			// Create write reference by calling compatible constructor with literal.
+			FDataWriteReference DataRef = ReferenceCreator.CreateDataReference(InParams.OperatorSettings);
+			return MakeUnique<TInputOperator<DataType>>(InputNode.GetVertexName(), DataRef);
+		}
 	}
-
-	template<typename DataType>
-	TUniquePtr<IOperator> TInputOperatorFactory<DataType>::CreateOperator(const FCreateOperatorParams& InParams, FBuildErrorArray& OutErrors) 
-	{
-		using FInputNodeType = TInputNode<DataType>;
-
-		// Create write reference by calling copy constructor.
-		FDataWriteReference DataRef = TDataWriteReferenceFactory<DataType>::CreateExplicitArgs(InParams.OperatorSettings, Data);
-
-		const FInputNodeType& InputNode = static_cast<const FInputNodeType&>(InParams.Node);
-
-		return MakeUnique<TInputOperator<DataType>>(InputNode.GetVertexName(), DataRef);
-	}
-
 } // namespace Metasound
 #undef LOCTEXT_NAMESPACE //MetasoundOutputNode
diff --git a/Engine/Plugins/Runtime/Metasound/Source/MetasoundGraphCore/Private/MetasoundDataReferenceCollection.cpp b/Engine/Plugins/Runtime/Metasound/Source/MetasoundGraphCore/Private/MetasoundDataReferenceCollection.cpp
--- a/Engine/Plugins/Runtime/Metasound/Source/MetasoundGraphCore/Private/MetasoundDataReferenceCollection.cpp
+++ b/Engine/Plugins/Runtime/Metasound/Source/MetasoundGraphCore/Private/MetasoundDataReferenceCollection.cpp
@@ -44,6 +44,12 @@ namespace Metasound
 
 			return true;
 		}
+		else if (OtherCollection.ContainsDataWriteReference(OtherName, OtherTypeName))
+		{
+			DataReadRefMap.Add(InName, OtherCollection.DataWriteRefMap[OtherName]);
+
+			return true;
+		}
 
 		return false;
 	}
diff --git a/Engine/Plugins/Runtime/Metasound/Source/MetasoundGraphCore/Private/MetasoundGraph.cpp b/Engine/Plugins/Runtime/Metasound/Source/MetasoundGraphCore/Private/MetasoundGraph.cpp
--- a/Engine/Plugins/Runtime/Metasound/Source/MetasoundGraphCore/Private/MetasoundGraph.cpp
+++ b/Engine/Plugins/Runtime/Metasound/Source/MetasoundGraphCore/Private/MetasoundGraph.cpp
@@ -1,8 +1,7 @@
 // Copyright Epic Games, Inc. All Rights Reserved.
 #include "MetasoundGraph.h"
 
-#include "Internationalization/Text.h"
-
+#include "MetasoundOperatorBuilder.h"
 
 namespace Metasound
 {
@@ -134,4 +133,27 @@ namespace Metasound
 	{
 		return InInterface == Metadata.DefaultInterface;
 	}
+
+	FOperatorFactorySharedRef FGraph::GetDefaultOperatorFactory() const 
+	{
+		return MakeShared<FGraph::FFactory>();
+	}
+
+	TUniquePtr<IOperator> FGraph::FFactory::CreateOperator(const FCreateOperatorParams& InParams, FBuildErrorArray& OutErrors) 
+	{
+		const FGraph& Graph = static_cast<const FGraph&>(InParams.Node);
+
+		FBuildGraphParams BuildParams{Graph, InParams.OperatorSettings, InParams.InputDataReferences, InParams.Environment};
+
+		if (nullptr != InParams.Builder)
+		{
+			// Use the provided builder if it exists. 
+			return InParams.Builder->BuildGraphOperator(BuildParams, OutErrors);
+		}
+		else
+		{
+			// Use a builder with default settings if no builder was provided. 
+			return FOperatorBuilder(FOperatorBuilderSettings::GetDefaultSettings()).BuildGraphOperator(BuildParams, OutErrors);
+		}
+	}
 }
diff --git a/Engine/Plugins/Runtime/Metasound/Source/MetasoundGraphCore/Private/MetasoundOperatorBuilder.cpp b/Engine/Plugins/Runtime/Metasound/Source/MetasoundGraphCore/Private/MetasoundOperatorBuilder.cpp
--- a/Engine/Plugins/Runtime/Metasound/Source/MetasoundGraphCore/Private/MetasoundOperatorBuilder.cpp
+++ b/Engine/Plugins/Runtime/Metasound/Source/MetasoundGraphCore/Private/MetasoundOperatorBuilder.cpp
@@ -120,18 +120,17 @@ namespace Metasound
 	{
 	}
 
-	TUniquePtr<IOperator> FOperatorBuilder::BuildGraphOperator(const IGraph& InGraph, const FOperatorSettings& InOperatorSettings, const FMetasoundEnvironment& InEnvironment, TArray<FBuildErrorPtr>& OutErrors) 
+	TUniquePtr<IOperator> FOperatorBuilder::BuildGraphOperator(const FBuildGraphParams& InParams, FBuildErrorArray& OutErrors) const
 	{
 		using namespace OperatorBuilderPrivate;
 
 		FBuildStatus BuildStatus = FBuildStatus::NoError;
 
-
 		// Validate that the sources and destinations declared in an edge actually
 		// exist in the node.
 		if (BuilderSettings.bValidateVerticesExist)
 		{
-			if (!FGraphLinter::ValidateVerticesExist(InGraph, OutErrors))
+			if (!FGraphLinter::ValidateVerticesExist(InParams.Graph, OutErrors))
 			{
 				BuildStatus |= FBuildStatus::FatalError;
 			}
@@ -140,7 +139,7 @@ namespace Metasound
 		// Validate that the data types for a source and destination match.
 		if (BuilderSettings.bValidateEdgeDataTypesMatch)
 		{
-			if (!FGraphLinter::ValidateEdgeDataTypesMatch(InGraph, OutErrors))
+			if (!FGraphLinter::ValidateEdgeDataTypesMatch(InParams.Graph, OutErrors))
 			{
 				BuildStatus |= FBuildStatus::FatalError;
 			}
@@ -149,7 +148,7 @@ namespace Metasound
 		// Validate that node inputs only have one source
 		if (BuilderSettings.bValidateNoDuplicateInputs)
 		{
-			if (!FGraphLinter::ValidateNoDuplicateInputs(InGraph, OutErrors))
+			if (!FGraphLinter::ValidateNoDuplicateInputs(InParams.Graph, OutErrors))
 			{
 				BuildStatus |= FBuildStatus::FatalError;
 			}
@@ -163,7 +162,7 @@ namespace Metasound
 
 
 		// Create algo adapter view of graph to cache graph operations.
-		TPimplPtr<FDirectedGraphAlgoAdapter> AlgoAdapter = FDirectedGraphAlgo::CreateDirectedGraphAlgoAdapter(InGraph);
+		TPimplPtr<FDirectedGraphAlgoAdapter> AlgoAdapter = FDirectedGraphAlgo::CreateDirectedGraphAlgoAdapter(InParams.Graph);
 		
 		if (!AlgoAdapter.IsValid())
 		{
@@ -171,7 +170,7 @@ namespace Metasound
 			return TUniquePtr<IOperator>(nullptr);
 		}
 
-		FBuildContext BuildContext(InGraph, *AlgoAdapter, InOperatorSettings, InEnvironment, OutErrors);
+		FBuildContext BuildContext(InParams.Graph, *AlgoAdapter, InParams.OperatorSettings, InParams.Environment, OutErrors);
 
 		TArray<const INode*> SortedNodes;
 
@@ -191,7 +190,7 @@ namespace Metasound
 		FNodeDataReferenceMap NodeData;
 
 		// Create node operators from factories. 
-		BuildStatus |= CreateOperators(BuildContext, SortedNodes, Operators, NodeData);
+		BuildStatus |= CreateOperators(BuildContext, SortedNodes, InParams.InputDataReferences, Operators, NodeData);
 
 		if (BuildStatus > MaxBuildStatusErrorLevel)
 		{
@@ -361,7 +360,35 @@ namespace Metasound
 		return BuildStatus;
 	}
 
-	FOperatorBuilder::FBuildStatus FOperatorBuilder::CreateOperators(FBuildContext& InContext, const TArray<const INode*>& InSortedNodes, TArray<TUniquePtr<IOperator>>& OutOperators, FNodeDataReferenceMap& OutDataReferences) const
+	FOperatorBuilder::FBuildStatus FOperatorBuilder::GatherExternalInputDataReferences(FBuildContext& InContext, const INode* InNode, const FNodeDestinationMap& InNodeDestinationMap, const FDataReferenceCollection& InExternalCollection, FDataReferenceCollection& OutDataReferences) const
+	{
+		using namespace OperatorBuilderPrivate;
+
+		FBuildStatus BuildStatus;
+
+		if (const FInputDataDestination* const* DestinationPtr = InNodeDestinationMap.Find(InNode))
+		{
+			if (const FInputDataDestination* Destination = *DestinationPtr)
+			{
+				const bool bReferenceExists = InExternalCollection.ContainsDataReadReference(Destination->Vertex.GetVertexName(), Destination->Vertex.GetDataTypeName());
+				if (bReferenceExists)
+				{
+					const bool bSuccess = OutDataReferences.AddDataReadReferenceFrom(Destination->Vertex.GetVertexName(), InExternalCollection, Destination->Vertex.GetVertexName(), Destination->Vertex.GetDataTypeName());
+
+					if (!bSuccess)
+					{
+						AddBuildError<FInternalError>(InContext.Errors, __FILE__, __LINE__);
+
+						BuildStatus |= FBuildStatus::NonFatalError;
+					}
+				}
+			}
+		}
+
+		return BuildStatus;
+	}
+
+	FOperatorBuilder::FBuildStatus FOperatorBuilder::CreateOperators(FBuildContext& InContext, const TArray<const INode*>& InSortedNodes, const FDataReferenceCollection& InExternalInputDataReferences, TArray<TUniquePtr<IOperator>>& OutOperators, FNodeDataReferenceMap& OutDataReferences) const
 	{
 		FBuildStatus BuildStatus;
 
@@ -373,14 +400,23 @@ namespace Metasound
 			NodeInputEdges.Add(Edge.To.Node, &Edge);
 		}
 
+		// Map input nodes to graph destinations
+		FNodeDestinationMap NodeDestinations;
+		for (const auto& InputDestinationKV : InContext.Graph.GetInputDataDestinations())
+		{
+			const FInputDataDestination& Destination = InputDestinationKV.Value;
+			NodeDestinations.Add(Destination.Node, &Destination);
+		}
+
 		// Call operator factory for each node.
 		for (const INode* Node : InSortedNodes)
 		{
-
 			FDataReferenceCollection InputCollection;
 
 			// Gather the input parameters for this IOperator from the output parameters of already created IOperators. 
 			BuildStatus |= GatherInputDataReferences(InContext, Node, NodeInputEdges, OutDataReferences, InputCollection);
+			// Gather the input parameters for this IOperator from the graph inputs.
+			BuildStatus |= GatherExternalInputDataReferences(InContext, Node, NodeDestinations, InExternalInputDataReferences, InputCollection);
 
 			if (BuildStatus >= FBuildStatus::FatalError)
 			{
@@ -389,7 +425,7 @@ namespace Metasound
 
 			FOperatorFactorySharedRef Factory = Node->GetDefaultOperatorFactory();
 
-			FCreateOperatorParams CreateParams(*Node, InContext.Settings, InputCollection, InContext.Environment);
+			FCreateOperatorParams CreateParams{*Node, InContext.Settings, InputCollection, InContext.Environment, this};
 
 			FOperatorPtr Operator = Factory->CreateOperator(CreateParams, InContext.Errors);
 
@@ -432,8 +468,15 @@ namespace Metasound
 
 			FDataReferenceCollection& Collection = InNodeDataReferences[InputDestination.Node].Inputs;
 
+			// Get either readable or writable values for inputs. 
+			// Readable inputs can occur when the data reference is provided outside
+			// of the graph being built. 
+			const bool bContainsWritableInput = Collection.ContainsDataWriteReference(InputDestination.Vertex.GetVertexName(), InputDestination.Vertex.GetDataTypeName());
+			const bool bContainsReadableInput = Collection.ContainsDataReadReference(InputDestination.Vertex.GetVertexName(), InputDestination.Vertex.GetDataTypeName());
 
-			if (!Collection.ContainsDataWriteReference(InputDestination.Vertex.GetVertexName(), InputDestination.Vertex.GetDataTypeName()))
+			const bool bContainsInput = bContainsWritableInput || bContainsReadableInput;
+
+			if (!bContainsInput)
 			{
 				AddBuildError<FMissingInputDataReferenceError>(InContext.Errors, InputDestination);
 
@@ -441,7 +484,15 @@ namespace Metasound
 				continue;
 			}
 
-			bool bSuccess = OutGraphInputs.AddDataWriteReferenceFrom(InputDestination.Vertex.GetVertexName(), Collection, InputDestination.Vertex.GetVertexName(), InputDestination.Vertex.GetDataTypeName());
+			bool bSuccess = true;
+			if (bContainsWritableInput)
+			{
+				bSuccess = OutGraphInputs.AddDataWriteReferenceFrom(InputDestination.Vertex.GetVertexName(), Collection, InputDestination.Vertex.GetVertexName(), InputDestination.Vertex.GetDataTypeName());
+			}
+			else
+			{
+				bSuccess = OutGraphInputs.AddDataReadReferenceFrom(InputDestination.Vertex.GetVertexName(), Collection, InputDestination.Vertex.GetVertexName(), InputDestination.Vertex.GetDataTypeName());
+			}
 
 			if (!bSuccess)
 			{
diff --git a/Engine/Plugins/Runtime/Metasound/Source/MetasoundGraphCore/Public/MetasoundBuilderInterface.h b/Engine/Plugins/Runtime/Metasound/Source/MetasoundGraphCore/Public/MetasoundBuilderInterface.h
--- a/Engine/Plugins/Runtime/Metasound/Source/MetasoundGraphCore/Public/MetasoundBuilderInterface.h
+++ b/Engine/Plugins/Runtime/Metasound/Source/MetasoundGraphCore/Public/MetasoundBuilderInterface.h
@@ -9,6 +9,9 @@
 
 namespace Metasound
 {
+	// Forward Declare
+	class IOperatorBuilder;
+
 	/** IOperatorBuildError
 	 *
 	 * This interface is intended for errors encountered when building an GraphOperator.
@@ -42,6 +45,7 @@ namespace Metasound
 	 */
 	struct FCreateOperatorParams
 	{
+
 		/** The node associated with this factory and the desired IOperator. */
 		const INode& Node;
 
@@ -54,15 +58,11 @@ namespace Metasound
 		/** Environment settings available. */
 		const FMetasoundEnvironment& Environment;
 
-		FCreateOperatorParams(const INode& InNode, const FOperatorSettings& InOperatorSettings, const FDataReferenceCollection& InInputDataReferences, const FMetasoundEnvironment& InEnvironment)
-		:	Node(InNode)
-		,	OperatorSettings(InOperatorSettings)
-		,	InputDataReferences(InInputDataReferences)
-		,	Environment(InEnvironment)
-		{
-		}
+		/** Pointer to builder actively building graph. */
+		const IOperatorBuilder* Builder = nullptr;
 	};
 
+	/** Array of build errors. */
 	typedef TArray<TUniquePtr<IOperatorBuildError>> FBuildErrorArray;
 
 	/** Convenience template for adding build errors.
@@ -102,6 +102,20 @@ namespace Metasound
 			virtual TUniquePtr<IOperator> CreateOperator(const FCreateOperatorParams& InParams, FBuildErrorArray& OutErrors) = 0;
 	};
 
+	struct FBuildGraphParams
+	{
+		/** Reference to graph being built. */
+		const IGraph& Graph;
+
+		/** General operator settings for the graph. */
+		const FOperatorSettings& OperatorSettings;
+
+		/** Collection of input parameters available for to an IOperator. */
+		const FDataReferenceCollection& InputDataReferences;
+
+		/** Environment settings available. */
+		const FMetasoundEnvironment& Environment;
+	};
 
 	/** IOperatorBuilder
 	 *
@@ -117,14 +131,12 @@ namespace Metasound
 
 			/** Build a graph operator from a graph. 
 			 *
-			 * @params InGraph - The input graph object containing edges, input vertices and output vertices.
-			 * @param InOperatorSettings - Settings to be passed to all operators on creation.
-			 * @param InEnvironment - The environment variables to use during construction. 
+			 * @params InParams - Input parameters for building a graph.
 			 * @param OutErrors - An array of errors. Errors can be added if issues occur while creating the IOperator.
 			 *
 			 * @return A unique pointer to an IOperator. 
 			 */
-			virtual TUniquePtr<IOperator> BuildGraphOperator(const IGraph& InGraph, const FOperatorSettings& InOperatorSettings, const FMetasoundEnvironment& InEnvironment, TArray<FBuildErrorPtr>& OutErrors) = 0;
+			virtual TUniquePtr<IOperator> BuildGraphOperator(const FBuildGraphParams& InParams, FBuildErrorArray& OutErrors) const = 0;
 	};
 }
 
diff --git a/Engine/Plugins/Runtime/Metasound/Source/MetasoundGraphCore/Public/MetasoundGraph.h b/Engine/Plugins/Runtime/Metasound/Source/MetasoundGraphCore/Public/MetasoundGraph.h
--- a/Engine/Plugins/Runtime/Metasound/Source/MetasoundGraphCore/Public/MetasoundGraph.h
+++ b/Engine/Plugins/Runtime/Metasound/Source/MetasoundGraphCore/Public/MetasoundGraph.h
@@ -2,6 +2,7 @@
 #pragma once
 
 #include "CoreMinimal.h"
+#include "MetasoundBuilderInterface.h"
 #include "MetasoundNode.h"
 
 // Forward Declarations
@@ -106,7 +107,18 @@ namespace Metasound
 
 			// TODO: Add ability to remove things.
 
+			/** Return a reference to the default operator factory. */
+			virtual FOperatorFactorySharedRef GetDefaultOperatorFactory() const override;
+
 		private:
+			class FFactory : public IOperatorFactory
+			{
+			public:
+				virtual ~FFactory() = default;
+
+				virtual TUniquePtr<IOperator> CreateOperator(const FCreateOperatorParams& InParams, FBuildErrorArray& OutErrors) override;
+			};
+
 			FString InstanceName;
 			FGuid InstanceID;
 
diff --git a/Engine/Plugins/Runtime/Metasound/Source/MetasoundGraphCore/Public/MetasoundNodeInterface.h b/Engine/Plugins/Runtime/Metasound/Source/MetasoundGraphCore/Public/MetasoundNodeInterface.h
--- a/Engine/Plugins/Runtime/Metasound/Source/MetasoundGraphCore/Public/MetasoundNodeInterface.h
+++ b/Engine/Plugins/Runtime/Metasound/Source/MetasoundGraphCore/Public/MetasoundNodeInterface.h
@@ -320,7 +320,7 @@ namespace Metasound
 
 
 	/** Interface for graph of nodes. */
-	class IGraph : public INodeBase
+	class IGraph : public INode
 	{
 		public:
 			virtual ~IGraph() {}
diff --git a/Engine/Plugins/Runtime/Metasound/Source/MetasoundGraphCore/Public/MetasoundOperatorBuilder.h b/Engine/Plugins/Runtime/Metasound/Source/MetasoundGraphCore/Public/MetasoundOperatorBuilder.h
--- a/Engine/Plugins/Runtime/Metasound/Source/MetasoundGraphCore/Public/MetasoundOperatorBuilder.h
+++ b/Engine/Plugins/Runtime/Metasound/Source/MetasoundGraphCore/Public/MetasoundOperatorBuilder.h
@@ -113,7 +113,7 @@ namespace Metasound
 			 * @return A TUniquePtr to an IOperator. If the processes was unsuccessful, 
 			 *         the returned pointer will contain a nullptr and be invalid.
 			 */
-			virtual TUniquePtr<IOperator> BuildGraphOperator(const IGraph& InGraph, const FOperatorSettings& InOperatorSettings, const FMetasoundEnvironment& InEnvironment, TArray<FBuildErrorPtr>& OutErrors) override;
+			virtual TUniquePtr<IOperator> BuildGraphOperator(const FBuildGraphParams& InParams, FBuildErrorArray& OutErrors) const override;
 			
 
 		private:
@@ -136,6 +136,7 @@ namespace Metasound
 			};
 
 			using FNodeEdgeMultiMap = TMultiMap<const INode*, const FDataEdge*>;
+			using FNodeDestinationMap = TMap<const INode*, const FInputDataDestination*>;
 			using FNodeDataReferenceMap = TMap<const INode*, FOperatorDataReferences>;
 			using FOperatorPtr = TUniquePtr<IOperator>;
 
@@ -213,14 +214,17 @@ namespace Metasound
 			// Prune unreachable nodes from InOutNodes
 			FBuildStatus PruneNodes(FBuildContext& InContext, TArray<const INode*>& InOutNodes) const;
 
-			// Get all input data references for a given node.
+			// Get all input data references for a given node for inputs provided internally to the graph.
 			FBuildStatus GatherInputDataReferences(FBuildContext& InContext, const INode* InNode, const FNodeEdgeMultiMap& InEdgeMap, const FNodeDataReferenceMap& InDataReferenceMap, FDataReferenceCollection& OutCollection) const;
 
+			// Get all input data references for a given node for inputs provided externally to the graph.
+			FBuildStatus GatherExternalInputDataReferences(FBuildContext& InContext, const INode* InNode, const FNodeDestinationMap& InNodeDestinationMap, const FDataReferenceCollection& InExternalCollection, FDataReferenceCollection& OutDataReferences) const;
+
 			// Get all input/output data references for a given graph.
 			FBuildStatus GatherGraphDataReferences(FBuildContext& InContext, FNodeDataReferenceMap& InNodeDataReferences, FDataReferenceCollection& OutGraphInputs, FDataReferenceCollection& OutGraphOutputs) const;
 
 			// Call the operator factories for the nodes
-			FBuildStatus CreateOperators(FBuildContext& InContext, const TArray<const INode*>& InSortedNodes, TArray<FOperatorPtr>& OutOperators, FNodeDataReferenceMap& OutDataReferences) const;
+			FBuildStatus CreateOperators(FBuildContext& InContext, const TArray<const INode*>& InSortedNodes, const FDataReferenceCollection& InGraphInputs, TArray<FOperatorPtr>& OutOperators, FNodeDataReferenceMap& OutDataReferences) const;
 
 			// Create the final graph operator from the individual operators.
 			TUniquePtr<IOperator> CreateGraphOperator(FBuildContext& InContext, TArray<FOperatorPtr>& InOperators, FNodeDataReferenceMap& InNodeDataReferences) const;
