Add ability to incrementally update frontend queries.

diff --git a/Engine/Plugins/Runtime/Metasound/Source/MetasoundEditor/Private/MetasoundEditorGraphBuilder.cpp b/Engine/Plugins/Runtime/Metasound/Source/MetasoundEditor/Private/MetasoundEditorGraphBuilder.cpp
--- a/Engine/Plugins/Runtime/Metasound/Source/MetasoundEditor/Private/MetasoundEditorGraphBuilder.cpp
+++ b/Engine/Plugins/Runtime/Metasound/Source/MetasoundEditor/Private/MetasoundEditorGraphBuilder.cpp
@@ -18,6 +18,7 @@
 #include "MetasoundFrontendQuery.h"
 #include "MetasoundFrontendQuerySteps.h"
 #include "MetasoundFrontendRegistries.h"
+#include "MetasoundFrontendSearchEngine.h"
 #include "MetasoundLiteral.h"
 #include "MetasoundUObjectRegistry.h"
 #include "Modules/ModuleManager.h"
@@ -609,26 +610,17 @@ namespace Metasound
 
 			else if (UMetasoundEditorGraphExternalNode* ExternalNode = Cast<UMetasoundEditorGraphExternalNode>(&InGraphNode))
 			{
-				FFrontendQuery Query;
-				const FFrontendQuerySelectionView Result = Query
-					.AddStep<FGenerateAllAvailableNodeClasses>()
-					.AddStep<FFilterClassesByClassName>(ExternalNode->ClassName)
-					.ExecuteQuery();
-
-				TArrayView<const FFrontendQueryEntry* const> Selection = Result.GetSelection();
-				if (ensure(!Selection.IsEmpty()))
+				FMetasoundFrontendClass FrontendClass;
+				bool bDidFindClassWithName = ISearchEngine::Get().FindClassWithHighestVersion(ExternalNode->ClassName.ToNodeClassName(), FrontendClass);
+				if (ensure(bDidFindClassWithName))
 				{
-					const FMetasoundFrontendClass* FrontendClass = Selection[0]->Value.TryGet<FMetasoundFrontendClass>();
-					if (ensure(FrontendClass))
-					{
-						FMetasoundAssetBase* MetasoundAsset = IMetasoundUObjectRegistry::Get().GetObjectAsAssetBase(&InMetasound);
-						check(MetasoundAsset);
+					FMetasoundAssetBase* MetasoundAsset = IMetasoundUObjectRegistry::Get().GetObjectAsAssetBase(&InMetasound);
+					check(MetasoundAsset);
 
-						Frontend::FNodeHandle NewNode = MetasoundAsset->GetRootGraphHandle()->AddNode(FrontendClass->Metadata);
-						ExternalNode->SetNodeID(NewNode->GetID());
+					Frontend::FNodeHandle NewNode = MetasoundAsset->GetRootGraphHandle()->AddNode(FrontendClass.Metadata);
+					ExternalNode->SetNodeID(NewNode->GetID());
 
-						NodeHandle = NewNode;
-					}
+					NodeHandle = NewNode;
 				}
 			}
 
diff --git a/Engine/Plugins/Runtime/Metasound/Source/MetasoundFrontend/Private/MetasoundFrontendQuery.cpp b/Engine/Plugins/Runtime/Metasound/Source/MetasoundFrontend/Private/MetasoundFrontendQuery.cpp
--- a/Engine/Plugins/Runtime/Metasound/Source/MetasoundFrontend/Private/MetasoundFrontendQuery.cpp
+++ b/Engine/Plugins/Runtime/Metasound/Source/MetasoundFrontend/Private/MetasoundFrontendQuery.cpp
@@ -2,34 +2,37 @@
 
 #include "MetasoundFrontendQuery.h"
 
+#include "Algo/IsSorted.h"
 #include "Algo/Sort.h"
+#include "Algo/Transform.h"
 #include "CoreMinimal.h"
 #include "MetasoundFrontendDocument.h"
 
 namespace Metasound
 {
-
 	namespace FrontendQueryPrivate
 	{
-		struct FGenerateFunctionFrontendQueryStep: IFrontendQueryGenerateStep
+		// Wrapper for step defined by function
+		struct FStreamFunctionFrontendQueryStep: IFrontendQueryStreamStep
 		{
-			using FGenerateFunction = FFrontendQueryStep::FGenerateFunction;
+			using FStreamFunction = FFrontendQueryStep::FStreamFunction;
 
-			FGenerateFunctionFrontendQueryStep(FGenerateFunction InFunc)
-			:	Func(InFunc)
+			FStreamFunctionFrontendQueryStep(FStreamFunction&& InFunc)
+			:	Func(MoveTemp(InFunc))
 			{
 			}
 
-			void Generate(TArray<FFrontendQueryEntry>& OutEntries) const override
+			void Stream(TArray<FFrontendQueryEntry>& OutEntries) override
 			{
 				Func(OutEntries);
 			}
 
 			private:
 
-			FGenerateFunction Func;
+			FStreamFunction Func;
 		};
 
+		// Wrapper for step defined by function
 		struct FMapFunctionFrontendQueryStep: IFrontendQueryMapStep
 		{
 			using FMapFunction = FFrontendQueryStep::FMapFunction;
@@ -47,6 +50,7 @@ namespace Metasound
 			FMapFunction Func;
 		};
 
+		// Wrapper for step defined by function
 		struct FReduceFunctionFrontendQueryStep: IFrontendQueryReduceStep
 		{
 			using FReduceFunction = FFrontendQueryStep::FReduceFunction;
@@ -56,7 +60,7 @@ namespace Metasound
 			{
 			}
 
-			void Reduce(FFrontendQueryEntry::FKey InKey, TArrayView<FFrontendQueryEntry*>& InEntries, FReduceOutputView& OutResult) const override
+			void Reduce(FFrontendQueryEntry::FKey InKey, TArrayView<FFrontendQueryEntry * const>& InEntries, FReduceOutputView& OutResult) const override
 			{
 				return Func(InKey, InEntries, OutResult);
 			}
@@ -64,6 +68,7 @@ namespace Metasound
 			FReduceFunction Func;
 		};
 
+		// Wrapper for step defined by function
 		struct FFilterFunctionFrontendQueryStep: IFrontendQueryFilterStep
 		{
 			using FFilterFunction = FFrontendQueryStep::FFilterFunction;
@@ -81,6 +86,7 @@ namespace Metasound
 			FFilterFunction Func;
 		};
 
+		// Wrapper for step defined by function
 		struct FScoreFunctionFrontendQueryStep: IFrontendQueryScoreStep
 		{
 			using FScoreFunction = FFrontendQueryStep::FScoreFunction;
@@ -98,6 +104,7 @@ namespace Metasound
 			FScoreFunction Func;
 		};
 
+		// Wrapper for step defined by function
 		struct FSortFunctionFrontendQueryStep: IFrontendQuerySortStep
 		{
 			using FSortFunction = FFrontendQueryStep::FSortFunction;
@@ -115,6 +122,7 @@ namespace Metasound
 			FSortFunction Func;
 		};
 
+		// Wrapper for step defined by function
 		struct FLimitFunctionFrontendQueryStep: IFrontendQueryLimitStep
 		{
 			using FLimitFunction = FFrontendQueryStep::FLimitFunction;
@@ -132,32 +140,128 @@ namespace Metasound
 			FLimitFunction Func;
 		};
 
+		// Base implementation of a query step execution.
+		struct FStepExecuterBase : public FFrontendQueryStep::IStepExecuter
+		{
+
+			// Appends incremental results to the output results.
+			EResultModificationState Append(const FFrontendQuerySelection& InIncremental, FFrontendQuerySelection& InOutResult) const
+			{
+				EResultModificationState ResultState = EResultModificationState::Unmodified;
+
+				if (InIncremental.GetSelection().Num() > 0)
+				{
+					InOutResult.AppendToStorageAndSelection(InIncremental);
+					ResultState = EResultModificationState::Modified;
+				}
+
+				return ResultState;
+			}
+		};
+
+		// Implements execution of a source
+		struct FSourceStepExecuter : public FStepExecuterBase
+		{
+			FSourceStepExecuter(TUniquePtr<IFrontendQuerySource>&& InSource)
+			: Source(MoveTemp(InSource))
+			{
+			}
+
+			// Get the latest entries from the source
+			virtual EResultModificationState Increment(FFrontendQuerySelection& InOutResult) override
+			{
+				if (Source.IsValid())
+				{
+					TArray<FFrontendQueryEntry> NewEntries;
+					Source->Stream(NewEntries);
+
+					if (NewEntries.Num() > 0)
+					{
+						InOutResult.AppendToStorageAndSelection(NewEntries);
+						return EResultModificationState::Modified;
+					}
+				}
+
+				return EResultModificationState::Unmodified;
+			}
+
+			virtual EResultModificationState Merge(const FFrontendQuerySelection& InIncremental, FFrontendQuerySelection& InOutResult) override
+			{
+				return Append(InIncremental, InOutResult);
+			}
+
+			virtual bool IsMergeRequiredForIncremental() const override
+			{
+				return false;
+			}
+
+			virtual EResultModificationState Execute(FFrontendQuerySelection& InOutResult) override
+			{
+				Reset(); // Reset source to original state.
+				return Increment(InOutResult);
+			}
+
+			virtual void Reset() override 
+			{ 
+				if (Source.IsValid())
+				{
+					Source->Reset();
+				}
+			}
+
+		private:
+
+			TUniquePtr<IFrontendQuerySource> Source;
+		};
+
 		template<typename StepType>
-		struct TStepExecuter : public FFrontendQueryStep::IStepExecuter
+		struct TStepExecuter : public FStepExecuterBase
 		{
 			TStepExecuter(TUniquePtr<StepType>&& InStep)
 			:	Step(MoveTemp(InStep))
 			{
 			}
 
+			virtual void Reset() override { }
+
+			virtual bool IsMergeRequiredForIncremental() const override { return false; }
+
+		protected:
+
 			TUniquePtr<StepType> Step;
 		};
 
-
-		struct FGenerateStepExecuter : TStepExecuter<IFrontendQueryGenerateStep>
+		struct FStreamStepExecuter : TStepExecuter<IFrontendQueryStreamStep>
 		{
-			using TStepExecuter<IFrontendQueryGenerateStep>::TStepExecuter;
+			using TStepExecuter<IFrontendQueryStreamStep>::TStepExecuter;
 
-			void ExecuteStep(FFrontendQuerySelection& InOutResult) const override
+			virtual EResultModificationState Increment(FFrontendQuerySelection& InOutResult) override
 			{
+				EResultModificationState ResultState = EResultModificationState::Unmodified;
+
 				if (Step.IsValid())
 				{
 					TArray<FFrontendQueryEntry> NewEntries;
+					Step->Stream(NewEntries);
 
-					Step->Generate(NewEntries);
-
-					InOutResult.AppendToStorageAndSelection(NewEntries);
+					if (NewEntries.Num() > 0)
+					{
+						ResultState = EResultModificationState::Modified;
+						InOutResult.AppendToStorageAndSelection(NewEntries);
+					}
 				}
+
+				return ResultState;
+			}
+
+			virtual EResultModificationState Merge(const FFrontendQuerySelection& InIncremental, FFrontendQuerySelection& InOutResult) override
+			{
+				return Append(InIncremental, InOutResult);
+			}
+
+			virtual EResultModificationState Execute(FFrontendQuerySelection& InOutResult) override
+			{
+				return Increment(InOutResult);
 			}
 		};
 
@@ -165,10 +269,17 @@ namespace Metasound
 		{
 			using TStepExecuter<IFrontendQueryMapStep>::TStepExecuter;
 
-			void ExecuteStep(FFrontendQuerySelection& InOutResult) const override
+			virtual EResultModificationState Increment(FFrontendQuerySelection& InOutResult) override
 			{
+				EResultModificationState ResultState = EResultModificationState::Unmodified;
+
 				if (Step.IsValid())
 				{
+					if (InOutResult.GetSelection().Num() > 0)
+					{
+						ResultState = EResultModificationState::Modified;
+					}
+
 					for (FFrontendQueryEntry* Entry : InOutResult.GetSelection())
 					{
 						if (nullptr != Entry)
@@ -177,66 +288,161 @@ namespace Metasound
 						}
 					}
 				}
+
+				return ResultState;
+			}
+
+			virtual EResultModificationState Merge(const FFrontendQuerySelection& InIncremental, FFrontendQuerySelection& InOutResult) override
+			{
+				return Append(InIncremental, InOutResult);
+			}
+
+			virtual EResultModificationState Execute(FFrontendQuerySelection& InOutResult) override
+			{
+				return Increment(InOutResult);
 			}
 		};
 
 		struct FReduceStepExecuter : TStepExecuter<IFrontendQueryReduceStep>
 		{
+			using FKey = FFrontendQueryEntry::FKey;
+
+		private:
+
+			static FFrontendQueryEntry::FKey GetKey(const FFrontendQueryEntry* InEntry)
+			{
+				return (nullptr != InEntry) ? InEntry->Key : FFrontendQueryEntry::InvalidKey;
+			}
+
+			// Reduce entries.
+			//
+			// @parma InKey - The key associated with InEntriesToReduce
+			// @parma InEntriesToReduce - Entries to either append or reduce.
+			// @param OutResult - The output query selection to append to.
+			void Reduce(const FFrontendQueryEntry::FKey& InKey, TArrayView<FFrontendQueryEntry* const> InEntriesToReduce, FFrontendQuerySelection& OutResult)
+			{
+				check(Step.IsValid());
+
+				// Note: InEntriesToReduce is has valid storage in OutResult.
+				FFrontendQuerySelection::FReduceOutputView ReduceResult(InEntriesToReduce);
+				Step->Reduce(InKey, InEntriesToReduce, ReduceResult);
+				OutResult.AppendToStorageAndSelection(MoveTemp(ReduceResult));
+			}
+
+			// Iterate over sorted entries.
+			template<typename FuncType>
+			void IterSortedEntries(TArrayView<FFrontendQueryEntry* const> SortedSelection, FuncType InFunc)
+			{
+				// input selection must be sorted by key.
+				check(Algo::IsSortedBy(SortedSelection, FReduceStepExecuter::GetKey));
+
+				if (SortedSelection.Num() > 0)
+				{
+					int32 StartIndex = 0;
+					
+					FFrontendQueryEntry::FKey CurrentKey = GetKey(SortedSelection[0]);
+
+					const int32 Num = SortedSelection.Num();
+					for (int32 EndIndex = 1; EndIndex < Num; EndIndex++)
+					{
+						FFrontendQueryEntry::FKey ThisKey = GetKey(SortedSelection[EndIndex]);
+						if (CurrentKey != ThisKey)
+						{
+							InFunc(CurrentKey, SortedSelection.Slice(StartIndex, EndIndex - StartIndex));
+
+							CurrentKey = ThisKey;
+							StartIndex = EndIndex;
+						}
+					}
+
+					InFunc(CurrentKey, SortedSelection.Slice(StartIndex, SortedSelection.Num() - StartIndex));
+				}
+			}
+
+			TArray<FFrontendQueryEntry*> GetSortedSelection(TArrayView<FFrontendQueryEntry* const> InSelection)
+			{
+				TArray<FFrontendQueryEntry*> SortedSelection(InSelection.GetData(), InSelection.Num()); 
+				Algo::SortBy(SortedSelection, GetKey);
+				return SortedSelection;
+			}
+
+		public:
+
 			using TStepExecuter<IFrontendQueryReduceStep>::TStepExecuter;
 
-			void ExecuteStep(FFrontendQuerySelection& InOutResult) const override
+			virtual EResultModificationState Increment(FFrontendQuerySelection& InOutResult) override
+			{
+				check(IsMergeRequiredForIncremental());
+				return Execute(InOutResult);
+			}
+
+			virtual EResultModificationState Merge(const FFrontendQuerySelection& InIncremental, FFrontendQuerySelection& InOutResult) override
 			{
+				EResultModificationState ResultState = EResultModificationState::Unmodified;
+
 				if (Step.IsValid())
-				{
-					if (InOutResult.GetSelection().Num() > 0)
+				{ 	
+					if (InIncremental.GetSelection().Num() > 0)
 					{
-						TArrayView<FFrontendQueryEntry*> SelectionView = InOutResult.GetSelection();
-						TArray<FFrontendQueryEntry*> SortedSelection(SelectionView.GetData(), SelectionView.Num()); 
+						TSet<FKey> NewKeys;
+
+						Algo::Transform(InIncremental.GetSelection(), NewKeys, GetKey);
+
+						InOutResult.AppendToStorageAndSelection(InIncremental);
+						TArray<FFrontendQueryEntry*> SortedSelection = GetSortedSelection(InOutResult.GetSelection());
 						InOutResult.ResetSelection();
 
-						Algo::SortBy(SortedSelection, [](const FFrontendQueryEntry* InEntry)
+						IterSortedEntries(SortedSelection, [&](const FKey& InKey, TArrayView<FFrontendQueryEntry* const> InEntriesToReduce)
+						{
+							if (NewKeys.Contains(InKey))
 							{
-								if (nullptr == InEntry)
-								{
-									return FFrontendQueryEntry::InvalidKey;
-								}
-
-								return InEntry->Key;
+								// Some of the entries re new. Need to re-reduce.
+								this->Reduce(InKey, InEntriesToReduce, InOutResult);
 							}
-						);
+							else
+							{
+								// All these entries are owned by the result already.
+								InOutResult.AppendToSelection(InEntriesToReduce);
+							}
+						});
 
-						if (!SortedSelection.IsEmpty())
-						{
-							int32 StartIndex = 0;
-							FFrontendQueryEntry::FKey CurrentKey = SortedSelection[0] ? SortedSelection[0]->Key : FFrontendQueryEntry::InvalidKey;
+						ResultState = EResultModificationState::Modified;
+					}
+				}
 
-							auto ReduceFunc = [&](int32 Index)
-							{
-								TArrayView<FFrontendQueryEntry*> ResultsToReduce(&SortedSelection[StartIndex], Index - StartIndex);
-								FFrontendQuerySelection::FReduceOutputView ReduceResult(ResultsToReduce);
+				return ResultState;
+			}
 
-								Step->Reduce(CurrentKey, ResultsToReduce, ReduceResult);
+			virtual EResultModificationState Execute(FFrontendQuerySelection& InOutResult) override
+			{
+				EResultModificationState ResultState = EResultModificationState::Unmodified;
 
-								InOutResult.AppendToStorageAndSelection(ReduceResult);
-							};
+				if (Step.IsValid())
+				{ 	
+					if (InOutResult.GetSelection().Num() > 0)
+					{
+						TArray<FFrontendQueryEntry*> SortedSelection = GetSortedSelection(InOutResult.GetSelection());
+						InOutResult.ResetSelection();
 
-							const int32 Num = SortedSelection.Num();
-							for (int32 ResultIndex = 1; ResultIndex < Num; ResultIndex++)
-							{
-								FFrontendQueryEntry::FKey ThisKey = SortedSelection[ResultIndex] ? SortedSelection[ResultIndex]->Key : FFrontendQueryEntry::InvalidKey;
-								if (CurrentKey != ThisKey)
-								{
-									ReduceFunc(ResultIndex);
-
-									CurrentKey = ThisKey;
-									StartIndex = ResultIndex;
-								}
-							}
+						IterSortedEntries(SortedSelection, [&](const FKey& InKey, TArrayView<FFrontendQueryEntry* const> InEntriesToReduce)
+						{
+							this->Reduce(InKey, InEntriesToReduce, InOutResult);
+						});
 
-							ReduceFunc(Num);
-						}
+						ResultState = EResultModificationState::Modified;
 					}
 				}
+
+				return ResultState;
+			}
+
+
+			virtual bool IsMergeRequiredForIncremental() const override
+			{
+				// Merge is required during an incremental update because
+				// existing results and new results may share the same key and
+				// need re-reducing.
+				return true;
 			}
 		};
 
@@ -244,21 +450,35 @@ namespace Metasound
 		{
 			using TStepExecuter<IFrontendQueryFilterStep>::TStepExecuter;
 
-			void ExecuteStep(FFrontendQuerySelection& InOutResult) const override
+			virtual EResultModificationState Increment(FFrontendQuerySelection& InOutResult) override
 			{
+				EResultModificationState ResultState = EResultModificationState::Unmodified;
+
 				if (Step.IsValid())
 				{
-					TArrayView<FFrontendQueryEntry*> Selection = InOutResult.GetSelection();
-
-					InOutResult.ResetSelection();
+					auto FilterFunc = [&](const FFrontendQueryEntry* Entry)
+					{
+						return (nullptr != Entry) && Step->Filter(*Entry);
+					};
 
-					InOutResult.AppendToSelection(Selection.FilterByPredicate(
-						[&](const FFrontendQueryEntry* Entry)
-						{
-							return (nullptr != Entry) && Step->Filter(*Entry);
-						}
-					));
+					int32 NumRemoved = InOutResult.FilterSelection(FilterFunc);
+					if (NumRemoved > 0)
+					{
+						ResultState = EResultModificationState::Modified;
+					}
 				}
+
+				return ResultState;
+			}
+
+			virtual EResultModificationState Merge(const FFrontendQuerySelection& InIncremental, FFrontendQuerySelection& InOutResult) override
+			{
+				return Append(InIncremental, InOutResult);
+			}
+
+			virtual EResultModificationState Execute(FFrontendQuerySelection& InOutResult) override
+			{
+				return Increment(InOutResult);
 			}
 		};
 
@@ -266,10 +486,17 @@ namespace Metasound
 		{
 			using TStepExecuter<IFrontendQueryScoreStep>::TStepExecuter;
 
-			void ExecuteStep(FFrontendQuerySelection& InOutResult) const override
+			virtual EResultModificationState Increment(FFrontendQuerySelection& InOutResult) override
 			{
+				EResultModificationState ResultState = EResultModificationState::Unmodified;
+
 				if (Step.IsValid())
 				{
+					if (InOutResult.GetSelection().Num() > 0)
+					{
+						ResultState = EResultModificationState::Modified;
+					}
+
 					for (FFrontendQueryEntry* Entry : InOutResult.GetSelection())
 					{
 						if (nullptr != Entry)
@@ -278,6 +505,18 @@ namespace Metasound
 						}
 					}
 				}
+
+				return ResultState;
+			}
+
+			virtual EResultModificationState Merge(const FFrontendQuerySelection& InIncremental, FFrontendQuerySelection& InOutResult) override
+			{
+				return Append(InIncremental, InOutResult);
+			}
+
+			virtual EResultModificationState Execute(FFrontendQuerySelection& InOutResult) override
+			{
+				return Increment(InOutResult);
 			}
 		};
 		
@@ -285,17 +524,66 @@ namespace Metasound
 		{
 			using TStepExecuter<IFrontendQuerySortStep>::TStepExecuter;
 
-			void ExecuteStep(FFrontendQuerySelection& InOutResult) const override
+			EResultModificationState Increment(FFrontendQuerySelection& InOutResult) override
+			{
+				check(IsMergeRequiredForIncremental());
+				return Execute(InOutResult);
+			}
+
+			virtual EResultModificationState Merge(const FFrontendQuerySelection& InIncremental, FFrontendQuerySelection& InOutResult) override
 			{
+				EResultModificationState ResultState = EResultModificationState::Unmodified;
+
 				if (Step.IsValid())
 				{
-					InOutResult.GetSelection().Sort(
-						[&](const FFrontendQueryEntry& InLHS, const FFrontendQueryEntry& InRHS)
-						{
-							return Step->Sort(InLHS, InRHS);
-						}
-					);
+					const bool bIncrementalHasNonEmptySelection = InIncremental.GetSelection().Num() > 0;
+
+					if (bIncrementalHasNonEmptySelection)
+					{
+						// Append incremental data to this result
+						InOutResult.AppendToStorageAndSelection(InIncremental);
+
+						InOutResult.GetSelection().Sort(
+							[&](const FFrontendQueryEntry& InLHS, const FFrontendQueryEntry& InRHS)
+							{
+								return Step->Sort(InLHS, InRHS);
+							}
+						);
+
+						ResultState = EResultModificationState::Modified;
+					}
+				}
+
+				return ResultState;
+			}
+
+			virtual EResultModificationState Execute(FFrontendQuerySelection& InOutResult) override
+			{
+				EResultModificationState ResultState = EResultModificationState::Unmodified;
+
+				if (Step.IsValid())
+				{
+					if (InOutResult.GetSelection().Num() > 0)
+					{
+						InOutResult.GetSelection().Sort(
+							[&](const FFrontendQueryEntry& InLHS, const FFrontendQueryEntry& InRHS)
+							{
+								return Step->Sort(InLHS, InRHS);
+							}
+						);
+
+						ResultState = EResultModificationState::Modified;
+					}
 				}
+
+				return ResultState;
+			}
+
+			virtual bool IsMergeRequiredForIncremental() const 
+			{
+				// Merging is required during incremental update because sort
+				// order may change after merge.
+				return true;
 			}
 		};
 
@@ -303,21 +591,76 @@ namespace Metasound
 		{
 			using TStepExecuter<IFrontendQueryLimitStep>::TStepExecuter;
 
-			void ExecuteStep(FFrontendQuerySelection& InOutResult) const override
+			EResultModificationState Increment(FFrontendQuerySelection& InOutResult) override
+			{
+				check(IsMergeRequiredForIncremental());
+				return Execute(InOutResult);
+			}
+
+			virtual EResultModificationState Merge(const FFrontendQuerySelection& InIncremental, FFrontendQuerySelection& InOutResult) override
 			{
+				EResultModificationState ResultState = EResultModificationState::Unmodified;
+
 				if (Step.IsValid())
 				{
 					int32 Limit = Step->Limit();
+
+					if (InIncremental.GetSelection().Num() > 0)
+					{
+						ResultState = EResultModificationState::Modified;
+						InOutResult.AppendToStorageAndSelection(InIncremental);
+
+						if (InOutResult.GetSelection().Num() > Limit)
+						{
+							InOutResult.SetSelection(InOutResult.GetSelection().Slice(0, Limit));
+						}
+					}
+				}
+
+				return ResultState;
+			}
+
+			virtual EResultModificationState Execute(FFrontendQuerySelection& InOutResult) override
+			{
+				EResultModificationState ResultState = EResultModificationState::Unmodified;
+
+				if (Step.IsValid())
+				{
+					int32 Limit = Step->Limit();
+
 					if (InOutResult.GetSelection().Num() > Limit)
 					{
+						ResultState = EResultModificationState::Modified;
 						InOutResult.SetSelection(InOutResult.GetSelection().Slice(0, Limit));
 					}
 				}
+
+				return ResultState;
+			}
+
+			virtual bool IsMergeRequiredForIncremental() const 
+			{
+				// Merge is required because values under limit may be different
+				// after merge.
+				return true;
 			}
 		};
 	}
 
-	FFrontendQuerySelection::FReduceOutputView::FReduceOutputView(TArrayView<FFrontendQueryEntry*> InEntries)
+	FFrontendQuerySelection::FFrontendQuerySelection(const FFrontendQuerySelection& InOther)
+	{
+		AppendToStorageAndSelection(InOther);
+	}
+
+	FFrontendQuerySelection& FFrontendQuerySelection::operator=(const FFrontendQuerySelection& InOther)
+	{
+		ResetStorageAndSelection();
+		AppendToStorageAndSelection(InOther);
+
+		return *this;
+	}
+
+	FFrontendQuerySelection::FReduceOutputView::FReduceOutputView(TArrayView<FFrontendQueryEntry* const> InEntries)
 	:	InitialEntries(InEntries)
 	{
 	}
@@ -326,33 +669,14 @@ namespace Metasound
 	{
 		if (!InitialEntries.Contains(&InResult))
 		{
-			NewEntryStorage.Add(InResult);
+			SelectedNewEntries.Add(MakeUnique<FFrontendQueryEntry>(InResult));
 		}
 		else
 		{
-			ExistingEntryPointers.Add(&InResult);
+			SelectedExistingEntries.Add(&InResult);
 		}
 	}
 
-	TArrayView<FFrontendQueryEntry*> FFrontendQuerySelection::FReduceOutputView::GetSelectedInitialEntries()
-	{
-		return ExistingEntryPointers;
-	}
-
-	TArrayView<FFrontendQueryEntry> FFrontendQuerySelection::FReduceOutputView::GetSelectedNewEntries() 
-	{
-		return NewEntryStorage;
-	}
-
-	TArrayView<FFrontendQueryEntry> FFrontendQuerySelection::GetStorage()
-	{
-		return Storage;
-	}
-
-	TArrayView<const FFrontendQueryEntry> FFrontendQuerySelection::GetStorage() const
-	{
-		return Storage;
-	}
 
 	TArrayView<FFrontendQueryEntry*> FFrontendQuerySelection::GetSelection()
 	{
@@ -367,55 +691,106 @@ namespace Metasound
 	void FFrontendQuerySelection::AppendToStorageAndSelection(const FFrontendQuerySelection& InSelection)
 	{
 		TArrayView<const FFrontendQueryEntry* const> InEntries = InSelection.GetSelection();
-			
-		int32 StorageIndex = Storage.Num();
-		int32 SelectionIndex = Selection.Num();
+		AppendToStorageAndSelection(InSelection.GetSelection());
+	}
 
-		for (const FFrontendQueryEntry* Entry : InEntries)
-		{
-			Storage.Add(*Entry);
-		}
+	void FFrontendQuerySelection::AppendToStorageAndSelection(FFrontendQuerySelection&& InSelection)
+	{
+		InSelection.ShrinkStorageToSelection();
+		AppendToStorageAndSelection(MoveTemp(InSelection.Storage));
+	}
 
-		const int32 Num = Storage.Num() - StorageIndex;
+	void FFrontendQuerySelection::AppendToStorageAndSelection(TArrayView<const FFrontendQueryEntry> InEntries)
+	{
+		const int32 Num = InEntries.Num();
 
-		if (Num > 0)
+		if (Num >  0)
 		{
+			int32 StorageIndex = Storage.Num();
+			int32 SelectionIndex = Selection.Num();
+
+			for (const FFrontendQueryEntry& Entry : InEntries)
+			{
+				Storage.Add(MakeUnique<FFrontendQueryEntry>(Entry));
+			}
 			Selection.AddZeroed(Num);
 
 			for (int32 i = 0; i < Num; i++)
 			{
-				Selection[SelectionIndex] = &Storage[StorageIndex];
+				Selection[SelectionIndex] = Storage[StorageIndex].Get();
 				SelectionIndex++;
 				StorageIndex++;
 			}
 		}
 	}
 
-	void FFrontendQuerySelection::AppendToStorageAndSelection(TArrayView<const FFrontendQueryEntry> InEntries)
+	void FFrontendQuerySelection::AppendToStorageAndSelection(TArrayView<const FFrontendQueryEntry* const> InEntries)
 	{
-		const int32 Num = InEntries.Num();
-
-		if (Num >  0)
+		if (InEntries.Num() >  0)
 		{
 			int32 StorageIndex = Storage.Num();
 			int32 SelectionIndex = Selection.Num();
 
-			Storage.Append(InEntries.GetData(), InEntries.Num());
-			Selection.AddZeroed(Num);
+			int32 NumAdded = 0;
+			for (const FFrontendQueryEntry* Entry : InEntries)
+			{
+				if (nullptr != Entry)
+				{
+					Storage.Add(MakeUnique<FFrontendQueryEntry>(*Entry));
+					NumAdded++;
+				}
+			}
 
-			for (int32 i = 0; i < Num; i++)
+			if (NumAdded > 0)
 			{
-				Selection[SelectionIndex] = &Storage[StorageIndex];
-				SelectionIndex++;
-				StorageIndex++;
+				Selection.AddZeroed(NumAdded);
+
+				for (int32 i = 0; i < NumAdded; i++)
+				{
+					Selection[SelectionIndex] = Storage[StorageIndex].Get();
+					SelectionIndex++;
+					StorageIndex++;
+				}
 			}
 		}
 	}
 
-	void FFrontendQuerySelection::AppendToStorageAndSelection(FReduceOutputView& InReduceView)
+	void FFrontendQuerySelection::AppendToStorageAndSelection(FReduceOutputView&& InReduceView)
 	{
-		AppendToStorageAndSelection(InReduceView.GetSelectedNewEntries());
-		AppendToSelection(InReduceView.GetSelectedInitialEntries());
+		AppendToStorageAndSelection(MoveTemp(InReduceView.SelectedNewEntries));
+		AppendToSelection(InReduceView.SelectedExistingEntries);
+	}
+
+	void FFrontendQuerySelection::AppendToStorageAndSelection(TArray<TUniquePtr<FFrontendQueryEntry>>&& InEntries)
+	{
+		if (InEntries.Num() >  0)
+		{
+			int32 StorageIndex = Storage.Num();
+			int32 SelectionIndex = Selection.Num();
+
+			int32 NumAdded = 0;
+			for (TUniquePtr<FFrontendQueryEntry>& Entry : InEntries)
+			{
+				if (Entry.IsValid())
+				{
+					Storage.Add(MoveTemp(Entry));
+					NumAdded++;
+				}
+			}
+
+			if (NumAdded > 0)
+			{
+				Selection.AddZeroed(NumAdded);
+
+				for (int32 i = 0; i < NumAdded; i++)
+				{
+					Selection[SelectionIndex] = Storage[StorageIndex].Get();
+					SelectionIndex++;
+					StorageIndex++;
+				}
+			}
+		}
+
 	}
 
 	void FFrontendQuerySelection::ResetStorageAndSelection()
@@ -442,11 +817,19 @@ namespace Metasound
 		}
 	}
 
-	void FFrontendQuerySelection::AppendToSelection(TArrayView<FFrontendQueryEntry*> InEntries)
+	void FFrontendQuerySelection::AppendToSelection(TArrayView<FFrontendQueryEntry * const> InEntries)
 	{
 		Selection.Append(InEntries.GetData(), InEntries.Num());
 	}
 
+	void FFrontendQuerySelection::ShrinkStorageToSelection()
+	{
+		TSet<FFrontendQueryEntry*> SelectionSet(Selection);
+		Storage.RemoveAll([&](const TUniquePtr<FFrontendQueryEntry>& InEntry)
+		{
+			return !SelectionSet.Contains(InEntry.Get());
+		});
+	}
 
 
 	FFrontendQuerySelectionView::FFrontendQuerySelectionView(TSharedRef<const FFrontendQuerySelection, ESPMode::ThreadSafe> InResult)
@@ -454,53 +837,53 @@ namespace Metasound
 	{
 	}
 
-	TArrayView<const FFrontendQueryEntry> FFrontendQuerySelectionView::GetStorage() const
-	{
-		return Result->GetStorage();
-	}
-
 	TArrayView<const FFrontendQueryEntry* const> FFrontendQuerySelectionView::GetSelection() const
 	{
 		return Result->GetSelection();
 	}
 
-	FFrontendQueryStep::FFrontendQueryStep(FGenerateFunction InFunc)
-	:	StepExecuter(MakeUnique<FrontendQueryPrivate::FGenerateStepExecuter>(MakeUnique<FrontendQueryPrivate::FGenerateFunctionFrontendQueryStep>(InFunc)))
+	FFrontendQueryStep::FFrontendQueryStep(FStreamFunction&& InFunc)
+	:	StepExecuter(MakeUnique<FrontendQueryPrivate::FStreamStepExecuter>(MakeUnique<FrontendQueryPrivate::FStreamFunctionFrontendQueryStep>(MoveTemp(InFunc))))
 	{
 	}
 
-	FFrontendQueryStep::FFrontendQueryStep(FMapFunction InFunc)
-	:	StepExecuter(MakeUnique<FrontendQueryPrivate::FMapStepExecuter>(MakeUnique<FrontendQueryPrivate::FMapFunctionFrontendQueryStep>(InFunc)))
+	FFrontendQueryStep::FFrontendQueryStep(FMapFunction&& InFunc)
+	:	StepExecuter(MakeUnique<FrontendQueryPrivate::FMapStepExecuter>(MakeUnique<FrontendQueryPrivate::FMapFunctionFrontendQueryStep>(MoveTemp(InFunc))))
 	{
 	}
 
-	FFrontendQueryStep::FFrontendQueryStep(FReduceFunction InFunc)
-	:	StepExecuter(MakeUnique<FrontendQueryPrivate::FReduceStepExecuter>(MakeUnique<FrontendQueryPrivate::FReduceFunctionFrontendQueryStep>(InFunc)))
+	FFrontendQueryStep::FFrontendQueryStep(FReduceFunction&& InFunc)
+	:	StepExecuter(MakeUnique<FrontendQueryPrivate::FReduceStepExecuter>(MakeUnique<FrontendQueryPrivate::FReduceFunctionFrontendQueryStep>(MoveTemp(InFunc))))
 	{
 	}
 
-	FFrontendQueryStep::FFrontendQueryStep(FFilterFunction InFunc)
-	:	StepExecuter(MakeUnique<FrontendQueryPrivate::FFilterStepExecuter>(MakeUnique<FrontendQueryPrivate::FFilterFunctionFrontendQueryStep>(InFunc)))
+	FFrontendQueryStep::FFrontendQueryStep(FFilterFunction&& InFunc)
+	:	StepExecuter(MakeUnique<FrontendQueryPrivate::FFilterStepExecuter>(MakeUnique<FrontendQueryPrivate::FFilterFunctionFrontendQueryStep>(MoveTemp(InFunc))))
 	{
 	}
 	
-	FFrontendQueryStep::FFrontendQueryStep(FScoreFunction InFunc)
-	:	StepExecuter(MakeUnique<FrontendQueryPrivate::FScoreStepExecuter>(MakeUnique<FrontendQueryPrivate::FScoreFunctionFrontendQueryStep>(InFunc)))
+	FFrontendQueryStep::FFrontendQueryStep(FScoreFunction&& InFunc)
+	:	StepExecuter(MakeUnique<FrontendQueryPrivate::FScoreStepExecuter>(MakeUnique<FrontendQueryPrivate::FScoreFunctionFrontendQueryStep>(MoveTemp(InFunc))))
+	{
+	}
+
+	FFrontendQueryStep::FFrontendQueryStep(FSortFunction&& InFunc)
+	:	StepExecuter(MakeUnique<FrontendQueryPrivate::FSortStepExecuter>(MakeUnique<FrontendQueryPrivate::FSortFunctionFrontendQueryStep>(MoveTemp(InFunc))))
 	{
 	}
 
-	FFrontendQueryStep::FFrontendQueryStep(FSortFunction InFunc)
-	:	StepExecuter(MakeUnique<FrontendQueryPrivate::FSortStepExecuter>(MakeUnique<FrontendQueryPrivate::FSortFunctionFrontendQueryStep>(InFunc)))
+	FFrontendQueryStep::FFrontendQueryStep(FLimitFunction&& InFunc)
+	:	StepExecuter(MakeUnique<FrontendQueryPrivate::FLimitStepExecuter>(MakeUnique<FrontendQueryPrivate::FLimitFunctionFrontendQueryStep>(MoveTemp(InFunc))))
 	{
 	}
 
-	FFrontendQueryStep::FFrontendQueryStep(FLimitFunction InFunc)
-	:	StepExecuter(MakeUnique<FrontendQueryPrivate::FLimitStepExecuter>(MakeUnique<FrontendQueryPrivate::FLimitFunctionFrontendQueryStep>(InFunc)))
+	FFrontendQueryStep::FFrontendQueryStep(TUniquePtr<IFrontendQuerySource>&& InSource)
+	:	StepExecuter(MakeUnique<FrontendQueryPrivate::FSourceStepExecuter>(MoveTemp(InSource)))
 	{
 	}
 
-	FFrontendQueryStep::FFrontendQueryStep(TUniquePtr<IFrontendQueryGenerateStep>&& InStep)
-	:	StepExecuter(MakeUnique<FrontendQueryPrivate::FGenerateStepExecuter>(MoveTemp(InStep)))
+	FFrontendQueryStep::FFrontendQueryStep(TUniquePtr<IFrontendQueryStreamStep>&& InStep)
+	:	StepExecuter(MakeUnique<FrontendQueryPrivate::FStreamStepExecuter>(MoveTemp(InStep)))
 	{
 	}
 	
@@ -534,11 +917,51 @@ namespace Metasound
 	{
 	}
 
-	void FFrontendQueryStep::ExecuteStep(FFrontendQuerySelection& InOutResult)
+	FFrontendQueryStep::EResultModificationState FFrontendQueryStep::Execute(FFrontendQuerySelection& InOutResult)
+	{
+		if (StepExecuter.IsValid())
+		{
+			return StepExecuter->Execute(InOutResult);
+		}
+
+		return EResultModificationState::Unmodified;
+	}
+
+	FFrontendQueryStep::EResultModificationState FFrontendQueryStep::Increment(FFrontendQuerySelection& InOutResult)
+	{
+		if (StepExecuter.IsValid())
+		{
+			return StepExecuter->Increment(InOutResult);
+		}
+
+		return EResultModificationState::Unmodified;
+	}
+
+	FFrontendQueryStep::EResultModificationState FFrontendQueryStep::Merge(const FFrontendQuerySelection& InIncremental, FFrontendQuerySelection& InOutResult)
+	{
+		if (StepExecuter.IsValid())
+		{
+			return StepExecuter->Merge(InIncremental, InOutResult);
+		}
+
+		return EResultModificationState::Unmodified;
+	}
+
+	bool FFrontendQueryStep::IsMergeRequiredForIncremental() const
 	{
 		if (StepExecuter.IsValid())
 		{
-			StepExecuter->ExecuteStep(InOutResult);
+			return StepExecuter->IsMergeRequiredForIncremental();
+		}
+
+		return false;
+	}
+
+	void FFrontendQueryStep::Reset()
+	{
+		if (StepExecuter.IsValid())
+		{
+			StepExecuter->Reset();
 		}
 	}
 
@@ -552,52 +975,60 @@ namespace Metasound
 		return Steps;
 	}
 
-	FFrontendQuery& FFrontendQuery::AddGenerateLambdaStep(FGenerateFunction InFunc)
+
+	FFrontendQuery& FFrontendQuery::AddStreamLambdaStep(FStreamFunction&& InFunc)
 	{
-		return AddFunctionStep(InFunc);
+		return AddFunctionStep(MoveTemp(InFunc));
 	}
 
-	FFrontendQuery& FFrontendQuery::AddMapLambdaStep(FMapFunction InFunc)
+	FFrontendQuery& FFrontendQuery::AddMapLambdaStep(FMapFunction&& InFunc)
 	{
-		return AddFunctionStep(InFunc);
+		return AddFunctionStep(MoveTemp(InFunc));
 	}
 
-	FFrontendQuery& FFrontendQuery::AddReduceLambdaStep(FReduceFunction InFunc)
+	FFrontendQuery& FFrontendQuery::AddReduceLambdaStep(FReduceFunction&& InFunc)
 	{
-		return AddFunctionStep(InFunc);
+		return AddFunctionStep(MoveTemp(InFunc));
 	}
 
-	FFrontendQuery& FFrontendQuery::AddFilterLambdaStep(FFilterFunction InFunc)
+	FFrontendQuery& FFrontendQuery::AddFilterLambdaStep(FFilterFunction&& InFunc)
 	{
-		return AddFunctionStep(InFunc);
+		return AddFunctionStep(MoveTemp(InFunc));
 	}
 
-	FFrontendQuery& FFrontendQuery::AddScoreLambdaStep(FScoreFunction InFunc)
+	FFrontendQuery& FFrontendQuery::AddScoreLambdaStep(FScoreFunction&& InFunc)
 	{
-		return AddFunctionStep(InFunc);
+		return AddFunctionStep(MoveTemp(InFunc));
 	}
 
-	FFrontendQuery& FFrontendQuery::AddSortLambdaStep(FSortFunction InFunc)
+	FFrontendQuery& FFrontendQuery::AddSortLambdaStep(FSortFunction&& InFunc)
 	{
-		return AddFunctionStep(InFunc);
+		return AddFunctionStep(MoveTemp(InFunc));
 	}
 
-	FFrontendQuery& FFrontendQuery::AddLimitLambdaStep(FLimitFunction InFunc)
+	FFrontendQuery& FFrontendQuery::AddLimitLambdaStep(FLimitFunction&& InFunc)
 	{
-		return AddFunctionStep(InFunc);
+		return AddFunctionStep(MoveTemp(InFunc));
 	}
 
 	FFrontendQuery& FFrontendQuery::AddStep(TUniquePtr<FFrontendQueryStep>&& InStep)
 	{
-		Steps.Add(MoveTemp(InStep));
+		if (ensure(InStep.IsValid()))
+		{
+			int32 StepIndex = Steps.Num();
+			if (InStep->IsMergeRequiredForIncremental())
+			{
+				StepResultCache.Add(StepIndex);
+			}
+
+			Steps.Add(MoveTemp(InStep));
+		}
 		return *this;
 	}
 
-	FFrontendQuerySelectionView FFrontendQuery::ExecuteQuery()
+	FFrontendQuerySelectionView FFrontendQuery::Execute()
 	{
-		Reset();
-
-		Execute(*Result);
+		UpdateResult();
 
 		return FFrontendQuerySelectionView(Result);
 	}
@@ -605,18 +1036,15 @@ namespace Metasound
 	FFrontendQuerySelectionView FFrontendQuery::Reset()
 	{
 		Result->ResetStorageAndSelection();
-		
-		return FFrontendQuerySelectionView(Result);
-	}
-
-	FFrontendQuerySelectionView FFrontendQuery::ExecuteQueryAndAppend()
-	{
-		FFrontendQuerySelection NewSelection;
-
-		Execute(NewSelection);
-
-		Result->AppendToStorageAndSelection(NewSelection);
 
+		for (int32 StepIndex = 0; StepIndex < Steps.Num(); StepIndex++)
+		{
+			if (Steps[StepIndex].IsValid())
+			{
+				Steps[StepIndex]->Reset();
+			}
+		}
+		
 		return FFrontendQuerySelectionView(Result);
 	}
 
@@ -625,16 +1053,80 @@ namespace Metasound
 		return FFrontendQuerySelectionView(Result);
 	}
 
-	void FFrontendQuery::Execute(FFrontendQuerySelection& OutSelection)
+	void FFrontendQuery::UpdateResult()
 	{
+		FFrontendQuerySelection IncrementalResult;
+
+		EResultModificationState IncrementalResultState = EResultModificationState::Unmodified;
+
+		// Perform incremental update sequentially
 		for (int32 StepIndex = 0; StepIndex < Steps.Num(); StepIndex++)
 		{
-			if (!Steps[StepIndex].IsValid())
+			if (ensure(Steps[StepIndex].IsValid()))
 			{
-				continue;
+				FFrontendQueryStep& Step = *Steps[StepIndex];
+
+				IncrementalResultState = Step.Increment(IncrementalResult);
+
+				const bool bMergeIncrementalResults = Step.IsMergeRequiredForIncremental() && (EResultModificationState::Modified == IncrementalResultState);
+
+				if (bMergeIncrementalResults)
+				{
+					const bool bIsResultCacheEmpty = StepResultCache[StepIndex].GetSelection().Num() == 0;
+
+					if (bIsResultCacheEmpty)
+					{
+						// If the prior result for this step is empty, we can continue
+						// using the incremental path.
+						StepResultCache[StepIndex] = IncrementalResult;
+					}
+					else
+					{
+						Step.Merge(IncrementalResult, StepResultCache[StepIndex]);
+
+						*Result = StepResultCache[StepIndex];
+
+						// After results are merged, downstream steps can no longer 
+						// work on incremental data. Run downstream steps on full
+						// result set.
+						ExecuteSteps(StepIndex + 1);
+						return;
+					}
+				}
+			}
+		}
+
+		if (Steps.Num() > 0)
+		{
+			if (ensure(Steps.Last().IsValid()))
+			{
+				FFrontendQueryStep& Step = *Steps.Last();
+				if (EResultModificationState::Modified == IncrementalResultState)
+				{
+					// If all incremental steps are performed and result is modified,
+					// merge the incremental result with the output result.
+					Step.Merge(IncrementalResult, *Result);
+				}
 			}
+		}
+	}
 
-			Steps[StepIndex]->ExecuteStep(OutSelection);
+	void FFrontendQuery::ExecuteSteps(int32 InStartStepIndex)
+	{
+		for (int32 StepIndex = InStartStepIndex; StepIndex < Steps.Num(); StepIndex++)
+		{
+			if (ensure(Steps[StepIndex].IsValid()))
+			{
+				FFrontendQueryStep& Step = *Steps[StepIndex];
+
+				EResultModificationState ExecuteResultState = Step.Execute(*Result);
+
+				const bool bNeedsStepResultCached = (EResultModificationState::Modified == ExecuteResultState) && Step.IsMergeRequiredForIncremental();
+				if (bNeedsStepResultCached)
+				{
+					StepResultCache[StepIndex] = *Result;
+				}
+			}
 		}
 	}
 }
diff --git a/Engine/Plugins/Runtime/Metasound/Source/MetasoundFrontend/Private/MetasoundFrontendQuerySteps.cpp b/Engine/Plugins/Runtime/Metasound/Source/MetasoundFrontend/Private/MetasoundFrontendQuerySteps.cpp
--- a/Engine/Plugins/Runtime/Metasound/Source/MetasoundFrontend/Private/MetasoundFrontendQuerySteps.cpp
+++ b/Engine/Plugins/Runtime/Metasound/Source/MetasoundFrontend/Private/MetasoundFrontendQuerySteps.cpp
@@ -9,31 +9,28 @@
 
 namespace Metasound
 {
-	void FGenerateAllAvailableNodeClasses::Generate(TArray<FFrontendQueryEntry>& OutEntries) const
-	{
-		const TArray<Frontend::FNodeClassInfo> ClassInfos = Frontend::GetAllAvailableNodeClasses();
 
-		for (const Frontend::FNodeClassInfo& ClassInfo : ClassInfos)
-		{
-			OutEntries.Emplace(TInPlaceType<FMetasoundFrontendClass>(), Frontend::GenerateClassDescription(ClassInfo));
-		}
-	}
-
-	FGenerateNewlyAvailableNodeClasses::FGenerateNewlyAvailableNodeClasses()
+	FNodeClassRegistrationEvents::FNodeClassRegistrationEvents()
 	: CurrentTransactionID(Frontend::GetOriginRegistryTransactionID())
 	{
 	}
 
-	void FGenerateNewlyAvailableNodeClasses::Generate(TArray<FFrontendQueryEntry>& OutEntries) const
+	void FNodeClassRegistrationEvents::Stream(TArray<FFrontendQueryEntry>& OutEntries)
 	{
 		const TArray<Frontend::FNodeClassInfo> ClassInfos = Frontend::GetNodeClassesRegisteredSince(CurrentTransactionID, &CurrentTransactionID);
 
 		for (const Frontend::FNodeClassInfo& ClassInfo : ClassInfos)
 		{
-			OutEntries.Emplace(TInPlaceType<FMetasoundFrontendClass>(), Frontend::GenerateClassDescription(ClassInfo));
+			FFrontendQueryEntry::FValue Value(TInPlaceType<FMetasoundFrontendClass>(), Frontend::GenerateClassDescription(ClassInfo));
+			OutEntries.Emplace(MoveTemp(Value));
 		}
 	}
 
+	void FNodeClassRegistrationEvents::Reset()
+	{
+		CurrentTransactionID = Frontend::GetOriginRegistryTransactionID();
+	}
+
 	FFilterClassesByInputVertexDataType::FFilterClassesByInputVertexDataType(const FName& InTypeName)
 	:	InputVertexTypeName(InTypeName)
 	{
@@ -93,7 +90,7 @@ namespace Metasound
 		return static_cast<FFrontendQueryEntry::FKey>(HashKey);
 	}
 
-	void FReduceClassesToHighestVersion::Reduce(FFrontendQueryEntry::FKey InKey, TArrayView<FFrontendQueryEntry*>& InEntries, FReduceOutputView& OutResult) const
+	void FReduceClassesToHighestVersion::Reduce(FFrontendQueryEntry::FKey InKey, TArrayView<FFrontendQueryEntry * const>& InEntries, FReduceOutputView& OutResult) const
 	{
 		FFrontendQueryEntry* HighestVersionEntry = nullptr;
 		int32 HighestMajorVersion = -1;
@@ -120,7 +117,7 @@ namespace Metasound
 	{
 	}
 
-	void FReduceClassesToMajorVersion::Reduce(FFrontendQueryEntry::FKey InKey, TArrayView<FFrontendQueryEntry*>& InEntries, FReduceOutputView& OutResult) const
+	void FReduceClassesToMajorVersion::Reduce(FFrontendQueryEntry::FKey InKey, TArrayView<FFrontendQueryEntry * const>& InEntries, FReduceOutputView& OutResult) const
 	{
 		for (FFrontendQueryEntry* Entry : InEntries)
 		{
diff --git a/Engine/Plugins/Runtime/Metasound/Source/MetasoundFrontend/Private/MetasoundFrontendRegistries.cpp b/Engine/Plugins/Runtime/Metasound/Source/MetasoundFrontend/Private/MetasoundFrontendRegistries.cpp
--- a/Engine/Plugins/Runtime/Metasound/Source/MetasoundFrontend/Private/MetasoundFrontendRegistries.cpp
+++ b/Engine/Plugins/Runtime/Metasound/Source/MetasoundFrontend/Private/MetasoundFrontendRegistries.cpp
@@ -574,7 +574,7 @@ namespace Metasound
 					// Store update to newly registered node in history so nodes
 					// can be queried by transaction ID
 					FNodeClassInfo ClassInfo = { EMetasoundFrontendClassType::External, Key };
-					RegistryTransactionHistory.Add(MakeNodeRegistrationTransaction(ClassInfo));
+					RegistryTransactionHistory.Add(MakeAddNodeRegistrationTransaction(ClassInfo));
 
 					// Store registry elements in map so nodes can be queried using registry key.
 					ExternalNodeRegistry.Add(MoveTemp(Key), MoveTemp(RegistryElement));
diff --git a/Engine/Plugins/Runtime/Metasound/Source/MetasoundFrontend/Private/MetasoundFrontendRegistryTransaction.cpp b/Engine/Plugins/Runtime/Metasound/Source/MetasoundFrontend/Private/MetasoundFrontendRegistryTransaction.cpp
--- a/Engine/Plugins/Runtime/Metasound/Source/MetasoundFrontend/Private/MetasoundFrontendRegistryTransaction.cpp
+++ b/Engine/Plugins/Runtime/Metasound/Source/MetasoundFrontend/Private/MetasoundFrontendRegistryTransaction.cpp
@@ -95,7 +95,7 @@ namespace Metasound
 			
 		}
 
-		FRegistryTransactionPtr MakeNodeRegistrationTransaction(const FNodeClassInfo& InInfo)
+		FRegistryTransactionPtr MakeAddNodeRegistrationTransaction(const FNodeClassInfo& InInfo)
 		{
 			class FNodeRegistrationTransaction : public IRegistryTransaction
 			{
@@ -105,6 +105,44 @@ namespace Metasound
 				{
 				}
 
+				ETransactionType GetTransactionType() const override
+				{
+					return ETransactionType::Add;
+				}
+
+				TUniquePtr<IRegistryTransaction> Clone() const override
+				{
+					return MakeUnique<FNodeRegistrationTransaction>(*this);
+				}
+
+				const FNodeClassInfo* GetNodeClassInfo() const 
+				{
+					return &NodeClassInfo;
+				}
+
+			private:
+
+				FNodeClassInfo NodeClassInfo;
+			};
+
+			return MakeUnique<FNodeRegistrationTransaction>(InInfo);
+		}
+
+		FRegistryTransactionPtr MakeRemoveNodeRegistrationTransaction(const FNodeClassInfo& InInfo)
+		{
+			class FNodeRegistrationTransaction : public IRegistryTransaction
+			{
+			public:
+				FNodeRegistrationTransaction(const FNodeClassInfo& InNodeClassInfo)
+				: NodeClassInfo(InNodeClassInfo)
+				{
+				}
+
+				ETransactionType GetTransactionType() const override
+				{
+					return ETransactionType::Remove;
+				}
+
 				TUniquePtr<IRegistryTransaction> Clone() const override
 				{
 					return MakeUnique<FNodeRegistrationTransaction>(*this);
diff --git a/Engine/Plugins/Runtime/Metasound/Source/MetasoundFrontend/Private/MetasoundFrontendRegistryTransaction.h b/Engine/Plugins/Runtime/Metasound/Source/MetasoundFrontend/Private/MetasoundFrontendRegistryTransaction.h
--- a/Engine/Plugins/Runtime/Metasound/Source/MetasoundFrontend/Private/MetasoundFrontendRegistryTransaction.h
+++ b/Engine/Plugins/Runtime/Metasound/Source/MetasoundFrontend/Private/MetasoundFrontendRegistryTransaction.h
@@ -14,6 +14,13 @@ namespace Metasound
 		/** Returns an ID representing the beginning of the transaction history. */
 		FRegistryTransactionID GetOriginRegistryTransactionID();
 
+		/** Describes the type of transaction. */
+		enum class ETransactionType : uint8
+		{
+			Add,     //< Something was added to the registry.
+			Remove   //< Something was removed from the registry.
+		};
+
 		/** Interface for a registry transaction. */
 		class IRegistryTransaction
 		{
@@ -21,6 +28,9 @@ namespace Metasound
 			virtual ~IRegistryTransaction() = default;
 			virtual TUniquePtr<IRegistryTransaction> Clone() const = 0;
 
+			/** Returns the type of transaction */
+			virtual ETransactionType GetTransactionType() const = 0;
+
 			/** If a FNodeClassInfo added during the transaction, this will return
 			 * a non-null pointer to the FNodeClassInfo. */
 			virtual const FNodeClassInfo* GetNodeClassInfo() const = 0;
@@ -29,7 +39,10 @@ namespace Metasound
 		using FRegistryTransactionPtr = TUniquePtr<IRegistryTransaction>;
 
 		/** Create a node registration transaction. */
-		FRegistryTransactionPtr MakeNodeRegistrationTransaction(const FNodeClassInfo& InInfo);
+		FRegistryTransactionPtr MakeAddNodeRegistrationTransaction(const FNodeClassInfo& InInfo);
+
+		/** Create a node unregistration transaction. */
+		FRegistryTransactionPtr MakeRemoveNodeRegistrationTransaction(const FNodeClassInfo& InInfo);
 
 		/** Maintains a history of IRegistryTransactions. Calls are threadsafe (excluding
 		 * the constructor and destructor.)
diff --git a/Engine/Plugins/Runtime/Metasound/Source/MetasoundFrontend/Private/MetasoundFrontendSearchEngine.cpp b/Engine/Plugins/Runtime/Metasound/Source/MetasoundFrontend/Private/MetasoundFrontendSearchEngine.cpp
--- a/Engine/Plugins/Runtime/Metasound/Source/MetasoundFrontend/Private/MetasoundFrontendSearchEngine.cpp
+++ b/Engine/Plugins/Runtime/Metasound/Source/MetasoundFrontend/Private/MetasoundFrontendSearchEngine.cpp
@@ -64,7 +64,7 @@ namespace Metasound
 			{
 				TUniquePtr<FFrontendQuery> NewQuery = MakeUnique<FFrontendQuery>();
 
-				NewQuery->AddStep<FGenerateNewlyAvailableNodeClasses>();
+				NewQuery->AddStep<FNodeClassRegistrationEvents>();
 
 				if (!bInIncludeDeprecated)
 				{
@@ -79,7 +79,7 @@ namespace Metasound
 
 			if (ensure(Query))
 			{
-				FFrontendQuerySelectionView Selection = Query->ExecuteQueryAndAppend();
+				FFrontendQuerySelectionView Selection = Query->Execute();
 				TArrayView<const FFrontendQueryEntry* const> Entries = Selection.GetSelection();
 				for (const FFrontendQueryEntry* Entry : Selection.GetSelection())
 				{
@@ -104,7 +104,7 @@ namespace Metasound
 			if (!Query)
 			{
 				TUniquePtr<FFrontendQuery> NewQuery = MakeUnique<FFrontendQuery>();
-				NewQuery->AddStep<FGenerateNewlyAvailableNodeClasses>()
+				NewQuery->AddStep<FNodeClassRegistrationEvents>()
 					.AddStep<FFilterClassesByClassName>(InName);
 
 				if (bInSortByVersion)
@@ -119,7 +119,7 @@ namespace Metasound
 
 			if (ensure(nullptr != Query))
 			{
-				FFrontendQuerySelectionView Selection = Query->ExecuteQueryAndAppend();
+				FFrontendQuerySelectionView Selection = Query->Execute();
 				TArrayView<const FFrontendQueryEntry* const> Entries = Selection.GetSelection();
 				for (const FFrontendQueryEntry* Entry : Selection.GetSelection())
 				{
@@ -139,7 +139,7 @@ namespace Metasound
 			{
 				TUniquePtr<FFrontendQuery> NewQuery = MakeUnique<FFrontendQuery>();
 
-				NewQuery->AddStep<FGenerateNewlyAvailableNodeClasses>()
+				NewQuery->AddStep<FNodeClassRegistrationEvents>()
 					.AddStep<FFilterClassesByClassName>(InName)
 					.AddStep<FMapClassNameToMajorVersion>()
 					.AddStep<FReduceClassesToHighestVersion>();
@@ -149,7 +149,7 @@ namespace Metasound
 
 			if (ensure(Query))
 			{
-				FFrontendQuerySelectionView Selection = Query->ExecuteQueryAndAppend();
+				FFrontendQuerySelectionView Selection = Query->Execute();
 				TArrayView<const FFrontendQueryEntry* const> Entries = Selection.GetSelection();
 				if (Entries.IsEmpty())
 				{
@@ -174,7 +174,7 @@ namespace Metasound
 			{
 				TUniquePtr<FFrontendQuery> NewQuery = MakeUnique<FFrontendQuery>();
 
-				NewQuery->AddStep<FGenerateNewlyAvailableNodeClasses>()
+				NewQuery->AddStep<FNodeClassRegistrationEvents>()
 					.AddStep<FFilterClassesByClassName>(InName)
 					.AddStep<FMapClassNameToMajorVersion>()
 					.AddStep<FReduceClassesToMajorVersion>(InMajorVersion);
@@ -184,7 +184,7 @@ namespace Metasound
 
 			if (ensure(Query))
 			{
-				FFrontendQuerySelectionView Selection = Query->ExecuteQueryAndAppend();
+				FFrontendQuerySelectionView Selection = Query->Execute();
 				TArrayView<const FFrontendQueryEntry* const> Entries = Selection.GetSelection();
 				if (Entries.IsEmpty())
 				{
diff --git a/Engine/Plugins/Runtime/Metasound/Source/MetasoundFrontend/Public/MetasoundFrontendQuery.h b/Engine/Plugins/Runtime/Metasound/Source/MetasoundFrontend/Public/MetasoundFrontendQuery.h
--- a/Engine/Plugins/Runtime/Metasound/Source/MetasoundFrontend/Public/MetasoundFrontendQuery.h
+++ b/Engine/Plugins/Runtime/Metasound/Source/MetasoundFrontend/Public/MetasoundFrontendQuery.h
@@ -23,9 +23,16 @@ namespace Metasound
 		FValue Value;
 		float Score = 0.f;
 
-		template<typename... ArgTypes>
-		FFrontendQueryEntry(ArgTypes&&... Args)
-		:	Value(Forward<ArgTypes>(Args)...)
+		FFrontendQueryEntry(const FFrontendQueryEntry&) = default;
+		FFrontendQueryEntry& operator=(const FFrontendQueryEntry&) = default;
+
+		FFrontendQueryEntry(FValue&& InValue)
+		:	Value(MoveTemp(InValue))
+		{
+		}
+
+		FFrontendQueryEntry(const FValue& InValue)
+		:	Value(InValue)
 		{
 		}
 	};
@@ -54,31 +61,26 @@ namespace Metasound
 
 		public:
 			// Set initial entries which are stored elsewhere
-			FReduceOutputView(TArrayView<FFrontendQueryEntry*> InEntries);
+			FReduceOutputView(TArrayView<FFrontendQueryEntry* const> InEntries);
 			~FReduceOutputView() = default;
 
 			/** Add a result to the selection */
 			void Add(FFrontendQueryEntry& InResult);
 
-			/** Get selection which was also in initial entries. */
-			TArrayView<FFrontendQueryEntry*> GetSelectedInitialEntries();
 
-			/** Get selection which was not in the initial entries. */
-			TArrayView<FFrontendQueryEntry> GetSelectedNewEntries();
 
 		private:
+			friend class FFrontendQuerySelection;
 
 			TArray<FFrontendQueryEntry*> InitialEntries;
 
-			TArray<FFrontendQueryEntry*> ExistingEntryPointers;
-			TArray<FFrontendQueryEntry> NewEntryStorage;
+			TArray<FFrontendQueryEntry*> SelectedExistingEntries;
+			TArray<TUniquePtr<FFrontendQueryEntry>> SelectedNewEntries;
 		};
 
-		/** Get an array of all entries stored in this selection. */
-		TArrayView<FFrontendQueryEntry> GetStorage();
-
-		/** Get an array of all entries stored in this selection. */
-		TArrayView<const FFrontendQueryEntry> GetStorage() const;
+		FFrontendQuerySelection() = default;
+		FFrontendQuerySelection(const FFrontendQuerySelection&);
+		FFrontendQuerySelection& operator=(const FFrontendQuerySelection&);
 
 		/** Get an array of all entries selected. */
 		TArrayView<FFrontendQueryEntry*> GetSelection();
@@ -89,11 +91,18 @@ namespace Metasound
 		/** Add entries to both storage and current selection. */
 		void AppendToStorageAndSelection(const FFrontendQuerySelection& InSelection);
 
+		/** Add entries to both storage and current selection. */
+		void AppendToStorageAndSelection(FFrontendQuerySelection&& InSelection);
+
 		/** Add entries to both storage and current selection. */
 		void AppendToStorageAndSelection(TArrayView<const FFrontendQueryEntry> InEntries);
 
 		/** Add entries to both storage and current selection. */
-		void AppendToStorageAndSelection(FReduceOutputView& InReduceView);
+		void AppendToStorageAndSelection(TArrayView<const FFrontendQueryEntry* const> InEntries);
+
+		/** Add entries to both storage and current selection. */
+		void AppendToStorageAndSelection(FReduceOutputView&& InReduceView);
+
 
 		/** Reset both storage and selection. */
 		void ResetStorageAndSelection();
@@ -101,6 +110,19 @@ namespace Metasound
 		/** Reset selection, but keep storage. */
 		void ResetSelection();
 
+		/** Filters the selection by calling Func on each entry. If Func returns
+		 * false, the entry is removed.
+		 */
+		template<typename FilterFuncType>
+		int32 FilterSelection(FilterFuncType Func)
+		{
+			auto InvertedFilter = [&](const FFrontendQueryEntry* InEntry) -> bool
+			{
+				return !Func(InEntry);
+			};
+			return Selection.RemoveAll(InvertedFilter);
+		}
+
 		/** Set selection. All entries must already be in storage. */
 		void SetSelection(TArrayView<FFrontendQueryEntry*> InEntries);
 
@@ -108,11 +130,14 @@ namespace Metasound
 		void AddToSelection(FFrontendQueryEntry* InEntry);
 
 		/** Append entries to selection. Each selection must already be in storage. */
-		void AppendToSelection(TArrayView<FFrontendQueryEntry*> InEntries);
+		void AppendToSelection(TArrayView<FFrontendQueryEntry * const> InEntries);
 
 	private:
+		void AppendToStorageAndSelection(TArray<TUniquePtr<FFrontendQueryEntry>>&& InEntries);
+		void ShrinkStorageToSelection();
+
 		TArray<FFrontendQueryEntry*> Selection;
-		TArray<FFrontendQueryEntry> Storage;
+		TArray<TUniquePtr<FFrontendQueryEntry>> Storage;
 	};
 
 	/** FFrontendQuerySelectionView holds the result of query. */
@@ -126,9 +151,6 @@ namespace Metasound
 		 */
 		FFrontendQuerySelectionView(TSharedRef<const FFrontendQuerySelection, ESPMode::ThreadSafe> InResult);
 
-		/** Get all the stored results. */
-		TArrayView<const FFrontendQueryEntry> GetStorage() const;
-
 		/** Get the selected results. */
 		TArrayView<const FFrontendQueryEntry* const> GetSelection() const;
 
@@ -136,6 +158,15 @@ namespace Metasound
 		TSharedRef<const FFrontendQuerySelection, ESPMode::ThreadSafe> Result;
 	};
 
+	class IFrontendQuerySource
+	{
+		public:
+			virtual ~IFrontendQuerySource() = default;
+
+			virtual void Stream(TArray<FFrontendQueryEntry>& OutEntries) = 0;
+			virtual void Reset() = 0;
+	};
+
 	/** Interface for an individual step in a query */
 	class IFrontendQueryStep
 	{
@@ -143,12 +174,12 @@ namespace Metasound
 			virtual ~IFrontendQueryStep() = default;
 	};
 
-	/** Interface for a query step which generates entries. */
-	class IFrontendQueryGenerateStep : public IFrontendQueryStep
+	/** Interface for a query step which streams new entries. */
+	class IFrontendQueryStreamStep : public IFrontendQueryStep
 	{
 		public:
-			virtual ~IFrontendQueryGenerateStep() = default;
-			virtual void Generate(TArray<FFrontendQueryEntry>& OutEntries) const = 0;
+			virtual ~IFrontendQueryStreamStep() = default;
+			virtual void Stream(TArray<FFrontendQueryEntry>& OutEntries) = 0;
 	};
 
 	/** Interface for a query step which maps entries to keys. */
@@ -166,7 +197,7 @@ namespace Metasound
 			using FReduceOutputView = FFrontendQuerySelection::FReduceOutputView;
 
 			virtual ~IFrontendQueryReduceStep() = default;
-			virtual void Reduce(FFrontendQueryEntry::FKey InKey, TArrayView<FFrontendQueryEntry*>& InEntries, FReduceOutputView& OutResult) const = 0;
+			virtual void Reduce(FFrontendQueryEntry::FKey InKey, TArrayView<FFrontendQueryEntry * const>& InEntries, FReduceOutputView& OutResult) const = 0;
 	};
 
 	/** Interface for a query step which filters entries. */
@@ -206,29 +237,62 @@ namespace Metasound
 	class METASOUNDFRONTEND_API FFrontendQueryStep
 	{
 		FFrontendQueryStep() = delete;
+
 	public:
+		enum class EResultModificationState : uint8
+		{
+			Modified,
+			Unmodified,
+		};
+
+		/* Interface for executing a step in the query. */
+		struct IStepExecuter
+		{
+			using EResultModificationState = FFrontendQueryStep::EResultModificationState;
+
+			virtual ~IStepExecuter() = default;
+
+			// Perform an incremental step. Assume a previous result already exists. 
+			virtual EResultModificationState Increment(FFrontendQuerySelection& InOutResult) = 0;
+
+			// Merge an incremental result with the prior result from this step.
+			virtual EResultModificationState Merge(const FFrontendQuerySelection& InIncremental, FFrontendQuerySelection& InOutResult) = 0;
+
+			// Execute step. Assume not other prior results exist.
+			virtual EResultModificationState Execute(FFrontendQuerySelection& InOutResult) = 0;
+
+			// Reset internal state.
+			virtual void Reset() = 0;
+
+			// Returns true if a merge is required for a modified incremental update.
+			// 
+			// If this is true, Merge() will be called after Incremental() if Incremental
+			// returns EResultModificationState::Modified. 
+			virtual bool IsMergeRequiredForIncremental() const = 0;
+		};
 
 		using FReduceOutputView = FFrontendQuerySelection::FReduceOutputView;
 
-		using FGenerateFunction = TFunction<void (TArray<FFrontendQueryEntry>&)>;
+		using FStreamFunction = TUniqueFunction<void (TArray<FFrontendQueryEntry>&)>;
 		using FMapFunction = TFunction<FFrontendQueryEntry::FKey (const FFrontendQueryEntry&)>;
-		using FReduceFunction = TFunction<void (FFrontendQueryEntry::FKey, TArrayView<FFrontendQueryEntry*>, FReduceOutputView& )>;
+		using FReduceFunction = TFunction<void (FFrontendQueryEntry::FKey, TArrayView<FFrontendQueryEntry * const>, FReduceOutputView& )>;
 		using FFilterFunction = TFunction<bool (const FFrontendQueryEntry&)>;
 		using FScoreFunction = TFunction<float (const FFrontendQueryEntry&)>;
 		using FSortFunction = TFunction<bool (const FFrontendQueryEntry& InEntryLHS, const FFrontendQueryEntry& InEntryRHS)>;
 		using FLimitFunction = TFunction<int32 ()>;
 
 		/** Create query step using TFunction or lambda. */
-		FFrontendQueryStep(FGenerateFunction InFunc);
-		FFrontendQueryStep(FMapFunction InFunc);
-		FFrontendQueryStep(FReduceFunction InFunc);
-		FFrontendQueryStep(FFilterFunction InFilt);
-		FFrontendQueryStep(FScoreFunction InScore);
-		FFrontendQueryStep(FSortFunction InSort);
-		FFrontendQueryStep(FLimitFunction InLimit);
+		FFrontendQueryStep(FStreamFunction&& InFunc);
+		FFrontendQueryStep(FMapFunction&& InFunc);
+		FFrontendQueryStep(FReduceFunction&& InFunc);
+		FFrontendQueryStep(FFilterFunction&& InFilt);
+		FFrontendQueryStep(FScoreFunction&& InScore);
+		FFrontendQueryStep(FSortFunction&& InSort);
+		FFrontendQueryStep(FLimitFunction&& InLimit);
 
 		/** Create a query step using a IFrotnedQueryStep */
-		FFrontendQueryStep(TUniquePtr<IFrontendQueryGenerateStep>&& InStep);
+		FFrontendQueryStep(TUniquePtr<IFrontendQuerySource>&& InSource);
+		FFrontendQueryStep(TUniquePtr<IFrontendQueryStreamStep>&& InStep);
 		FFrontendQueryStep(TUniquePtr<IFrontendQueryMapStep>&& InStep);
 		FFrontendQueryStep(TUniquePtr<IFrontendQueryReduceStep>&& InStep);
 		FFrontendQueryStep(TUniquePtr<IFrontendQueryFilterStep>&& InStep);
@@ -236,14 +300,24 @@ namespace Metasound
 		FFrontendQueryStep(TUniquePtr<IFrontendQuerySortStep>&& InStep);
 		FFrontendQueryStep(TUniquePtr<IFrontendQueryLimitStep>&& InStep);
 
-		/** Execute a step on a selection. */
-		void ExecuteStep(FFrontendQuerySelection& InOutResult);
 
-		struct IStepExecuter
-		{
-			virtual ~IStepExecuter() = default;
-			virtual void ExecuteStep(FFrontendQuerySelection& InOutResult) const = 0;
-		};
+		// Perform an incremental step. Assume a previous result already exists. 
+		EResultModificationState Increment(FFrontendQuerySelection& InOutResult);
+
+		// Merge an incremental result with the prior result from this step.
+		EResultModificationState Merge(const FFrontendQuerySelection& InIncremental, FFrontendQuerySelection& InOutResult);
+
+		// Execute step. Assume not other prior results exist.
+		EResultModificationState Execute(FFrontendQuerySelection& InOutResult);
+
+		// Returns true if a merge is required for a modified incremental update.
+		// 
+		// If this is true, Merge() will be called after Incremental() if Incremental
+		// returns EResultModificationState::Modified. 
+		bool IsMergeRequiredForIncremental() const;
+
+		// Reset internal state.
+		void Reset();
 
 	private:
 
@@ -257,7 +331,8 @@ namespace Metasound
 	{
 	public:
 
-		using FGenerateFunction = FFrontendQueryStep::FGenerateFunction;
+		using EResultModificationState = FFrontendQueryStep::EResultModificationState;
+		using FStreamFunction = FFrontendQueryStep::FStreamFunction;
 		using FMapFunction = FFrontendQueryStep::FMapFunction;
 		using FReduceFunction = FFrontendQueryStep::FReduceFunction;
 		using FFilterFunction = FFrontendQueryStep::FFilterFunction;
@@ -281,40 +356,39 @@ namespace Metasound
 		}
 
 		template<typename FuncType>
-		FFrontendQuery& AddFunctionStep(TFunction<FuncType> InFunc)
+		FFrontendQuery& AddFunctionStep(FuncType&& InFunc)
 		{
-			return AddStep(MakeUnique<FFrontendQueryStep>(InFunc));
+			return AddStep(MakeUnique<FFrontendQueryStep>(MoveTemp(InFunc)));
 		}
 
-		FFrontendQuery& AddGenerateLambdaStep(FGenerateFunction InFunc);
-		FFrontendQuery& AddMapLambdaStep(FMapFunction InFunc);
-		FFrontendQuery& AddReduceLambdaStep(FReduceFunction InFunc);
-		FFrontendQuery& AddFilterLambdaStep(FFilterFunction InFunc);
-		FFrontendQuery& AddScoreLambdaStep(FScoreFunction InFunc);
-		FFrontendQuery& AddSortLambdaStep(FSortFunction InFunc);
-		FFrontendQuery& AddLimitLambdaStep(FLimitFunction InFunc);
+		FFrontendQuery& AddStreamLambdaStep(FStreamFunction&& InFunc);
+		FFrontendQuery& AddMapLambdaStep(FMapFunction&& InFunc);
+		FFrontendQuery& AddReduceLambdaStep(FReduceFunction&& InFunc);
+		FFrontendQuery& AddFilterLambdaStep(FFilterFunction&& InFunc);
+		FFrontendQuery& AddScoreLambdaStep(FScoreFunction&& InFunc);
+		FFrontendQuery& AddSortLambdaStep(FSortFunction&& InFunc);
+		FFrontendQuery& AddLimitLambdaStep(FLimitFunction&& InFunc);
 
 		/** Add a step to the query. */
 		FFrontendQuery& AddStep(TUniquePtr<FFrontendQueryStep>&& InStep);
 
 		/** Calls all steps in the query and returns the selection. */
-		FFrontendQuerySelectionView ExecuteQuery();
+		FFrontendQuerySelectionView Execute();
 
 		/** Resets the query result by removing all entries. */
 		FFrontendQuerySelectionView Reset();
 
-		/** Executes the query again and appends the results to the existing results. */
-		FFrontendQuerySelectionView ExecuteQueryAndAppend();
-
 		/** Returns the current result. */
 		FFrontendQuerySelectionView GetSelection() const;
 
 	private:
 
-		void Execute(FFrontendQuerySelection& OutSelection);
+		void UpdateResult();
+		void ExecuteSteps(int32 InStartStepIndex);
 
 		TSharedRef<FFrontendQuerySelection, ESPMode::ThreadSafe> Result;
 
 		TArray<TUniquePtr<FFrontendQueryStep>> Steps;
+		TMap<int32, FFrontendQuerySelection> StepResultCache;
 	};
 }
diff --git a/Engine/Plugins/Runtime/Metasound/Source/MetasoundFrontend/Public/MetasoundFrontendQuerySteps.h b/Engine/Plugins/Runtime/Metasound/Source/MetasoundFrontend/Public/MetasoundFrontendQuerySteps.h
--- a/Engine/Plugins/Runtime/Metasound/Source/MetasoundFrontend/Public/MetasoundFrontendQuerySteps.h
+++ b/Engine/Plugins/Runtime/Metasound/Source/MetasoundFrontend/Public/MetasoundFrontendQuerySteps.h
@@ -9,25 +9,18 @@
 
 namespace Metasound
 {
-	/** Generates all registered node classes */
-	class METASOUNDFRONTEND_API FGenerateAllAvailableNodeClasses : public IFrontendQueryGenerateStep
-	{
-	public:
-
-		void Generate(TArray<FFrontendQueryEntry>& OutEntries) const override;
-	};
 
-	/** Generates node classes that have been newly registered since last call
-	 * to this FGenerateNewlyAvailableClasses::Generate()
+	/** Streams node classes that have been newly registered or unregistered since last call to Stream()
 	 */
-	class METASOUNDFRONTEND_API FGenerateNewlyAvailableNodeClasses : public IFrontendQueryGenerateStep
+	class METASOUNDFRONTEND_API FNodeClassRegistrationEvents : public IFrontendQuerySource
 	{
 	public:
-		FGenerateNewlyAvailableNodeClasses();
-		void Generate(TArray<FFrontendQueryEntry>& OutEntries) const override;
+		FNodeClassRegistrationEvents();
+		void Stream(TArray<FFrontendQueryEntry>& OutEntries) override;
+		void Reset() override;
 
 	private:
-		mutable Frontend::FRegistryTransactionID CurrentTransactionID;
+		Frontend::FRegistryTransactionID CurrentTransactionID;
 	};
 
 	class METASOUNDFRONTEND_API FFilterClassesByInputVertexDataType : public IFrontendQueryFilterStep
@@ -95,7 +88,7 @@ namespace Metasound
 	class METASOUNDFRONTEND_API FReduceClassesToHighestVersion : public IFrontendQueryReduceStep
 	{
 	public:
-		void Reduce(FFrontendQueryEntry::FKey InKey, TArrayView<FFrontendQueryEntry*>& InEntries, FReduceOutputView& OutResult) const override;
+		void Reduce(FFrontendQueryEntry::FKey InKey, TArrayView<FFrontendQueryEntry * const>& InEntries, FReduceOutputView& OutResult) const override;
 	};
 
 	class METASOUNDFRONTEND_API FReduceClassesToMajorVersion : public IFrontendQueryReduceStep
@@ -103,7 +96,7 @@ namespace Metasound
 	public:
 		FReduceClassesToMajorVersion(int32 InMajorVersion);
 
-		void Reduce(FFrontendQueryEntry::FKey InKey, TArrayView<FFrontendQueryEntry*>& InEntries, FReduceOutputView& OutResult) const override;
+		void Reduce(FFrontendQueryEntry::FKey InKey, TArrayView<FFrontendQueryEntry * const>& InEntries, FReduceOutputView& OutResult) const override;
 
 	private:
 		int32 MajorVersion = -1;
