- Migrate MetaSound InputLiterals to DefaultLiterals housed on FrontendNodes
- Add FrontendGraphBuilder support for generating VariableNodes for DefaultLiterals if no connection is provided on External/Output nodes
- Move node names to be guids (they're never seen by the user and this will be easier to migrate later)
- Fix DisplayName logic to be unique against one another SEPARATELY from node name logic
- Fix for MetaSounds not compiling if trigger outputs do not have input connected
- MetaSound Document Transform Support
- Fix for for MetaSound Generators not stopping source if failed to create generator due to build errors

diff --git a/Engine/Plugins/Runtime/Metasound/Source/MetasoundEditor/Private/MetasoundEditor.cpp b/Engine/Plugins/Runtime/Metasound/Source/MetasoundEditor/Private/MetasoundEditor.cpp
--- a/Engine/Plugins/Runtime/Metasound/Source/MetasoundEditor/Private/MetasoundEditor.cpp
+++ b/Engine/Plugins/Runtime/Metasound/Source/MetasoundEditor/Private/MetasoundEditor.cpp
@@ -36,6 +36,7 @@
 #include "MetasoundFrontendDocument.h"
 #include "MetasoundFrontendRegistries.h"
 #include "MetasoundFrontendSearchEngine.h"
+#include "MetasoundFrontendTransform.h"
 #include "MetasoundLog.h"
 #include "MetasoundUObjectRegistry.h"
 #include "Misc/Attribute.h"
@@ -58,14 +59,6 @@
 
 #define LOCTEXT_NAMESPACE "MetaSoundEditor"
 
-static int32 ShowLiteralMetasoundInputsInEditorCVar = 0;
-// FAutoConsoleVariableRef CVarShowLiteralMetasoundInputsInEditor(
-// 	TEXT("au.Debug.Editor.Metasounds.ShowLiteralInputs"),
-// 	ShowLiteralMetasoundInputsInEditorCVar,
-// 	TEXT("Show literal inputs in the Metasound Editor.\n")
-// 	TEXT("0: Disabled (default), !0: Enabled"),
-// 	ECVF_Default);
-
 
 namespace Metasound
 {
@@ -350,13 +343,13 @@ namespace Metasound
 			Metasound = ObjectToEdit;
 
 			FMetasoundAssetBase* MetasoundAsset = IMetasoundUObjectRegistry::Get().GetObjectAsAssetBase(Metasound);
-			if (!MetasoundAsset->GetGraph())
+			check(MetasoundAsset);
+
+			if (Frontend::FVersionDocument().Transform(MetasoundAsset->GetDocumentHandle()))
 			{
-				UMetasoundEditorGraph* Graph = NewObject<UMetasoundEditorGraph>(Metasound, FName(), RF_Transactional);
-				Graph->Schema = UMetasoundEditorGraphSchema::StaticClass();
-				MetasoundAsset->SetGraph(Graph);
-				FGraphBuilder::SynchronizeGraph(*Metasound);
+				Metasound->MarkPackageDirty();
 			}
+			FGraphBuilder::SynchronizeGraph(*Metasound);
 
 			GEditor->RegisterForUndo(this);
 
@@ -504,16 +497,15 @@ namespace Metasound
 			{
 				// If a property change event occurs outside of the metasound UEdGraph and results in the metasound document changing, 
 				// then the document and the UEdGraph need to be synchronized. There may be a better trigger for this call to reduce
-				// the number of times the graph is synchronized. 
-				UObject* MetasoundObj = GetMetasoundObject();
-				if (nullptr != MetasoundObj)
+				// the number of times the graph is synchronized.
+				if (Metasound)
 				{
-					if (IMetasoundUObjectRegistry::Get().IsRegisteredClass(MetasoundObj))
+					if (IMetasoundUObjectRegistry::Get().IsRegisteredClass(Metasound))
 					{
-						if (FGraphBuilder::SynchronizeGraph(*MetasoundObj))
+						if (FGraphBuilder::SynchronizeGraph(*Metasound))
 						{
 							NotifyUserModifiedBySync();
-							MetasoundObj->MarkPackageDirty();
+							Metasound->MarkPackageDirty();
 						}
 						MetasoundGraphEditor->NotifyGraphChanged();
 						MetasoundInterfaceMenu->RefreshAllActions(true /* bPreserveExpansion */);
@@ -1235,20 +1227,14 @@ namespace Metasound
 				{
 					for (int32 i = IndexToDelete + 1; i < Handles.Num(); ++i)
 					{
-						if (Handles[i]->GetNodeStyle().Display.Visibility != EMetasoundFrontendNodeStyleDisplayVisibility::Hidden)
-						{
-							const FName NameToSelect(*(Handles[i]->GetDisplayName().ToString()));
-							return FNameNodeIDPair { NameToSelect, Handles[i]->GetID() };
-						}
+						const FName NameToSelect(*(Handles[i]->GetDisplayName().ToString()));
+						return FNameNodeIDPair { NameToSelect, Handles[i]->GetID() };
 					}
 
 					for (int32 i = IndexToDelete - 1; i >= 0; --i)
 					{
-						if (Handles[i]->GetNodeStyle().Display.Visibility != EMetasoundFrontendNodeStyleDisplayVisibility::Hidden)
-						{
-							FName NameToSelect = (*(Handles[i]->GetDisplayName().ToString()));
-							return FNameNodeIDPair { NameToSelect, Handles[i]->GetID() };
-						}
+						FName NameToSelect = (*(Handles[i]->GetDisplayName().ToString()));
+						return FNameNodeIDPair { NameToSelect, Handles[i]->GetID() };
 					}
 				}
 
@@ -1620,18 +1606,21 @@ namespace Metasound
 						continue;
 					}
 
-					if (Pin->LinkedTo.IsEmpty())
+					FInputHandle InputHandle = FGraphBuilder::GetInputHandleFromPin(Pin);
+					if (InputHandle->IsValid())
 					{
-						FGraphBuilder::AddOrUpdateLiteralInput(*Metasound, MetasoundNode->GetNodeHandle(), *Pin, true /* bForcePinValueAsDefault */);
+						FMetasoundFrontendLiteral LiteralValue;
+						if (FGraphBuilder::GetPinLiteral(*Pin, LiteralValue))
+						{
+							InputHandle->SetLiteral(LiteralValue);
+						}
 					}
-					else
+
+					for (UEdGraphPin* LinkedPin : Pin->LinkedTo)
 					{
-						for (UEdGraphPin* LinkedPin : Pin->LinkedTo)
+						if (UMetasoundEditorGraphNode* Node = Cast<UMetasoundEditorGraphNode>(LinkedPin->GetOwningNode()))
 						{
-							if (UMetasoundEditorGraphNode* Node = Cast<UMetasoundEditorGraphNode>(LinkedPin->GetOwningNode()))
-							{
-								FGraphBuilder::ConnectNodes(*Pin, *LinkedPin, false /* bConnectEdPins */);
-							}
+							FGraphBuilder::ConnectNodes(*Pin, *LinkedPin, false /* bConnectEdPins */);
 						}
 					}
 				}
@@ -1748,12 +1737,6 @@ namespace Metasound
 
 			MetasoundAsset->GetRootGraphHandle()->IterateConstNodes([this, ActionList = &OutAllActions](const Frontend::FConstNodeHandle& Input)
 			{
-				const bool bIsLiteralInput = Input->GetNodeStyle().Display.Visibility == EMetasoundFrontendNodeStyleDisplayVisibility::Hidden;
-				if (!ShowLiteralMetasoundInputsInEditorCVar && bIsLiteralInput)
-				{
-					return;
-				}
-
 				const FText Tooltip = Input->GetDescription();
 				const FText InputName = Input->GetDisplayName();
 				UMetasoundEditorGraph& Graph = GetMetaSoundGraphChecked();
@@ -1938,10 +1921,7 @@ namespace Metasound
 					Metasound->Modify();
 
 					const FName DataTypeName = "Bool";
-					FString NewNodeName = FGraphBuilder::GenerateUniqueInputName(*Metasound);
-					NameToSelect = FName(*NewNodeName);
-
-					Frontend::FNodeHandle NodeHandle = FGraphBuilder::AddInputNodeHandle(*Metasound, NewNodeName, DataTypeName, FText::GetEmpty());
+					Frontend::FNodeHandle NodeHandle = FGraphBuilder::AddInputNodeHandle(*Metasound, DataTypeName, FText::GetEmpty());
 					if (ensure(NodeHandle->IsValid()))
 					{
 						TObjectPtr<UMetasoundEditorGraphInput> Input = Graph.FindOrAddInput(NodeHandle);
@@ -1964,12 +1944,13 @@ namespace Metasound
 					Metasound->Modify();
 
 					const FName DataTypeName = "Bool";
-					FString NewNodeName = FGraphBuilder::GenerateUniqueOutputName(*Metasound);
-					NameToSelect = FName(*NewNodeName);
-
-					Frontend::FNodeHandle NodeHandle = FGraphBuilder::AddOutputNodeHandle(*Metasound, NewNodeName, DataTypeName, FText::GetEmpty());
+					Frontend::FNodeHandle NodeHandle = FGraphBuilder::AddOutputNodeHandle(*Metasound, DataTypeName, FText::GetEmpty());
 					if (ensure(NodeHandle->IsValid()))
 					{
+						const FText DisplayName = FGraphBuilder::GenerateUniqueOutputDisplayName(*Metasound);
+						NodeHandle->SetDisplayName(DisplayName);
+						NameToSelect = *DisplayName.ToString();
+
 						TObjectPtr<UMetasoundEditorGraphOutput> Output = Graph.FindOrAddOutput(NodeHandle);
 						if (ensure(Output))
 						{
diff --git a/Engine/Plugins/Runtime/Metasound/Source/MetasoundEditor/Private/MetasoundEditorGraph.cpp b/Engine/Plugins/Runtime/Metasound/Source/MetasoundEditor/Private/MetasoundEditorGraph.cpp
--- a/Engine/Plugins/Runtime/Metasound/Source/MetasoundEditor/Private/MetasoundEditorGraph.cpp
+++ b/Engine/Plugins/Runtime/Metasound/Source/MetasoundEditor/Private/MetasoundEditorGraph.cpp
@@ -84,7 +84,6 @@ void UMetasoundEditorGraphVariable::SetDataType(FName InNewType)
 
 	// 2. Cache the old version's Frontend data.
 	FNodeHandle NodeHandle = GetNodeHandle();
-	const FString NodeName = NodeHandle->GetNodeName();
 	const FText NodeDisplayName = NodeHandle->GetDisplayName();
 
 	// 3. Delete the Frontend variable
@@ -92,7 +91,8 @@ void UMetasoundEditorGraphVariable::SetDataType(FName InNewType)
 
 	// 4. Add the new input node with the same identifier data but new datatype.
 	UObject& Metasound = Graph->GetMetasoundChecked();
-	FNodeHandle NewNodeHandle = AddNodeHandle(NodeName, NodeDisplayName, InNewType);
+	FNodeHandle NewNodeHandle = AddNodeHandle(NodeDisplayName, InNewType);
+
 	if (!ensure(NewNodeHandle->IsValid()))
 	{
 		return;
@@ -225,7 +225,7 @@ void UMetasoundEditorGraphInput::UpdatePreviewInstance(const Metasound::FVertexK
 	}
 }
 
-Metasound::Frontend::FNodeHandle UMetasoundEditorGraphInput::AddNodeHandle(const FString& InNodeName, const FText& InNodeDisplayName, FName InDataType)
+Metasound::Frontend::FNodeHandle UMetasoundEditorGraphInput::AddNodeHandle(const FText& InNodeDisplayName, FName InDataType)
 {
 	using namespace Metasound::Editor;
 	using namespace Metasound::Frontend;
@@ -233,11 +233,11 @@ Metasound::Frontend::FNodeHandle UMetasoundEditorGraphInput::AddNodeHandle(const
 	UMetasoundEditorGraph* Graph = Cast<UMetasoundEditorGraph>(GetOuter());
 	if (!ensure(Graph))
 	{
-		return Super::AddNodeHandle(InNodeName, InNodeDisplayName, InDataType);
+		return UMetasoundEditorGraphVariable::AddNodeHandle(InNodeDisplayName, InDataType);
 	}
 
 	UObject& Metasound = Graph->GetMetasoundChecked();
-	Metasound::Frontend::FNodeHandle NewNodeHandle = FGraphBuilder::AddInputNodeHandle(Metasound, InNodeName, InDataType, FText::GetEmpty());
+	Metasound::Frontend::FNodeHandle NewNodeHandle = FGraphBuilder::AddInputNodeHandle(Metasound, InDataType, FText::GetEmpty());
 	NewNodeHandle->SetDisplayName(InNodeDisplayName);
 	return NewNodeHandle;
 }
@@ -305,7 +305,7 @@ void UMetasoundEditorGraphInput::OnDataTypeChanged()
 	Literal = NewObject<UMetasoundEditorGraphInputLiteral>(this, InputLiteralClass, FName(), RF_Transactional);
 }
 
-Metasound::Frontend::FNodeHandle UMetasoundEditorGraphOutput::AddNodeHandle(const FString& InNodeName, const FText& InNodeDisplayName, FName InDataType)
+Metasound::Frontend::FNodeHandle UMetasoundEditorGraphOutput::AddNodeHandle(const FText& InNodeDisplayName, FName InDataType)
 {
 	using namespace Metasound::Editor;
 	using namespace Metasound::Frontend;
@@ -313,11 +313,11 @@ Metasound::Frontend::FNodeHandle UMetasoundEditorGraphOutput::AddNodeHandle(cons
 	UMetasoundEditorGraph* Graph = Cast<UMetasoundEditorGraph>(GetOuter());
 	if (!ensure(Graph))
 	{
-		return Super::AddNodeHandle(InNodeName, InNodeDisplayName, InDataType);
+		return UMetasoundEditorGraphVariable::AddNodeHandle(InNodeDisplayName, InDataType);
 	}
 
 	UObject& Metasound = Graph->GetMetasoundChecked();
-	Metasound::Frontend::FNodeHandle NewNodeHandle = FGraphBuilder::AddOutputNodeHandle(Metasound, InNodeName, InDataType, FText::GetEmpty());
+	Metasound::Frontend::FNodeHandle NewNodeHandle = FGraphBuilder::AddOutputNodeHandle(Metasound, InDataType, FText::GetEmpty());
 	NewNodeHandle->SetDisplayName(InNodeDisplayName);
 	return NewNodeHandle;
 }
diff --git a/Engine/Plugins/Runtime/Metasound/Source/MetasoundEditor/Private/MetasoundEditorGraphBuilder.cpp b/Engine/Plugins/Runtime/Metasound/Source/MetasoundEditor/Private/MetasoundEditorGraphBuilder.cpp
--- a/Engine/Plugins/Runtime/Metasound/Source/MetasoundEditor/Private/MetasoundEditorGraphBuilder.cpp
+++ b/Engine/Plugins/Runtime/Metasound/Source/MetasoundEditor/Private/MetasoundEditorGraphBuilder.cpp
@@ -13,6 +13,7 @@
 #include "MetasoundEditorGraph.h"
 #include "MetasoundEditorGraphNode.h"
 #include "MetasoundEditorGraphInputNodes.h"
+#include "MetasoundEditorGraphSchema.h"
 #include "MetasoundEditorModule.h"
 #include "MetasoundEditorSettings.h"
 #include "MetasoundFrontendQuery.h"
@@ -90,7 +91,10 @@ namespace Metasound
 
 		UMetasoundEditorGraphExternalNode* FGraphBuilder::AddExternalNode(UObject& InMetasound, const FMetasoundFrontendClassMetadata& InMetadata, FVector2D InLocation, bool bInSelectNewNode)
 		{
-			Frontend::FNodeHandle NodeHandle = AddExternalNodeHandle(InMetasound, InMetadata);
+			FMetasoundAssetBase* MetasoundAsset = IMetasoundUObjectRegistry::Get().GetObjectAsAssetBase(&InMetasound);
+			check(MetasoundAsset);
+
+			Frontend::FNodeHandle NodeHandle = MetasoundAsset->GetRootGraphHandle()->AddNode(InMetadata);
 			return AddExternalNode(InMetasound, NodeHandle, InLocation, bInSelectNewNode);
 		}
 
@@ -129,14 +133,6 @@ namespace Metasound
 			RebuildNodePins(*NewGraphNode);
 		}
 
-		Frontend::FNodeHandle FGraphBuilder::AddExternalNodeHandle(UObject& InMetasound, const FMetasoundFrontendClassMetadata& InMetadata)
-		{
-			FMetasoundAssetBase* MetasoundAsset = IMetasoundUObjectRegistry::Get().GetObjectAsAssetBase(&InMetasound);
-			check(MetasoundAsset);
-
-			return MetasoundAsset->GetRootGraphHandle()->AddNode(InMetadata);
-		}
-
 		TArray<FString> FGraphBuilder::GetDataTypeNameCategories(const FName& InDataTypeName)
 		{
 			FString CategoryString = InDataTypeName.ToString();
@@ -153,44 +149,60 @@ namespace Metasound
 			return Categories;
 		}
 
-		FString FGraphBuilder::GenerateUniqueInputName(const UObject& InMetasound, const FString* InBaseName)
+		FText FGraphBuilder::GenerateUniqueInputDisplayName(const UObject& InMetaSound, const FText* InBaseName)
 		{
-			static const FString DefaultBaseName = TEXT("Input");
-			const FString* NameBase = InBaseName ? InBaseName : &DefaultBaseName;
-
-			const FMetasoundAssetBase* MetasoundAsset = IMetasoundUObjectRegistry::Get().GetObjectAsAssetBase(&InMetasound);
-			check(MetasoundAsset);
-
-			const Frontend::FConstGraphHandle GraphHandle = MetasoundAsset->GetRootGraphHandle();
-
-			int32 i = 1;
-			FString NewInputName = *NameBase + FString::Printf(TEXT(" %02d"), i);
-			while (GraphHandle->ContainsInputVertexWithName(NewInputName))
+			static const FText DefaultBaseName = FText::FromString(TEXT("Input"));
+			const FText& NameBase = InBaseName ? *InBaseName : DefaultBaseName;
+			return GenerateUniqueNameByFilter(InMetaSound, NameBase, [](const Frontend::FConstGraphHandle& GraphHandle, const FText& NewName)
 			{
-				NewInputName = *NameBase + FString::Printf(TEXT(" %02d"), ++i);
-			}
-
-			return NewInputName;
+				bool bIsNameInvalid = false;
+				GraphHandle->IterateConstNodes([&](Frontend::FConstNodeHandle Node)
+				{
+					if (NewName.CompareToCaseIgnored(Node->GetDisplayName()) == 0)
+					{
+						bIsNameInvalid = true;
+					}
+				}, EMetasoundFrontendClassType::Input);
+				return !bIsNameInvalid;
+			});
 		}
 
-		FString FGraphBuilder::GenerateUniqueOutputName(const UObject& InMetasound, const FString* InBaseName)
+		FText FGraphBuilder::GenerateUniqueOutputDisplayName(const UObject& InMetaSound, const FText* InBaseName)
 		{
-			static const FString DefaultBaseName = TEXT("Input");
-			const FString* NameBase = InBaseName ? InBaseName : &DefaultBaseName;
+			static const FText DefaultBaseName = FText::FromString(TEXT("Output"));
+			const FText& NameBase = InBaseName ? *InBaseName : DefaultBaseName;
+			return GenerateUniqueNameByFilter(InMetaSound, NameBase, [](const Frontend::FConstGraphHandle& GraphHandle, const FText& NewName)
+			{
+				bool bIsNameInvalid = false;
+				GraphHandle->IterateConstNodes([&](Frontend::FConstNodeHandle Node)
+				{
+					if (NewName.CompareToCaseIgnored(Node->GetDisplayName()) == 0)
+					{
+						bIsNameInvalid = true;
+					}
+				}, EMetasoundFrontendClassType::Output);
+				return !bIsNameInvalid;
+			});
+		}
 
-			const FMetasoundAssetBase* MetasoundAsset = IMetasoundUObjectRegistry::Get().GetObjectAsAssetBase(&InMetasound);
+		FText FGraphBuilder::GenerateUniqueNameByFilter(const UObject& InMetaSound, const FText& InBaseText, TUniqueFunction<bool(const Frontend::FConstGraphHandle&, const FText&)> InIsValidNameFilter)
+		{
+			const FMetasoundAssetBase* MetasoundAsset = IMetasoundUObjectRegistry::Get().GetObjectAsAssetBase(&InMetaSound);
 			check(MetasoundAsset);
 
 			const Frontend::FConstGraphHandle GraphHandle = MetasoundAsset->GetRootGraphHandle();
 
+			static const FText VariableUniqueNameFormat = LOCTEXT("VariableUniqueNameFormat", "{0}{1}");
+
 			int32 i = 1;
-			FString NewOutputName = *NameBase + FString::Printf(TEXT(" %02d"), i);
-			while (GraphHandle->ContainsOutputVertexWithName(NewOutputName))
+			FText NewName = FText::Format(VariableUniqueNameFormat, InBaseText, i);
+
+			while (!InIsValidNameFilter(GraphHandle, NewName))
 			{
-				NewOutputName = *NameBase + FString::Printf(TEXT(" %02d"), ++i);
+				NewName = FText::Format(VariableUniqueNameFormat, InBaseText, ++i);
 			}
 
-			return NewOutputName;
+			return NewName;
 		}
 
 		TSharedPtr<FEditor> FGraphBuilder::GetEditorForGraph(const UObject& Metasound)
@@ -353,93 +365,7 @@ namespace Metasound
 			return nullptr;
 		}
 
-		bool FGraphBuilder::IsLiteralInput(Frontend::FNodeHandle InNodeHandle)
-		{
-			if (InNodeHandle->GetClassMetadata().Type == EMetasoundFrontendClassType::Input)
-			{
-				return InNodeHandle->GetNodeStyle().Display.Visibility == EMetasoundFrontendNodeStyleDisplayVisibility::Hidden;
-			}
-
-			return false;
-		}
-
-		void FGraphBuilder::AddOrUpdateLiteralInput(UObject& InMetasound, Frontend::FNodeHandle InNodeHandle, UEdGraphPin& InInputPin, bool bForcePinValueAsDefault)
-		{
-			using namespace Frontend;
-
-			FInputHandle InputHandle = GetInputHandleFromPin(&InInputPin);
-			if (!ensure(InputHandle->IsValid()))
-			{
-				return;
-			}
-
-			FOutputHandle OutputHandle = InputHandle->GetConnectedOutput();
-			if (OutputHandle->IsValid())
-			{
-				FNodeHandle InputNode = OutputHandle->GetOwningNode();
-				if (IsLiteralInput(InputNode))
-				{
-					Frontend::FGraphHandle GraphHandle = InputNode->GetOwningGraph();
-					const FGuid VertexID = GraphHandle->GetVertexIDForInputVertex(InputNode->GetNodeName());
-
-					FMetasoundFrontendLiteral PinDataTypeDefaultLiteral;
-					if (GetPinDefaultLiteral(InInputPin, PinDataTypeDefaultLiteral))
-					{
-						GraphHandle->SetDefaultInput(VertexID, PinDataTypeDefaultLiteral);
-					}
-				}
-				return;
-			}
-
-			const FName TypeName = InputHandle->GetDataType();
-			const FString NewInputName = GenerateUniqueInputName(InMetasound);
-			Frontend::FNodeHandle NewInputNode = INodeController::GetInvalidHandle();
-
-			if (bForcePinValueAsDefault)
-			{
-				FMetasoundFrontendLiteral PinDataTypeDefaultLiteral;
-				if (GetPinDefaultLiteral(InInputPin, PinDataTypeDefaultLiteral))
-				{
-					NewInputNode = AddInputNodeHandle(InMetasound, NewInputName, TypeName, FText::GetEmpty(), true /* bisLiteralInput */, &PinDataTypeDefaultLiteral);
-				}
-			}
-			else
-			{
-				NewInputNode = AddInputNodeHandle(InMetasound, NewInputName, TypeName, FText::GetEmpty(), true /* bInIsLiteral */);
-
-				Frontend::FGraphHandle GraphHandle = InNodeHandle->GetOwningGraph();
-				const FGuid VertexID = GraphHandle->GetVertexIDForInputVertex(NewInputNode->GetNodeName());
-
-				FMetasoundFrontendLiteral DefaultLiteral;
-
-				// If the node's input is set/not the base implementation (none), set the new graph input's value/pin's default value to the node input's default literal.
-				const FMetasoundFrontendLiteral* NodeInputDefault = InputHandle->GetDefaultLiteral();
-				if (NodeInputDefault && NodeInputDefault->GetType() != EMetasoundFrontendLiteralType::None)
-				{
-					DefaultLiteral = *NodeInputDefault;
-				}
-				// Otherwise, get the class default for the given input node's type, and set the new graph input's value/input pin's default to that.
-				else
-				{
-					const Metasound::FLiteral Literal = Frontend::GetDefaultParamForDataType(TypeName);
-					DefaultLiteral.SetFromLiteral(Literal);
-				}
-				InInputPin.DefaultValue = DefaultLiteral.ToString();
-				GraphHandle->SetDefaultInput(VertexID, DefaultLiteral);
-			}
-
-			if (NewInputNode->IsValid())
-			{
-				TArray<FOutputHandle> OutputHandles = NewInputNode->GetOutputs();
-				if (ensure(OutputHandles.Num() == 1))
-				{
-					OutputHandle = OutputHandles[0];
-					ensure(InputHandle->Connect(*OutputHandle));
-				}
-			}
-		}
-
-		bool FGraphBuilder::GetPinDefaultLiteral(UEdGraphPin& InInputPin, FMetasoundFrontendLiteral& OutDefaultLiteral)
+		bool FGraphBuilder::GetPinLiteral(UEdGraphPin& InInputPin, FMetasoundFrontendLiteral& OutDefaultLiteral)
 		{
 			using namespace Frontend;
 
@@ -589,7 +515,8 @@ namespace Metasound
 					UMetasoundEditorGraphInput* Input = InputNode->Input;
 					if (ensure(Input))
 					{
-						NodeHandle = FGraphBuilder::AddInputNodeHandle(InMetasound, Pin->GetName(), Input->TypeName, InGraphNode.GetTooltipText());
+						NodeHandle = FGraphBuilder::AddInputNodeHandle(InMetasound, Input->TypeName, InGraphNode.GetTooltipText());
+						NodeHandle->SetDisplayName(FText::FromString(Pin->GetName()));
 					}
 				}
 			}
@@ -603,7 +530,8 @@ namespace Metasound
 					UMetasoundEditorGraphOutput* Output = OutputNode->Output;
 					if (ensure(Output))
 					{
-						NodeHandle = FGraphBuilder::AddOutputNodeHandle(InMetasound, Pin->GetName(), Output->TypeName, InGraphNode.GetTooltipText());
+						NodeHandle = FGraphBuilder::AddOutputNodeHandle(InMetasound, Output->TypeName, InGraphNode.GetTooltipText());
+						NodeHandle->SetDisplayName(FText::FromString(Pin->GetName()));
 					}
 				}
 			}
@@ -634,7 +562,7 @@ namespace Metasound
 			return NodeHandle;
 		}
 
-		Frontend::FNodeHandle FGraphBuilder::AddInputNodeHandle(UObject& InMetasound, const FString& InName, const FName InTypeName, const FText& InToolTip, bool bInIsLiteral, const FMetasoundFrontendLiteral* InDefaultValue)
+		Frontend::FNodeHandle FGraphBuilder::AddInputNodeHandle(UObject& InMetasound, const FName InTypeName, const FText& InToolTip, const FMetasoundFrontendLiteral* InDefaultValue)
 		{
 			FMetasoundAssetBase* MetasoundAsset = IMetasoundUObjectRegistry::Get().GetObjectAsAssetBase(&InMetasound);
 			check(MetasoundAsset);
@@ -643,7 +571,12 @@ namespace Metasound
 			const FGuid VertexID = FGuid::NewGuid();
 
 			FMetasoundFrontendClassInput Description;
-			Description.Name = InName;
+
+			// The user currently never interfaces with the name, only the display name,
+			// so just set to VertexID to make unique, avoid confusion, and eventually
+			// this field will most likely go away (or become the display name).
+			Description.Name = VertexID.ToString();
+
 			Description.TypeName = InTypeName;
 			Description.Metadata.Description = InToolTip;
 			Description.VertexID = VertexID;
@@ -651,12 +584,12 @@ namespace Metasound
 			Frontend::FNodeHandle NodeHandle = GraphHandle->AddInputVertex(Description);
 
 			FMetasoundFrontendNodeStyle Style = NodeHandle->GetNodeStyle();
-			Style.Display.Visibility = bInIsLiteral ? EMetasoundFrontendNodeStyleDisplayVisibility::Hidden : EMetasoundFrontendNodeStyleDisplayVisibility::Visible;
 			NodeHandle->SetNodeStyle(Style);
 
 			if (ensure(NodeHandle->IsValid()))
 			{
-				GraphHandle->SetInputDisplayName(InName, FText::FromString(InName));
+				const FText DisplayName = FGraphBuilder::GenerateUniqueInputDisplayName(InMetasound);
+				GraphHandle->SetInputDisplayName(Description.Name, DisplayName);
 
 				if (InDefaultValue)
 				{
@@ -671,21 +604,29 @@ namespace Metasound
 			return NodeHandle;
 		}
 
-		Frontend::FNodeHandle FGraphBuilder::AddOutputNodeHandle(UObject& InMetasound, const FString& InName, const FName InTypeName, const FText& InToolTip)
+		Frontend::FNodeHandle FGraphBuilder::AddOutputNodeHandle(UObject& InMetasound, const FName InTypeName, const FText& InToolTip)
 		{
 			FMetasoundAssetBase* MetasoundAsset = IMetasoundUObjectRegistry::Get().GetObjectAsAssetBase(&InMetasound);
 			check(MetasoundAsset);
 
 			Frontend::FGraphHandle GraphHandle = MetasoundAsset->GetRootGraphHandle();
+			const FGuid VertexID = FGuid::NewGuid();
 
 			FMetasoundFrontendClassOutput Description;
-			Description.Name = InName;
+
+			// The user currently never interfaces with the name, only the display name,
+			// so just set to VertexID to make unique, avoid confusion, and eventually
+			// this field will most likely go away (or become the display name).
+			Description.Name = VertexID.ToString();
+
 			Description.TypeName = InTypeName;
 			Description.Metadata.Description = InToolTip;
-			Description.VertexID = FGuid::NewGuid();
+			Description.VertexID = VertexID;
 
 			Frontend::FNodeHandle NodeHandle = GraphHandle->AddOutputVertex(Description);
-			GraphHandle->SetOutputDisplayName(InName, FText::FromString(InName));
+
+			const FText DisplayName = FGraphBuilder::GenerateUniqueInputDisplayName(InMetasound);
+			GraphHandle->SetOutputDisplayName(Description.Name, DisplayName);
 
 			return NodeHandle;
 		}
@@ -714,10 +655,11 @@ namespace Metasound
 
 				case EMetasoundFrontendClassType::Invalid:
 				case EMetasoundFrontendClassType::Graph:
+				case EMetasoundFrontendClassType::Variable:
 				default:
 				{
 					checkNoEntry();
-					static_assert(static_cast<int32>(EMetasoundFrontendClassType::Invalid) == 4, "Possible missing FMetasoundFrontendClassType case coverage");
+					static_assert(static_cast<int32>(EMetasoundFrontendClassType::Invalid) == 5, "Possible missing FMetasoundFrontendClassType case coverage");
 				}
 				break;
 			}
@@ -752,19 +694,6 @@ namespace Metasound
 				return false;
 			}
 
-			FOutputHandle ExistingOutput = InputHandle->GetConnectedOutput();
-			if (ExistingOutput->IsValid())
-			{
-				FNodeHandle NodeHandle = ExistingOutput->GetOwningNode();
-				if (IsLiteralInput(NodeHandle))
-				{
-					if (UMetasoundEditorGraph* Graph = Cast<UMetasoundEditorGraph>(InInputPin.GetOwningNode()->GetGraph()))
-					{
-						GraphBuilderPrivate::DeleteNode(Graph->GetMetasoundChecked(), NodeHandle);
-					}
-				}
-			}
-
 			if (!ensure(InputHandle->Connect(*OutputHandle)))
 			{
 				InInputPin.BreakLinkTo(&InOutputPin);
@@ -808,45 +737,40 @@ namespace Metasound
 				FConstOutputHandle OutputHandle = InputHandle->GetConnectedOutput();
 				const FMetasoundFrontendNodeStyle& Style = OutputHandle->GetOwningNode()->GetNodeStyle();
 
-				// Hidden nodes are not "connected" from the perspective of EdGraph,
-				// and therefore should be ignored when disconnecting links.
-				if (Style.Display.Visibility != EMetasoundFrontendNodeStyleDisplayVisibility::Hidden)
-				{
-					InputHandle->Disconnect();
+				InputHandle->Disconnect();
 
-					if (bAddLiteralInputs)
-					{
-						FNodeHandle NodeHandle = InputHandle->GetOwningNode();
-						FGraphBuilder::AddOrUpdateLiteralInput(Metasound, NodeHandle, *InputPins[i]);
-					}
+				if (bAddLiteralInputs)
+				{
+					FNodeHandle NodeHandle = InputHandle->GetOwningNode();
+					SynchronizePinLiteral(*InputPins[i]);
 				}
 			}
 		}
 
-		void FGraphBuilder::ConstructGraph(UObject& InMetasound)
+		void FGraphBuilder::InitMetaSound(UObject& InMetaSound, const FString& InAuthor)
 		{
 			using namespace Frontend;
 
-			FMetasoundAssetBase* MetasoundAsset = IMetasoundUObjectRegistry::Get().GetObjectAsAssetBase(&InMetasound);
-			check(MetasoundAsset);
+			FMetasoundFrontendClassMetadata Metadata;
 
-			FGraphHandle GraphHandle = MetasoundAsset->GetRootGraphHandle();
-			UMetasoundEditorGraph* Graph = CastChecked<UMetasoundEditorGraph>(MetasoundAsset->GetGraph());
+			Metadata.ClassName = FMetasoundFrontendClassName(FName(), InMetaSound.GetFName(), FName());
+			Metadata.Version.Major = 1;
+			Metadata.Version.Minor = 0;
+			Metadata.Type = EMetasoundFrontendClassType::Graph;
+			Metadata.Author = FText::FromString(InAuthor);
+
+			FMetasoundAssetBase* MetaSoundAsset = IMetasoundUObjectRegistry::Get().GetObjectAsAssetBase(&InMetaSound);
+			check(MetaSoundAsset);
 
-			Graph->Nodes.Reset();
+			MetaSoundAsset->SetMetadata(Metadata);
+			MetaSoundAsset->ConformDocumentToMetasoundArchetype();
+			FGraphHandle GraphHandle = MetaSoundAsset->GetRootGraphHandle();
 
 			// TODO: Space graph nodes in a procedural and readable way
 			FVector2D InputNodeLocation = FVector2D::ZeroVector;
 			FVector2D OpNodeLocation = FVector2D(250.0f, 0.0f);
 			FVector2D OutputNodeLocation = FVector2D(500.0f, 0.0f);
 
-			struct FNodePair
-			{
-				FNodeHandle NodeHandle = INodeController::GetInvalidHandle();
-				UEdGraphNode* GraphNode = nullptr;
-			};
-
-			TMap<FGuid, FNodePair> NewIdNodeMap;
 			TArray<FNodeHandle> NodeHandles = GraphHandle->GetNodes();
 			for (FNodeHandle& NodeHandle : NodeHandles)
 			{
@@ -867,52 +791,9 @@ namespace Metasound
 					NewLocation = OpNodeLocation;
 					OpNodeLocation.Y += 100.0f;
 				}
-
-				UEdGraphNode* NewNode = AddNode(InMetasound, NodeHandle, NewLocation, false /* bInSelectNewNode */);
-				if (ensure(NewNode))
-				{
-					NewIdNodeMap.Add(NodeHandle->GetID(), FNodePair{ NodeHandle, NewNode });
-				}
-			}
-
-			for (const TPair<FGuid, FNodePair>& IdNodePair : NewIdNodeMap)
-			{
-				UEdGraphNode* GraphNode = IdNodePair.Value.GraphNode;
-				check(GraphNode);
-
-				FNodeHandle NodeHandle = IdNodePair.Value.NodeHandle;
-				TArray<UEdGraphPin*> Pins = GraphNode->GetAllPins();
-
-				const TArray<FInputHandle> NodeInputs = NodeHandle->GetInputs();
-
-				int32 InputIndex = 0;
-				for (UEdGraphPin* Pin : Pins)
-				{
-					switch (Pin->Direction)
-					{
-						case EEdGraphPinDirection::EGPD_Input:
-						{
-							FOutputHandle OutputHandle = NodeInputs[InputIndex]->GetConnectedOutput();
-							if (OutputHandle->IsValid())
-							{
-								const FMetasoundFrontendNodeStyle& Style = OutputHandle->GetOwningNode()->GetNodeStyle();
-								if (!IsLiteralInput(OutputHandle->GetOwningNode()))
-								{
-									UEdGraphNode* OutputGraphNode = NewIdNodeMap.FindChecked(OutputHandle->GetOwningNodeID()).GraphNode;
-									UEdGraphPin* OutputPin = OutputGraphNode->FindPinChecked(OutputHandle->GetName(), EEdGraphPinDirection::EGPD_Output);
-									Pin->MakeLinkTo(OutputPin);
-								}
-							}
-
-							InputIndex++;
-						}
-						break;
-
-						case EEdGraphPinDirection::EGPD_Output:
-							// Do nothing.  Connecting all inputs will naturally connect all outputs where required
-						break;
-					}
-				}
+				FMetasoundFrontendNodeStyle Style = NodeHandle->GetNodeStyle();
+				Style.Display.Locations.Add(FGuid::NewGuid(), NewLocation);
+				NodeHandle->SetNodeStyle(Style);
 			}
 		}
 
@@ -941,28 +822,6 @@ namespace Metasound
 			GraphHandle->RemoveNode(*NodeHandle);
 		}
 
-		void FGraphBuilder::DeleteLiteralInputs(UEdGraphNode& InNode)
-		{
-			using namespace Frontend;
-
-			UMetasoundEditorGraphNode* Node = CastChecked<UMetasoundEditorGraphNode>(&InNode);
-			FNodeHandle NodeHandle = Node->GetNodeHandle();
-			TArray<FInputHandle> Inputs = NodeHandle->GetInputs();
-			for (FInputHandle& Input : Inputs)
-			{
-				FOutputHandle Output = Input->GetConnectedOutput();
-				if (Output->IsValid())
-				{
-					FNodeHandle InputHandle = Output->GetOwningNode();
-					if (IsLiteralInput(InputHandle))
-					{
-						FGraphHandle Graph = InputHandle->GetOwningGraph();
-						Graph->RemoveNode(*InputHandle);
-					}
-				}
-			}
-		}
-
 		bool FGraphBuilder::DeleteNode(UEdGraphNode& InNode)
 		{
 			using namespace Frontend;
@@ -989,10 +848,6 @@ namespace Metasound
 				return false;
 			}
 
-			// Delete literal inputs prior to disconnecting/removing Frontend node
-			// representation to avoid orphaning them, causing them to not get deleted.
-			DeleteLiteralInputs(InNode);
-
 			// Remove connects only to pins associated with this EdGraph node
 			// only (Iterate pins and not Frontend representation to preserve
 			// other input/output EditorGraph reference node associations)
@@ -1022,10 +877,11 @@ namespace Metasound
 					break;
 
 					case EMetasoundFrontendClassType::Graph:
+					case EMetasoundFrontendClassType::Variable:
 					case EMetasoundFrontendClassType::External:
 					default:
 					{
-						static_assert(static_cast<int32>(EMetasoundFrontendClassType::Invalid) == 4, "Possible missing MetasoundFrontendClassType switch case coverage.");
+						static_assert(static_cast<int32>(EMetasoundFrontendClassType::Invalid) == 5, "Possible missing MetasoundFrontendClassType switch case coverage.");
 
 						if (ensure(GraphHandle->RemoveNode(*NodeHandle)))
 						{
@@ -1042,7 +898,7 @@ namespace Metasound
 			// the graph is not malformed once all errors are addressed by the user.
 			if (bSuccess && bWasErroredNode)
 			{
-				FGraphBuilder::SynchronizeGraph(Graph->GetMetasoundChecked());
+				SynchronizeGraph(Graph->GetMetasoundChecked());
 			}
 
 			return bSuccess;
@@ -1052,14 +908,19 @@ namespace Metasound
 		{
 			using namespace Frontend;
 		
-			DeleteLiteralInputs(InGraphNode);
-
-			FNodeHandle NodeHandle = InGraphNode.GetNodeHandle();
 			for (int32 i = InGraphNode.Pins.Num() - 1; i >= 0; i--)
 			{
 				InGraphNode.RemovePin(InGraphNode.Pins[i]);
 			}
 
+			// TODO: Make this a utility in Frontend (ClearInputLiterals())
+			FNodeHandle NodeHandle = InGraphNode.GetNodeHandle();
+			TArray<FInputHandle> Inputs = NodeHandle->GetInputs();
+			for (FInputHandle& Input : Inputs)
+			{
+				NodeHandle->ClearInputLiteral(Input->GetID());
+			}
+
 			// Only add input pins of the node is not an input node. Input nodes
 			// have their input pins hidden because they cannot be connected internal
 			// to the graph.
@@ -1159,7 +1020,7 @@ namespace Metasound
 				RefreshPinMetadata(*NewPin, InInputHandle->GetMetadata());
 
 				FNodeHandle NodeHandle = InInputHandle->GetOwningNode();
-				FGraphBuilder::AddOrUpdateLiteralInput(InEditorNode.GetMetasoundChecked(), NodeHandle, *NewPin);
+				SynchronizePinLiteral(*NewPin);
 			}
 
 			return NewPin;
@@ -1248,22 +1109,8 @@ namespace Metasound
 					// Remove pin so it isn't used twice.
 					Pins.Remove(MatchingPin);
 
-					bool bShowConnectionInEditor = true;
-					bool bIsConnectedToLiteralInput = false;
-
 					FOutputHandle OutputHandle = NodeInput->GetConnectedOutput();
 					if (OutputHandle->IsValid())
-					{
-						FNodeHandle InputNodeHandle = OutputHandle->GetOwningNode();
-						bShowConnectionInEditor = !IsLiteralInput(InputNodeHandle);
-						bIsConnectedToLiteralInput = !bShowConnectionInEditor;
-					}
-					else
-					{
-						bShowConnectionInEditor = false;
-					}
-						
-					if (bShowConnectionInEditor)
 					{
 						bool bAddLink = false;
 
@@ -1306,29 +1153,9 @@ namespace Metasound
 							UE_LOG(LogMetasoundEditor, Display, TEXT("Synchronizing Node '%s' Connection: Breaking all pin links to '%s'"), *OwningNodeName.ToString(), *NodeInput->GetDisplayName().ToString());
 							bIsNodeDirty = true;
 						}
-
-						// Hidden literal input Link exists, so match pin default value to literal input value
-						if (bIsConnectedToLiteralInput)
-						{
-							FNodeHandle InputNode = OutputHandle->GetOwningNode();
-							const FString& InputName = InputNode->GetNodeName();
-							const FGuid VertexID = GraphHandle->GetVertexIDForInputVertex(InputName);
-							const FMetasoundFrontendLiteral DefaultLiteral = GraphHandle->GetDefaultInput(VertexID);
-							const FString DefaultLiteralString = DefaultLiteral.ToString();
-
-							FMetasoundFrontendLiteral PinDataTypeDefaultLiteral;
-							if (GetPinDefaultLiteral(*MatchingPin, PinDataTypeDefaultLiteral))
-							{
-								if (PinDataTypeDefaultLiteral.ToString() != DefaultLiteralString)
-								{
-									const FText& OwningNodeName = NodeInput->GetOwningNode()->GetDisplayName();
-									UE_LOG(LogMetasoundEditor, Display, TEXT("Synchronizing Node '%s' Connection: Setting pin '%s' default value to '%s'"), *OwningNodeName.ToString(), *MatchingPin->GetName(), *DefaultLiteralString);
-									MatchingPin->DefaultValue = DefaultLiteralString;
-									bIsNodeDirty = true;
-								}
-							}
-						}
 					}
+
+					SynchronizePinLiteral(*MatchingPin);
 				}
 
 				bIsGraphDirty |= bIsNodeDirty;
@@ -1341,11 +1168,20 @@ namespace Metasound
 		{
 			using namespace Frontend;
 
-			bool bIsEditorGraphDirty = SynchronizeVariables(InMetasound);
-
 			FMetasoundAssetBase* MetasoundAsset = IMetasoundUObjectRegistry::Get().GetObjectAsAssetBase(&InMetasound);
 			check(MetasoundAsset);
 
+			// Initial graph generation is not something to be managed by the transaction stack, so don't track dirty state until
+			// after initial setup if necessary.
+			if (!MetasoundAsset->GetGraph())
+			{
+				UMetasoundEditorGraph* Graph = NewObject<UMetasoundEditorGraph>(&InMetasound, FName(), RF_Transient | RF_Transactional);
+				Graph->Schema = UMetasoundEditorGraphSchema::StaticClass();
+				MetasoundAsset->SetGraph(Graph);
+			}
+
+			bool bIsEditorGraphDirty = SynchronizeVariables(InMetasound);
+
 			// Get all external nodes from Frontend graph.  Input and output references will only be added/synchronized
 			// if required when synchronizing connections (as they are not required to inhabit editor graph).
 			FGraphHandle GraphHandle = MetasoundAsset->GetRootGraphHandle();
@@ -1438,11 +1274,6 @@ namespace Metasound
 				}
 
 				FMetasoundFrontendNodeStyle NewStyle = CurrentStyle;
-				if (NewStyle.Display.Visibility == EMetasoundFrontendNodeStyleDisplayVisibility::Hidden)
-				{
-					continue;
-				}
-
 				bIsEditorGraphDirty = true;
 
 				TArray<UMetasoundEditorGraphNode*> AddedNodes;
@@ -1588,6 +1419,36 @@ namespace Metasound
 			return bIsNodeDirty;
 		}
 
+		bool FGraphBuilder::SynchronizePinLiteral(UEdGraphPin& InPin)
+		{
+			using namespace Metasound::Frontend;
+
+			if (!ensure(InPin.Direction == EGPD_Input))
+			{
+				return false;
+			}
+
+			const FString OldValue = InPin.DefaultValue;
+
+			FInputHandle InputHandle = GetInputHandleFromPin(&InPin);
+			if (const FMetasoundFrontendLiteral* NodeDefaultLiteral = InputHandle->GetLiteral())
+			{
+				InPin.DefaultValue = NodeDefaultLiteral->ToString();
+				return OldValue != InPin.DefaultValue;
+			}
+
+			if (const FMetasoundFrontendLiteral* ClassDefaultLiteral = InputHandle->GetClassDefaultLiteral())
+			{
+				InPin.DefaultValue = ClassDefaultLiteral->ToString();
+				return OldValue != InPin.DefaultValue;
+			}
+
+			FMetasoundFrontendLiteral DefaultLiteral;
+			DefaultLiteral.SetFromLiteral(Frontend::GetDefaultParamForDataType(InputHandle->GetDataType()));
+			InPin.DefaultValue = DefaultLiteral.ToString();
+			return OldValue != InPin.DefaultValue;
+		}
+
 		bool FGraphBuilder::SynchronizeVariables(UObject& InMetasound)
 		{
 			using namespace Frontend;
@@ -1606,12 +1467,6 @@ namespace Metasound
 			FGraphHandle GraphHandle = MetasoundAsset->GetRootGraphHandle();
 			GraphHandle->IterateNodes([&](FNodeHandle NodeHandle)
 			{
-				// Literal Inputs are not described on EdGraph side
-				if (IsLiteralInput(NodeHandle))
-				{
-					return;
-				}
-
 				if (UMetasoundEditorGraphInput* Input = Graph->FindInput(NodeHandle->GetID()))
 				{
 					Inputs.Add(Input);
diff --git a/Engine/Plugins/Runtime/Metasound/Source/MetasoundEditor/Private/MetasoundEditorGraphBuilder.h b/Engine/Plugins/Runtime/Metasound/Source/MetasoundEditor/Private/MetasoundEditorGraphBuilder.h
--- a/Engine/Plugins/Runtime/Metasound/Source/MetasoundEditor/Private/MetasoundEditorGraphBuilder.h
+++ b/Engine/Plugins/Runtime/Metasound/Source/MetasoundEditor/Private/MetasoundEditorGraphBuilder.h
@@ -69,10 +69,8 @@ namespace Metasound
 			// Generates FNodeHandle for the given external node data. Does not bind or create EdGraph representation of given node.
 			static Frontend::FNodeHandle AddInputNodeHandle(
 				UObject& InMetasound,
-				const FString& InName,
 				const FName InTypeName,
 				const FText& InToolTip,
-				bool bInIsLiteral = false,
 				const FMetasoundFrontendLiteral* InDefaultValue = nullptr);
 
 			// Adds a corresponding UMetasoundEditorGraphExternalNode for the provided node handle.
@@ -85,14 +83,11 @@ namespace Metasound
 			// Synchronizes node location data
 			static void SynchronizeNodeLocation(FVector2D InLocation, Frontend::FNodeHandle InNodeHandle, UMetasoundEditorGraphNode& InNode);
 
-			// Generates FNodeHandle for the given external node data. Does not bind or create EdGraph representation of given node.
-			static Frontend::FNodeHandle AddExternalNodeHandle(UObject& InMetasound, const FMetasoundFrontendClassMetadata& InMetadata);
-
 			// Adds an output node to the editor graph that corresponds to the provided node handle.
 			static UMetasoundEditorGraphOutputNode* AddOutputNode(UObject& InMetasound, Frontend::FNodeHandle& InNodeHandle, FVector2D InLocation, bool bInSelectNewNode = true);
 
 			// Generates analogous FNodeHandle for the given internal node data. Does not bind nor create EdGraph representation of given node.
-			static Frontend::FNodeHandle AddOutputNodeHandle(UObject& InMetasound, const FString& InName, const FName InTypeName, const FText& InToolTip);
+			static Frontend::FNodeHandle AddOutputNodeHandle(UObject& InMetasound, const FName InTypeName, const FText& InToolTip);
 
 			// Attempts to connect Frontend node counterparts together for provided pins.  Returns true if succeeded,
 			// and breaks pin link and returns false if failed.  If bConnectEdPins is set, will attempt to connect
@@ -105,9 +100,14 @@ namespace Metasound
 			// applicable post disconnection.
 			static void DisconnectPin(UEdGraphPin& InPin, bool bAddLiteralInputs = true);
 
-			static FString GenerateUniqueInputName(const UObject& InMetasound, const FString* InBaseName = nullptr);
+			// Generates a unique input display name for the given MetaSound object
+			static FText GenerateUniqueInputDisplayName(const UObject& InMetaSound, const FText* InBaseName = nullptr);
 
-			static FString GenerateUniqueOutputName(const UObject& InMetasound, const FString* InBaseName = nullptr);
+			// Generates a unique output name for the given MetaSound object
+			static FText GenerateUniqueOutputDisplayName(const UObject& InMetaSound, const FText* InBaseName = nullptr);
+
+			// Generates a unique name by retrieving the provided MetaSound's associated GraphHandle and filtering by the provided function.
+			static FText GenerateUniqueNameByFilter(const UObject& InMetaSound, const FText& InBaseText, TUniqueFunction<bool(const Frontend::FConstGraphHandle&, const FText&)> InIsValidNameFilter);
 
 			static TArray<FString> GetDataTypeNameCategories(const FName& InDataTypeName);
 
@@ -121,16 +121,8 @@ namespace Metasound
 			static Frontend::FOutputHandle GetOutputHandleFromPin(const UEdGraphPin* InPin);
 			static Frontend::FConstOutputHandle GetConstOutputHandleFromPin(const UEdGraphPin* InPin);
 
-			// Returns whether or not the given handle refers to a literal input
-			static bool IsLiteralInput(Frontend::FNodeHandle InNodeHandle);
-
-			// Adds or updates a hidden input that is set to the default literal value.  If the literal input is added, value defaults to
-			// 1. what is set on the input vertex of the node's inputer vertex interface, and if not defined there,
-			// 2. the pin's DataType literal default.
-			static void AddOrUpdateLiteralInput(UObject& InMetasound, Frontend::FNodeHandle InNodeHandle, UEdGraphPin& InInputPin, bool bForcePinValueAsDefault = false);
-
 			// Returns the default literal stored on the respective Frontend Node's Input.
-			static bool GetPinDefaultLiteral(UEdGraphPin& InInputPin, FMetasoundFrontendLiteral& OutLiteralDefault);
+			static bool GetPinLiteral(UEdGraphPin& InInputPin, FMetasoundFrontendLiteral& OutLiteralDefault);
 
 			// Deletes Editor Graph Variable's associated Frontend node, as well as any
 			// Editor Graph nodes referencing the given variable.
@@ -139,15 +131,12 @@ namespace Metasound
 			// Retrieves the proper pin color for the given PinType
 			static FLinearColor GetPinCategoryColor(const FEdGraphPinType& PinType);
 
-			// Constructs graph with default inputs & outputs.
-			static void ConstructGraph(UObject& InMetasound);
+			// Initializes MetaSound with default inputs & outputs.
+			static void InitMetaSound(UObject& InMetaSound, const FString& InAuthor);
 
 			// Rebuilds all editor node pins based on the provided node handle's class definition.
 			static void RebuildNodePins(UMetasoundEditorGraphNode& InGraphNode);
 
-			// Removes all literal inputs connected to the given node
-			static void DeleteLiteralInputs(UEdGraphNode& InNode);
-
 			// Deletes both the editor graph & frontend nodes from respective graphs
 			static bool DeleteNode(UEdGraphNode& InNode);
 
@@ -176,6 +165,9 @@ namespace Metasound
 			// @return True if the UEdGraph was altered. False otherwise.
 			static bool SynchronizeConnections(UObject& InMetasound);
 
+			// Synchronizes literal for a given input with the EdGraph's pin value.
+			static bool SynchronizePinLiteral(UEdGraphPin& InPin);
+
 			// Synchronizes inputs and outputs for the given Metasound.
 			//
 			// @return True if the UEdGraph was altered. False otherwise.
diff --git a/Engine/Plugins/Runtime/Metasound/Source/MetasoundEditor/Private/MetasoundEditorGraphNode.cpp b/Engine/Plugins/Runtime/Metasound/Source/MetasoundEditor/Private/MetasoundEditorGraphNode.cpp
--- a/Engine/Plugins/Runtime/Metasound/Source/MetasoundEditor/Private/MetasoundEditorGraphNode.cpp
+++ b/Engine/Plugins/Runtime/Metasound/Source/MetasoundEditor/Private/MetasoundEditorGraphNode.cpp
@@ -248,9 +248,15 @@ void UMetasoundEditorGraphNode::PinDefaultValueChanged(UEdGraphPin* Pin)
 	{
 		GetMetasoundChecked().Modify();
 
-		FNodeHandle NodeHandle = GetNodeHandle();
-		IMetasoundEditorModule& EditorModule = FModuleManager::GetModuleChecked<IMetasoundEditorModule>("MetaSoundEditor");
-		FGraphBuilder::AddOrUpdateLiteralInput(GetMetasoundChecked(), NodeHandle, *Pin);
+		FInputHandle InputHandle = FGraphBuilder::GetInputHandleFromPin(Pin);
+		if (InputHandle->IsValid())
+		{
+			FMetasoundFrontendLiteral LiteralValue;
+			if (FGraphBuilder::GetPinLiteral(*Pin, LiteralValue))
+			{
+				InputHandle->SetLiteral(LiteralValue);
+			}
+		}
 	}
 }
 
@@ -351,12 +357,15 @@ void UMetasoundEditorGraphNode::PostEditChangeChainProperty(FPropertyChangedChai
 
 void UMetasoundEditorGraphNode::PostEditUndo()
 {
+	using namespace Metasound::Editor;
+	using namespace Metasound::Frontend;
+
 	UEdGraphPin::ResolveAllPinReferences();
 
 	// This can trigger and the handle is no longer valid if transaction
 	// is being undone on a graph node that is orphaned.  If orphaned,
 	// bail early.
-	Metasound::Frontend::FNodeHandle NodeHandle = GetNodeHandle();
+	FNodeHandle NodeHandle = GetNodeHandle();
 	if (!NodeHandle->IsValid())
 	{
 		return;
@@ -366,8 +375,7 @@ void UMetasoundEditorGraphNode::PostEditUndo()
 	{
 		if (Pin && Pin->Direction == EGPD_Input)
 		{
-			UObject& Metasound = GetMetasoundChecked();
-			Metasound::Editor::FGraphBuilder::AddOrUpdateLiteralInput(Metasound, NodeHandle, *Pin);
+			FGraphBuilder::SynchronizePinLiteral(*Pin);
 		}
 	}
 }
@@ -651,17 +659,16 @@ UMetasoundEditorGraphExternalNode* UMetasoundEditorGraphExternalNode::UpdateToVe
 			FInputHandle InputHandle = FGraphBuilder::GetInputHandleFromPin(&Pin);
 			if (ConnectionInfo->DataType == InputHandle->GetDataType())
 			{
-				if (ConnectionInfo->PinsLinkedTo.IsEmpty())
+				Pin.DefaultValue = ConnectionInfo->DefaultValue;
+				FMetasoundFrontendLiteral LiteralValue;
+				if (FGraphBuilder::GetPinLiteral(Pin, LiteralValue))
 				{
-					Pin.DefaultValue = ConnectionInfo->DefaultValue;
-					FGraphBuilder::AddOrUpdateLiteralInput(Metasound, InputHandle->GetOwningNode(), Pin, true /* bForcePinValueAsDefault */);
+					InputHandle->SetLiteral(LiteralValue);
 				}
-				else
+
+				for (UEdGraphPin* OutputPin : ConnectionInfo->PinsLinkedTo)
 				{
-					for (UEdGraphPin* OutputPin : ConnectionInfo->PinsLinkedTo)
-					{
-						FGraphBuilder::ConnectNodes(Pin, *OutputPin, true /* bInConnectEdPins */);
-					}
+					FGraphBuilder::ConnectNodes(Pin, *OutputPin, true /* bInConnectEdPins */);
 				}
 			}
 		}
diff --git a/Engine/Plugins/Runtime/Metasound/Source/MetasoundEditor/Private/MetasoundEditorGraphSchema.cpp b/Engine/Plugins/Runtime/Metasound/Source/MetasoundEditor/Private/MetasoundEditorGraphSchema.cpp
--- a/Engine/Plugins/Runtime/Metasound/Source/MetasoundEditor/Private/MetasoundEditorGraphSchema.cpp
+++ b/Engine/Plugins/Runtime/Metasound/Source/MetasoundEditor/Private/MetasoundEditorGraphSchema.cpp
@@ -258,15 +258,16 @@ UEdGraphNode* FMetasoundGraphSchemaAction_PromoteToInput::PerformAction(UEdGraph
 
 	const FText& InputName = InputHandle->GetOwningNode()->GetDisplayName();
 	const FText& InputNodeName = InputHandle->GetDisplayName();
-	const FString InputNameBase = FText::Format(LOCTEXT("PromoteInputNameBaseFormat", "{0} {1}"), InputName, InputNodeName).ToString();
-	FString NewNodeName = FGraphBuilder::GenerateUniqueInputName(ParentMetasound, &InputNameBase);
 
 	FMetasoundFrontendLiteral DefaultValue;
-	FGraphBuilder::GetPinDefaultLiteral(*FromPin, DefaultValue);
+	FGraphBuilder::GetPinLiteral(*FromPin, DefaultValue);
 
-	FNodeHandle NodeHandle = FGraphBuilder::AddInputNodeHandle(ParentMetasound, NewNodeName, InputHandle->GetDataType(), FText::GetEmpty(), false /* bIsLiteralInput */, &DefaultValue);
+	FNodeHandle NodeHandle = FGraphBuilder::AddInputNodeHandle(ParentMetasound, InputHandle->GetDataType(), FText::GetEmpty(), &DefaultValue);
 	if (ensure(NodeHandle->IsValid()))
 	{
+		const FText InputNameBase = FText::Format(LOCTEXT("PromoteInputNameBaseFormat", "{0} {1}"), InputName, InputNodeName);
+		const FText NewNodeName = FGraphBuilder::GenerateUniqueInputDisplayName(ParentMetasound, &InputNameBase);
+		NodeHandle->SetDisplayName(NewNodeName);
 		UMetasoundEditorGraphInput* Input = MetasoundGraph->FindOrAddInput(NodeHandle);
 		if (ensure(Input))
 		{
@@ -741,12 +742,15 @@ FText UMetasoundEditorGraphSchema::GetPinDisplayName(const UEdGraphPin* Pin) con
 			}
 		}
 
-		case EMetasoundFrontendClassType::External:
+		case EMetasoundFrontendClassType::Variable:
 		case EMetasoundFrontendClassType::Graph:
+			// TODO: Implement above
+
+		case EMetasoundFrontendClassType::External:
 		case EMetasoundFrontendClassType::Invalid:
 		default:
 		{
-			static_assert(static_cast<int32>(EMetasoundFrontendClassType::Invalid) == 4, "Possible missing EMetasoundFrontendClassType case coverage");
+			static_assert(static_cast<int32>(EMetasoundFrontendClassType::Invalid) == 5, "Possible missing EMetasoundFrontendClassType case coverage");
 			return Super::GetPinDisplayName(Pin);
 		}
 	}
@@ -859,15 +863,6 @@ void UMetasoundEditorGraphSchema::GetDataTypeInputNodeActions(FGraphContextMenuB
 	using namespace Metasound::Frontend;
 
 	TArray<FConstNodeHandle> Inputs = InGraphHandle->GetConstInputNodes();
-	for (int32 i = Inputs.Num() - 1; i >= 0; --i)
-	{
-		const FConstNodeHandle& NodeHandle = Inputs[i];
-		if (NodeHandle->GetNodeStyle().Display.Visibility == EMetasoundFrontendNodeStyleDisplayVisibility::Hidden)
-		{
-			Inputs.RemoveAtSwap(i, 1, false);
-		}
-	}
-
 	const SchemaPrivate::FDataTypeActionQuery ActionQuery
 	{
 		ActionMenuBuilder,
diff --git a/Engine/Plugins/Runtime/Metasound/Source/MetasoundEditor/Private/MetasoundFactory.cpp b/Engine/Plugins/Runtime/Metasound/Source/MetasoundEditor/Private/MetasoundFactory.cpp
--- a/Engine/Plugins/Runtime/Metasound/Source/MetasoundEditor/Private/MetasoundFactory.cpp
+++ b/Engine/Plugins/Runtime/Metasound/Source/MetasoundEditor/Private/MetasoundFactory.cpp
@@ -4,47 +4,9 @@
 #include "Kismet/KismetSystemLibrary.h"
 #include "Metasound.h"
 #include "MetasoundEditorGraphBuilder.h"
-#include "MetasoundEditorGraphNode.h"
-#include "MetasoundEditorGraphSchema.h"
 #include "MetasoundSource.h"
-#include "MetasoundUObjectRegistry.h"
 
 
-namespace Metasound
-{
-	namespace Editor
-	{
-		namespace FactoryPrivate
-		{
-			UObject* InitMetasound(UMetaSoundSource* MetasoundSource, EObjectFlags Flags)
-			{
-				FMetasoundFrontendClassMetadata Metadata;
-
-				Metadata.ClassName = FMetasoundFrontendClassName(FName(), *MetasoundSource->GetName(), FName());
-				Metadata.Version.Major = 1;
-				Metadata.Version.Minor = 0;
-				Metadata.Type = EMetasoundFrontendClassType::Graph;
-				Metadata.Author = FText::FromString(UKismetSystemLibrary::GetPlatformUserName());
-
-				MetasoundSource->SetMetadata(Metadata);
-				MetasoundSource->ConformDocumentToMetasoundArchetype();
-
-				FMetasoundAssetBase* MetasoundAsset = IMetasoundUObjectRegistry::Get().GetObjectAsAssetBase(MetasoundSource);
-				check(MetasoundAsset);
-				check(!MetasoundAsset->GetGraph());
-
-				UObject* MetasoundObject = CastChecked<UObject>(MetasoundSource);
-
-				UMetasoundEditorGraph* Graph = NewObject<UMetasoundEditorGraph>(MetasoundSource, FName(), Flags);
-				Graph->Schema = UMetasoundEditorGraphSchema::StaticClass();
-				MetasoundAsset->SetGraph(Graph);
-
-				return MetasoundSource;
-			}
-		} // namespace FactoryPrivate
-	} // namespace Editor
-} // namespace Metasound
-
 // TODO: Re-enable and potentially rename once composition is supported
 // UMetasoundFactory::UMetasoundFactory(const FObjectInitializer& ObjectInitializer)
 // 	: Super(ObjectInitializer)
@@ -59,12 +21,7 @@ namespace Metasound
 // UObject* UMetasoundFactory::FactoryCreateNew(UClass* InClass, UObject* InParent, FName Name, EObjectFlags Flags, UObject* InContext, FFeedbackContext* InFeedbackContext)
 // {
 // 	UMetaSound* Metasound = NewObject<UMetaSound>(InParent, Name, Flags);
-//	Metasound::Editor::FactoryPrivate::InitMetasound(Metasound, Flags);
-// 
-//	CastChecked<UMetasoundEditorGraph>(&Metasound->GetGraphChecked())->ParentMetasound = MetasoundSource;
-// 
-//	FGraphBuilder::ConstructGraph(*Metasound);
-//	FGraphBuilder::SynchronizeGraph(*Metasound);
+//	FGraphBuilder::InitMetasound(Metasound);
 // }
 
 UMetasoundSourceFactory::UMetasoundSourceFactory(const FObjectInitializer& ObjectInitializer)
@@ -79,15 +36,10 @@ UMetasoundSourceFactory::UMetasoundSourceFactory(const FObjectInitializer& Objec
 
 UObject* UMetasoundSourceFactory::FactoryCreateNew(UClass* InClass, UObject* InParent, FName Name, EObjectFlags Flags, UObject* InContext, FFeedbackContext* InFeedbackContext)
 {
-	using namespace Metasound;
 	using namespace Metasound::Editor;
-	using namespace Metasound::Frontend;
-
-	UMetaSoundSource* MetasoundSource = NewObject<UMetaSoundSource>(InParent, Name, Flags);
-	Metasound::Editor::FactoryPrivate::InitMetasound(MetasoundSource, Flags);
-
-	FGraphBuilder::ConstructGraph(*MetasoundSource);
-	FGraphBuilder::SynchronizeGraph(*MetasoundSource);
 
-	return MetasoundSource;
+	UMetaSoundSource* MetaSoundSource = NewObject<UMetaSoundSource>(InParent, Name, Flags);
+	check(MetaSoundSource);
+	FGraphBuilder::InitMetaSound(*MetaSoundSource, UKismetSystemLibrary::GetPlatformUserName());
+	return MetaSoundSource;
 }
diff --git a/Engine/Plugins/Runtime/Metasound/Source/MetasoundEditor/Public/MetasoundEditorGraph.h b/Engine/Plugins/Runtime/Metasound/Source/MetasoundEditor/Public/MetasoundEditorGraph.h
--- a/Engine/Plugins/Runtime/Metasound/Source/MetasoundEditor/Public/MetasoundEditorGraph.h
+++ b/Engine/Plugins/Runtime/Metasound/Source/MetasoundEditor/Public/MetasoundEditorGraph.h
@@ -31,13 +31,15 @@ namespace Metasound
 
 DECLARE_MULTICAST_DELEGATE_OneParam(FOnMetasoundNodeNameChanged, FGuid /* NodeID */);
 
+// TODO: Implement for core variables, possibly splitting into two classes, UMetasoundEditorGraphVariableBase
+// & UMetasoundEditorGraphVariable (for actual variables).
 UCLASS(Abstract)
 class METASOUNDEDITOR_API UMetasoundEditorGraphVariable : public UObject
 {
 	GENERATED_BODY()
 
 protected:
-	virtual Metasound::Frontend::FNodeHandle AddNodeHandle(const FString& InNodeName, const FText& InNodeDisplayName, FName InDataType)
+	virtual Metasound::Frontend::FNodeHandle AddNodeHandle(const FText& InNodeDisplayName, FName InDataType)
 	{
 		return Metasound::Frontend::INodeController::GetInvalidHandle();
 	}
@@ -103,10 +105,9 @@ class METASOUNDEDITOR_API UMetasoundEditorGraphInput : public UMetasoundEditorGr
 	GENERATED_BODY()
 
 protected:
-	virtual Metasound::Frontend::FNodeHandle AddNodeHandle(const FString& InNodeName, const FText& InNodeDisplayName, FName InDataType) override;
+	virtual Metasound::Frontend::FNodeHandle AddNodeHandle(const FText& InNodeDisplayName, FName InDataType) override;
 	virtual EMetasoundFrontendClassType GetClassType() const override { return EMetasoundFrontendClassType::Input; }
 	virtual const FText& GetVariableLabel() const override;
-
 public:
 	UPROPERTY(VisibleAnywhere, Category = DefaultValue)
 	UMetasoundEditorGraphInputLiteral* Literal;
@@ -128,7 +129,7 @@ class METASOUNDEDITOR_API UMetasoundEditorGraphOutput : public UMetasoundEditorG
 	GENERATED_BODY()
 
 protected:
-	virtual Metasound::Frontend::FNodeHandle AddNodeHandle(const FString& InNodeName, const FText& InNodeDisplayName, FName InDataType) override;
+	virtual Metasound::Frontend::FNodeHandle AddNodeHandle(const FText& InNodeDisplayName, FName InDataType) override;
 	virtual EMetasoundFrontendClassType GetClassType() const override { return EMetasoundFrontendClassType::Output; }
 	virtual const FText& GetVariableLabel() const override;
 };
diff --git a/Engine/Plugins/Runtime/Metasound/Source/MetasoundEngine/Private/MetasoundSource.cpp b/Engine/Plugins/Runtime/Metasound/Source/MetasoundEngine/Private/MetasoundSource.cpp
--- a/Engine/Plugins/Runtime/Metasound/Source/MetasoundEngine/Private/MetasoundSource.cpp
+++ b/Engine/Plugins/Runtime/Metasound/Source/MetasoundEngine/Private/MetasoundSource.cpp
@@ -334,10 +334,7 @@ TArray<FString> UMetaSoundSource::GetTransmittableInputVertexNames() const
 					Frontend::FConstNodeHandle InputNode = RootGraph->GetNodeWithID(ClassInput->NodeID);
 					if (InputNode->IsValid())
 					{
-						if (EMetasoundFrontendNodeStyleDisplayVisibility::Visible == InputNode->GetNodeStyle().Display.Visibility)
-						{
-							return true;
-						}
+						return true;
 					}
 				}
 			}
diff --git a/Engine/Plugins/Runtime/Metasound/Source/MetasoundFrontend/Private/MetasoundFrontendBaseClasses.cpp b/Engine/Plugins/Runtime/Metasound/Source/MetasoundFrontend/Private/MetasoundFrontendBaseClasses.cpp
--- a/Engine/Plugins/Runtime/Metasound/Source/MetasoundFrontend/Private/MetasoundFrontendBaseClasses.cpp
+++ b/Engine/Plugins/Runtime/Metasound/Source/MetasoundFrontend/Private/MetasoundFrontendBaseClasses.cpp
@@ -1,12 +1,11 @@
 // Copyright Epic Games, Inc. All Rights Reserved.
 
 #include "MetasoundFrontendBaseClasses.h"
-#include "MetasoundDataReference.h"
 
+#include "MetasoundDataReference.h"
 #include "MetasoundFrontendRegistries.h"
 
 
-
 namespace Metasound
 {
 	namespace Frontend
@@ -16,9 +15,14 @@ namespace Metasound
 			return FMetasoundFrontendRegistryContainer::Get()->ConstructInputNode(InInputType, MoveTemp(InParams));
 		}
 
-		TUniquePtr<INode> ConstructOutputNode(const FName& InOutputType, const FOutputNodeConstructorParams& InParams)
+		TUniquePtr<INode> ConstructVariableNode(const FName& InVariableType, FVariableNodeConstructorParams&& InParams)
+		{
+			return FMetasoundFrontendRegistryContainer::Get()->ConstructInputNode(InVariableType, MoveTemp(InParams));
+		}
+
+		TUniquePtr<INode> ConstructOutputNode(const FName& InOutputType, FOutputNodeConstructorParams&& InParams)
 		{
-			return FMetasoundFrontendRegistryContainer::Get()->ConstructOutputNode(InOutputType, InParams);
+			return FMetasoundFrontendRegistryContainer::Get()->ConstructOutputNode(InOutputType, MoveTemp(InParams));
 		}
 
 		void RegisterPendingNodes()
@@ -27,5 +31,3 @@ namespace Metasound
 		}
 	}
 }
-
-
diff --git a/Engine/Plugins/Runtime/Metasound/Source/MetasoundFrontend/Private/MetasoundFrontendDocument.cpp b/Engine/Plugins/Runtime/Metasound/Source/MetasoundFrontend/Private/MetasoundFrontendDocument.cpp
--- a/Engine/Plugins/Runtime/Metasound/Source/MetasoundFrontend/Private/MetasoundFrontendDocument.cpp
+++ b/Engine/Plugins/Runtime/Metasound/Source/MetasoundFrontend/Private/MetasoundFrontendDocument.cpp
@@ -104,6 +104,14 @@ FMetasoundFrontendClassInput::FMetasoundFrontendClassInput(const FMetasoundFront
 	DefaultLiteral.SetType(DefaultType);
 }
 
+FMetasoundFrontendClassVariable::FMetasoundFrontendClassVariable(const FMetasoundFrontendClassVertex& InOther)
+	: FMetasoundFrontendClassVertex(InOther)
+{
+	EMetasoundFrontendLiteralType DefaultType = Metasound::Frontend::GetMetasoundFrontendLiteralType(FMetasoundFrontendRegistryContainer::Get()->GetDesiredLiteralTypeForDataType(InOther.TypeName));
+
+	DefaultLiteral.SetType(DefaultType);
+}
+
 FMetasoundFrontendGraphClass::FMetasoundFrontendGraphClass()
 {
 	Metadata.Type = EMetasoundFrontendClassType::Graph;
diff --git a/Engine/Plugins/Runtime/Metasound/Source/MetasoundFrontend/Private/MetasoundFrontendGraph.cpp b/Engine/Plugins/Runtime/Metasound/Source/MetasoundFrontend/Private/MetasoundFrontendGraph.cpp
--- a/Engine/Plugins/Runtime/Metasound/Source/MetasoundFrontend/Private/MetasoundFrontendGraph.cpp
+++ b/Engine/Plugins/Runtime/Metasound/Source/MetasoundFrontend/Private/MetasoundFrontendGraph.cpp
@@ -14,6 +14,21 @@
 
 namespace Metasound
 {
+	namespace FrontendGraphPrivate
+	{
+		FNodeInitData CreateNodeInitData(const FMetasoundFrontendNode& InNode)
+		{
+			FNodeInitData InitData;
+
+			InitData.InstanceName.Append(InNode.Name);
+			InitData.InstanceName.AppendChar('_');
+			InitData.InstanceName.Append(InNode.ID.ToString());
+
+			InitData.InstanceID = InNode.ID;
+
+			return InitData;
+		}
+	}
 
 	FFrontendGraph::FFrontendGraph(const FString& InInstanceName, const FGuid& InInstanceID)
 	:	FGraph(InInstanceName, InInstanceID)
@@ -143,7 +158,7 @@ namespace Metasound
 		NodeStorage.Add(InNode);
 	}
 
-	TUniquePtr<INode> FFrontendGraphBuilder::CreateInputNode(const FMetasoundFrontendNode& InNode, const FMetasoundFrontendClass& InClass, const FMetasoundFrontendClassInput& InOwningGraphClassInput) const
+	TUniquePtr<INode> FFrontendGraphBuilder::CreateInputNode(const FMetasoundFrontendNode& InNode, const FMetasoundFrontendClass& InClass, const FMetasoundFrontendClassInput& InOwningGraphClassInput)
 	{
 		const FMetasoundFrontendLiteral* FrontendLiteral = FindInputLiteralForInputNode(InNode, InClass, InOwningGraphClassInput);
 
@@ -184,7 +199,7 @@ namespace Metasound
 		return TUniquePtr<INode>(nullptr);
 	}
 
-	TUniquePtr<INode> FFrontendGraphBuilder::CreateOutputNode(const FMetasoundFrontendNode& InNode, const FMetasoundFrontendClass& InClass) const
+	TUniquePtr<INode> FFrontendGraphBuilder::CreateOutputNode(const FMetasoundFrontendNode& InNode, const FMetasoundFrontendClass& InClass, FBuildGraphContext& InGraphContext)
 	{
 		check(InClass.Metadata.Type == EMetasoundFrontendClassType::Output);
 		check(InNode.ClassID == InClass.ID);
@@ -197,27 +212,39 @@ namespace Metasound
 			{
 				InNode.Name,
 				InNode.ID,
-				OutputVertex.Name,
+				OutputVertex.Name
 			};
 
+			{
+				const FNodeInitData InitData = FrontendGraphPrivate::CreateNodeInitData(InNode);
+				TArray<FDefaultVariableData> DefaultVariableData = GetInputDefaultVariableData(InNode, InitData);
+				for (FDefaultVariableData& Data : DefaultVariableData)
+				{
+					FGuid VertexID = Data.DestinationVertexID;
+					InGraphContext.DefaultInputs.Emplace(MoveTemp(VertexID), MoveTemp(Data));
+				}
+			}
+
 			// TODO: create output node using external class definition
-			return FMetasoundFrontendRegistryContainer::Get()->ConstructOutputNode(OutputVertex.TypeName, InitParams);
+			return FMetasoundFrontendRegistryContainer::Get()->ConstructOutputNode(OutputVertex.TypeName, MoveTemp(InitParams));
 		}
 		return TUniquePtr<INode>(nullptr);
 	}
 
-	TUniquePtr<INode> FFrontendGraphBuilder::CreateExternalNode(const FMetasoundFrontendNode& InNode, const FMetasoundFrontendClass& InClass) const
+	TUniquePtr<INode> FFrontendGraphBuilder::CreateExternalNode(const FMetasoundFrontendNode& InNode, const FMetasoundFrontendClass& InClass, FBuildGraphContext& InGraphContext)
 	{
 		check(InClass.Metadata.Type == EMetasoundFrontendClassType::External);
 		check(InNode.ClassID == InClass.ID);
 
-		FNodeInitData InitData;
-
-		InitData.InstanceName.Append(InNode.Name);
-		InitData.InstanceName.AppendChar('_');
-		InitData.InstanceName.Append(InNode.ID.ToString());
-
-		InitData.InstanceID = InNode.ID;
+		const FNodeInitData InitData = FrontendGraphPrivate::CreateNodeInitData(InNode);
+		{
+			TArray<FDefaultVariableData> DefaultVariableData = GetInputDefaultVariableData(InNode, InitData);
+			for (FDefaultVariableData& Data : DefaultVariableData)
+			{
+				FGuid VertexID = Data.DestinationVertexID;
+				InGraphContext.DefaultInputs.Emplace(MoveTemp(VertexID), MoveTemp(Data));
+			}
+		}
 
 		// TODO: handle check to see if node interface conforms to class interface here. 
 		// TODO: check to see if external object supports class interface.
@@ -227,8 +254,7 @@ namespace Metasound
 		return FMetasoundFrontendRegistryContainer::Get()->ConstructExternalNode(Key, InitData);
 	}
 
-
-	const FMetasoundFrontendClassInput* FFrontendGraphBuilder::FindClassInputForInputNode(const FMetasoundFrontendGraphClass& InOwningGraph, const FMetasoundFrontendNode& InInputNode, int32& OutClassInputIndex) const
+	const FMetasoundFrontendClassInput* FFrontendGraphBuilder::FindClassInputForInputNode(const FMetasoundFrontendGraphClass& InOwningGraph, const FMetasoundFrontendNode& InInputNode, int32& OutClassInputIndex)
 	{
 		OutClassInputIndex = INDEX_NONE;
 
@@ -251,7 +277,7 @@ namespace Metasound
 		return nullptr;
 	}
 
-	const FMetasoundFrontendClassOutput* FFrontendGraphBuilder::FindClassOutputForOutputNode(const FMetasoundFrontendGraphClass& InOwningGraph, const FMetasoundFrontendNode& InOutputNode, int32& OutClassOutputIndex) const
+	const FMetasoundFrontendClassOutput* FFrontendGraphBuilder::FindClassOutputForOutputNode(const FMetasoundFrontendGraphClass& InOwningGraph, const FMetasoundFrontendNode& InOutputNode, int32& OutClassOutputIndex)
 	{
 		OutClassOutputIndex = INDEX_NONE;
 
@@ -274,8 +300,7 @@ namespace Metasound
 		return nullptr;
 	}
 
-
-	const FMetasoundFrontendLiteral* FFrontendGraphBuilder::FindInputLiteralForInputNode(const FMetasoundFrontendNode& InInputNode, const FMetasoundFrontendClass& InInputNodeClass, const FMetasoundFrontendClassInput& InOwningGraphClassInput) const
+	const FMetasoundFrontendLiteral* FFrontendGraphBuilder::FindInputLiteralForInputNode(const FMetasoundFrontendNode& InInputNode, const FMetasoundFrontendClass& InInputNodeClass, const FMetasoundFrontendClassInput& InOwningGraphClassInput)
 	{
 		// Default value priority is:
 		// 1. A value set directly on the node
@@ -328,79 +353,86 @@ namespace Metasound
 	}
 
 	// TODO: add errors here. Most will be a "PromptIfMissing"...
-	void FFrontendGraphBuilder::AddNodesToGraph(const FMetasoundFrontendGraphClass& InGraphClass, const FDependencyByIDMap& InClasses, const FSharedNodeByIDMap& InSubgraphs, FFrontendGraph& OutGraph) const
+	void FFrontendGraphBuilder::AddNodesToGraph(FBuildGraphContext& InGraphContext)
 	{
-		for (const FMetasoundFrontendNode& Node : InGraphClass.Graph.Nodes)
+		for (const FMetasoundFrontendNode& Node : InGraphContext.GraphClass.Graph.Nodes)
 		{
-			const FMetasoundFrontendClass* NodeClass = InClasses.FindRef(Node.ClassID);
+			const FMetasoundFrontendClass* NodeClass = InGraphContext.BuildContext.FrontendClasses.FindRef(Node.ClassID);
 
 			if (ensure(nullptr != NodeClass))
 			{
 				switch (NodeClass->Metadata.Type)
 				{
 					case EMetasoundFrontendClassType::Input:
+					{
+						int32 InputIndex = INDEX_NONE;
+						const FMetasoundFrontendClassInput* ClassInput = FindClassInputForInputNode(InGraphContext.GraphClass, Node, InputIndex);
+
+						if ((nullptr != ClassInput) && (INDEX_NONE != InputIndex))
 						{
-							int32 InputIndex = INDEX_NONE;
-							const FMetasoundFrontendClassInput* ClassInput = FindClassInputForInputNode(InGraphClass, Node, InputIndex);
-
-							if ((nullptr != ClassInput) && (INDEX_NONE != InputIndex))
-							{
-								TSharedPtr<const INode> InputNode(CreateInputNode(Node, *NodeClass, *ClassInput).Release());
-								OutGraph.AddInputNode(Node.ID, InputIndex, ClassInput->Name, InputNode);
-							}
-							else
-							{
-								UE_LOG(LogMetaSound, Error, TEXT("Failed to match input node [NodeID:%s, NodeName:%s] to owning graph [ClassID:%s] output."), *Node.ID.ToString(), *Node.Name, *InGraphClass.ID.ToString());
-							}
+							TSharedPtr<const INode> InputNode(CreateInputNode(Node, *NodeClass, *ClassInput).Release());
+							InGraphContext.Graph->AddInputNode(Node.ID, InputIndex, ClassInput->Name, InputNode);
 						}
-
-						break;
+						else
+						{
+							const FString GraphClassIDString = InGraphContext.GraphClass.ID.ToString();
+							UE_LOG(LogMetaSound, Error, TEXT("Failed to match input node [NodeID:%s, NodeName:%s] to owning graph [ClassID:%s] output."), *Node.ID.ToString(), *Node.Name, *GraphClassIDString);
+						}
+					}
+					break;
 
 					case EMetasoundFrontendClassType::Output:
+					{
+						int32 OutputIndex = INDEX_NONE;
+						const FMetasoundFrontendClassOutput* ClassOutput = FindClassOutputForOutputNode(InGraphContext.GraphClass, Node, OutputIndex);
+						if ((nullptr != ClassOutput) && (INDEX_NONE != OutputIndex))
+						{
+							TSharedPtr<const INode> OutputNode(CreateOutputNode(Node, *NodeClass, InGraphContext).Release());
+							InGraphContext.Graph->AddOutputNode(Node.ID, OutputIndex, ClassOutput->Name, OutputNode);
+						}
+						else
 						{
-							int32 OutputIndex = INDEX_NONE;
-							const FMetasoundFrontendClassOutput* ClassOutput = FindClassOutputForOutputNode(InGraphClass, Node, OutputIndex);
-
-							if ((nullptr != ClassOutput) && (INDEX_NONE != OutputIndex))
-							{
-								TSharedPtr<const INode> OutputNode(CreateOutputNode(Node, *NodeClass).Release());
-								OutGraph.AddOutputNode(Node.ID, OutputIndex, ClassOutput->Name, OutputNode);
-							}
-							else
-							{
-								UE_LOG(LogMetaSound, Error, TEXT("Failed to match output node [NodeID:%s, NodeName:%s] to owning graph [ClassID:%s] output."), *Node.ID.ToString(), *Node.Name, *InGraphClass.ID.ToString());
-							}
+							const FString GraphClassIDString = InGraphContext.GraphClass.ID.ToString();
+							UE_LOG(LogMetaSound, Error, TEXT("Failed to match output node [NodeID:%s, NodeName:%s] to owning graph [ClassID:%s] output."), *Node.ID.ToString(), *Node.Name, *GraphClassIDString);
 						}
+					}
+					break;
+
+					case EMetasoundFrontendClassType::Variable:
+					{
+						ensureAlwaysMsgf(false, TEXT("TODO: Implement ability to create variables directly in Frontend"));
+					}
 
-						break;
+					break;
 
 					case EMetasoundFrontendClassType::Graph:
+					{
+						const TSharedPtr<const INode> SubgraphPtr = InGraphContext.BuildContext.Graphs.FindRef(Node.ClassID);
+
+						if (SubgraphPtr.IsValid())
 						{
-							const TSharedPtr<const INode>* SubgraphPtr = InSubgraphs.Find(Node.ClassID);
-
-							if (nullptr == SubgraphPtr)
-							{
-								UE_LOG(LogMetaSound, Error, TEXT("Failed to find subgraph for node [NodeID:%s, NodeName:%s, ClassID:%s]"), *Node.ID.ToString(), *Node.Name, *Node.ClassID.ToString());
-							}
-							else
-							{
-								OutGraph.AddNode(Node.ID, *SubgraphPtr);
-							}
+							InGraphContext.Graph->AddNode(Node.ID, SubgraphPtr);
 						}
-						break;
+						else
+						{
+							UE_LOG(LogMetaSound, Error, TEXT("Failed to find subgraph for node [NodeID:%s, NodeName:%s, ClassID:%s]"), *Node.ID.ToString(), *Node.Name, *Node.ClassID.ToString());
+						}
+					}
+					break;
 
 					case EMetasoundFrontendClassType::External:
 					default:
-						{
-							TSharedPtr<const INode> ExternalNode(CreateExternalNode(Node, *NodeClass).Release());
-							OutGraph.AddNode(Node.ID, ExternalNode);
-						}
+					{
+						TSharedPtr<const INode> ExternalNode(CreateExternalNode(Node, *NodeClass, InGraphContext).Release());
+						InGraphContext.Graph->AddNode(Node.ID, ExternalNode);
+					}
+					break;
 				}
 			}
 		}
 	}
 
-	void FFrontendGraphBuilder::AddEdgesToGraph(const FMetasoundFrontendGraph& InFrontendGraph, FFrontendGraph& OutGraph) const
+	void FFrontendGraphBuilder::AddEdgesToGraph(FBuildGraphContext& InGraphContext)
 	{
 		// Pair of frontend node and core node. The frontend node can be one of
 		// several types.
@@ -416,12 +448,12 @@ namespace Metasound
 		TMap<FNodeIDVertexID, FCoreNodeAndFrontendVertex> NodeDestinationsByID;
 
 		// Add nodes to NodeID/VertexID map
-		for (const FMetasoundFrontendNode& Node : InFrontendGraph.Nodes)
+		for (const FMetasoundFrontendNode& Node : InGraphContext.GraphClass.Graph.Nodes)
 		{
-			const INode* CoreNode = OutGraph.FindNode(Node.ID);
+			const INode* CoreNode = InGraphContext.Graph->FindNode(Node.ID);
 			if (nullptr == CoreNode)
 			{
-				UE_LOG(LogMetaSound, Display, TEXT("Could not find referenced node [NodeID:%s]"), *Node.ID.ToString());
+				UE_LOG(LogMetaSound, Warning, TEXT("Could not find referenced node [Name:%s, NodeID:%s]"), *Node.Name, *Node.ID.ToString());
 				continue;
 			}
 
@@ -435,9 +467,8 @@ namespace Metasound
 				NodeSourcesByID.Add(FNodeIDVertexID(Node.ID, Vertex.VertexID), FCoreNodeAndFrontendVertex({CoreNode, &Vertex}));
 			}
 		};
-		
 
-		for (const FMetasoundFrontendEdge& Edge : InFrontendGraph.Edges)
+		for (const FMetasoundFrontendEdge& Edge : InGraphContext.GraphClass.Graph.Edges)
 		{
 			const FNodeIDVertexID DestinationKey(Edge.ToNodeID, Edge.ToVertexID);
 			const FCoreNodeAndFrontendVertex* DestinationNodeAndVertex = NodeDestinationsByID.Find(DestinationKey);
@@ -478,17 +509,106 @@ namespace Metasound
 			const INode* ToNode = DestinationNodeAndVertex->Node;
 			const FVertexKey ToVertexKey = DestinationNodeAndVertex->Vertex->Name;
 
-			bool bSuccess = OutGraph.AddDataEdge(*FromNode, FromVertexKey,  *ToNode, ToVertexKey);
+			bool bSuccess = InGraphContext.Graph->AddDataEdge(*FromNode, FromVertexKey,  *ToNode, ToVertexKey);
 
-			if (!bSuccess)
+			// If succeeded, remove input as viable vertex to construct default variable, as it has been superceded by a connection.
+			if (bSuccess)
+			{
+				InGraphContext.DefaultInputs.Remove(DestinationNodeAndVertex->Vertex->VertexID);
+			}
+			else
 			{
 				UE_LOG(LogMetaSound, Error, TEXT("Failed to connect edge from [NodeID:%s, VertexID:%s] to [NodeID:%s, VertexID:%s]"), *Edge.FromNodeID.ToString(), *Edge.FromVertexID.ToString(), *Edge.ToNodeID.ToString(), *Edge.ToVertexID.ToString());
 			}
 		}
 	}
 
+	void FFrontendGraphBuilder::AddDefaultInputVariables(FBuildGraphContext& InGraphContext)
+	{
+		for (const TPair<FGuid, FDefaultVariableData>& Pair : InGraphContext.DefaultInputs)
+		{
+			const FDefaultVariableData& VariableData = Pair.Value;
+			const FGuid VariableNodeID = FGuid::NewGuid();
+
+			// 1. Construct and add the default variable to the graph
+			{
+				FVariableNodeConstructorParams InitParams = VariableData.InitParams.Clone();
+				TUniquePtr<const INode> DefaultVariable = FMetasoundFrontendRegistryContainer::Get()->ConstructVariableNode(VariableData.TypeName, MoveTemp(InitParams));
+				InGraphContext.Graph->AddNode(VariableNodeID, TSharedPtr<const INode>(DefaultVariable.Release()));
+			}
+
+			// 2. Connect the default variable to the expected input
+			const INode* FromNode = InGraphContext.Graph->FindNode(VariableNodeID);
+			if (!ensure(FromNode))
+			{
+				continue;
+			}
+			const FVertexKey& FromVertexKey = VariableData.InitParams.VertexName;
+
+			const INode* ToNode = InGraphContext.Graph->FindNode(VariableData.DestinationNodeID);
+			if (!ensure(ToNode))
+			{
+				continue;
+			}
+			const FVertexKey& ToVertexKey = VariableData.DestinationVertexKey;
+
+			bool bSuccess = InGraphContext.Graph->AddDataEdge(*FromNode, FromVertexKey, *ToNode, ToVertexKey);
+			if (!bSuccess)
+			{
+				UE_LOG(LogMetaSound, Error, TEXT("Failed to connect default variable edge: from '%s' to '%s'"), *FromVertexKey, *ToVertexKey);
+			}
+		}
+	}
+
+	TArray<FFrontendGraphBuilder::FDefaultVariableData> FFrontendGraphBuilder::GetInputDefaultVariableData(const FMetasoundFrontendNode& InNode, const FNodeInitData& InInitData)
+	{
+		TArray<FDefaultVariableData> DefaultVariableData;
+
+		TArray<FMetasoundFrontendVertex> InputVertices = InNode.Interface.Inputs;
+		for (const FMetasoundFrontendVertexLiteral& Literal : InNode.InputLiterals)
+		{
+			FVariableNodeConstructorParams InitParams;
+			FName TypeName;
+
+			FGuid VertexID = Literal.VertexID;
+			FVertexKey DestinationVertexKey;
+			auto RemoveAndBuildParams = [&](const FMetasoundFrontendVertex& Vertex)
+			{
+				if (Vertex.VertexID == VertexID)
+				{
+					InitParams.InitParam = Literal.Value.ToLiteral(Vertex.TypeName);
+					InitParams.InstanceID = FGuid::NewGuid();
+					InitParams.NodeName = InInitData.InstanceName + Vertex.Name + InitParams.InstanceID.ToString();
+					InitParams.VertexName = InitParams.NodeName;
+					TypeName = Vertex.TypeName;
+
+					DestinationVertexKey = Vertex.Name;
+
+					return true;
+				}
+
+				return false;
+			};
+
+			const bool bRemoved = InputVertices.RemoveAllSwap(RemoveAndBuildParams, false /* bAllowShrinking */) > 0;
+			if (ensure(bRemoved))
+			{
+				DefaultVariableData.Emplace(FDefaultVariableData
+				{
+					InNode.ID,
+					VertexID,
+					DestinationVertexKey,
+					TypeName,
+					MoveTemp(InitParams)
+				});
+			}
+		}
+
+		return DefaultVariableData;
+	}
+
 	/** Check that all dependencies are C++ class dependencies. */
-	bool FFrontendGraphBuilder::IsFlat(const FMetasoundFrontendDocument& InDocument) const
+	bool FFrontendGraphBuilder::IsFlat(const FMetasoundFrontendDocument& InDocument)
 	{
 		if (InDocument.Subgraphs.Num() > 0)
 		{
@@ -498,7 +618,7 @@ namespace Metasound
 		return IsFlat(InDocument.RootGraph, InDocument.Dependencies);
 	}
 
-	bool FFrontendGraphBuilder::IsFlat(const FMetasoundFrontendGraphClass& InRoot, const TArray<FMetasoundFrontendClass>& InDependencies) const
+	bool FFrontendGraphBuilder::IsFlat(const FMetasoundFrontendGraphClass& InRoot, const TArray<FMetasoundFrontendClass>& InDependencies)
 	{
 		// All dependencies are external dependencies in a flat graph
 		auto IsClassExternal = [&](const FMetasoundFrontendClass& InDesc) 
@@ -529,7 +649,7 @@ namespace Metasound
 		return bIsEveryDependencyMet;
 	}
 
-	bool FFrontendGraphBuilder::SortSubgraphDependencies(TArray<const FMetasoundFrontendGraphClass*>& Subgraphs) const
+	bool FFrontendGraphBuilder::SortSubgraphDependencies(TArray<const FMetasoundFrontendGraphClass*>& Subgraphs)
 	{
 		// Helper for caching and querying subgraph dependencies
 		struct FSubgraphDependencyLookup
@@ -577,24 +697,28 @@ namespace Metasound
 		return bSuccess;
 	}
 
-	TUniquePtr<FFrontendGraph> FFrontendGraphBuilder::CreateGraph(FBuildContext& InContext, const FMetasoundFrontendGraphClass& InGraph) const
+	TUniquePtr<FFrontendGraph> FFrontendGraphBuilder::CreateGraph(FBuildContext& InContext, const FMetasoundFrontendGraphClass& InGraphClass)
 	{
+		const FString GraphName = InGraphClass.Metadata.ClassName.GetFullName().ToString();
 
-		TUniquePtr<FFrontendGraph> MetasoundGraph = MakeUnique<FFrontendGraph>(InGraph.Metadata.ClassName.GetFullName().ToString(), FGuid::NewGuid());
+		FBuildGraphContext BuildGraphContext
+		{
+			MakeUnique<FFrontendGraph>(GraphName, FGuid::NewGuid()),
+			InGraphClass,
+			InContext
+		};
 
 		// TODO: will likely want to bubble up errors here for case where
-		// a datatype or node is not registered. 
-		AddNodesToGraph(InGraph, InContext.FrontendClasses, InContext.Graphs, *MetasoundGraph);
-
-		AddEdgesToGraph(InGraph.Graph, *MetasoundGraph);
+		// a datatype or node is not registered.
+		AddNodesToGraph(BuildGraphContext);
+		AddEdgesToGraph(BuildGraphContext);
+		AddDefaultInputVariables(BuildGraphContext);
 
-		check(MetasoundGraph->OwnsAllReferencedNodes());
-
-		return MoveTemp(MetasoundGraph);
+		check(BuildGraphContext.Graph->OwnsAllReferencedNodes());
+		return MoveTemp(BuildGraphContext.Graph);
 	}
 
-
-	TUniquePtr<FFrontendGraph> FFrontendGraphBuilder::CreateGraph(const FMetasoundFrontendGraphClass& InGraph, const TArray<FMetasoundFrontendGraphClass>& InSubgraphs, const TArray<FMetasoundFrontendClass>& InDependencies) const
+	TUniquePtr<FFrontendGraph> FFrontendGraphBuilder::CreateGraph(const FMetasoundFrontendGraphClass& InGraph, const TArray<FMetasoundFrontendGraphClass>& InSubgraphs, const TArray<FMetasoundFrontendClass>& InDependencies)
 	{
 		FBuildContext Context;
 
@@ -635,11 +759,11 @@ namespace Metasound
 		}
 
 		// Create parent graph.
-		return CreateGraph(Context, InGraph); 
+		return CreateGraph(Context, InGraph);
 	}
 	
 	/* Metasound document should be inflated by now. */
-	TUniquePtr<FFrontendGraph> FFrontendGraphBuilder::CreateGraph(const FMetasoundFrontendDocument& InDocument) const
+	TUniquePtr<FFrontendGraph> FFrontendGraphBuilder::CreateGraph(const FMetasoundFrontendDocument& InDocument)
 	{
 		return CreateGraph(InDocument.RootGraph, InDocument.Subgraphs, InDocument.Dependencies);
 	}
diff --git a/Engine/Plugins/Runtime/Metasound/Source/MetasoundFrontend/Private/MetasoundFrontendGraph.h b/Engine/Plugins/Runtime/Metasound/Source/MetasoundFrontend/Private/MetasoundFrontendGraph.h
--- a/Engine/Plugins/Runtime/Metasound/Source/MetasoundFrontend/Private/MetasoundFrontendGraph.h
+++ b/Engine/Plugins/Runtime/Metasound/Source/MetasoundFrontend/Private/MetasoundFrontendGraph.h
@@ -104,46 +104,66 @@ namespace Metasound
 		 *
 		 * @return True if all dependencies are C++ classes. False otherwise.
 		 */
-		bool IsFlat(const FMetasoundFrontendDocument& InDocument) const;
+		static bool IsFlat(const FMetasoundFrontendDocument& InDocument);
 
-		bool IsFlat(const FMetasoundFrontendGraphClass& InRoot, const TArray<FMetasoundFrontendClass>& InDependencies) const;
+		static bool IsFlat(const FMetasoundFrontendGraphClass& InRoot, const TArray<FMetasoundFrontendClass>& InDependencies);
 
 		/* Metasound document should be in order to create this graph. */
-		TUniquePtr<FFrontendGraph> CreateGraph(const FMetasoundFrontendDocument& InDocument) const;
+		static TUniquePtr<FFrontendGraph> CreateGraph(const FMetasoundFrontendDocument& InDocument);
 
 		/* Metasound document should be in order to create this graph. */
-		TUniquePtr<FFrontendGraph> CreateGraph(const FMetasoundFrontendGraphClass& InGraph, const TArray<FMetasoundFrontendGraphClass>& InSubgraphs, const TArray<FMetasoundFrontendClass>& InDependencies) const;
-
+		static TUniquePtr<FFrontendGraph> CreateGraph(const FMetasoundFrontendGraphClass& InGraphClass, const TArray<FMetasoundFrontendGraphClass>& InSubgraphs, const TArray<FMetasoundFrontendClass>& InDependencies);
 
 	private:
+		struct FDefaultVariableData
+		{
+			FGuid DestinationNodeID;
+			FGuid DestinationVertexID;
+			FVertexKey DestinationVertexKey;
+			FName TypeName;
+			FVariableNodeConstructorParams InitParams;
+		};
+
+		// Map of Input VertexID to variable data required to construct and connect default variable
 		using FDependencyByIDMap = TMap<FGuid, const FMetasoundFrontendClass*>;
 		using FSharedNodeByIDMap = TMap<FGuid, TSharedPtr<const INode>>;
+		using FDefaultInputByIDMap = TMap<FGuid, FDefaultVariableData>;
 
+		// Context used throughout entire graph build process
+		// (for both a root and nested subgraphs)
 		struct FBuildContext
 		{
 			FDependencyByIDMap FrontendClasses;
 			FSharedNodeByIDMap Graphs;
 		};
 
-		bool SortSubgraphDependencies(TArray<const FMetasoundFrontendGraphClass*>& Subgraphs) const;
-
-		TUniquePtr<FFrontendGraph> CreateGraph(FBuildContext& InContext, const FMetasoundFrontendGraphClass& InSubgraph) const;
+		// Transient context used for building a specific graph
+		struct FBuildGraphContext
+		{
+			TUniquePtr<FFrontendGraph> Graph;
+			const FMetasoundFrontendGraphClass& GraphClass;
+			FBuildContext& BuildContext;
 
-		const FMetasoundFrontendClassInput* FindClassInputForInputNode(const FMetasoundFrontendGraphClass& InOwningGraph, const FMetasoundFrontendNode& InInputNode, int32& OutClassInputIndex) const;
+			FDefaultInputByIDMap DefaultInputs;
+		};
 
-		const FMetasoundFrontendClassOutput* FindClassOutputForOutputNode(const FMetasoundFrontendGraphClass& InOwningGraph, const FMetasoundFrontendNode& InOutputNode, int32& OutClassOutputIndex) const;
+		static TArray<FDefaultVariableData> GetInputDefaultVariableData(const FMetasoundFrontendNode& InNode, const FNodeInitData& InInitData);
 
-		const FMetasoundFrontendLiteral* FindInputLiteralForInputNode(const FMetasoundFrontendNode& InInputNode, const FMetasoundFrontendClass& InInputNodeClass, const FMetasoundFrontendClassInput& InOwningGraphClassInput) const;
+		static bool SortSubgraphDependencies(TArray<const FMetasoundFrontendGraphClass*>& Subgraphs);
 
-		TUniquePtr<INode> CreateInputNode(const FMetasoundFrontendNode& InNode, const FMetasoundFrontendClass& InClass, const FMetasoundFrontendClassInput& InOwningGraphClassInput) const;
+		static TUniquePtr<FFrontendGraph> CreateGraph(FBuildContext& InContext, const FMetasoundFrontendGraphClass& InSubgraph);
 
-		TUniquePtr<INode> CreateOutputNode(const FMetasoundFrontendNode& InNode, const FMetasoundFrontendClass& InExternal) const;
+		static const FMetasoundFrontendClassInput* FindClassInputForInputNode(const FMetasoundFrontendGraphClass& InOwningGraph, const FMetasoundFrontendNode& InInputNode, int32& OutClassInputIndex);
+		static const FMetasoundFrontendClassOutput* FindClassOutputForOutputNode(const FMetasoundFrontendGraphClass& InOwningGraph, const FMetasoundFrontendNode& InOutputNode, int32& OutClassOutputIndex);
+		static const FMetasoundFrontendLiteral* FindInputLiteralForInputNode(const FMetasoundFrontendNode& InInputNode, const FMetasoundFrontendClass& InInputNodeClass, const FMetasoundFrontendClassInput& InOwningGraphClassInput);
 
-		TUniquePtr<INode> CreateExternalNode(const FMetasoundFrontendNode& InNode, const FMetasoundFrontendClass& InClass) const;
+		static TUniquePtr<INode> CreateInputNode(const FMetasoundFrontendNode& InNode, const FMetasoundFrontendClass& InClass, const FMetasoundFrontendClassInput& InOwningGraphClassInput);
+		static TUniquePtr<INode> CreateOutputNode(const FMetasoundFrontendNode& InNode, const FMetasoundFrontendClass& InClass, FBuildGraphContext& InGraphContext);
+		static TUniquePtr<INode> CreateExternalNode(const FMetasoundFrontendNode& InNode, const FMetasoundFrontendClass& InClass, FBuildGraphContext& InGraphContext);
 
 		// TODO: add errors here. Most will be a "PromptIfMissing"...
-		void AddNodesToGraph(const FMetasoundFrontendGraphClass& InGraphClass, const FDependencyByIDMap& InClasses, const FSharedNodeByIDMap& InSubgraphs, FFrontendGraph& OutGraph) const;
-
-		void AddEdgesToGraph(const FMetasoundFrontendGraph& InGraphDescription, FFrontendGraph& OutGraph) const;
+		static void AddNodesToGraph(FBuildGraphContext& InGraphContext);
+		static void AddEdgesToGraph(FBuildGraphContext& InGraphContext);
+		static void AddDefaultInputVariables(FBuildGraphContext& InGraphContext);
 	};
 }
diff --git a/Engine/Plugins/Runtime/Metasound/Source/MetasoundFrontend/Private/MetasoundFrontendInvalidController.h b/Engine/Plugins/Runtime/Metasound/Source/MetasoundFrontend/Private/MetasoundFrontendInvalidController.h
--- a/Engine/Plugins/Runtime/Metasound/Source/MetasoundFrontend/Private/MetasoundFrontendInvalidController.h
+++ b/Engine/Plugins/Runtime/Metasound/Source/MetasoundFrontend/Private/MetasoundFrontendInvalidController.h
@@ -94,7 +94,9 @@ namespace Metasound
 			const FText& GetDisplayName() const override { return MetasoundFrontendInvalidControllerPrivate::GetInvalid<FText>(); }
 			const FText& GetTooltip() const override { return MetasoundFrontendInvalidControllerPrivate::GetInvalid<FText>(); }
 			const FMetasoundFrontendVertexMetadata& GetMetadata() const override { return MetasoundFrontendInvalidControllerPrivate::GetInvalid<FMetasoundFrontendVertexMetadata>(); }
-			const FMetasoundFrontendLiteral* GetDefaultLiteral() const override { return nullptr; }
+			const FMetasoundFrontendLiteral* GetLiteral() const override { return nullptr; }
+			void SetLiteral(const FMetasoundFrontendLiteral& InLiteral) { };
+			const FMetasoundFrontendLiteral* GetClassDefaultLiteral() const override { return nullptr; }
 			FGuid GetOwningNodeID() const override { return Metasound::FrontendInvalidID; }
 			TSharedRef<INodeController> GetOwningNode() override;
 			TSharedRef<const INodeController> GetOwningNode() const override;
@@ -156,6 +158,10 @@ namespace Metasound
 			FInputHandle AddOutput(const FString& InVertexName, const FMetasoundFrontendLiteral* InDefault) override { return FInvalidInputController::GetInvalid(); }
 			bool RemoveOutput(FGuid InVertexID) override { return false; }
 
+			bool ClearInputLiteral(FGuid InVertexID) override { return false; };
+			const FMetasoundFrontendLiteral* GetInputLiteral(const FGuid& InVertexID) const { return nullptr; }
+			void SetInputLiteral(const FMetasoundFrontendVertexLiteral& InVertexLiteral) override { }
+
 			const FMetasoundFrontendClassInterface& GetClassInterface() const override { return MetasoundFrontendInvalidControllerPrivate::GetInvalid<FMetasoundFrontendClassInterface>(); }
 			const FMetasoundFrontendClassMetadata& GetClassMetadata() const override { return MetasoundFrontendInvalidControllerPrivate::GetInvalid<FMetasoundFrontendClassMetadata>(); }
 			const FMetasoundFrontendInterfaceStyle& GetInputStyle() const override { return MetasoundFrontendInvalidControllerPrivate::GetInvalid<FMetasoundFrontendInterfaceStyle>(); }
@@ -339,6 +345,9 @@ namespace Metasound
 				FConstClassAccessPtr FindOrAddClass(const FMetasoundFrontendClassMetadata& InMetadata) override{ return FConstClassAccessPtr(); }
 				FGraphHandle AddDuplicateSubgraph(const IGraphController& InGraph) override { return FInvalidGraphController::GetInvalid(); }
 
+				void SetMetadata(const FMetasoundFrontendDocumentMetadata& InMetadata) override { }
+				const FMetasoundFrontendDocumentMetadata& GetMetadata() const override { return MetasoundFrontendInvalidControllerPrivate::GetInvalid<FMetasoundFrontendDocumentMetadata>(); }
+
 				void SynchronizeDependencies() override { }
 
 				TArray<FGraphHandle> GetSubgraphHandles() override { return TArray<FGraphHandle>(); }
diff --git a/Engine/Plugins/Runtime/Metasound/Source/MetasoundFrontend/Private/MetasoundFrontendLiteral.cpp b/Engine/Plugins/Runtime/Metasound/Source/MetasoundFrontend/Private/MetasoundFrontendLiteral.cpp
--- a/Engine/Plugins/Runtime/Metasound/Source/MetasoundFrontend/Private/MetasoundFrontendLiteral.cpp
+++ b/Engine/Plugins/Runtime/Metasound/Source/MetasoundFrontend/Private/MetasoundFrontendLiteral.cpp
@@ -106,6 +106,45 @@ bool FMetasoundFrontendLiteral::IsArray() const
 	}
 }
 
+bool FMetasoundFrontendLiteral::IsEquivalent(const FMetasoundFrontendLiteral& InOther) const
+{
+	if (InOther.GetType() != GetType())
+	{
+		return false;
+	}
+
+	switch (Type)
+	{
+		case EMetasoundFrontendLiteralType::Boolean:
+		case EMetasoundFrontendLiteralType::BooleanArray:
+			return AsBoolean == InOther.AsBoolean;
+
+		case EMetasoundFrontendLiteralType::Float:
+		case EMetasoundFrontendLiteralType::FloatArray:
+			return AsFloat == InOther.AsFloat;
+
+		case EMetasoundFrontendLiteralType::Integer:
+		case EMetasoundFrontendLiteralType::IntegerArray:
+			return AsInteger == InOther.AsInteger;
+
+		case EMetasoundFrontendLiteralType::None:
+		case EMetasoundFrontendLiteralType::NoneArray:
+			return AsNumDefault == InOther.AsNumDefault;
+
+		case EMetasoundFrontendLiteralType::String:
+		case EMetasoundFrontendLiteralType::StringArray:
+			return AsString == InOther.AsString;
+
+		case EMetasoundFrontendLiteralType::UObject:
+		case EMetasoundFrontendLiteralType::UObjectArray:
+			return AsUObject == InOther.AsUObject;
+
+		case EMetasoundFrontendLiteralType::Invalid:
+		default:
+			return true;
+	}
+}
+
 bool FMetasoundFrontendLiteral::IsValid() const
 {
 	return (Type != EMetasoundFrontendLiteralType::Invalid);
diff --git a/Engine/Plugins/Runtime/Metasound/Source/MetasoundFrontend/Private/MetasoundFrontendRegistries.cpp b/Engine/Plugins/Runtime/Metasound/Source/MetasoundFrontend/Private/MetasoundFrontendRegistries.cpp
--- a/Engine/Plugins/Runtime/Metasound/Source/MetasoundFrontend/Private/MetasoundFrontendRegistries.cpp
+++ b/Engine/Plugins/Runtime/Metasound/Source/MetasoundFrontend/Private/MetasoundFrontendRegistries.cpp
@@ -186,11 +186,13 @@ namespace Metasound
 				bool FindFrontendClassFromRegistered(const FNodeClassInfo& InClassInfo, FMetasoundFrontendClass& OutClass) override;
 				bool FindFrontendClassFromRegistered(const FMetasoundFrontendClassMetadata& InMetadata, FMetasoundFrontendClass& OutClass) override;
 				bool FindInputNodeClassMetadataForDataType(const FName& InDataTypeName, FMetasoundFrontendClassMetadata& OutMetadata) override;
+				bool FindVariableNodeClassMetadataForDataType(const FName& InDataTypeName, FMetasoundFrontendClassMetadata& OutMetadata) override;
 				bool FindOutputNodeClassMetadataForDataType(const FName& InDataTypeName, FMetasoundFrontendClassMetadata& OutMetadata) override;
 
 				// Create a new instance of a C++ implemented node from the registry.
 				TUniquePtr<Metasound::INode> ConstructInputNode(const FName& InInputType, Metasound::FInputNodeConstructorParams&& InParams) override;
-				TUniquePtr<Metasound::INode> ConstructOutputNode(const FName& InOutputType, const Metasound::FOutputNodeConstructorParams& InParams) override;
+				TUniquePtr<Metasound::INode> ConstructVariableNode(const FName& InVariableType, FVariableNodeConstructorParams&& InParams) override;
+				TUniquePtr<Metasound::INode> ConstructOutputNode(const FName& InOutputType, Metasound::FOutputNodeConstructorParams&& InParams) override;
 				TUniquePtr<Metasound::INode> ConstructExternalNode(const FNodeRegistryKey& InKey, const Metasound::FNodeInitData& InInitData) override;
 
 				// Returns a list of possible nodes to use to convert from FromDataType to ToDataType.
@@ -316,7 +318,6 @@ namespace Metasound
 				return OutClasses;
 			}
 
-
 			TUniquePtr<Metasound::INode> FRegistryContainerImpl::ConstructInputNode(const FName& InInputType, Metasound::FInputNodeConstructorParams&& InParams)
 			{
 				if (ensureAlwaysMsgf(DataTypeRegistry.Contains(InInputType), TEXT("Couldn't find data type %s!"), *InInputType.ToString()))
@@ -329,11 +330,23 @@ namespace Metasound
 				}
 			}
 
-			TUniquePtr<Metasound::INode> FRegistryContainerImpl::ConstructOutputNode(const FName& InOutputType, const Metasound::FOutputNodeConstructorParams& InParams)
+			TUniquePtr<Metasound::INode> FRegistryContainerImpl::ConstructVariableNode(const FName& InVariableType, Metasound::FVariableNodeConstructorParams&& InParams)
+			{
+				if (ensureAlwaysMsgf(DataTypeRegistry.Contains(InVariableType), TEXT("Couldn't find data type %s!"), *InVariableType.ToString()))
+				{
+					return DataTypeRegistry[InVariableType].Callbacks.CreateVariableNode(MoveTemp(InParams));
+				}
+				else
+				{
+					return nullptr;
+				}
+			}
+
+			TUniquePtr<Metasound::INode> FRegistryContainerImpl::ConstructOutputNode(const FName& InOutputType, Metasound::FOutputNodeConstructorParams&& InParams)
 			{
 				if (ensureAlwaysMsgf(DataTypeRegistry.Contains(InOutputType), TEXT("Couldn't find data type %s!"), *InOutputType.ToString()))
 				{
-					return DataTypeRegistry[InOutputType].Callbacks.CreateOutputNode(InParams);
+					return DataTypeRegistry[InOutputType].Callbacks.CreateOutputNode(MoveTemp(InParams));
 				}
 				else
 				{
@@ -537,6 +550,9 @@ namespace Metasound
 					Metasound::Frontend::FNodeRegistryKey InputNodeRegistryKey = GetRegistryKey(InCallbacks.CreateFrontendInputClass().Metadata);
 					DataTypeNodeRegistry.Add(InputNodeRegistryKey, InElement);
 
+					Metasound::Frontend::FNodeRegistryKey VariableNodeRegistryKey = GetRegistryKey(InCallbacks.CreateFrontendVariableClass().Metadata);
+					DataTypeNodeRegistry.Add(VariableNodeRegistryKey, InElement);
+
 					Metasound::Frontend::FNodeRegistryKey OutputNodeRegistryKey = GetRegistryKey(InCallbacks.CreateFrontendOutputClass().Metadata);
 					DataTypeNodeRegistry.Add(OutputNodeRegistryKey, InElement);
 
@@ -701,9 +717,19 @@ namespace Metasound
 
 			bool FRegistryContainerImpl::FindInputNodeClassMetadataForDataType(const FName& InDataTypeName, FMetasoundFrontendClassMetadata& OutMetadata)
 			{
-				if (DataTypeRegistry.Contains(InDataTypeName))
+				if (const FDataTypeRegistryElement* Element = DataTypeRegistry.Find(InDataTypeName))
+				{
+					OutMetadata = Element->Callbacks.CreateFrontendInputClass().Metadata;
+					return true;
+				}
+				return false;
+			}
+
+			bool FRegistryContainerImpl::FindVariableNodeClassMetadataForDataType(const FName& InDataTypeName, FMetasoundFrontendClassMetadata& OutMetadata)
+			{
+				if (const FDataTypeRegistryElement* Element = DataTypeRegistry.Find(InDataTypeName))
 				{
-					OutMetadata = DataTypeRegistry[InDataTypeName].Callbacks.CreateFrontendInputClass().Metadata;
+					OutMetadata = Element->Callbacks.CreateFrontendVariableClass().Metadata;
 					return true;
 				}
 				return false;
@@ -711,9 +737,9 @@ namespace Metasound
 
 			bool FRegistryContainerImpl::FindOutputNodeClassMetadataForDataType(const FName& InDataTypeName, FMetasoundFrontendClassMetadata& OutMetadata)
 			{
-				if (DataTypeRegistry.Contains(InDataTypeName))
+				if (const FDataTypeRegistryElement* Element = DataTypeRegistry.Find(InDataTypeName))
 				{
-					OutMetadata = DataTypeRegistry[InDataTypeName].Callbacks.CreateFrontendOutputClass().Metadata;
+					OutMetadata = Element->Callbacks.CreateFrontendOutputClass().Metadata;
 					return true;
 				}
 				return false;
@@ -745,6 +771,14 @@ namespace Metasound
 					}
 				};
 
+				auto IterateVariableClasses = [=]
+				{
+					for (const TPair<FNodeRegistryKey, FDataTypeRegistryElement>& Pair : DataTypeNodeRegistry)
+					{
+						InIterFunc(Pair.Value.Callbacks.CreateFrontendVariableClass());
+					}
+				};
+
 				switch (InClassType)
 				{
 					case EMetasoundFrontendClassType::External:
@@ -772,16 +806,23 @@ namespace Metasound
 					}
 					break;
 
+					case EMetasoundFrontendClassType::Variable:
+					{
+						IterateVariableClasses();
+					}
+					break;
+
 					default:
 					case EMetasoundFrontendClassType::Invalid:
 					{
 						IterateExternalClasses();
+						IterateVariableClasses();
 						IterateInputClasses();
 						IterateOutputClasses();
 					}
 					break;
 				}
-				static_assert(static_cast<uint32>(EMetasoundFrontendClassType::Invalid) == 4, "Possible missing switch case coverage for EMetasoundFrontendClassType");
+				static_assert(static_cast<uint32>(EMetasoundFrontendClassType::Invalid) == 5, "Possible missing switch case coverage for EMetasoundFrontendClassType");
 			}
 		}
 	}
@@ -853,6 +894,15 @@ bool FMetasoundFrontendRegistryContainer::GetInputNodeClassMetadataForDataType(c
 	return false;
 }
 
+bool FMetasoundFrontendRegistryContainer::GetVariableNodeClassMetadataForDataType(const FName& InDataTypeName, FMetasoundFrontendClassMetadata& OutMetadata)
+{
+	if (FMetasoundFrontendRegistryContainer* Registry = FMetasoundFrontendRegistryContainer::Get())
+	{
+		return Registry->FindVariableNodeClassMetadataForDataType(InDataTypeName, OutMetadata);
+	}
+	return false;
+}
+
 bool FMetasoundFrontendRegistryContainer::GetOutputNodeClassMetadataForDataType(const FName& InDataTypeName, FMetasoundFrontendClassMetadata& OutMetadata)
 {
 	if (FMetasoundFrontendRegistryContainer* Registry = FMetasoundFrontendRegistryContainer::Get())
diff --git a/Engine/Plugins/Runtime/Metasound/Source/MetasoundFrontend/Private/MetasoundFrontendStandardController.cpp b/Engine/Plugins/Runtime/Metasound/Source/MetasoundFrontend/Private/MetasoundFrontendStandardController.cpp
--- a/Engine/Plugins/Runtime/Metasound/Source/MetasoundFrontend/Private/MetasoundFrontendStandardController.cpp
+++ b/Engine/Plugins/Runtime/Metasound/Source/MetasoundFrontend/Private/MetasoundFrontendStandardController.cpp
@@ -462,7 +462,37 @@ namespace Metasound
 			return FrontendControllerIntrinsics::GetInvalidValueConstRef<FText>();
 		}
 
-		const FMetasoundFrontendLiteral* FBaseInputController::GetDefaultLiteral() const 
+		const FMetasoundFrontendLiteral* FBaseInputController::GetLiteral() const
+		{
+			if (NodeVertexPtr.IsValid() && OwningNode->IsValid())
+			{
+				const FGuid& VertexID = NodeVertexPtr->VertexID;
+				return OwningNode->GetInputLiteral(VertexID);
+			}
+
+			return nullptr;
+		}
+
+		void FBaseInputController::SetLiteral(const FMetasoundFrontendLiteral& InLiteral)
+		{
+			if (NodeVertexPtr.IsValid() && OwningNode->IsValid())
+			{
+				const FGuid& VertexID = NodeVertexPtr->VertexID;
+				if (const FMetasoundFrontendLiteral* ClassLiteral = GetClassDefaultLiteral())
+				{
+					// Clear if equivalent to class default as fallback is the class default
+					if (ClassLiteral->IsEquivalent(InLiteral))
+					{
+						OwningNode->ClearInputLiteral(VertexID);
+						return;
+					}
+				}
+
+				OwningNode->SetInputLiteral(FMetasoundFrontendVertexLiteral{ VertexID, InLiteral });
+			}
+		}
+
+		const FMetasoundFrontendLiteral* FBaseInputController::GetClassDefaultLiteral() const
 		{
 			if (const FMetasoundFrontendClassInput* ClassInput = ClassInputPtr.Get())
 			{
@@ -671,7 +701,6 @@ namespace Metasound
 			return false;
 		}
 
-
 		const FMetasoundFrontendEdge* FBaseInputController::FindEdge() const
 		{
 			if (GraphPtr.IsValid())
@@ -786,11 +815,6 @@ namespace Metasound
 			return FrontendControllerIntrinsics::GetInvalidValueConstRef<FMetasoundFrontendVertexMetadata>();
 		}
 
-		const FMetasoundFrontendLiteral* FOutputNodeInputController::GetDefaultLiteral() const
-		{
-			return nullptr;
-		}
-
 		FDocumentAccess FOutputNodeInputController::ShareAccess() 
 		{
 			FDocumentAccess Access = FBaseInputController::ShareAccess();
@@ -927,6 +951,68 @@ namespace Metasound
 			return Metasound::FrontendInvalidID;
 		}
 
+		const FMetasoundFrontendLiteral* FBaseNodeController::GetInputLiteral(const FGuid& InVertexID) const
+		{
+			if (NodePtr.IsValid())
+			{
+				for (const FMetasoundFrontendVertexLiteral& VertexLiteral : NodePtr->InputLiterals)
+				{
+					if (VertexLiteral.VertexID == InVertexID)
+					{
+						return &VertexLiteral.Value;
+					}
+				}
+			}
+
+			return nullptr;
+		}
+
+		void FBaseNodeController::SetInputLiteral(const FMetasoundFrontendVertexLiteral& InVertexLiteral)
+		{
+			if (NodePtr.IsValid())
+			{
+				auto IsInputVertex = [InVertexLiteral] (const FMetasoundFrontendVertex& Vertex)
+				{
+					return InVertexLiteral.VertexID == Vertex.VertexID;
+				};
+
+				FMetasoundFrontendNodeInterface& NodeInterface = NodePtr->Interface;
+				if (!ensure(NodeInterface.Inputs.ContainsByPredicate(IsInputVertex)))
+				{
+					return;
+				}
+
+				for (FMetasoundFrontendVertexLiteral& VertexLiteral : NodePtr->InputLiterals)
+				{
+					if (VertexLiteral.VertexID == InVertexLiteral.VertexID)
+					{
+						if (ensure(VertexLiteral.Value.GetType() == InVertexLiteral.Value.GetType()))
+						{
+							VertexLiteral = InVertexLiteral;
+						}
+						return;
+					}
+				}
+
+				NodePtr->InputLiterals.Add(InVertexLiteral);
+			}
+		}
+
+		bool FBaseNodeController::ClearInputLiteral(FGuid InVertexID)
+		{
+			if (NodePtr.IsValid())
+			{
+				auto IsInputVertex = [InVertexID](const FMetasoundFrontendVertexLiteral& VertexLiteral)
+				{
+					return InVertexID == VertexLiteral.VertexID;
+				};
+
+				return NodePtr->InputLiterals.RemoveAllSwap(IsInputVertex, false) > 0;
+			}
+
+			return false;
+		}
+
 		const FMetasoundFrontendClassInterface& FBaseNodeController::GetClassInterface() const
 		{
 			if (ClassPtr.IsValid())
@@ -1901,7 +1987,6 @@ namespace Metasound
 			return Access;
 		}
 
-
 		// 
 		// FGraphController
 		//
@@ -2315,7 +2400,7 @@ namespace Metasound
 							NewOutput.NodeID = Node.ID;
 
 							FNodeAccessPtr NodePtr = GraphClassPtr.GetNodeWithNodeID(Node.ID);
-							return GetNodeHandle(FGraphController::FNodeAndClass{NodePtr, OutputClass});
+							return GetNodeHandle(FGraphController::FNodeAndClass { NodePtr, OutputClass });
 						}
 					}
 					else 
@@ -2561,6 +2646,7 @@ namespace Metasound
 						return RemoveOutput(*Desc);
 					}
 
+					case EMetasoundFrontendClassType::Variable:
 					case EMetasoundFrontendClassType::External:
 					case EMetasoundFrontendClassType::Graph:
 					{
@@ -2570,7 +2656,7 @@ namespace Metasound
 					default:
 					case EMetasoundFrontendClassType::Invalid:
 					{
-						static_assert(static_cast<int32>(EMetasoundFrontendClassType::Invalid) == 4, "Possible missing switch case coverage for EMetasoundFrontendClassType.");
+						static_assert(static_cast<int32>(EMetasoundFrontendClassType::Invalid) == 5, "Possible missing switch case coverage for EMetasoundFrontendClassType.");
 						checkNoEntry();
 					}
 				}
@@ -2626,15 +2712,13 @@ namespace Metasound
 			TArray<FMetasoundFrontendGraphClass> Subgraphs = OwningDocument->GetSubgraphs();
 			TArray<FMetasoundFrontendClass> Dependencies = OwningDocument->GetDependencies();
 
-			FFrontendGraphBuilder GraphBuilder;
-			TUniquePtr<FFrontendGraph> Graph = GraphBuilder.CreateGraph(*GraphClassPtr, Subgraphs, Dependencies);
+			TUniquePtr<FFrontendGraph> Graph = FFrontendGraphBuilder::CreateGraph(*GraphClassPtr, Subgraphs, Dependencies);
 
 			if (!Graph.IsValid())
 			{
 				return TUniquePtr<IOperator>(nullptr);
 			}
 
-			// Step 5: Invoke Operator Builder
 			FOperatorBuilder OperatorBuilder(FOperatorBuilderSettings::GetDefaultSettings());
 			FBuildGraphParams BuildParams{*Graph, InSettings, FDataReferenceCollection{}, InEnvironment};
 			return OperatorBuilder.BuildGraphOperator(BuildParams, OutBuildErrors);
@@ -2661,7 +2745,6 @@ namespace Metasound
 					Node.ID = FGuid::NewGuid();
 
 					FNodeAccessPtr NodePtr = GraphClassPtr.GetNodeWithNodeID(Node.ID);
-
 					return GetNodeHandle(FGraphController::FNodeAndClass{NodePtr, InExistingDependency});
 				}
 			}
@@ -3382,6 +3465,24 @@ namespace Metasound
 			return MetasoundClass;
 		}
 
+		void FDocumentController::SetMetadata(const FMetasoundFrontendDocumentMetadata& InMetadata)
+		{
+			if (DocumentPtr.IsValid())
+			{
+				DocumentPtr->Metadata = InMetadata;
+			}
+		}
+
+		const FMetasoundFrontendDocumentMetadata& FDocumentController::GetMetadata() const
+		{
+			if (DocumentPtr.IsValid())
+			{
+				return DocumentPtr->Metadata;
+			}
+
+			return FrontendControllerIntrinsics::GetInvalidValueConstRef<FMetasoundFrontendDocumentMetadata>();
+		}
+
 		FConstClassAccessPtr FDocumentController::FindClass(const FNodeClassInfo& InNodeClass) const
 		{
 			return DocumentPtr.GetClassWithInfo(InNodeClass);
diff --git a/Engine/Plugins/Runtime/Metasound/Source/MetasoundFrontend/Private/MetasoundFrontendStandardController.h b/Engine/Plugins/Runtime/Metasound/Source/MetasoundFrontend/Private/MetasoundFrontendStandardController.h
--- a/Engine/Plugins/Runtime/Metasound/Source/MetasoundFrontend/Private/MetasoundFrontendStandardController.h
+++ b/Engine/Plugins/Runtime/Metasound/Source/MetasoundFrontend/Private/MetasoundFrontendStandardController.h
@@ -198,7 +198,10 @@ namespace Metasound
 			const FName& GetDataType() const override;
 			const FString& GetName() const override;
 
-			const FMetasoundFrontendLiteral* GetDefaultLiteral() const override;
+			const FMetasoundFrontendLiteral* GetLiteral() const override;
+			void SetLiteral(const FMetasoundFrontendLiteral& InLiteral) override;
+
+			const FMetasoundFrontendLiteral* GetClassDefaultLiteral() const override;
 
 			// Input metadata
 			const FText& GetDisplayName() const override;
@@ -270,8 +273,6 @@ namespace Metasound
 			const FText& GetTooltip() const override;
 			const FMetasoundFrontendVertexMetadata& GetMetadata() const override;
 
-			const FMetasoundFrontendLiteral* GetDefaultLiteral() const override;
-
 		protected:
 
 			FDocumentAccess ShareAccess() override;
@@ -353,6 +354,10 @@ namespace Metasound
 			FGuid GetID() const override;
 			FGuid GetClassID() const override;
 
+			bool ClearInputLiteral(FGuid InVertexID) override;
+			const FMetasoundFrontendLiteral* GetInputLiteral(const FGuid& InVertexID) const override;
+			void SetInputLiteral(const FMetasoundFrontendVertexLiteral& InVertexLiteral) override;
+
 			const FMetasoundFrontendClassInterface& GetClassInterface() const override;
 			const FMetasoundFrontendClassMetadata& GetClassMetadata() const override;
 			const FMetasoundFrontendInterfaceStyle& GetInputStyle() const override;
@@ -637,13 +642,17 @@ namespace Metasound
 			void SetDescription(const FText& InDescription) override;
 			void SetDisplayName(const FText& InText) override;
 
+			// No-ops as inputs do not handle literals the same way as other nodes
+			bool ClearInputLiteral(FGuid InVertexID) override { return false; }
+			const FMetasoundFrontendLiteral* GetInputLiteral(const FGuid& InVertexID) const override { return nullptr; }
+			void SetInputLiteral(const FMetasoundFrontendVertexLiteral& InVertexLiteral) override { }
+
 		protected:
 
 			FDocumentAccess ShareAccess() override;
 			FConstDocumentAccess ShareAccess() const override;
 
 
-
 		private:
 			FInputHandle CreateInputController(FGuid InVertexID, FConstVertexAccessPtr InNodeVertexPtr, FConstClassInputAccessPtr InClassInputPtr, FNodeHandle InOwningNode) const override;
 
@@ -793,11 +802,8 @@ namespace Metasound
 
 		private:
 
-			// Add/remove nodes
 			FNodeHandle AddNode(FConstClassAccessPtr InExistingDependency);
 			bool RemoveNode(const FMetasoundFrontendNode& InDesc);
-
-			// Remove inputs
 			bool RemoveInput(const FMetasoundFrontendNode& InNode);
 			bool RemoveOutput(const FMetasoundFrontendNode& InNode);
 
@@ -909,6 +915,9 @@ namespace Metasound
 
 			virtual FGraphHandle AddDuplicateSubgraph(const IGraphController& InGraph) override;
 
+			void SetMetadata(const FMetasoundFrontendDocumentMetadata& InMetadata) override;
+			const FMetasoundFrontendDocumentMetadata& GetMetadata() const override;
+
 			void SynchronizeDependencies() override;
 
 			FGraphHandle GetRootGraph() override;
diff --git a/Engine/Plugins/Runtime/Metasound/Source/MetasoundFrontend/Private/MetasoundFrontendTransform.cpp b/Engine/Plugins/Runtime/Metasound/Source/MetasoundFrontend/Private/MetasoundFrontendTransform.cpp
--- a/Engine/Plugins/Runtime/Metasound/Source/MetasoundFrontend/Private/MetasoundFrontendTransform.cpp
+++ b/Engine/Plugins/Runtime/Metasound/Source/MetasoundFrontend/Private/MetasoundFrontendTransform.cpp
@@ -79,5 +79,97 @@ namespace Metasound
 
 			return bDidEdit;
 		}
+
+		// Versioning Transforms
+		class FVersionDocumentTransform : public IDocumentTransform
+		{
+			protected:
+				virtual FMetasoundFrontendVersionNumber GetTargetVersion() const = 0;
+				virtual void TransformInternal(FDocumentHandle InDocument) const = 0;
+
+			public:
+				bool Transform(FDocumentHandle InDocument) const override
+				{
+					const FMetasoundFrontendDocumentMetadata& Metadata = InDocument->GetMetadata();
+					if (Metadata.Version.Number >= GetTargetVersion())
+					{
+						return false;
+					}
+
+					TransformInternal(InDocument);
+
+					FMetasoundFrontendDocumentMetadata NewMetadata = Metadata;
+					NewMetadata.Version.Number = GetTargetVersion();
+					InDocument->SetMetadata(NewMetadata);
+
+					return true;
+				}
+		};
+
+		/** Versions document from 1.0 to 1.1. */
+		class FVersionDocument_1_1 : public FVersionDocumentTransform
+		{
+		public:
+			FMetasoundFrontendVersionNumber GetTargetVersion() const override
+			{
+				return FMetasoundFrontendVersionNumber { 1, 1 };
+			}
+
+			void TransformInternal(FDocumentHandle InDocument) const override
+			{
+				FGraphHandle GraphHandle = InDocument->GetRootGraph();
+				TArray<FNodeHandle> FrontendNodes = GraphHandle->GetNodes();
+
+				// If was connected to a hidden node, then remove that node and set the literal value accordingly.
+				for (FNodeHandle& NodeHandle : FrontendNodes)
+				{
+					if (NodeHandle->GetNodeStyle().Display.Visibility == EMetasoundFrontendNodeStyleDisplayVisibility::Hidden)
+					{
+						const FGuid VertexID = GraphHandle->GetVertexIDForInputVertex(NodeHandle->GetNodeName());
+						const FMetasoundFrontendLiteral DefaultLiteral = GraphHandle->GetDefaultInput(VertexID);
+
+						TArray<FOutputHandle> OutputHandles = NodeHandle->GetOutputs();
+						if (ensure(OutputHandles.Num() == 1))
+						{
+							FOutputHandle OutputHandle = OutputHandles[0];
+							TArray<FInputHandle> Inputs = OutputHandle->GetConnectedInputs();
+							OutputHandle->Disconnect();
+
+							for (FInputHandle& Input : Inputs)
+							{
+								if (const FMetasoundFrontendLiteral* Literal = Input->GetClassDefaultLiteral())
+								{
+									if (!Literal->IsEquivalent(DefaultLiteral))
+									{
+										Input->SetLiteral(DefaultLiteral);
+									}
+								}
+								else
+								{
+									Input->SetLiteral(DefaultLiteral);
+								}
+							}
+
+							GraphHandle->RemoveNode(*NodeHandle);
+						}
+					}
+				}
+			}
+		};
+
+		bool FVersionDocument::Transform(FDocumentHandle InDocument) const
+		{
+			if (!ensure(InDocument->IsValid()))
+			{
+				return false;
+			}
+
+			bool bWasUpdated = false;
+
+			// Add additional transforms here after defining them above, example below.
+			bWasUpdated |= FVersionDocument_1_1().Transform(InDocument);
+
+			return bWasUpdated;
+		}
 	} // namespace Frontend
 } // namespace Metasound
diff --git a/Engine/Plugins/Runtime/Metasound/Source/MetasoundFrontend/Public/MetasoundAutoConverterNode.h b/Engine/Plugins/Runtime/Metasound/Source/MetasoundFrontend/Public/MetasoundAutoConverterNode.h
--- a/Engine/Plugins/Runtime/Metasound/Source/MetasoundFrontend/Public/MetasoundAutoConverterNode.h
+++ b/Engine/Plugins/Runtime/Metasound/Source/MetasoundFrontend/Public/MetasoundAutoConverterNode.h
@@ -135,23 +135,28 @@ namespace Metasound
 				{
 					TDataWriteReference<ToDataType> WriteReference = TDataWriteReferenceFactory<ToDataType>::CreateAny(InParams.OperatorSettings);
 
-					if (!InParams.InputDataReferences.ContainsDataReadReference<FromDataType>(GetInputName()))
+					const FString& InputName = GetInputName();
+					const bool bContainsRef = InParams.InputDataReferences.ContainsDataReadReference<FromDataType>(InputName);
+					if (bContainsRef)
+					{
+						TDataReadReference<FromDataType> ReadReference = InParams.InputDataReferences.GetDataReadReference<FromDataType>(InputName);
+						return MakeUnique<FConverterOperator>(ReadReference, WriteReference);
+					}
+
+					if constexpr (TIsParsable<FromDataType>::Value)
 					{
-						if (ensure(InParams.Node.GetVertexInterface().ContainsInputVertex(GetInputName())))
-						{
-							// There should be something hooked up to the converter node. Report it as an error. 
-							FInputDataDestination Dest(InParams.Node, InParams.Node.GetVertexInterface().GetInputVertex(GetInputName()));
-							AddBuildError<FMissingInputDataReferenceError>(OutErrors, Dest);
-						}
-
-						// We can still build something even though there is an error. 
 						TDataReadReference<FromDataType> ReadReference = TDataReadReferenceFactory<FromDataType>::CreateAny(InParams.OperatorSettings);
 						return MakeUnique<FConverterOperator>(ReadReference, WriteReference);
 					}
 
-					TDataReadReference<FromDataType> ReadReference = InParams.InputDataReferences.GetDataReadReference<FromDataType>(GetInputName());
+					// Converter node requires parsable reference if input not connected. Report as an error.
+					if (ensure(InParams.Node.GetVertexInterface().ContainsInputVertex(InputName)))
+					{
+						FInputDataDestination Dest(InParams.Node, InParams.Node.GetVertexInterface().GetInputVertex(GetInputName()));
+						AddBuildError<FMissingInputDataReferenceError>(OutErrors, Dest);
+					}
 
-					return MakeUnique<FConverterOperator>(ReadReference, WriteReference);
+					return TUniquePtr<IOperator>(nullptr);
 				}
 		};
 
diff --git a/Engine/Plugins/Runtime/Metasound/Source/MetasoundFrontend/Public/MetasoundDataTypeRegistrationMacro.h b/Engine/Plugins/Runtime/Metasound/Source/MetasoundFrontend/Public/MetasoundDataTypeRegistrationMacro.h
--- a/Engine/Plugins/Runtime/Metasound/Source/MetasoundFrontend/Public/MetasoundDataTypeRegistrationMacro.h
+++ b/Engine/Plugins/Runtime/Metasound/Source/MetasoundFrontend/Public/MetasoundDataTypeRegistrationMacro.h
@@ -23,6 +23,7 @@
 #include "MetasoundRouter.h"
 #include "MetasoundSendNode.h"
 #include "MetasoundTransmissionRegistration.h"
+#include "MetasoundVariableNode.h"
 
 #include <type_traits>
 
@@ -296,10 +297,10 @@ namespace Metasound
 
 			bAlreadyRegisteredThisDataType = true;
 
-			// Lambdas that generate our template-instantiated input and output nodes:
+			// Lambdas that generate our template-instantiated input, variable, and output nodes:
 			FCreateInputNodeFunction InputNodeConstructor = [](FInputNodeConstructorParams&& InParams) -> TUniquePtr<INode>
 			{
-				return TUniquePtr<INode>(new TInputNode<TDataType>(InParams.InNodeName, InParams.InInstanceID, InParams.InVertexName, MoveTemp(InParams.InitParam)));
+				return TUniquePtr<INode>(new TInputNode<TDataType>(InParams.NodeName, InParams.InstanceID, InParams.VertexName, MoveTemp(InParams.InitParam)));
 			};
 
 			FCreateMetasoundFrontendClassFunction CreateFrontendInputClass = []() -> FMetasoundFrontendClass
@@ -310,15 +311,27 @@ namespace Metasound
 				return Metasound::Frontend::GenerateClassDescription(Prototype.GetMetadata(), EMetasoundFrontendClassType::Input);
 			};
 
+			FCreateVariableNodeFunction VariableNodeConstructor = [](FVariableNodeConstructorParams&& InParams) -> TUniquePtr<INode>
+			{
+				return TUniquePtr<INode>(new TGetVariableNode<TDataType>(InParams.NodeName, InParams.InstanceID, InParams.VertexName, MoveTemp(InParams.InitParam)));
+			};
+
+			FCreateMetasoundFrontendClassFunction CreateFrontendVariableClass = []() -> FMetasoundFrontendClass
+			{
+				// Create class info using prototype node
+				static TGetVariableNode<TDataType> Prototype(TEXT(""), FGuid(), TEXT(""), FLiteral());
+				return Metasound::Frontend::GenerateClassDescription(Prototype.GetMetadata(), EMetasoundFrontendClassType::Variable);
+			};
+
 			FCreateOutputNodeFunction OutputNodeConstructor = [](const FOutputNodeConstructorParams& InParams) -> TUniquePtr<INode>
 			{
-				return TUniquePtr<INode>(new TOutputNode<TDataType>(InParams.InNodeName, InParams.InInstanceID, InParams.InVertexName));
+				return TUniquePtr<INode>(new TOutputNode<TDataType>(InParams.NodeName, InParams.InstanceID, InParams.VertexName));
 			};
 
 			FCreateMetasoundFrontendClassFunction CreateFrontendOutputClass = []() -> FMetasoundFrontendClass
 			{
 				// Create class info using prototype node
-				// TODO: register input nodes with static class info.
+				// TODO: register output nodes with static class info.
 				static TOutputNode<TDataType> Prototype(TEXT(""), FGuid(), TEXT(""));
 				return Metasound::Frontend::GenerateClassDescription(Prototype.GetMetadata(), EMetasoundFrontendClassType::Output);
 			};
@@ -360,9 +373,11 @@ namespace Metasound
 			// Pack all of our various constructor lambdas to a single struct.
 			FDataTypeConstructorCallbacks Callbacks = 
 			{ 
-				MoveTemp(InputNodeConstructor), 
+				MoveTemp(InputNodeConstructor),
 				MoveTemp(CreateFrontendInputClass),
-				MoveTemp(OutputNodeConstructor), 
+				MoveTemp(VariableNodeConstructor),
+				MoveTemp(CreateFrontendVariableClass),
+				MoveTemp(OutputNodeConstructor),
 				MoveTemp(CreateFrontendOutputClass),
 				MoveTemp(ProxyGenerator),
 				MoveTemp(CreateDataChannelFunc)
diff --git a/Engine/Plugins/Runtime/Metasound/Source/MetasoundFrontend/Public/MetasoundFrontendBaseClasses.h b/Engine/Plugins/Runtime/Metasound/Source/MetasoundFrontend/Public/MetasoundFrontendBaseClasses.h
--- a/Engine/Plugins/Runtime/Metasound/Source/MetasoundFrontend/Public/MetasoundFrontendBaseClasses.h
+++ b/Engine/Plugins/Runtime/Metasound/Source/MetasoundFrontend/Public/MetasoundFrontendBaseClasses.h
@@ -18,6 +18,7 @@ namespace Metasound
 		// Convenience functions to create an INode corresponding to a specific input or output for a metasound graph.
 		// @returns nullptr if the type given wasn't found.
 		METASOUNDFRONTEND_API TUniquePtr<INode> ConstructInputNode(const FName& InInputType, FInputNodeConstructorParams&& InParams);
-		METASOUNDFRONTEND_API TUniquePtr<INode> ConstructOutputNode(const FName& InOutputType, const FOutputNodeConstructorParams& InParams);
+		METASOUNDFRONTEND_API TUniquePtr<INode> ConstructVariableNode(const FName& InVariableType, FVariableNodeConstructorParams&& InParams);
+		METASOUNDFRONTEND_API TUniquePtr<INode> ConstructOutputNode(const FName& InOutputType, FOutputNodeConstructorParams&& InParams);
 	}
 }
diff --git a/Engine/Plugins/Runtime/Metasound/Source/MetasoundFrontend/Public/MetasoundFrontendController.h b/Engine/Plugins/Runtime/Metasound/Source/MetasoundFrontend/Public/MetasoundFrontendController.h
--- a/Engine/Plugins/Runtime/Metasound/Source/MetasoundFrontend/Public/MetasoundFrontendController.h
+++ b/Engine/Plugins/Runtime/Metasound/Source/MetasoundFrontend/Public/MetasoundFrontendController.h
@@ -220,8 +220,14 @@ namespace Metasound
 			/** Returns the data type name associated with this input. */
 			virtual const FText& GetDisplayName() const = 0;
 
-			/** Returns the default value of the given input. */
-			virtual const FMetasoundFrontendLiteral* GetDefaultLiteral() const = 0;
+			/** Returns the value for the given input instance if set. */
+			virtual const FMetasoundFrontendLiteral* GetLiteral() const = 0;
+
+			/** Sets the value for the given input instance (effectively overriding the class default). */
+			virtual void SetLiteral(const FMetasoundFrontendLiteral& InLiteral) = 0;
+
+			/** Returns the class default value of the given input. */
+			virtual const FMetasoundFrontendLiteral* GetClassDefaultLiteral() const = 0;
 
 			/** Returns the data type name associated with this input. */
 			virtual const FText& GetTooltip() const = 0;
@@ -371,6 +377,15 @@ namespace Metasound
 			virtual FInputHandle AddOutput(const FString& InVertexName, const FMetasoundFrontendLiteral* InDefault) = 0;
 			virtual bool RemoveOutput(FGuid InVertexID) = 0;
 
+			// Returns an input's default literal if set, null if not.
+			virtual const FMetasoundFrontendLiteral* GetInputLiteral(const FGuid& InVertexID) const = 0;
+
+			// Sets an input's default literal
+			virtual void SetInputLiteral(const FMetasoundFrontendVertexLiteral& InVertexLiteral) = 0;
+
+			// Clears an input's default literal
+			virtual bool ClearInputLiteral(FGuid InVertexID) = 0;
+
 			/** Returns associated node class data */
 			virtual const FMetasoundFrontendClassInterface& GetClassInterface() const = 0;
 			virtual const FMetasoundFrontendClassMetadata& GetClassMetadata() const = 0;
@@ -700,6 +715,8 @@ namespace Metasound
 			virtual FConstGraphClassAccessPtr FindSubgraphWithID(FGuid InClassID) const = 0;
 			virtual FConstClassAccessPtr FindClassWithID(FGuid InClassID) const = 0;
 
+			virtual void SetMetadata(const FMetasoundFrontendDocumentMetadata& InMetadata) = 0;
+			virtual const FMetasoundFrontendDocumentMetadata& GetMetadata() const = 0;
 
 			/** Returns an existing Metasound class description corresponding to 
 			 * a dependency which matches the provided class information.
diff --git a/Engine/Plugins/Runtime/Metasound/Source/MetasoundFrontend/Public/MetasoundFrontendDocument.h b/Engine/Plugins/Runtime/Metasound/Source/MetasoundFrontend/Public/MetasoundFrontendDocument.h
--- a/Engine/Plugins/Runtime/Metasound/Source/MetasoundFrontend/Public/MetasoundFrontendDocument.h
+++ b/Engine/Plugins/Runtime/Metasound/Source/MetasoundFrontend/Public/MetasoundFrontendDocument.h
@@ -35,7 +35,10 @@ enum class EMetasoundFrontendClassType : uint8
 	// The Metasound class is an input into a graph in the containing document.
 	Input,
 
-	// The Metasound class is an input into a graph in the containing document.
+	// The Metasound class is an internal variable of a graph in the containing document.
+	Variable,
+
+	// The Metasound class is an output from a graph in the containing document.
 	Output,
 
 	Invalid UMETA(Hidden)
@@ -201,6 +204,7 @@ struct METASOUNDFRONTEND_API FMetasoundFrontendNodeInterface
 	TArray<FMetasoundFrontendVertex> Environment;
 };
 
+// DEPRECATED in Document Model v1.1
 UENUM()
 enum class EMetasoundFrontendNodeStyleDisplayVisibility : uint8
 {
@@ -213,7 +217,7 @@ struct METASOUNDFRONTEND_API FMetasoundFrontendNodeStyleDisplay
 {
 	GENERATED_BODY()
 
-	// Visibility state of node
+	// DEPRECATED in Document Model v1.1: Visibility state of node
 	UPROPERTY()
 	EMetasoundFrontendNodeStyleDisplayVisibility Visibility = EMetasoundFrontendNodeStyleDisplayVisibility::Visible;
 
@@ -465,6 +469,24 @@ struct METASOUNDFRONTEND_API FMetasoundFrontendClassInput : public FMetasoundFro
 	FMetasoundFrontendLiteral DefaultLiteral;
 };
 
+// Contains info for variable vertex of a Metasound class.
+USTRUCT() 
+struct METASOUNDFRONTEND_API FMetasoundFrontendClassVariable : public FMetasoundFrontendClassVertex
+{
+	GENERATED_BODY()
+
+	FMetasoundFrontendClassVariable() = default;
+
+	FMetasoundFrontendClassVariable(const FMetasoundFrontendClassVertex& InOther);
+
+	virtual ~FMetasoundFrontendClassVariable() = default;
+
+	// Default value for this variable.
+	UPROPERTY(EditAnywhere, Category = Parameters)
+	FMetasoundFrontendLiteral DefaultLiteral;
+};
+
+// Contains info for output vertex of a Metasound class.
 USTRUCT()
 struct METASOUNDFRONTEND_API FMetasoundFrontendClassOutput : public FMetasoundFrontendClassVertex
 {
@@ -734,7 +756,7 @@ struct FMetasoundFrontendDocumentMetadata
 	GENERATED_BODY()
 
 	UPROPERTY()
-	FMetasoundFrontendVersion Format;
+	FMetasoundFrontendVersion Version;
 };
 
 
@@ -791,7 +813,6 @@ struct METASOUNDFRONTEND_API FMetasoundFrontendDocument
 	UPROPERTY()
 	TArray<FMetasoundFrontendGraphClass> Subgraphs;
 
-
 	UPROPERTY()
 	FMetasoundFrontendEditorData EditorData;
 
diff --git a/Engine/Plugins/Runtime/Metasound/Source/MetasoundFrontend/Public/MetasoundFrontendLiteral.h b/Engine/Plugins/Runtime/Metasound/Source/MetasoundFrontend/Public/MetasoundFrontendLiteral.h
--- a/Engine/Plugins/Runtime/Metasound/Source/MetasoundFrontend/Public/MetasoundFrontendLiteral.h
+++ b/Engine/Plugins/Runtime/Metasound/Source/MetasoundFrontend/Public/MetasoundFrontendLiteral.h
@@ -77,6 +77,9 @@ public:
 	// Returns true if the stored Type is an array type.
 	bool IsArray() const;
 
+	// Returns whether the other literal is value equivalent
+	bool IsEquivalent(const FMetasoundFrontendLiteral& InOther) const;
+
 	// Returns true if the literal is in a valid state (Type != EMetasoundFrontendLiteralType::Invalid)
 	bool IsValid() const;
 
diff --git a/Engine/Plugins/Runtime/Metasound/Source/MetasoundFrontend/Public/MetasoundFrontendRegistries.h b/Engine/Plugins/Runtime/Metasound/Source/MetasoundFrontend/Public/MetasoundFrontendRegistries.h
--- a/Engine/Plugins/Runtime/Metasound/Source/MetasoundFrontend/Public/MetasoundFrontendRegistries.h
+++ b/Engine/Plugins/Runtime/Metasound/Source/MetasoundFrontend/Public/MetasoundFrontendRegistries.h
@@ -19,25 +19,39 @@ namespace Metasound
 		using FRegistryTransactionID = int32;
 	}
 
-	// This struct is passed to FInputNodeConstructorCallback and FOutputNodeConstructorCallback.
-	struct FInputNodeConstructorParams
+	// Base implementation for NodeConstructorCallbacks
+	struct FDefaultNodeConstructorParams
 	{
-		// the instance name and name of the specific connection that we should use.
-		const FString& InNodeName;
-		const FGuid& InInstanceID;
-		const FString& InVertexName;
-
-		FLiteral InitParam;
+		// the instance name and name of the specific connection that should be used.
+		FString NodeName;
+		FGuid InstanceID;
+		FString VertexName;
 	};
 
-	struct FOutputNodeConstructorParams
+	struct FDefaultLiteralNodeConstructorParams
 	{
-		// the instance name and name of the specific connection that we should use.
-		const FString& InNodeName;
-		const FGuid& InInstanceID;
-		const FString& InVertexName;
+		// the instance name and name of the specific connection that should be used.
+		FString NodeName;
+		FGuid InstanceID;
+		FString VertexName;
+		FLiteral InitParam = FLiteral::CreateInvalid();
+
+		FDefaultLiteralNodeConstructorParams Clone() const
+		{
+			return
+			{
+				NodeName,
+				InstanceID,
+				VertexName,
+				InitParam.Clone()
+			};
+		}
 	};
 
+	using FOutputNodeConstructorParams = FDefaultNodeConstructorParams;
+	using FVariableNodeConstructorParams = FDefaultLiteralNodeConstructorParams;
+	using FInputNodeConstructorParams = FDefaultLiteralNodeConstructorParams;
+
 	/*
 	class IMetasoundDataTypeRegistryEntry
 	{
@@ -62,7 +76,8 @@ namespace Metasound
 	*/
 
 	using FCreateInputNodeFunction = TFunction<TUniquePtr<Metasound::INode>(::Metasound::FInputNodeConstructorParams&&)>;
-	using FCreateOutputNodeFunction = TFunction<TUniquePtr<Metasound::INode>(const ::Metasound::FOutputNodeConstructorParams&)>;
+	using FCreateVariableNodeFunction = TFunction<TUniquePtr<Metasound::INode>(::Metasound::FVariableNodeConstructorParams&&)>;
+	using FCreateOutputNodeFunction = TFunction<TUniquePtr<Metasound::INode>(::Metasound::FOutputNodeConstructorParams&&)>;
 
 	// This function is used to create a proxy from a datatype's base uclass.
 	using FCreateAudioProxyFunction = TFunction<Audio::IProxyDataPtr(UObject*)> ;
@@ -258,12 +273,17 @@ namespace Metasound
 
 	struct FDataTypeConstructorCallbacks
 	{
-		// This constructs a TInputNode<> with the corresponding datatype.
+		// Constructs a TInputNode<> with the corresponding datatype.
 		FCreateInputNodeFunction CreateInputNode;
 
 		FCreateMetasoundFrontendClassFunction CreateFrontendInputClass;
 
-		// This constructs a TOutputNode<> with the corresponding datatype.
+		// Constructs a TGetVariableNode<> with the corresponding datatype.
+		FCreateVariableNodeFunction CreateVariableNode;
+
+		FCreateMetasoundFrontendClassFunction CreateFrontendVariableClass;
+
+		// Constructs a TOutputNode<> with the corresponding datatype.
 		FCreateOutputNodeFunction CreateOutputNode;
 
 		FCreateMetasoundFrontendClassFunction CreateFrontendOutputClass;
@@ -303,6 +323,7 @@ public:
 	static bool GetRegistryKey(const FNodeRegistryElement& InElement, FNodeRegistryKey& OutKey);
 	static bool GetFrontendClassFromRegistered(const FMetasoundFrontendClassMetadata& InMetadata, FMetasoundFrontendClass& OutClass);
 	static bool GetInputNodeClassMetadataForDataType(const FName& InDataTypeName, FMetasoundFrontendClassMetadata& OutMetadata);
+	static bool GetVariableNodeClassMetadataForDataType(const FName& InDataTypeName, FMetasoundFrontendClassMetadata& OutMetadata);
 	static bool GetOutputNodeClassMetadataForDataType(const FName& InDataTypeName, FMetasoundFrontendClassMetadata& OutMetadata);
 
 
@@ -346,11 +367,13 @@ public:
 	virtual bool FindFrontendClassFromRegistered(const Metasound::Frontend::FNodeClassInfo& InClassInfo, FMetasoundFrontendClass& OutClass) = 0;
 	virtual bool FindFrontendClassFromRegistered(const FMetasoundFrontendClassMetadata& InMetadata, FMetasoundFrontendClass& OutClass) = 0;
 	virtual bool FindInputNodeClassMetadataForDataType(const FName& InDataTypeName, FMetasoundFrontendClassMetadata& OutMetadata) = 0;
+	virtual bool FindVariableNodeClassMetadataForDataType(const FName& InDataTypeName, FMetasoundFrontendClassMetadata& OutMetadata) = 0;
 	virtual bool FindOutputNodeClassMetadataForDataType(const FName& InDataTypeName, FMetasoundFrontendClassMetadata& OutMetadata) = 0;
 
 	// Create a new instance of a C++ implemented node from the registry.
 	virtual TUniquePtr<Metasound::INode> ConstructInputNode(const FName& InInputType, Metasound::FInputNodeConstructorParams&& InParams) = 0;
-	virtual TUniquePtr<Metasound::INode> ConstructOutputNode(const FName& InOutputType, const Metasound::FOutputNodeConstructorParams& InParams) = 0;
+	virtual TUniquePtr<Metasound::INode> ConstructVariableNode(const FName& InVariableType, Metasound::FVariableNodeConstructorParams&& InParams) = 0;
+	virtual TUniquePtr<Metasound::INode> ConstructOutputNode(const FName& InOutputType, Metasound::FOutputNodeConstructorParams&& InParams) = 0;
 	virtual TUniquePtr<Metasound::INode> ConstructExternalNode(const Metasound::Frontend::FNodeRegistryKey& InRegistryKey, const Metasound::FNodeInitData& InInitData) = 0;
 
 	// Returns a list of possible nodes to use to convert from FromDataType to ToDataType.
diff --git a/Engine/Plugins/Runtime/Metasound/Source/MetasoundFrontend/Public/MetasoundFrontendTransform.h b/Engine/Plugins/Runtime/Metasound/Source/MetasoundFrontend/Public/MetasoundFrontendTransform.h
--- a/Engine/Plugins/Runtime/Metasound/Source/MetasoundFrontend/Public/MetasoundFrontendTransform.h
+++ b/Engine/Plugins/Runtime/Metasound/Source/MetasoundFrontend/Public/MetasoundFrontendTransform.h
@@ -22,5 +22,12 @@ namespace Metasound
 			public:
 				bool Transform(FDocumentHandle InDocument) const override;
 		};
+
+		/** Base class for versioning a document. */
+		class METASOUNDFRONTEND_API FVersionDocument : public IDocumentTransform
+		{
+			public:
+				bool Transform(FDocumentHandle InDocument) const;
+		};
 	}
 }
diff --git a/Engine/Plugins/Runtime/Metasound/Source/MetasoundFrontend/Public/MetasoundInputNode.h b/Engine/Plugins/Runtime/Metasound/Source/MetasoundFrontend/Public/MetasoundInputNode.h
--- a/Engine/Plugins/Runtime/Metasound/Source/MetasoundFrontend/Public/MetasoundInputNode.h
+++ b/Engine/Plugins/Runtime/Metasound/Source/MetasoundFrontend/Public/MetasoundInputNode.h
@@ -4,16 +4,16 @@
 #include "CoreMinimal.h"
 #include "Internationalization/Text.h"
 #include "MetasoundDataReference.h"
-#include "MetasoundOperatorInterface.h"
 #include "MetasoundNodeInterface.h"
+#include "MetasoundOperatorInterface.h"
 #include "MetasoundTrigger.h"
 
-#define LOCTEXT_NAMESPACE "MetasoundGraphCore"
+#define LOCTEXT_NAMESPACE "MetasoundFrontend"
 
 
 namespace Metasound
 {
-	/** TInputOperator supplies a writable input and a readable output. */
+	/** A writable input and a readable output. */
 	template<typename DataType>
 	class TInputOperator : public IOperator
 	{
@@ -88,7 +88,7 @@ namespace Metasound
 				return nullptr;
 			}
 
-		private:
+		protected:
 			FVertexKey DataReferenceName;
 
 			FDataWriteReference InputValue;
@@ -389,4 +389,4 @@ namespace Metasound
 		}
 	}
 } // namespace Metasound
-#undef LOCTEXT_NAMESPACE //MetasoundOutputNode
+#undef LOCTEXT_NAMESPACE // MetasoundFrontend
diff --git a/Engine/Plugins/Runtime/Metasound/Source/MetasoundFrontend/Public/MetasoundOutputNode.h b/Engine/Plugins/Runtime/Metasound/Source/MetasoundFrontend/Public/MetasoundOutputNode.h
--- a/Engine/Plugins/Runtime/Metasound/Source/MetasoundFrontend/Public/MetasoundOutputNode.h
+++ b/Engine/Plugins/Runtime/Metasound/Source/MetasoundFrontend/Public/MetasoundOutputNode.h
@@ -2,6 +2,7 @@
 #pragma once
 
 #include "CoreMinimal.h"
+#include "MetasoundBuildError.h"
 #include "MetasoundBuilderInterface.h"
 #include "MetasoundNodeInterface.h"
 #include "MetasoundOperatorInterface.h"
@@ -11,6 +12,26 @@
 
 namespace Metasound
 {
+	/** FMissingOutputNodeInputReferenceError
+	 *
+	 * Caused by Output not being able to generate an IOperator instance due to
+	 * the type requiring an input reference (i.e. it is not default constructable).
+	 */
+	class FMissingOutputNodeInputReferenceError : public FBuildErrorBase
+	{
+	public:
+		FMissingOutputNodeInputReferenceError(const INode& InNode, const FText& InDataType)
+			: FBuildErrorBase(
+				"MetasoundMissingOutputDataReferenceError",
+				FText::Format(LOCTEXT("MissingOutputNodeInputReferenceError", "Missing required output node input reference for type {0}."), InDataType))
+		{
+			AddNode(InNode);
+		}
+
+		virtual ~FMissingOutputNodeInputReferenceError() = default;
+
+	};
+
 	template<typename DataType>
 	class TOutputNode : public FNode
 	{
@@ -55,13 +76,22 @@ namespace Metasound
 
 				virtual TUniquePtr<IOperator> CreateOperator(const FCreateOperatorParams& InParams, FBuildErrorArray& OutErrors) override
 				{
-					if (!InParams.InputDataReferences.ContainsDataReadReference<DataType>(DataReferenceName))
+					const bool bContainsRef = InParams.InputDataReferences.ContainsDataReadReference<DataType>(DataReferenceName);
+					if (bContainsRef)
+					{
+						TDataReadReference<DataType> InputReadRef = InParams.InputDataReferences.GetDataReadReference<DataType>(DataReferenceName);
+						return MakeUnique<FOutputOperator>(DataReferenceName, InputReadRef);
+					}
+
+					// Only construct default if default construction is supported
+					if constexpr (TIsParsable<DataType>::Value)
 					{
-						// TODO: Add build error.
-						return TUniquePtr<IOperator>(nullptr);
+						TDataReadReference<DataType> DefaultReadRef = TDataReadReferenceFactory<DataType>::CreateAny(InParams.OperatorSettings);
+						return MakeUnique<FOutputOperator>(DataReferenceName, DefaultReadRef);
 					}
 
-					return MakeUnique<FOutputOperator>(DataReferenceName, InParams.InputDataReferences.GetDataReadReference<DataType>(DataReferenceName));
+					OutErrors.Emplace(MakeUnique<FMissingOutputNodeInputReferenceError>(InParams.Node, GetMetasoundDataTypeDisplayText<DataType>()));
+					return TUniquePtr<IOperator>(nullptr);
 				}
 
 			private:
diff --git a/Engine/Plugins/Runtime/Metasound/Source/MetasoundFrontend/Public/MetasoundSendNode.h b/Engine/Plugins/Runtime/Metasound/Source/MetasoundFrontend/Public/MetasoundSendNode.h
--- a/Engine/Plugins/Runtime/Metasound/Source/MetasoundFrontend/Public/MetasoundSendNode.h
+++ b/Engine/Plugins/Runtime/Metasound/Source/MetasoundFrontend/Public/MetasoundSendNode.h
@@ -142,7 +142,7 @@ namespace Metasound
 					if (InParams.InputDataReferences.ContainsDataReadReference<TDataType>(GetSendInputName()))
 					{
 						return MakeUnique<TSendOperator>(InParams.InputDataReferences.GetDataReadReference<TDataType>(GetSendInputName()),
-							InParams.InputDataReferences.GetDataReadReference<FSendAddress>(GetAddressInputName()),
+							InParams.InputDataReferences.GetDataReadReferenceOrConstruct<FSendAddress>(GetAddressInputName()),
 							InParams.OperatorSettings
 							);
 					}
diff --git a/Engine/Plugins/Runtime/Metasound/Source/MetasoundFrontend/Public/MetasoundVariableNode.h b/Engine/Plugins/Runtime/Metasound/Source/MetasoundFrontend/Public/MetasoundVariableNode.h
new file mode 100644
--- /dev/null
+++ b/Engine/Plugins/Runtime/Metasound/Source/MetasoundFrontend/Public/MetasoundVariableNode.h
@@ -0,0 +1,201 @@
+// Copyright Epic Games, Inc. All Rights Reserved.
+#pragma once
+
+#include "CoreMinimal.h"
+#include "Internationalization/Text.h"
+#include "MetasoundInputNode.h"
+#include "MetasoundNodeInterface.h"
+#include "MetasoundOperatorInterface.h"
+#include "MetasoundTrigger.h"
+#include "MetasoundVertex.h"
+
+#define LOCTEXT_NAMESPACE "MetasoundFrontend"
+
+
+namespace Metasound
+{
+	/** A writable variable (if non-const) and a readable output that is only accessible within a given graph. */
+	template<typename DataType>
+	class TGetVariableOperator : public TInputOperator<DataType>
+	{
+		public:
+			TGetVariableOperator(const FVertexKey& InDataReferenceName, TDataWriteReference<DataType> InDataReference)
+				: TInputOperator<DataType>(InDataReferenceName, InDataReference)
+			{
+			}
+
+			virtual ~TGetVariableOperator() = default;
+
+			virtual FDataReferenceCollection GetInputs() const override
+			{
+				FDataReferenceCollection Inputs;
+				Inputs.AddDataWriteReference<DataType>(TInputOperator<DataType>::DataReferenceName, TInputOperator<DataType>::InputValue);
+				return Inputs;
+			}
+	};
+
+	/** TGetVariableOperatorLiteralFactory creates an input by passing it a literal. */
+	template<typename DataType>
+	class TGetVariableOperatorLiteralFactory : public IOperatorFactory
+	{
+	public:
+		// If the data type is parsable from a literal type, then the data type 
+		// can be registered as an input type with the frontend.  To make a 
+		// DataType registrable, either create a constructor for the data type
+		// which accepts the one of the supported literal types with an optional 
+		// FOperatorSettings argument, or create a default constructor, or specialize
+		// this factory with an implementation for that specific data type.
+		static constexpr bool bCanRegister = TLiteralTraits<DataType>::bIsParsableFromAnyLiteralType;
+
+		using FDataWriteReference = TDataWriteReference<DataType>;
+
+		TGetVariableOperatorLiteralFactory(FLiteral&& InInitParam)
+			: InitParam(MoveTemp(InInitParam))
+		{
+		}
+
+		virtual TUniquePtr<IOperator> CreateOperator(const FCreateOperatorParams& InParams, FBuildErrorArray& OutErrors) override;
+
+	private:
+		FLiteral InitParam;
+	};
+
+	/** TGetVariableOperatorFactory initializes the DataType at construction. It uses
+	 * the DataType's copy operator to create additional version.
+	 */
+	template<typename DataType>
+	class TGetVariableOperatorFactory : public IOperatorFactory
+	{
+	public:
+
+		using FDataWriteReference = TDataWriteReference<DataType>;
+
+		template<typename... ArgTypes>
+		TGetVariableOperatorFactory(ArgTypes&&... Args)
+			: Data(Forward<ArgTypes>(Args)...)
+		{
+		}
+
+		virtual TUniquePtr<IOperator> CreateOperator(const FCreateOperatorParams& InParams, FBuildErrorArray& OutErrors) override;
+
+	private:
+		DataType Data;
+	};
+
+	/** TGetVariableNode represents a variable within a metasound graph. */
+	template<typename DataType>
+	class TGetVariableNode : public FNode
+	{
+		public:
+
+			// If true, this node can be instantiated by the Frontend.
+			static constexpr bool bCanRegister = TGetVariableOperatorLiteralFactory<DataType>::bCanRegister;
+
+			static FVertexInterface DeclareVertexInterface(const FVertexKey& InVertexName)
+			{
+				return FVertexInterface(
+					FInputVertexInterface(
+						TInputDataVertexModel<DataType>(InVertexName, FText::GetEmpty())
+					),
+					FOutputVertexInterface(
+						TOutputDataVertexModel<DataType>(InVertexName, FText::GetEmpty())
+					)
+				);
+			}
+
+			static FNodeClassMetadata GetNodeInfo(const FVertexKey& InVertexName)
+			{
+				FNodeClassMetadata Info;
+
+				Info.ClassName = { TEXT("Variable"), GetMetasoundDataTypeName<DataType>(), TEXT("")};
+				Info.MajorVersion = 1;
+				Info.MinorVersion = 0;
+				Info.DisplayName = FText::Format(LOCTEXT("Metasound_VariableNodeDisplayNameFormat", "Variable {0}"), FText::FromName(GetMetasoundDataTypeName<DataType>()));
+				Info.Description = LOCTEXT("Metasound_VariableNodeDescription", "Variable accessible within a parent Metasound graph.");
+				Info.Author = PluginAuthor;
+				Info.PromptIfMissing = PluginNodeMissingPrompt;
+				Info.DefaultInterface = DeclareVertexInterface(InVertexName);
+
+				return Info;
+			}
+
+			/* Construct a TGetVariableNode using the TGetVariableOperatorFactory<> and forwarding
+			 * Args to the TGetVariableOperatorFactory constructor.*/
+			template<typename... ArgTypes>
+			TGetVariableNode(const FString& InNodeDescription, const FGuid& InInstanceID, const FVertexKey& InVertexName, ArgTypes&&... Args)
+			:	FNode(InNodeDescription, InInstanceID, GetNodeInfo(InVertexName))
+			,	VertexName(InVertexName)
+			,	Interface(DeclareVertexInterface(InVertexName))
+			,	Factory(MakeOperatorFactoryRef<TGetVariableOperatorFactory<DataType>>(Forward<ArgTypes>(Args)...))
+			{
+			}
+
+			/* Construct a TGetVariableNode using the TGetVariableOperatorLiteralFactory<> and moving
+			 * InParam to the TGetVariableOperatorLiteralFactory constructor.*/
+			explicit TGetVariableNode(const FString& InNodeDescription, const FGuid& InInstanceID, const FVertexKey& InVertexName, FLiteral&& InParam)
+			:	FNode(InNodeDescription, InInstanceID, GetNodeInfo(InVertexName))
+			,	VertexName(InVertexName)
+			,	Interface(DeclareVertexInterface(InVertexName))
+			, 	Factory(MakeOperatorFactoryRef<TGetVariableOperatorLiteralFactory<DataType>>(MoveTemp(InParam)))
+			{
+			}
+
+			const FVertexKey& GetVertexName() const
+			{
+				return VertexName;
+			}
+
+			virtual const FVertexInterface& GetVertexInterface() const override
+			{
+				return Interface;
+			}
+
+			virtual bool SetVertexInterface(const FVertexInterface& InInterface) override
+			{
+				return Interface == InInterface;
+			}
+
+			virtual bool IsVertexInterfaceSupported(const FVertexInterface& InInterface) const override
+			{
+				return Interface == InInterface;
+			}
+
+			virtual TSharedRef<IOperatorFactory, ESPMode::ThreadSafe> GetDefaultOperatorFactory() const override
+			{
+				return Factory;
+			}
+
+		private:
+			FVertexKey VertexName;
+
+			FVertexInterface Interface;
+			FOperatorFactorySharedRef Factory;
+	};
+
+	template<typename DataType>
+	TUniquePtr<IOperator> TGetVariableOperatorLiteralFactory<DataType>::CreateOperator(const FCreateOperatorParams& InParams, FBuildErrorArray& OutErrors)
+	{
+		using FVariableNodeType = TGetVariableNode<DataType>;
+
+		// Create write reference by calling compatible constructor with literal.
+		FDataWriteReference DataRef = TDataWriteReferenceLiteralFactory<DataType>::CreateExplicitArgs(InParams.OperatorSettings, InitParam);
+
+		const FVariableNodeType& VariableNode = static_cast<const FVariableNodeType&>(InParams.Node);
+
+		return MakeUnique<TGetVariableOperator<DataType>>(VariableNode.GetVertexName(), DataRef);
+	}
+
+	template<typename DataType>
+	TUniquePtr<IOperator> TGetVariableOperatorFactory<DataType>::CreateOperator(const FCreateOperatorParams& InParams, FBuildErrorArray& OutErrors)
+	{
+		using FVariableNodeType = TGetVariableNode<DataType>;
+
+		// Create write reference by calling copy constructor.
+		FDataWriteReference DataRef = TDataWriteReferenceFactory<DataType>::CreateExplicitArgs(InParams.OperatorSettings, Data);
+
+		const FVariableNodeType& VariableNode = static_cast<const FVariableNodeType&>(InParams.Node);
+
+		return MakeUnique<TGetVariableOperator<DataType>>(VariableNode.GetVertexName(), DataRef);
+	}
+} // namespace Metasound
+#undef LOCTEXT_NAMESPACE // MetasoundFrontend
diff --git a/Engine/Plugins/Runtime/Metasound/Source/MetasoundGenerator/Private/MetasoundGenerator.cpp b/Engine/Plugins/Runtime/Metasound/Source/MetasoundGenerator/Private/MetasoundGenerator.cpp
--- a/Engine/Plugins/Runtime/Metasound/Source/MetasoundGenerator/Private/MetasoundGenerator.cpp
+++ b/Engine/Plugins/Runtime/Metasound/Source/MetasoundGenerator/Private/MetasoundGenerator.cpp
@@ -92,6 +92,8 @@ namespace Metasound
 		, NumSamplesPerExecute(0)
 		, OnPlayTriggerRef(FTriggerWriteRef::CreateNew(InParams.OperatorSettings))
 		, OnFinishedTriggerRef(FTriggerWriteRef::CreateNew(InParams.OperatorSettings))
+		, bPendingGraphTrigger(true)
+		, bIsNewGraphPending(true)
 	{
 		NumChannels = InParams.NumOutputChannels;
 		NumFramesPerExecute = InParams.OperatorSettings.GetNumFramesPerBlock();
@@ -125,7 +127,7 @@ namespace Metasound
 		}
 	}
 
-	void FMetasoundGenerator::UpdateGraphIfPending()
+	bool FMetasoundGenerator::UpdateGraphIfPending()
 	{
 		FScopeLock GraphLock(&PendingGraphMutex);
 		if (bIsNewGraphPending)
@@ -133,6 +135,8 @@ namespace Metasound
 			SetGraph(MoveTemp(PendingGraphData), bPendingGraphTrigger);
 			bIsNewGraphPending = false;
 		}
+
+		return bIsNewGraphPending;
 	}
 
 	void FMetasoundGenerator::SetGraph(TUniquePtr<FMetasoundGeneratorData>&& InData, bool bTriggerGraph)
@@ -207,13 +211,28 @@ namespace Metasound
 			return 0;
 		}
 
-		UpdateGraphIfPending();
+		const bool bPending = UpdateGraphIfPending();
 
 		// Output silent audio if we're still building a graph
-		if (RootExecuter.IsNoOp() || NumSamplesPerExecute < 1)
+		if (bPending)
 		{
-			FMemory::Memset(OutAudio, 0, sizeof(float) * NumSamplesRemaining);
-			return NumSamplesRemaining;
+			if (RootExecuter.IsNoOp() || NumSamplesPerExecute < 1)
+			{
+				FMemory::Memset(OutAudio, 0, sizeof(float)* NumSamplesRemaining);
+				return NumSamplesRemaining;
+			}
+		}
+		// If no longer pending and executer is no-op, kill the MetaSound.
+		// Covers case where there was an error when building, resulting in
+		// Executer operator being assigned to NoOp.
+		else
+		{
+			if (RootExecuter.IsNoOp())
+			{
+				bIsFinished = true;
+				FMemory::Memset(OutAudio, 0, sizeof(float) * NumSamplesRemaining);
+				return NumSamplesRemaining;
+			}
 		}
 
 		// If we have any audio left in the internal overflow buffer from 
diff --git a/Engine/Plugins/Runtime/Metasound/Source/MetasoundGenerator/Public/MetasoundGenerator.h b/Engine/Plugins/Runtime/Metasound/Source/MetasoundGenerator/Public/MetasoundGenerator.h
--- a/Engine/Plugins/Runtime/Metasound/Source/MetasoundGenerator/Public/MetasoundGenerator.h
+++ b/Engine/Plugins/Runtime/Metasound/Source/MetasoundGenerator/Public/MetasoundGenerator.h
@@ -143,9 +143,8 @@ namespace Metasound
 		 */
 		void SetPendingGraph(FMetasoundGeneratorData&& InData, bool bTriggerGraph);
 
-
 	private:
-		void UpdateGraphIfPending();
+		bool UpdateGraphIfPending();
 		// Internal set graph after checking compatibility.
 		void SetGraph(TUniquePtr<FMetasoundGeneratorData>&& InData, bool bTriggerGraph);
 
diff --git a/Engine/Plugins/Runtime/Metasound/Source/MetasoundGraphCore/Private/MetasoundOperatorBuilder.cpp b/Engine/Plugins/Runtime/Metasound/Source/MetasoundGraphCore/Private/MetasoundOperatorBuilder.cpp
--- a/Engine/Plugins/Runtime/Metasound/Source/MetasoundGraphCore/Private/MetasoundOperatorBuilder.cpp
+++ b/Engine/Plugins/Runtime/Metasound/Source/MetasoundGraphCore/Private/MetasoundOperatorBuilder.cpp
@@ -177,8 +177,9 @@ namespace Metasound
 		// Sort the nodes in a valid execution order
 		BuildStatus |= DepthFirstTopologicalSort(BuildContext, SortedNodes);
 
-		// Prune unneeded nodes. 
-		BuildStatus |= PruneNodes(BuildContext, SortedNodes);
+		// TODO: Add FindReachableNodesFromVariables in Prune.
+		// Otherwise, subgraphs incorrectly get pruned.
+		// BuildStatus |= PruneNodges(BuildContext, SortedNodes);
 
 		// Check build status in case build routine should be exited early.
 		if (BuildStatus > MaxBuildStatusErrorLevel)
diff --git a/Engine/Plugins/Runtime/Metasound/Source/MetasoundGraphCore/Public/MetasoundFacade.h b/Engine/Plugins/Runtime/Metasound/Source/MetasoundGraphCore/Public/MetasoundFacade.h
--- a/Engine/Plugins/Runtime/Metasound/Source/MetasoundGraphCore/Public/MetasoundFacade.h
+++ b/Engine/Plugins/Runtime/Metasound/Source/MetasoundGraphCore/Public/MetasoundFacade.h
@@ -125,6 +125,14 @@ namespace Metasound
 				VertexInterface = GetMetadata().DefaultInterface;
 			}
 
+			template<typename OperatorType>
+			FNodeFacade(FNodeInitData&& InInitData, TFacadeOperatorClass<OperatorType> OperatorClass)
+				: FNode(InInitData, OperatorType::GetNodeInfo())
+				, Factory(MakeShared<FFactory, ESPMode::ThreadSafe>(&OperatorType::CreateOperator))
+			{
+				VertexInterface = GetMetadata().DefaultInterface;
+			}
+
 			virtual ~FNodeFacade() = default;
 
 			virtual const FVertexInterface& GetVertexInterface() const override;
diff --git a/Engine/Plugins/Runtime/Metasound/Source/MetasoundGraphCore/Public/MetasoundLiteral.h b/Engine/Plugins/Runtime/Metasound/Source/MetasoundGraphCore/Public/MetasoundLiteral.h
--- a/Engine/Plugins/Runtime/Metasound/Source/MetasoundGraphCore/Public/MetasoundLiteral.h
+++ b/Engine/Plugins/Runtime/Metasound/Source/MetasoundGraphCore/Public/MetasoundLiteral.h
@@ -108,7 +108,10 @@ namespace Metasound
 #endif
 		}
 
-		FLiteral(const FLiteral& Other) = delete;
+		//FLiteral(const FLiteral& Other)
+		//{
+		//	*this = Other.Clone();
+		//}
 		FLiteral(FLiteral&& Other) = default;
 		FLiteral& operator=(FLiteral&& Other) = default;
 
@@ -124,7 +127,6 @@ namespace Metasound
 
 		FLiteral Clone() const;
 
-
 		template<typename ArgType>
 		void Set(ArgType&& Arg)
 		{
diff --git a/Engine/Plugins/Runtime/Metasound/Source/MetasoundGraphCore/Public/MetasoundNodeInterface.h b/Engine/Plugins/Runtime/Metasound/Source/MetasoundGraphCore/Public/MetasoundNodeInterface.h
--- a/Engine/Plugins/Runtime/Metasound/Source/MetasoundGraphCore/Public/MetasoundNodeInterface.h
+++ b/Engine/Plugins/Runtime/Metasound/Source/MetasoundGraphCore/Public/MetasoundNodeInterface.h
@@ -16,25 +16,6 @@ namespace Metasound
 	static const FText PluginAuthor = NSLOCTEXT("MetasoundGraphCore", "Metasound_DefaultAuthor", "Epic Games, Inc.");
 	static const FText PluginNodeMissingPrompt = NSLOCTEXT("MetasoundGraphCore", "Metasound_DefaultMissingNodePrompt", "The node was likely removed, renamed, or the Metasound plugin is not loaded.");
 
-	/**
-	 * This struct is used to pass in any arguments required for constructing a single node instance.
-	 * because of this, all FNode implementations have to implement a constructor that takes an FNodeInitData instance.
-	 */
-	struct FNodeInitData
-	{
-		FString InstanceName;
-		FGuid InstanceID;
-		TMap<FName, FLiteral> ParamMap;
-
-		template<typename ParamType>
-		ParamType GetParamValue(FName ParamName)
-		{
-			checkf(ParamMap.Contains(ParamName), TEXT("Tried to use node initialization parameter that didn't exist!"));
-
-			return TDataTypeLiteralFactory<ParamType>::CreateAny(ParamMap[ParamName]);
-		}
-	};
-
 	/**
 	 * Node style data
 	 */
diff --git a/Engine/Plugins/Runtime/Metasound/Source/MetasoundGraphCore/Public/MetasoundVertex.h b/Engine/Plugins/Runtime/Metasound/Source/MetasoundGraphCore/Public/MetasoundVertex.h
--- a/Engine/Plugins/Runtime/Metasound/Source/MetasoundGraphCore/Public/MetasoundVertex.h
+++ b/Engine/Plugins/Runtime/Metasound/Source/MetasoundGraphCore/Public/MetasoundVertex.h
@@ -4,6 +4,7 @@
 
 #include "CoreMinimal.h"
 #include "Containers/SortedMap.h"
+#include "MetasoundDataFactory.h"
 #include "MetasoundDataReference.h"
 #include "MetasoundEnvironment.h"
 #include "MetasoundLiteral.h"
@@ -867,4 +868,14 @@ namespace Metasound
 			FOutputVertexInterface OutputInterface;
 			FEnvironmentVertexInterface EnvironmentInterface;
 	};
+
+	/**
+	 * This struct is used to pass in any arguments required for constructing a single node instance.
+	 * because of this, all FNode implementations have to implement a constructor that takes an FNodeInitData instance.
+	 */
+	struct FNodeInitData
+	{
+		FString InstanceName;
+		FGuid InstanceID;
+	};
 }
