Execute versioning when attempting to build ISoundGenerator from MetaSoundSource in case asset hasn't been versioned & resaved.

diff --git a/Engine/Plugins/Runtime/Metasound/Source/MetasoundEngine/Private/MetasoundSource.cpp b/Engine/Plugins/Runtime/Metasound/Source/MetasoundEngine/Private/MetasoundSource.cpp
--- a/Engine/Plugins/Runtime/Metasound/Source/MetasoundEngine/Private/MetasoundSource.cpp
+++ b/Engine/Plugins/Runtime/Metasound/Source/MetasoundEngine/Private/MetasoundSource.cpp
@@ -11,6 +11,7 @@
 #include "MetasoundFrontendQuery.h"
 #include "MetasoundFrontendQuerySteps.h"
 #include "MetasoundFrontendSearchEngine.h"
+#include "MetasoundFrontendTransform.h"
 #include "MetasoundGenerator.h"
 #include "MetasoundInstanceTransmitter.h"
 #include "MetasoundLog.h"
@@ -24,7 +25,7 @@
 #include "EdGraph/EdGraph.h"
 #endif // WITH_EDITORONLY_DATA
 
-#define LOCTEXT_NAMESPACE "MetasoundSource"
+#define LOCTEXT_NAMESPACE "MetaSoundSource"
 
 
 static float MetaSoundBlockRateCVar = 100.f;
@@ -147,6 +148,12 @@ ISoundGeneratorPtr UMetaSoundSource::CreateSoundGenerator(const FSoundGeneratorI
 		return ISoundGeneratorPtr(nullptr);
 	}
 
+	// Check version and attempt to version up if out-of-date
+	if (Frontend::FVersionDocument().Transform(GetDocumentHandle()))
+	{
+		UE_LOG(LogMetaSound, Display, TEXT("Dynamically applied versioning for UMetaSoundSource [Name:%s].  Resaving the asset will avoid this requirement at runtime."), *GetName());
+	}
+
 	// Inject receive nodes for unused transmittable inputs. Perform edits on a copy
 	// of the document to avoid altering document.
 	// TODO: Use a light wrapper of the document instead of a copy.
@@ -530,4 +537,4 @@ const FMetasoundFrontendArchetype& UMetaSoundSource::GetStereoSourceArchetype()
 	return StereoArchetype;
 }
 
-#undef LOCTEXT_NAMESPACE
+#undef LOCTEXT_NAMESPACE // MetaSoundSource
