- Fix for sticky MetaSound literals when more than one external node attempts to set literals for a vertex with the same ID
- Fix for not serializing MetaSound EdGraph (again...)

diff --git a/Engine/Plugins/Runtime/Metasound/Source/MetasoundEditor/Private/MetasoundEditorGraphBuilder.cpp b/Engine/Plugins/Runtime/Metasound/Source/MetasoundEditor/Private/MetasoundEditorGraphBuilder.cpp
--- a/Engine/Plugins/Runtime/Metasound/Source/MetasoundEditor/Private/MetasoundEditorGraphBuilder.cpp
+++ b/Engine/Plugins/Runtime/Metasound/Source/MetasoundEditor/Private/MetasoundEditorGraphBuilder.cpp
@@ -1175,7 +1175,7 @@ namespace Metasound
 			// after initial setup if necessary.
 			if (!MetasoundAsset->GetGraph())
 			{
-				UMetasoundEditorGraph* Graph = NewObject<UMetasoundEditorGraph>(&InMetasound, FName(), RF_Transient | RF_Transactional);
+				UMetasoundEditorGraph* Graph = NewObject<UMetasoundEditorGraph>(&InMetasound, FName(), RF_Transactional);
 				Graph->Schema = UMetasoundEditorGraphSchema::StaticClass();
 				MetasoundAsset->SetGraph(Graph);
 			}
diff --git a/Engine/Plugins/Runtime/Metasound/Source/MetasoundFrontend/Private/MetasoundFrontendGraph.cpp b/Engine/Plugins/Runtime/Metasound/Source/MetasoundFrontend/Private/MetasoundFrontendGraph.cpp
--- a/Engine/Plugins/Runtime/Metasound/Source/MetasoundFrontend/Private/MetasoundFrontendGraph.cpp
+++ b/Engine/Plugins/Runtime/Metasound/Source/MetasoundFrontend/Private/MetasoundFrontendGraph.cpp
@@ -220,8 +220,7 @@ namespace Metasound
 				TArray<FDefaultVariableData> DefaultVariableData = GetInputDefaultVariableData(InNode, InitData);
 				for (FDefaultVariableData& Data : DefaultVariableData)
 				{
-					FGuid VertexID = Data.DestinationVertexID;
-					InGraphContext.DefaultInputs.Emplace(MoveTemp(VertexID), MoveTemp(Data));
+					InGraphContext.DefaultInputs.Emplace(FNodeIDVertexID { InNode.ID, Data.DestinationVertexID }, MoveTemp(Data));
 				}
 			}
 
@@ -241,8 +240,7 @@ namespace Metasound
 			TArray<FDefaultVariableData> DefaultVariableData = GetInputDefaultVariableData(InNode, InitData);
 			for (FDefaultVariableData& Data : DefaultVariableData)
 			{
-				FGuid VertexID = Data.DestinationVertexID;
-				InGraphContext.DefaultInputs.Emplace(MoveTemp(VertexID), MoveTemp(Data));
+				InGraphContext.DefaultInputs.Emplace(FNodeIDVertexID { InNode.ID, Data.DestinationVertexID }, MoveTemp(Data));
 			}
 		}
 
@@ -442,8 +440,6 @@ namespace Metasound
 			const FMetasoundFrontendVertex* Vertex = nullptr;
 		};
 
-		typedef TTuple<FGuid, FGuid> FNodeIDVertexID;
-
 		TMap<FNodeIDVertexID, FCoreNodeAndFrontendVertex> NodeSourcesByID;
 		TMap<FNodeIDVertexID, FCoreNodeAndFrontendVertex> NodeDestinationsByID;
 
@@ -514,7 +510,8 @@ namespace Metasound
 			// If succeeded, remove input as viable vertex to construct default variable, as it has been superceded by a connection.
 			if (bSuccess)
 			{
-				InGraphContext.DefaultInputs.Remove(DestinationNodeAndVertex->Vertex->VertexID);
+				FNodeIDVertexID DestinationPair { ToNode->GetInstanceID(), DestinationNodeAndVertex->Vertex->VertexID };
+				InGraphContext.DefaultInputs.Remove(DestinationPair);
 			}
 			else
 			{
@@ -525,7 +522,7 @@ namespace Metasound
 
 	void FFrontendGraphBuilder::AddDefaultInputVariables(FBuildGraphContext& InGraphContext)
 	{
-		for (const TPair<FGuid, FDefaultVariableData>& Pair : InGraphContext.DefaultInputs)
+		for (const TPair<FNodeIDVertexID, FDefaultVariableData>& Pair : InGraphContext.DefaultInputs)
 		{
 			const FDefaultVariableData& VariableData = Pair.Value;
 			const FGuid VariableNodeID = FGuid::NewGuid();
diff --git a/Engine/Plugins/Runtime/Metasound/Source/MetasoundFrontend/Private/MetasoundFrontendGraph.h b/Engine/Plugins/Runtime/Metasound/Source/MetasoundFrontend/Private/MetasoundFrontendGraph.h
--- a/Engine/Plugins/Runtime/Metasound/Source/MetasoundFrontend/Private/MetasoundFrontendGraph.h
+++ b/Engine/Plugins/Runtime/Metasound/Source/MetasoundFrontend/Private/MetasoundFrontendGraph.h
@@ -125,9 +125,10 @@ namespace Metasound
 		};
 
 		// Map of Input VertexID to variable data required to construct and connect default variable
+		using FNodeIDVertexID = TTuple<FGuid, FGuid>;
 		using FDependencyByIDMap = TMap<FGuid, const FMetasoundFrontendClass*>;
 		using FSharedNodeByIDMap = TMap<FGuid, TSharedPtr<const INode>>;
-		using FDefaultInputByIDMap = TMap<FGuid, FDefaultVariableData>;
+		using FDefaultInputByIDMap = TMap<FNodeIDVertexID, FDefaultVariableData>;
 
 		// Context used throughout entire graph build process
 		// (for both a root and nested subgraphs)
