Fix for MixerNode not utilizing default vertex value if no literal/variable node input is provided from builder

diff --git a/Engine/Plugins/Runtime/Metasound/Source/MetasoundStandardNodes/Private/MetasoundMixerNode.cpp b/Engine/Plugins/Runtime/Metasound/Source/MetasoundStandardNodes/Private/MetasoundMixerNode.cpp
--- a/Engine/Plugins/Runtime/Metasound/Source/MetasoundStandardNodes/Private/MetasoundMixerNode.cpp
+++ b/Engine/Plugins/Runtime/Metasound/Source/MetasoundStandardNodes/Private/MetasoundMixerNode.cpp
@@ -134,7 +134,7 @@ namespace Metasound
 					InputBuffers.Add(InputCollection.GetDataReadReferenceOrConstruct<FAudioBuffer>(GetAudioInputName(i, Chan), InParams.OperatorSettings));
 				}
 
-				InputGains.Add(InputCollection.GetDataReadReferenceOrConstruct<float>(GetGainInputName(i)));
+				InputGains.Add(InputCollection.GetDataReadReferenceOrConstructWithVertexDefault<float>(InputInterface, GetGainInputName(i), InParams.OperatorSettings));
 			}
 
 			return MakeUnique<TAudioMixerNodeOperator<NumInputs, NumChannels>>(InParams.OperatorSettings, MoveTemp(InputBuffers), MoveTemp(InputGains));
