- Up default MetaSound Document Version to avoid running Version Transform on new Doc
- Add field for setting graph as editable (to be used for presets)

diff --git a/Engine/Plugins/Runtime/Metasound/Source/MetasoundFrontend/Public/MetasoundFrontendDocument.h b/Engine/Plugins/Runtime/Metasound/Source/MetasoundFrontend/Public/MetasoundFrontendDocument.h
--- a/Engine/Plugins/Runtime/Metasound/Source/MetasoundFrontend/Public/MetasoundFrontendDocument.h
+++ b/Engine/Plugins/Runtime/Metasound/Source/MetasoundFrontend/Public/MetasoundFrontendDocument.h
@@ -340,6 +340,9 @@ struct FMetasoundFrontendGraphStyle
 {
 	GENERATED_BODY()
 
+	UPROPERTY()
+	bool bIsGraphEditable = true;
+
 	// Edge styles for graph.
 	UPROPERTY()
 	TArray<FMetasoundFrontendStyleEdgeClass> EdgeStyles;
@@ -757,6 +760,11 @@ struct FMetasoundFrontendDocumentMetadata
 
 	UPROPERTY()
 	FMetasoundFrontendVersion Version;
+
+	FMetasoundFrontendDocumentMetadata()
+	{
+		Version.Number = { 1, 1 };
+	}
 };
 
 
