Fix for unserialized MetaSound version info reporting default value, which is higher version than actual version, causing document transform code to not execute for existing assets.

diff --git a/Engine/Plugins/Runtime/Metasound/Source/MetasoundEditor/Private/MetasoundEditorGraphBuilder.cpp b/Engine/Plugins/Runtime/Metasound/Source/MetasoundEditor/Private/MetasoundEditorGraphBuilder.cpp
--- a/Engine/Plugins/Runtime/Metasound/Source/MetasoundEditor/Private/MetasoundEditorGraphBuilder.cpp
+++ b/Engine/Plugins/Runtime/Metasound/Source/MetasoundEditor/Private/MetasoundEditorGraphBuilder.cpp
@@ -755,7 +755,7 @@ namespace Metasound
 
 			Metadata.ClassName = FMetasoundFrontendClassName(FName(), InMetaSound.GetFName(), FName());
 			Metadata.Version.Major = 1;
-			Metadata.Version.Minor = 0;
+			Metadata.Version.Minor = 1;
 			Metadata.Type = EMetasoundFrontendClassType::Graph;
 			Metadata.Author = FText::FromString(InAuthor);
 
diff --git a/Engine/Plugins/Runtime/Metasound/Source/MetasoundFrontend/Public/MetasoundFrontendDocument.h b/Engine/Plugins/Runtime/Metasound/Source/MetasoundFrontend/Public/MetasoundFrontendDocument.h
--- a/Engine/Plugins/Runtime/Metasound/Source/MetasoundFrontend/Public/MetasoundFrontendDocument.h
+++ b/Engine/Plugins/Runtime/Metasound/Source/MetasoundFrontend/Public/MetasoundFrontendDocument.h
@@ -760,11 +760,6 @@ struct FMetasoundFrontendDocumentMetadata
 
 	UPROPERTY()
 	FMetasoundFrontendVersion Version;
-
-	FMetasoundFrontendDocumentMetadata()
-	{
-		Version.Number = { 1, 1 };
-	}
 };
 
 
