MetaSound UI/UX Composition/Preset Scaffolding
- Revive UMetaSound Asset type
- Revive MS Output generation
- Fix copy pasta'ed type name so MetaSound Source shows up properly in editor
- Clean-up to dynamic doc version reporting if required (show the new version number in log)
- Add versioning to update the RootGraphClass name & add DisplayName

MetaSound Registry Scaffolding Part 1 (From Phil Popp)

diff --git a/Engine/Plugins/Runtime/Metasound/Source/MetasoundEditor/Private/MetasoundAssetTypeActions.cpp b/Engine/Plugins/Runtime/Metasound/Source/MetasoundEditor/Private/MetasoundAssetTypeActions.cpp
--- a/Engine/Plugins/Runtime/Metasound/Source/MetasoundEditor/Private/MetasoundAssetTypeActions.cpp
+++ b/Engine/Plugins/Runtime/Metasound/Source/MetasoundEditor/Private/MetasoundAssetTypeActions.cpp
@@ -2,23 +2,116 @@
 
 #include "MetasoundAssetTypeActions.h"
 
+#include "ContentBrowserModule.h"
+#include "ContentBrowserMenuContexts.h"
+#include "IContentBrowserSingleton.h"
+#include "ObjectEditorUtils.h"
 #include "Metasound.h"
+#include "MetasoundAssetBase.h"
 #include "MetasoundSource.h"
 #include "MetasoundEditor.h"
+#include "MetasoundEditorGraphBuilder.h"
+#include "MetasoundFactory.h"
+#include "MetasoundUObjectRegistry.h"
+#include "ToolMenus.h"
+#include "ToolMenuSection.h"
 
 
-#define LOCTEXT_NAMESPACE "AssetTypeActions"
+#define LOCTEXT_NAMESPACE "MetaSoundEditor"
 
 namespace Metasound
 {
 	namespace Editor
 	{
-		UClass* FAssetTypeActions_Metasound::GetSupportedClass() const
+		namespace AssetTypeActionsPrivate
+		{
+			static const FText PresetLabel = LOCTEXT("MetaSound_CreatePreset", "Create MetaSound Preset");
+			static const FText PresetToolTip = LOCTEXT("MetaSound_CreatePresetToolTip", "Creates a MetaSound Preset using the selected MetaSound's root graph as a reference.");
+			static const FText SourcePresetLabel = LOCTEXT("MetaSoundSource_CreatePreset", "Create MetaSound Source Preset");
+			static const FText SourcePresetToolTip = LOCTEXT("MetaSoundSource_CreatePresetToolTip", "Creates a MetaSoundSource Preset using the selected MetaSound's root graph as a reference.");
+
+			template <typename TClass, typename TFactory>
+			void ExecuteCreatePreset(const FToolMenuContext& MenuContext)
+			{
+				UContentBrowserAssetContextMenuContext* Context = MenuContext.FindContext<UContentBrowserAssetContextMenuContext>();
+				if (!Context || Context->SelectedObjects.IsEmpty())
+				{
+					return;
+				}
+
+				FString PackagePath;
+				FString AssetName;
+
+				TWeakObjectPtr<UObject> ReferencedMetaSound = Context->SelectedObjects[0];
+				if (!ReferencedMetaSound.IsValid())
+				{
+					return;
+				}
+
+				FAssetToolsModule& AssetToolsModule = FModuleManager::Get().LoadModuleChecked<FAssetToolsModule>("AssetTools");
+				AssetToolsModule.Get().CreateUniqueAssetName(ReferencedMetaSound->GetOutermost()->GetName(), TEXT("_Preset"), PackagePath, AssetName);
+
+				TFactory* Factory = NewObject<TFactory>();
+				check(Factory);
+
+				Factory->ReferencedMetaSoundObject = ReferencedMetaSound.Get();
+
+				FContentBrowserModule& ContentBrowserModule = FModuleManager::Get().LoadModuleChecked<FContentBrowserModule>("ContentBrowser");
+				ContentBrowserModule.Get().CreateNewAsset(AssetName, FPackageName::GetLongPackagePath(PackagePath), TClass::StaticClass(), Factory);
+			}
+
+			template <typename TPresetClass, typename TFactory, typename TReferenceClass = TPresetClass>
+			void RegisterPresetAction(const TAttribute<FText>& InLabel, const TAttribute<FText>& InToolTip)
+			{
+				if (!UToolMenus::IsToolMenuUIEnabled())
+				{
+					return;
+				}
+
+				FString ClassName = TReferenceClass::StaticClass()->GetName();
+				const FString MenuName = TEXT("ContentBrowser.AssetContextMenu.") + ClassName;
+				UToolMenu* Menu = UToolMenus::Get()->ExtendMenu(*MenuName);
+				if (!ensure(Menu))
+				{
+					return;
+				}
+
+				FString PresetClassName = TPresetClass::StaticClass()->GetName();
+				const FString EntryName = FString::Printf(TEXT("%sTo%s_Preset"), *PresetClassName, *ClassName);
+				FToolMenuSection& Section = Menu->FindOrAddSection("GetAssetActions");
+				Section.AddDynamicEntry(*EntryName, FNewToolMenuSectionDelegate::CreateLambda([InPresetClassName = MoveTemp(PresetClassName), Label = InLabel, ToolTip = InToolTip](FToolMenuSection& InSection)
+				{
+					UContentBrowserAssetContextMenuContext* Context = InSection.FindContext<UContentBrowserAssetContextMenuContext>();
+					if (!Context || Context->SelectedObjects.IsEmpty())
+					{
+						return;
+					}
+
+					TWeakObjectPtr<UObject> MetaSoundObject = Context->SelectedObjects[0];
+					if (!MetaSoundObject.IsValid())
+					{
+						return;
+					}
+
+					// TODO: Make class icons for MetaSound types. For now just use SoundCue.
+					//const FString IconName = TEXT("ClassIcon.") + InPresetClassName;
+					const FString IconName = TEXT("ClassIcon.SoundCue");
+
+					const FSlateIcon Icon = FSlateIcon(FEditorStyle::GetStyleSetName(), *IconName);
+					const FToolMenuExecuteAction UIExecuteAction = FToolMenuExecuteAction::CreateStatic(&ExecuteCreatePreset<TPresetClass, TFactory>);
+
+					const FString PresetEntryName = InPresetClassName + TEXT("_CreatePreset");
+					InSection.AddMenuEntry(*PresetEntryName, Label, ToolTip, Icon, UIExecuteAction);
+				}));
+			}
+		} // namespace AssetTypeActionsPrivate
+
+		UClass* FAssetTypeActions_MetaSound::GetSupportedClass() const
 		{
 			return UMetaSound::StaticClass();
 		}
 
-		void FAssetTypeActions_Metasound::OpenAssetEditor(const TArray<UObject*>& InObjects, TSharedPtr<IToolkitHost> ToolkitHost)
+		void FAssetTypeActions_MetaSound::OpenAssetEditor(const TArray<UObject*>& InObjects, TSharedPtr<IToolkitHost> ToolkitHost)
 		{
 			const EToolkitMode::Type Mode = ToolkitHost.IsValid() ? EToolkitMode::WorldCentric : EToolkitMode::Standalone;
 
@@ -32,12 +125,18 @@ namespace Metasound
 			}
 		}
 
-		UClass* FAssetTypeActions_MetasoundSource::GetSupportedClass() const
+		void FAssetTypeActions_MetaSound::RegisterMenuActions()
+		{
+			using namespace AssetTypeActionsPrivate;
+			RegisterPresetAction<UMetaSound, UMetaSoundFactory>(PresetLabel, PresetToolTip);
+		}
+
+		UClass* FAssetTypeActions_MetaSoundSource::GetSupportedClass() const
 		{
 			return UMetaSoundSource::StaticClass();
 		}
 
-		void FAssetTypeActions_MetasoundSource::OpenAssetEditor(const TArray<UObject*>& InObjects, TSharedPtr<IToolkitHost> ToolkitHost)
+		void FAssetTypeActions_MetaSoundSource::OpenAssetEditor(const TArray<UObject*>& InObjects, TSharedPtr<IToolkitHost> ToolkitHost)
 		{
 			const EToolkitMode::Type Mode = ToolkitHost.IsValid() ? EToolkitMode::WorldCentric : EToolkitMode::Standalone;
 
@@ -50,8 +149,13 @@ namespace Metasound
 				}
 			}
 		}
+
+		void FAssetTypeActions_MetaSoundSource::RegisterMenuActions()
+		{
+			using namespace AssetTypeActionsPrivate;
+			// RegisterPresetAction<UMetaSoundSource, UMetaSoundSourceFactory>(SourcePresetLabel, SourcePresetToolTip);
+			// RegisterPresetAction<UMetaSound, UMetaSoundFactory, UMetaSoundSource>(PresetLabel, PresetToolTip);
+		}
 	} // namespace Editor
 } // namespace Metasound
-
-#undef LOCTEXT_NAMESPACE
-
+#undef LOCTEXT_NAMESPACE //MetaSoundEditor
diff --git a/Engine/Plugins/Runtime/Metasound/Source/MetasoundEditor/Private/MetasoundAssetTypeActions.h b/Engine/Plugins/Runtime/Metasound/Source/MetasoundEditor/Private/MetasoundAssetTypeActions.h
--- a/Engine/Plugins/Runtime/Metasound/Source/MetasoundEditor/Private/MetasoundAssetTypeActions.h
+++ b/Engine/Plugins/Runtime/Metasound/Source/MetasoundEditor/Private/MetasoundAssetTypeActions.h
@@ -12,27 +12,30 @@ namespace Metasound
 {
 	namespace Editor
 	{
-		class FAssetTypeActions_Metasound : public FAssetTypeActions_Base
+		class FAssetTypeActions_MetaSound : public FAssetTypeActions_Base
 		{
 		public:
 			// IAssetTypeActions Implementation
 			virtual FText GetName() const override { return NSLOCTEXT("AssetTypeActions", "AssetTypeActions_MetaSound", "MetaSound"); }
-			virtual FColor GetTypeColor() const override { return FColor(103, 214, 66); }
+			virtual FColor GetTypeColor() const override { return FColor(13, 55, 13); }
 			virtual UClass* GetSupportedClass() const override;
 			virtual uint32 GetCategories() override { return EAssetTypeCategories::Sounds; }
 			virtual void OpenAssetEditor(const TArray<UObject*>& InObjects, TSharedPtr<IToolkitHost> ToolkitHost) override;
+
+			static void RegisterMenuActions();
 		};
 
-		class FAssetTypeActions_MetasoundSource : public FAssetTypeActions_Base
+		class FAssetTypeActions_MetaSoundSource : public FAssetTypeActions_Base
 		{
 		public:
 			// IAssetTypeActions Implementation
-			virtual FText GetName() const override { return NSLOCTEXT("AssetTypeActions", "AssetTypeActions_MetaSoundSource", "MetaSound"); }
+			virtual FText GetName() const override { return NSLOCTEXT("AssetTypeActions", "AssetTypeActions_MetaSoundSource", "MetaSound Source"); }
 			virtual FColor GetTypeColor() const override { return FColor(103, 214, 66); }
 			virtual UClass* GetSupportedClass() const override;
 			virtual uint32 GetCategories() override { return EAssetTypeCategories::Sounds; }
 			virtual void OpenAssetEditor(const TArray<UObject*>& InObjects, TSharedPtr<IToolkitHost> ToolkitHost) override;
 
-	};
+			static void RegisterMenuActions();
+		};
 	} // namespace Editor
 } // namespace Metasound
diff --git a/Engine/Plugins/Runtime/Metasound/Source/MetasoundEditor/Private/MetasoundEditor.cpp b/Engine/Plugins/Runtime/Metasound/Source/MetasoundEditor/Private/MetasoundEditor.cpp
--- a/Engine/Plugins/Runtime/Metasound/Source/MetasoundEditor/Private/MetasoundEditor.cpp
+++ b/Engine/Plugins/Runtime/Metasound/Source/MetasoundEditor/Private/MetasoundEditor.cpp
@@ -134,7 +134,8 @@ namespace Metasound
 
 			bool IsRequired() const
 			{
-				return GetNodeHandle()->IsRequired();
+				const FMetasoundFrontendArchetype& Archetype = GetMetasoundAssetChecked().GetArchetype();
+				return GetNodeHandle()->IsRequired(Archetype);
 			}
 
 			// FEdGraphSchemaAction interface
@@ -345,7 +346,10 @@ namespace Metasound
 			FMetasoundAssetBase* MetasoundAsset = IMetasoundUObjectRegistry::Get().GetObjectAsAssetBase(Metasound);
 			check(MetasoundAsset);
 
-			if (Frontend::FVersionDocument().Transform(MetasoundAsset->GetDocumentHandle()))
+			Frontend::FDocumentHandle Document = MetasoundAsset->GetDocumentHandle();
+			FName AssetName = Metasound->GetFName();
+			FString AssetPath = Metasound->GetPathName();
+			if (Frontend::FVersionDocument(AssetName, AssetPath).Transform(Document))
 			{
 				Metasound->MarkPackageDirty();
 			}
@@ -643,13 +647,15 @@ namespace Metasound
 // 					}
 // 					ToolbarBuilder.EndSection();
 
-					ToolbarBuilder.BeginSection("Audition");
+					if (Metasound->IsA<USoundBase>())
 					{
-						ToolbarBuilder.AddToolBarButton(FEditorCommands::Get().Play);
-						ToolbarBuilder.AddToolBarButton(FEditorCommands::Get().Stop);
+						ToolbarBuilder.BeginSection("Audition");
+						{
+							ToolbarBuilder.AddToolBarButton(FEditorCommands::Get().Play);
+							ToolbarBuilder.AddToolBarButton(FEditorCommands::Get().Stop);
+						}
+						ToolbarBuilder.EndSection();
 					}
-					ToolbarBuilder.EndSection();
-
 					ToolbarBuilder.BeginSection("Utilities");
 					{
 						ToolbarBuilder.AddToolBarButton(
@@ -750,6 +756,8 @@ namespace Metasound
 			FMetasoundAssetBase* MetasoundAsset = IMetasoundUObjectRegistry::Get().GetObjectAsAssetBase(Metasound);
 			if (MetasoundAsset)
 			{
+				const FMetasoundFrontendArchetype& Archetype = MetasoundAsset->GetArchetype();
+
 				// TODO: Prompt OFD and provide path from user
 				const FString InputPath = FPaths::ProjectIntermediateDir() / TEXT("MetaSounds") + FPaths::ChangeExtension(Metasound->GetPathName(), FMetasoundAssetBase::FileExtension);
 				
@@ -760,11 +768,11 @@ namespace Metasound
 
 				if (Frontend::ImportJSONAssetToMetasound(InputPath, MetasoundDoc))
 				{
-					TArray<UClass*> ImportClasses = IMetasoundUObjectRegistry::Get().GetUClassesForArchetype(MetasoundDoc.Archetype.Name);
+					TArray<UClass*> ImportClasses = IMetasoundUObjectRegistry::Get().GetUClassesForArchetype(Archetype.Name);
 
 					if (ImportClasses.Num() < 1)
 					{
-						UE_LOG(LogMetaSound, Warning, TEXT("Cannot create UObject from MetaSound document. No UClass supports archetype \"%s\""), *MetasoundDoc.Archetype.Name.ToString());
+						UE_LOG(LogMetaSound, Warning, TEXT("Cannot create UObject from MetaSound document. No UClass supports archetype \"%s\""), *Archetype.Name.ToString());
 					}
 					else
 					{
@@ -773,11 +781,11 @@ namespace Metasound
 							for (UClass* Cls : ImportClasses)
 							{
 								// TODO: could do a modal dialog to give user choice of import type.
-								UE_LOG(LogMetaSound, Warning, TEXT("Duplicate UClass support archetype \"%s\" with UClass \"%s\""), *MetasoundDoc.Archetype.Name.ToString(), *Cls->GetName());
+								UE_LOG(LogMetaSound, Warning, TEXT("Duplicate UClass support archetype \"%s\" with UClass \"%s\""), *Archetype.Name.ToString(), *Cls->GetName());
 							}
 						}
 
-						IMetasoundUObjectRegistry::Get().NewObject(ImportClasses[0], MetasoundDoc, OutputPath);
+						IMetasoundUObjectRegistry::Get().NewObject(ImportClasses[0], MetasoundDoc, Archetype, OutputPath);
 					}
 				}
 				else
@@ -1133,9 +1141,6 @@ namespace Metasound
 					}));
 			}
 
-			FGraphAppearanceInfo AppearanceInfo;
-			AppearanceInfo.CornerText = LOCTEXT("AppearanceCornerText_MetaSound", "MetaSound");
-
 			SGraphEditor::FGraphEditorEvents InEvents;
 			InEvents.OnSelectionChanged = SGraphEditor::FOnSelectionChanged::CreateSP(this, &FEditor::OnSelectedNodesChanged);
 			InEvents.OnTextCommitted = FOnNodeTextCommitted::CreateSP(this, &FEditor::OnNodeTitleCommitted);
@@ -1146,11 +1151,11 @@ namespace Metasound
 
 			SAssignNew(MetasoundGraphEditor, SGraphEditor)
 				.AdditionalCommands(GraphEditorCommands)
-				.IsEditable(true)
-				.Appearance(AppearanceInfo)
-				.GraphToEdit(MetasoundAsset->GetGraph())
-				.GraphEvents(InEvents)
+				.Appearance(this, &FEditor::GetGraphAppearance)
 				.AutoExpandActionMenu(true)
+				.GraphEvents(InEvents)
+				.GraphToEdit(MetasoundAsset->GetGraph())
+				.IsEditable(this, &FEditor::IsGraphEditable)
 				.ShowGraphStateOverlay(false);
 
 			SAssignNew(PlayTimeWidget, STextBlock)
@@ -1159,6 +1164,20 @@ namespace Metasound
 				.ColorAndOpacity(FLinearColor(1, 1, 1, 0.30f));
 		}
 
+		FGraphAppearanceInfo FEditor::GetGraphAppearance() const
+		{
+			FGraphAppearanceInfo AppearanceInfo;
+
+			if (Metasound)
+			{
+				const FMetasoundAssetBase* MetasoundAsset = IMetasoundUObjectRegistry::Get().GetObjectAsAssetBase(Metasound);
+				check(MetasoundAsset);
+				AppearanceInfo.CornerText = MetasoundAsset->GetDisplayName();
+			}
+
+			return AppearanceInfo;
+		}
+
 		void FEditor::OnSelectedNodesChanged(const TSet<UObject*>& InSelectedNodes)
 		{
 			TArray<UObject*> Selection;
@@ -1272,6 +1291,11 @@ namespace Metasound
 		{
 			using namespace Frontend;
 
+			if (!IsGraphEditable())
+			{
+				return;
+			}
+
 			if (MetasoundInterfaceMenu.IsValid())
 			{
 				TArray<TSharedPtr<FEdGraphSchemaAction>> Actions;
@@ -1286,6 +1310,10 @@ namespace Metasound
 					UMetasoundEditorGraph& Graph = GetMetaSoundGraphChecked();
 					Graph.Modify();
 
+					FMetasoundAssetBase* MetasoundAsset = IMetasoundUObjectRegistry::Get().GetObjectAsAssetBase(Metasound);
+					check(MetasoundAsset);
+					const FMetasoundFrontendArchetype& Archetype = MetasoundAsset->GetArchetype();
+
 					TSharedPtr<FMetasoundGraphNodeSchemaAction> ActionToDelete;
 					for (const TSharedPtr<FEdGraphSchemaAction>& Action : Actions)
 					{
@@ -1295,7 +1323,7 @@ namespace Metasound
 							Frontend::FNodeHandle NodeHandle = MetasoundAction->GetNodeHandle();
 							if (ensure(NodeHandle->IsValid()))
 							{
-								if (NodeHandle->IsRequired())
+								if (NodeHandle->IsRequired(Archetype))
 								{
 									if (MetasoundGraphEditor.IsValid())
 									{
@@ -1416,6 +1444,11 @@ namespace Metasound
 
 		bool FEditor::CanDuplicateNodes() const
 		{
+			if (!IsGraphEditable())
+			{
+				return false;
+			}
+
 			// If any of the nodes can be duplicated then allow copying
 			const FGraphPanelSelectionSet& SelectedNodes = MetasoundGraphEditor->GetSelectedNodes();
 			for (FGraphPanelSelectionSet::TConstIterator SelectedIter(SelectedNodes); SelectedIter; ++SelectedIter)
@@ -1444,6 +1477,11 @@ namespace Metasound
 
 		bool FEditor::CanDeleteNodes() const
 		{
+			if (!IsGraphEditable())
+			{
+				return false;
+			}
+
 			if (MetasoundGraphEditor->GetSelectedNodes().IsEmpty())
 			{
 				return false;
@@ -1642,6 +1680,11 @@ namespace Metasound
 
 		bool FEditor::CanPasteNodes()
 		{
+			if (!IsGraphEditable())
+			{
+				return false;
+			}
+
 			UMetasoundEditorGraph& Graph = GetMetaSoundGraphChecked();
 			FPlatformApplicationMisc::ClipboardPaste(NodeTextToPaste);
 			if (FEdGraphUtilities::CanImportNodesFromText(&Graph, NodeTextToPaste))
@@ -1718,10 +1761,10 @@ namespace Metasound
 					Frontend::FNodeHandle NodeHandle = MetasoundAction->GetNodeHandle();
 					if (InNodeID == NodeHandle->GetID())
 					{
-						TArray<Frontend::FConstInputHandle> InputHandles = NodeHandle->GetConstInputs();
-						if (ensure(InputHandles.Num() == 1))
+						TArray<Frontend::FConstOutputHandle> OutputHandles = NodeHandle->GetConstOutputs();
+						if (ensure(OutputHandles.Num() == 1))
 						{
-							const FName ActionName = *InputHandles[0]->GetDisplayName().ToString();
+							const FName ActionName = *OutputHandles[0]->GetDisplayName().ToString();
 							MetasoundInterfaceMenu->SelectItemByName(ActionName, ESelectInfo::Direct, Action->GetSectionID());
 						}
 						break;
@@ -1864,6 +1907,18 @@ namespace Metasound
 			return FText::GetEmpty();
 		}
 
+		bool FEditor::IsGraphEditable() const
+		{
+			if (!Metasound)
+			{
+				return false;
+			}
+
+			FMetasoundAssetBase* MetasoundAsset = IMetasoundUObjectRegistry::Get().GetObjectAsAssetBase(Metasound);
+			check(MetasoundAsset);
+			return MetasoundAsset->GetRootGraphHandle()->GetGraphStyle().bIsGraphEditable;
+		}
+
 		TSharedRef<SWidget> FEditor::OnGetMenuSectionWidget(TSharedRef<SWidget> RowWidget, int32 InSectionID)
 		{
 			TWeakPtr<SWidget> WeakRowWidget = RowWidget;
@@ -1871,16 +1926,15 @@ namespace Metasound
 			FText AddNewText;
 			FName MetaDataTag;
 
-			switch (static_cast<ENodeSection>(InSectionID))
+			if (IsGraphEditable())
 			{
+				switch (static_cast<ENodeSection>(InSectionID))
+				{
 				case ENodeSection::Inputs:
 				{
 					AddNewText = LOCTEXT("AddNewInput", "Input");
 					MetaDataTag = "AddNewInput";
-
-					// TODO: Add back for outputs once reading outputs/Metasound
-					// composition is supported.
-					return CreateAddInputButton(InSectionID, AddNewText, MetaDataTag);
+					return CreateAddButton(InSectionID, AddNewText, MetaDataTag);
 				}
 				break;
 
@@ -1888,11 +1942,13 @@ namespace Metasound
 				{
 					AddNewText = LOCTEXT("AddNewOutput", "Output");
 					MetaDataTag = "AddNewOutput";
+					return CreateAddButton(InSectionID, AddNewText, MetaDataTag);
 				}
 				break;
 
 				default:
-				break;
+					break;
+				}
 			}
 
 			return SNullWidget::NullWidget;
@@ -1959,7 +2015,7 @@ namespace Metasound
 							{
 								Output->NameChanged.Remove(*NameChangeDelegate);
 							}
-							NameChangeDelegateHandles.FindOrAdd(NodeID) = Output->NameChanged.AddSP(this, &FEditor::OnInputNameChanged);
+							NameChangeDelegateHandles.FindOrAdd(NodeID) = Output->NameChanged.AddSP(this, &FEditor::OnOutputNameChanged);
 						}
 					}
 				}
@@ -1980,7 +2036,7 @@ namespace Metasound
 			return FReply::Handled();
 		}
 
-		TSharedRef<SWidget> FEditor::CreateAddInputButton(int32 InSectionID, FText AddNewText, FName MetaDataTag)
+		TSharedRef<SWidget> FEditor::CreateAddButton(int32 InSectionID, FText AddNewText, FName MetaDataTag)
 		{
 			return
 				SNew(SButton)
@@ -1993,7 +2049,7 @@ namespace Metasound
 				[
 					SNew(SImage)
 					.Image(FAppStyle::Get().GetBrush("Icons.PlusCircle"))
-				.ColorAndOpacity(FSlateColor::UseForeground())
+					.ColorAndOpacity(FSlateColor::UseForeground())
 				];
 		}
 	}
diff --git a/Engine/Plugins/Runtime/Metasound/Source/MetasoundEditor/Private/MetasoundEditor.h b/Engine/Plugins/Runtime/Metasound/Source/MetasoundEditor/Private/MetasoundEditor.h
--- a/Engine/Plugins/Runtime/Metasound/Source/MetasoundEditor/Private/MetasoundEditor.h
+++ b/Engine/Plugins/Runtime/Metasound/Source/MetasoundEditor/Private/MetasoundEditor.h
@@ -110,6 +110,9 @@ namespace Metasound
 			void PasteNodes(const FVector2D* InLocation = nullptr);
 			void PasteNodes(const FVector2D* InLocation, const FText& InTransactionText);
 
+			/* Whether the displayed graph is marked as editable */
+			bool IsGraphEditable() const;
+
 			int32 GetNumNodesSelected() const
 			{
 				return MetasoundGraphEditor->GetSelectedNodes().Num();
@@ -120,13 +123,14 @@ namespace Metasound
 
 			void OnMeterOutput(UMeterAnalyzer* InMeterAnalyzer, int32 ChannelIndex, const FMeterResults& MeterResults);
 
+
 		protected:
 			// Callbacks for action tree
 			bool CanRenameOnActionNode(TWeakPtr<FGraphActionNode> InSelectedNode) const;
 			bool CanAddNewElementToSection(int32 InSectionID) const;
 			void CollectAllActions(FGraphActionListBuilderBase& OutAllActions);
 			void CollectStaticSections(TArray<int32>& StaticSectionIDs);
-			TSharedRef<SWidget> CreateAddInputButton(int32 InSectionID, FText AddNewText, FName MetaDataTag);
+			TSharedRef<SWidget> CreateAddButton(int32 InSectionID, FText AddNewText, FName MetaDataTag);
 			FText GetFilterText() const;
 			bool HandleActionMatchesName(FEdGraphSchemaAction* InAction, const FName& InName) const;
 			FReply OnActionDragged(const TArray<TSharedPtr<FEdGraphSchemaAction>>& InActions, const FPointerEvent& MouseEvent);
@@ -140,6 +144,8 @@ namespace Metasound
 			/** Called when the selection changes in the GraphEditor */
 			void OnSelectedNodesChanged(const TSet<UObject*>& NewSelection);
 
+			FGraphAppearanceInfo GetGraphAppearance() const;
+
 			UMetasoundEditorGraph& GetMetaSoundGraphChecked();
 
 			/**
diff --git a/Engine/Plugins/Runtime/Metasound/Source/MetasoundEditor/Private/MetasoundEditorGraph.cpp b/Engine/Plugins/Runtime/Metasound/Source/MetasoundEditor/Private/MetasoundEditorGraph.cpp
--- a/Engine/Plugins/Runtime/Metasound/Source/MetasoundEditor/Private/MetasoundEditorGraph.cpp
+++ b/Engine/Plugins/Runtime/Metasound/Source/MetasoundEditor/Private/MetasoundEditorGraph.cpp
@@ -139,6 +139,17 @@ Metasound::Frontend::FConstNodeHandle UMetasoundEditorGraphVariable::GetConstNod
 	return GetNodeHandle();
 }
 
+bool UMetasoundEditorGraphVariable::IsRequired() const
+{
+	using namespace Metasound;
+
+	const FMetasoundAssetBase* MetasoundAsset = IMetasoundUObjectRegistry::Get().GetObjectAsAssetBase(GetOutermostObject());
+	check(MetasoundAsset);
+	const FMetasoundFrontendArchetype& Archetype = MetasoundAsset->GetArchetype();
+
+	return GetConstNodeHandle()->IsRequired(Archetype);
+}
+
 bool UMetasoundEditorGraphVariable::CanRename(const FText& InNewName, FText& OutError) const
 {
 	using namespace Metasound::Frontend;
@@ -149,15 +160,15 @@ bool UMetasoundEditorGraphVariable::CanRename(const FText& InNewName, FText& Out
 		return false;
 	}
 
-	FConstNodeHandle NodeHandle = GetConstNodeHandle();
-
-	if (NodeHandle->IsRequired())
+	if (IsRequired())
 	{
 		OutError = FText::Format(LOCTEXT("VariableRenameInvalid_VariableRequired", "{0} is required and cannot be renamed."), InNewName);
 		return false;
 	}
 
 	bool bIsNameValid = true;
+
+	FConstNodeHandle NodeHandle = GetConstNodeHandle();
 	FConstGraphHandle GraphHandle = NodeHandle->GetOwningGraph();
 	GraphHandle->IterateConstNodes([&](FConstNodeHandle NodeToCompare)
 	{
diff --git a/Engine/Plugins/Runtime/Metasound/Source/MetasoundEditor/Private/MetasoundEditorGraphBuilder.cpp b/Engine/Plugins/Runtime/Metasound/Source/MetasoundEditor/Private/MetasoundEditorGraphBuilder.cpp
--- a/Engine/Plugins/Runtime/Metasound/Source/MetasoundEditor/Private/MetasoundEditorGraphBuilder.cpp
+++ b/Engine/Plugins/Runtime/Metasound/Source/MetasoundEditor/Private/MetasoundEditorGraphBuilder.cpp
@@ -20,6 +20,7 @@
 #include "MetasoundFrontendQuerySteps.h"
 #include "MetasoundFrontendRegistries.h"
 #include "MetasoundFrontendSearchEngine.h"
+#include "MetasoundFrontendTransform.h"
 #include "MetasoundLiteral.h"
 #include "MetasoundUObjectRegistry.h"
 #include "Modules/ModuleManager.h"
@@ -52,6 +53,12 @@ namespace Metasound
 
 		namespace GraphBuilderPrivate
 		{
+			namespace NodeLayout
+			{
+				static const FVector2D BufferX { 250.0f, 0.0f };
+				static const FVector2D BufferY { 0.0f, 100.0f };
+			}
+
 			void DeleteNode(UObject& InMetasound, Frontend::FNodeHandle InNodeHandle)
 			{
 				if (InNodeHandle->IsValid())
@@ -63,6 +70,21 @@ namespace Metasound
 					}
 				}
 			}
+
+			void InitializeGraph(UObject& InMetaSound)
+			{
+				FMetasoundAssetBase* MetaSoundAsset = IMetasoundUObjectRegistry::Get().GetObjectAsAssetBase(&InMetaSound);
+				check(MetaSoundAsset);
+
+				// Initial graph generation is not something to be managed by the transaction stack, so don't track dirty state until
+				// after initial setup if necessary.
+				if (!MetaSoundAsset->GetGraph())
+				{
+					UMetasoundEditorGraph* Graph = NewObject<UMetasoundEditorGraph>(&InMetaSound, FName(), RF_Transactional);
+					Graph->Schema = UMetasoundEditorGraphSchema::StaticClass();
+					MetaSoundAsset->SetGraph(Graph);
+				}
+			}
 		}
 
 		UMetasoundEditorGraphExternalNode* FGraphBuilder::AddExternalNode(UObject& InMetasound, Frontend::FNodeHandle& InNodeHandle, FVector2D InLocation, bool bInSelectNewNode)
@@ -582,10 +604,6 @@ namespace Metasound
 			Description.VertexID = VertexID;
 
 			Frontend::FNodeHandle NodeHandle = GraphHandle->AddInputVertex(Description);
-
-			FMetasoundFrontendNodeStyle Style = NodeHandle->GetNodeStyle();
-			NodeHandle->SetNodeStyle(Style);
-
 			if (ensure(NodeHandle->IsValid()))
 			{
 				const FText DisplayName = FGraphBuilder::GenerateUniqueInputDisplayName(InMetasound);
@@ -625,7 +643,7 @@ namespace Metasound
 
 			Frontend::FNodeHandle NodeHandle = GraphHandle->AddOutputVertex(Description);
 
-			const FText DisplayName = FGraphBuilder::GenerateUniqueInputDisplayName(InMetasound);
+			const FText DisplayName = FGraphBuilder::GenerateUniqueOutputDisplayName(InMetasound);
 			GraphHandle->SetOutputDisplayName(Description.Name, DisplayName);
 
 			return NodeHandle;
@@ -747,29 +765,40 @@ namespace Metasound
 			}
 		}
 
-		void FGraphBuilder::InitMetaSound(UObject& InMetaSound, const FString& InAuthor)
+		void FGraphBuilder::InitMetaSound(UObject& InMetaSound, const FString& InAuthor, bool bConformToArchetype)
 		{
 			using namespace Frontend;
+			using namespace GraphBuilderPrivate;
 
 			FMetasoundFrontendClassMetadata Metadata;
 
-			Metadata.ClassName = FMetasoundFrontendClassName(FName(), InMetaSound.GetFName(), FName());
+			// 1. Set default class Metadata
+			Metadata.ClassName = FMetasoundFrontendClassName("GraphAsset", InMetaSound.GetFName(), *InMetaSound.GetPathName());
 			Metadata.Version.Major = 1;
-			Metadata.Version.Minor = 1;
+			Metadata.Version.Minor = 0;
+			Metadata.DisplayName = FText::FromString(InMetaSound.GetName());
 			Metadata.Type = EMetasoundFrontendClassType::Graph;
 			Metadata.Author = FText::FromString(InAuthor);
 
 			FMetasoundAssetBase* MetaSoundAsset = IMetasoundUObjectRegistry::Get().GetObjectAsAssetBase(&InMetaSound);
 			check(MetaSoundAsset);
-
 			MetaSoundAsset->SetMetadata(Metadata);
-			MetaSoundAsset->ConformDocumentToMetasoundArchetype();
-			FGraphHandle GraphHandle = MetaSoundAsset->GetRootGraphHandle();
 
-			// TODO: Space graph nodes in a procedural and readable way
+			// 2. Set default version Metadata
+			FDocumentHandle DocumentHandle = MetaSoundAsset->GetDocumentHandle();
+			FMetasoundFrontendDocumentMetadata DocMetadata = DocumentHandle->GetMetadata();
+			DocMetadata.Version.Number = FVersionDocument::GetMaxVersion();
+			DocumentHandle->SetMetadata(DocMetadata);
+
+			if (bConformToArchetype)
+			{
+				MetaSoundAsset->ConformDocumentToArchetype();
+			}
+
+			FGraphHandle GraphHandle = MetaSoundAsset->GetRootGraphHandle();
 			FVector2D InputNodeLocation = FVector2D::ZeroVector;
-			FVector2D OpNodeLocation = FVector2D(250.0f, 0.0f);
-			FVector2D OutputNodeLocation = FVector2D(500.0f, 0.0f);
+			FVector2D ExternalNodeLocation = InputNodeLocation + NodeLayout::BufferX;
+			FVector2D OutputNodeLocation = ExternalNodeLocation + NodeLayout::BufferX;
 
 			TArray<FNodeHandle> NodeHandles = GraphHandle->GetNodes();
 			for (FNodeHandle& NodeHandle : NodeHandles)
@@ -779,17 +808,17 @@ namespace Metasound
 				if (NodeType == EMetasoundFrontendClassType::Input)
 				{
 					NewLocation = InputNodeLocation;
-					InputNodeLocation.Y += 100.0f;
+					InputNodeLocation += NodeLayout::BufferY;
 				}
 				else if (NodeType == EMetasoundFrontendClassType::Output)
 				{
 					NewLocation = OutputNodeLocation;
-					OutputNodeLocation.Y += 100.0f;
+					OutputNodeLocation += NodeLayout::BufferY;
 				}
 				else
 				{
-					NewLocation = OpNodeLocation;
-					OpNodeLocation.Y += 100.0f;
+					NewLocation = ExternalNodeLocation;
+					ExternalNodeLocation += NodeLayout::BufferY;
 				}
 				FMetasoundFrontendNodeStyle Style = NodeHandle->GetNodeStyle();
 				Style.Display.Locations.Add(FGuid::NewGuid(), NewLocation);
@@ -797,6 +826,164 @@ namespace Metasound
 			}
 		}
 
+		void FGraphBuilder::InitMetaSoundPreset(const UObject& InMetaSoundReferenced, UObject& InMetaSoundPreset)
+		{
+			using namespace Frontend;
+			using namespace GraphBuilderPrivate;
+
+			const FMetasoundAssetBase* MetaSoundReferencedAsset = IMetasoundUObjectRegistry::Get().GetObjectAsAssetBase(&InMetaSoundReferenced);
+			check(MetaSoundReferencedAsset);
+
+			FConstGraphHandle ReferencedGraphHandle = MetaSoundReferencedAsset->GetRootGraphHandle();
+			if (!ensure(ReferencedGraphHandle->IsValid()))
+			{
+				return;
+			}
+
+			FMetasoundAssetBase* MetaSoundPresetAsset = IMetasoundUObjectRegistry::Get().GetObjectAsAssetBase(&InMetaSoundPreset);
+			check(MetaSoundPresetAsset);
+			FGraphHandle PresetGraphHandle = MetaSoundPresetAsset->GetRootGraphHandle();
+			if (!ensure(PresetGraphHandle->IsValid()))
+			{
+				return;
+			}
+
+			// Initialize EdGraph, but don't conform to archetype as this would provide the
+			// default interface (inputs/outputs), which will be added below when duplicating
+			// node's interface to the preset.
+			GraphBuilderPrivate::InitializeGraph(InMetaSoundPreset);
+
+			const FMetasoundFrontendArchetype& Archetype = MetaSoundPresetAsset->GetArchetype();
+
+			// 1a. Preset graph topology is considered "read-only", so mark as such
+			FMetasoundFrontendGraphStyle Style = PresetGraphHandle->GetGraphStyle();
+			Style.bIsGraphEditable = false;
+			PresetGraphHandle->SetGraphStyle(Style);
+
+			// 1b. Set archetype accordingly
+			const FMetasoundFrontendArchetype& RefArchetype = MetaSoundReferencedAsset->GetArchetype();
+			if (ensure(MetaSoundPresetAsset->IsArchetypeSupported(RefArchetype)))
+			{
+				MetaSoundPresetAsset->ConformDocumentToArchetype();
+			}
+
+			// 2. Generate referenced inputs & outputs
+			FVector2D InputNodeLocation = FVector2D::ZeroVector;
+			FVector2D ExternalNodeLocation = InputNodeLocation + NodeLayout::BufferX;
+			FVector2D OutputNodeLocation = ExternalNodeLocation + NodeLayout::BufferX;
+
+			ReferencedGraphHandle->IterateConstNodes([&](FConstNodeHandle InputNode)
+			{
+				// Required inputs are added via the prior step (archetype conformation).
+				// For these, find the corrisponding input & create an editor node only.
+				const bool bIsRequired = InputNode->IsRequired(Archetype);
+
+				const FText& DisplayName = InputNode->GetDisplayName();
+				const FText& Description = InputNode->GetDescription();
+
+				// Should only ever be one
+				ensure(InputNode->GetNumInputs() == 1);
+				InputNode->IterateConstInputs([&](FConstInputHandle InputHandle)
+				{
+					FNodeHandle NodeHandle = INodeController::GetInvalidHandle();
+					if (bIsRequired)
+					{
+						NodeHandle = PresetGraphHandle->GetInputNodeWithName(InputHandle->GetName());
+					}
+					else
+					{
+						FGuid InputID = ReferencedGraphHandle->GetVertexIDForInputVertex(InputHandle->GetName());
+						const FMetasoundFrontendLiteral DefaultLiteral = ReferencedGraphHandle->GetDefaultInput(InputID);
+
+						const FName DataTypeName = InputHandle->GetDataType();
+						NodeHandle = AddInputNodeHandle(InMetaSoundPreset, DataTypeName, Description, &DefaultLiteral);
+						NodeHandle->SetDisplayName(DisplayName);
+					}
+
+					if (ensure(NodeHandle->IsValid()))
+					{
+						AddInputNode(InMetaSoundPreset, NodeHandle, InputNodeLocation, false /* bSelectNewNode */);
+						InputNodeLocation += NodeLayout::BufferY;
+					}
+				});
+			},
+			EMetasoundFrontendClassType::Input);
+
+			ReferencedGraphHandle->IterateConstNodes([&](FConstNodeHandle OutputNode)
+			{
+				// Required inputs are added via the prior step (archetype conformation).
+				// For these, find the corrisponding input & create an editor node only.
+				const bool bIsRequired = OutputNode->IsRequired(Archetype);
+
+				const FText& DisplayName = OutputNode->GetDisplayName();
+				const FText& Description = OutputNode->GetDescription();
+
+				// Should only ever be one
+				ensure(OutputNode->GetNumOutputs() == 1);
+				OutputNode->IterateConstOutputs([&](FConstOutputHandle OutputHandle)
+				{
+					Frontend::FNodeHandle NodeHandle = Frontend::INodeController::GetInvalidHandle();
+					if (bIsRequired)
+					{
+						NodeHandle = PresetGraphHandle->GetOutputNodeWithName(OutputHandle->GetName());
+					}
+					else
+					{
+						const FName DataTypeName = OutputHandle->GetDataType();
+						NodeHandle = AddOutputNodeHandle(InMetaSoundPreset, DataTypeName, Description);
+						NodeHandle->SetDisplayName(DisplayName);
+					}
+
+					if (ensure(NodeHandle->IsValid()))
+					{
+						AddOutputNode(InMetaSoundPreset, NodeHandle, OutputNodeLocation, false /* bSelectNewNode */);
+						OutputNodeLocation += NodeLayout::BufferY;
+					}
+				});
+			},
+			EMetasoundFrontendClassType::Output);
+
+			// 3. Generate a referencing node to the given referenced MetaSound.
+			const FMetasoundFrontendClassMetadata ReferencedClassMetadata = ReferencedGraphHandle->GetGraphMetadata();
+			UMetasoundEditorGraphExternalNode* ReferencedNode = AddExternalNode(InMetaSoundPreset, ReferencedClassMetadata, ExternalNodeLocation, false /* bSelectNewNode */);
+			FNodeHandle ReferencedNodeHandle = ReferencedNode->GetNodeHandle();
+
+			// 4. Connect preset's respective inputs & outputs to generated referencing node.
+			PresetGraphHandle->IterateNodes([&](FNodeHandle InputNode)
+			{
+				const FString& Name = InputNode->GetNodeName();
+
+				// Should only ever be one
+				ensure(InputNode->GetNumOutputs() == 1);
+				InputNode->IterateOutputs([&](FOutputHandle OutputHandle)
+				{
+					const TArray<FInputHandle> ReferenceInputs = ReferencedNodeHandle->GetInputsWithVertexName(Name);
+					if (ensure(ReferenceInputs.Num() == 1))
+					{
+						ensure(ReferenceInputs[0]->Connect(*OutputHandle));
+					}
+				});
+			},
+			EMetasoundFrontendClassType::Input);
+
+			PresetGraphHandle->IterateNodes([&](FNodeHandle OutputNode)
+			{
+				const FString& Name = OutputNode->GetNodeName();
+
+				// Should only ever be one
+				ensure(OutputNode->GetNumInputs() == 1);
+				OutputNode->IterateInputs([&](FInputHandle InputHandle)
+				{
+					const TArray<FOutputHandle> ReferenceInputs = ReferencedNodeHandle->GetOutputsWithVertexName(Name);
+					if (ensure(ReferenceInputs.Num() == 1))
+					{
+						ensure(InputHandle->Connect(*ReferenceInputs[0]));
+					}
+				});
+			},
+			EMetasoundFrontendClassType::Output);
+		}
+
 		void FGraphBuilder::DeleteVariableNodeHandle(UMetasoundEditorGraphVariable& InVariable)
 		{
 			using namespace Frontend;
@@ -1168,20 +1355,13 @@ namespace Metasound
 		{
 			using namespace Frontend;
 
-			FMetasoundAssetBase* MetasoundAsset = IMetasoundUObjectRegistry::Get().GetObjectAsAssetBase(&InMetasound);
-			check(MetasoundAsset);
-
-			// Initial graph generation is not something to be managed by the transaction stack, so don't track dirty state until
-			// after initial setup if necessary.
-			if (!MetasoundAsset->GetGraph())
-			{
-				UMetasoundEditorGraph* Graph = NewObject<UMetasoundEditorGraph>(&InMetasound, FName(), RF_Transactional);
-				Graph->Schema = UMetasoundEditorGraphSchema::StaticClass();
-				MetasoundAsset->SetGraph(Graph);
-			}
+			GraphBuilderPrivate::InitializeGraph(InMetasound);
 
 			bool bIsEditorGraphDirty = SynchronizeVariables(InMetasound);
 
+			FMetasoundAssetBase* MetasoundAsset = IMetasoundUObjectRegistry::Get().GetObjectAsAssetBase(&InMetasound);
+			check(MetasoundAsset);
+
 			// Get all external nodes from Frontend graph.  Input and output references will only be added/synchronized
 			// if required when synchronizing connections (as they are not required to inhabit editor graph).
 			FGraphHandle GraphHandle = MetasoundAsset->GetRootGraphHandle();
@@ -1382,37 +1562,45 @@ namespace Metasound
 				}
 			}
 
+			const EMetasoundFrontendClassType ClassType = InNode->GetClassMetadata().Type;
+
 			// Only add input pins of the node is not an input node. Input nodes
 			// have their input pins hidden because they cannot be connected internal
 			// to the graph.
-			if (EMetasoundFrontendClassType::Input != InNode->GetClassMetadata().Type)
+			if (EMetasoundFrontendClassType::Input != ClassType)
 			{
-				bIsNodeDirty |= !InputHandles.IsEmpty();
-				InputHandles = InNode->GetOutputStyle().SortDefaults(InputHandles);
-				for (Frontend::FInputHandle& InputHandle : InputHandles)
+				if (!InputHandles.IsEmpty())
 				{
-					if (bLogChanges)
+					bIsNodeDirty = true;
+					InputHandles = InNode->GetInputStyle().SortDefaults(InputHandles);
+					for (Frontend::FInputHandle& InputHandle : InputHandles)
 					{
-						UE_LOG(LogMetasoundEditor, Display, TEXT("Synchronizing Node '%s' Pins: Adding missing Editor Input Pin '%s'"), *InNode->GetDisplayName().ToString(), *InputHandle->GetDisplayName().ToString());
+						if (bLogChanges)
+						{
+							UE_LOG(LogMetasoundEditor, Display, TEXT("Synchronizing Node '%s' Pins: Adding missing Editor Input Pin '%s'"), *InNode->GetDisplayName().ToString(), *InputHandle->GetDisplayName().ToString());
+						}
+						AddPinToNode(InEditorNode, InputHandle);
 					}
-					AddPinToNode(InEditorNode, InputHandle);
 				}
 			}
 
 			// Only add output pins of the node is not an output node. Output nodes
 			// have their output pins hidden because they cannot be connected internal
 			// to the graph.
-			if (EMetasoundFrontendClassType::Output != InNode->GetClassMetadata().Type)
+			if (EMetasoundFrontendClassType::Output != ClassType)
 			{
-				bIsNodeDirty |= !OutputHandles.IsEmpty();
-				OutputHandles = InNode->GetOutputStyle().SortDefaults(OutputHandles);
-				for (Frontend::FOutputHandle& OutputHandle : OutputHandles)
+				if (!OutputHandles.IsEmpty())
 				{
-					if (bLogChanges)
+					bIsNodeDirty = true;
+					OutputHandles = InNode->GetOutputStyle().SortDefaults(OutputHandles);
+					for (Frontend::FOutputHandle& OutputHandle : OutputHandles)
 					{
-						UE_LOG(LogMetasoundEditor, Display, TEXT("Synchronizing Node '%s' Pins: Adding missing Editor Output Pin '%s'"), *InNode->GetDisplayName().ToString(), *OutputHandle->GetDisplayName().ToString());
+						if (bLogChanges)
+						{
+							UE_LOG(LogMetasoundEditor, Display, TEXT("Synchronizing Node '%s' Pins: Adding missing Editor Output Pin '%s'"), *InNode->GetDisplayName().ToString(), *OutputHandle->GetDisplayName().ToString());
+						}
+						AddPinToNode(InEditorNode, OutputHandle);
 					}
-					AddPinToNode(InEditorNode, OutputHandle);
 				}
 			}
 
diff --git a/Engine/Plugins/Runtime/Metasound/Source/MetasoundEditor/Private/MetasoundEditorGraphBuilder.h b/Engine/Plugins/Runtime/Metasound/Source/MetasoundEditor/Private/MetasoundEditorGraphBuilder.h
--- a/Engine/Plugins/Runtime/Metasound/Source/MetasoundEditor/Private/MetasoundEditorGraphBuilder.h
+++ b/Engine/Plugins/Runtime/Metasound/Source/MetasoundEditor/Private/MetasoundEditorGraphBuilder.h
@@ -132,7 +132,11 @@ namespace Metasound
 			static FLinearColor GetPinCategoryColor(const FEdGraphPinType& PinType);
 
 			// Initializes MetaSound with default inputs & outputs.
-			static void InitMetaSound(UObject& InMetaSound, const FString& InAuthor);
+			static void InitMetaSound(UObject& InMetaSound, const FString& InAuthor, bool bConformToArchetype = true);
+
+			// Initializes a MetaSound Preset using the provided ReferencedMetaSound asset's
+			// root graph as the sole, encapsulated topology.
+			static void InitMetaSoundPreset(const UObject& InMetaSoundReferenced, UObject& InMetaSoundPreset);
 
 			// Rebuilds all editor node pins based on the provided node handle's class definition.
 			static void RebuildNodePins(UMetasoundEditorGraphNode& InGraphNode);
diff --git a/Engine/Plugins/Runtime/Metasound/Source/MetasoundEditor/Private/MetasoundEditorGraphNode.cpp b/Engine/Plugins/Runtime/Metasound/Source/MetasoundEditor/Private/MetasoundEditorGraphNode.cpp
--- a/Engine/Plugins/Runtime/Metasound/Source/MetasoundEditor/Private/MetasoundEditorGraphNode.cpp
+++ b/Engine/Plugins/Runtime/Metasound/Source/MetasoundEditor/Private/MetasoundEditorGraphNode.cpp
@@ -472,7 +472,11 @@ FGuid UMetasoundEditorGraphOutputNode::GetNodeID() const
 
 bool UMetasoundEditorGraphOutputNode::CanUserDeleteNode() const
 {
-	return !GetNodeHandle()->IsRequired();
+	const FMetasoundAssetBase* ConstMetasoundAsset = Metasound::IMetasoundUObjectRegistry::Get().GetObjectAsAssetBase(&GetMetasoundChecked());
+	FMetasoundAssetBase* MetasoundAsset = const_cast<FMetasoundAssetBase*>(ConstMetasoundAsset);
+	const FMetasoundFrontendArchetype& Archetype = MetasoundAsset->GetArchetype();
+
+	return !GetNodeHandle()->IsRequired(Archetype);
 }
 
 void UMetasoundEditorGraphOutputNode::SetNodeID(FGuid InNodeID)
diff --git a/Engine/Plugins/Runtime/Metasound/Source/MetasoundEditor/Private/MetasoundEditorModule.cpp b/Engine/Plugins/Runtime/Metasound/Source/MetasoundEditor/Private/MetasoundEditorModule.cpp
--- a/Engine/Plugins/Runtime/Metasound/Source/MetasoundEditor/Private/MetasoundEditorModule.cpp
+++ b/Engine/Plugins/Runtime/Metasound/Source/MetasoundEditor/Private/MetasoundEditorModule.cpp
@@ -1,6 +1,7 @@
 // Copyright Epic Games, Inc. All Rights Reserved.
 #include "MetasoundEditorModule.h"
 
+#include "AssetRegistryModule.h"
 #include "AssetTypeActions_Base.h"
 #include "Brushes/SlateImageBrush.h"
 #include "CoreMinimal.h"
@@ -192,6 +193,18 @@ namespace Metasound
 				DataTypeInfo.Emplace(InRegistryInfo.DataTypeName, Editor::FEditorDataType(MoveTemp(PinType), InRegistryInfo));
 			}
 
+			static bool GetClassAssetData(const FAssetData& InAssetData, FMetasoundFrontendClassAssetTags& OutTags)
+			{
+				bool bSuccess = false;
+				bSuccess |= InAssetData.GetTagValue(GET_MEMBER_NAME_CHECKED(FMetasoundFrontendClassAssetTags, ID), OutTags.ID);
+				bSuccess |= InAssetData.GetTagValue(GET_MEMBER_NAME_CHECKED(FMetasoundFrontendClassAssetTags, Namespace), OutTags.Namespace);
+				bSuccess |= InAssetData.GetTagValue(GET_MEMBER_NAME_CHECKED(FMetasoundFrontendClassAssetTags, Name), OutTags.Name);
+				bSuccess |= InAssetData.GetTagValue(GET_MEMBER_NAME_CHECKED(FMetasoundFrontendClassAssetTags, MajorVersion), OutTags.MajorVersion);
+				bSuccess |= InAssetData.GetTagValue(GET_MEMBER_NAME_CHECKED(FMetasoundFrontendClassAssetTags, MinorVersion), OutTags.MinorVersion);
+
+				return bSuccess;
+			}
+
 			void RegisterCoreDataTypes()
 			{
 				TArray<FName> DataTypeNames = Frontend::GetAllAvailableDataTypes();
@@ -296,12 +309,56 @@ namespace Metasound
 				}
 			}
 
+			//void AddOrUpdateAssetTags(const FAssetData& InAssetData)
+			//{
+			//	FMetasoundFrontendClassAssetTags TagData;
+			//	if (GetClassAssetData(InAssetData, TagData))
+			//	{
+			//		AssetTagData.FindOrAdd(TagData.ID) = MoveTemp(TagData);
+			//	}
+			//}
+
+			//void RemoveAssetTags(const FAssetData& InAssetData)
+			//{
+			//	FMetasoundFrontendClassAssetTags TagData;
+			//	if (!GetClassAssetData(InAssetData, TagData))
+			//	{
+			//		AssetTagData.Remove(TagData.ID);
+			//	}
+			//}
+
+			void RegisterAssetData()
+			{
+				//FAssetRegistryModule& AssetRegistryModule = FModuleManager::LoadModuleChecked<FAssetRegistryModule>("AssetRegistry");
+
+				//FARFilter Filter;
+				//Filter.ClassNames.Add(UMetaSound::StaticClass()->GetFName());
+				//Filter.ClassNames.Add(UMetaSoundSource::StaticClass()->GetFName());
+
+				//AssetRegistryModule.Get().EnumerateAssets(Filter, [this](const FAssetData& AssetData) { AddOrUpdateAssetTags(AssetData); return true; });
+				//AssetRegistryModule.Get().OnAssetAdded().AddRaw(this, &FModule::AddOrUpdateAssetTags);
+				//AssetRegistryModule.Get().OnAssetUpdated().AddRaw(this, &FModule::AddOrUpdateAssetTags);
+				//AssetRegistryModule.Get().OnAssetRemoved().AddRaw(this, &FModule::RemoveAssetTags);
+			}
+
+			void UnregisterAssetData()
+			{
+				//AssetTagData.Reset();
+
+				//if (FAssetRegistryModule* AssetRegistryModule = static_cast<FAssetRegistryModule*>(FModuleManager::Get().GetModule("AssetRegistry")))
+				//{
+				//	AssetRegistryModule->Get().OnAssetAdded().RemoveAll(this);
+				//	AssetRegistryModule->Get().OnAssetUpdated().RemoveAll(this);
+				//	AssetRegistryModule->Get().OnAssetRemoved().RemoveAll(this);
+				//}
+			}
+
 			virtual void StartupModule() override
 			{
 				// Register Metasound asset type actions
 				IAssetTools& AssetTools = FModuleManager::LoadModuleChecked<FAssetToolsModule>(AssetToolName).Get();
-				AddAssetAction<FAssetTypeActions_Metasound>(AssetTools, AssetActions);
-				AddAssetAction<FAssetTypeActions_MetasoundSource>(AssetTools, AssetActions);
+				//AddAssetAction<FAssetTypeActions_MetaSound>(AssetTools, AssetActions);
+				AddAssetAction<FAssetTypeActions_MetaSoundSource>(AssetTools, AssetActions);
 
 				FPropertyEditorModule& PropertyModule = FModuleManager::LoadModuleChecked<FPropertyEditorModule>("PropertyEditor");
 				
@@ -335,6 +392,7 @@ namespace Metasound
 
 				StyleSet = MakeShared<FSlateStyle>();
 
+				RegisterAssetData();
 				RegisterCoreDataTypes();
 				RegisterNodeInputClasses();
 
@@ -354,6 +412,9 @@ namespace Metasound
 					NSLOCTEXT("MetaSoundsEditor", "MetaSoundEditorSettingsDescription", "Customize MetaSound Editor."),
 					GetMutableDefault<UMetasoundEditorSettings>()
 				);
+
+				//FAssetTypeActions_MetaSound::RegisterMenuActions();
+				//FAssetTypeActions_MetaSoundSource::RegisterMenuActions();
 			}
 
 			virtual void ShutdownModule() override
@@ -389,10 +450,14 @@ namespace Metasound
 					GraphPanelPinFactory.Reset();
 				}
 
+				UnregisterAssetData();
+
 				AssetActions.Reset();
 				DataTypeInfo.Reset();
 			}
 
+			//TMap<FGuid, FMetasoundFrontendClassAssetTags> AssetTagData;
+
 			TArray<TSharedPtr<FAssetTypeActions_Base>> AssetActions;
 			TMap<FName, FEditorDataType> DataTypeInfo;
 			TMap<EMetasoundFrontendLiteralType, const TSubclassOf<UMetasoundEditorGraphInputLiteral>> NodeInputClassRegistry;
diff --git a/Engine/Plugins/Runtime/Metasound/Source/MetasoundEditor/Private/MetasoundFactory.cpp b/Engine/Plugins/Runtime/Metasound/Source/MetasoundEditor/Private/MetasoundFactory.cpp
--- a/Engine/Plugins/Runtime/Metasound/Source/MetasoundEditor/Private/MetasoundFactory.cpp
+++ b/Engine/Plugins/Runtime/Metasound/Source/MetasoundEditor/Private/MetasoundFactory.cpp
@@ -4,42 +4,83 @@
 #include "Kismet/KismetSystemLibrary.h"
 #include "Metasound.h"
 #include "MetasoundEditorGraphBuilder.h"
+#include "MetasoundFrontendTransform.h"
 #include "MetasoundSource.h"
 
 
-// TODO: Re-enable and potentially rename once composition is supported
-// UMetasoundFactory::UMetasoundFactory(const FObjectInitializer& ObjectInitializer)
-// 	: Super(ObjectInitializer)
-// {
-// 	SupportedClass = UMetaSound::StaticClass();
-// 
-// 	bCreateNew = true;
-// 	bEditorImport = false;
-// 	bEditAfterNew = true;
-// }
-// 
-// UObject* UMetasoundFactory::FactoryCreateNew(UClass* InClass, UObject* InParent, FName Name, EObjectFlags Flags, UObject* InContext, FFeedbackContext* InFeedbackContext)
-// {
-// 	UMetaSound* Metasound = NewObject<UMetaSound>(InParent, Name, Flags);
-//	FGraphBuilder::InitMetasound(Metasound);
-// }
-
-UMetasoundSourceFactory::UMetasoundSourceFactory(const FObjectInitializer& ObjectInitializer)
-	: Super(ObjectInitializer)
+namespace Metasound
 {
-	SupportedClass = UMetaSoundSource::StaticClass();
+	namespace FactoryPrivate
+	{
+		template <typename T>
+		T* CreateNewMetaSoundObject(UObject* InParent, FName InName, EObjectFlags InFlags, UObject* InReferencedMetaSoundObject)
+		{
+			using namespace Editor;
+
+			T* MetaSoundObject = NewObject<T>(InParent, InName, InFlags);
+			check(MetaSoundObject);
 
+			if (InReferencedMetaSoundObject)
+			{
+				// Ensure referenced MetaSound is up-to-date
+				FMetasoundAssetBase* RefMetasoundAsset = IMetasoundUObjectRegistry::Get().GetObjectAsAssetBase(InReferencedMetaSoundObject);
+				check(RefMetasoundAsset);
+				FName Name = InReferencedMetaSoundObject->GetFName();
+				FString Path = InReferencedMetaSoundObject->GetPathName();
+				if (Frontend::FVersionDocument(Name, Path).Transform(RefMetasoundAsset->GetDocumentHandle()))
+				{
+					InReferencedMetaSoundObject->MarkPackageDirty();
+				}
+
+				FGraphBuilder::InitMetaSound(*MetaSoundObject, UKismetSystemLibrary::GetPlatformUserName(), false /* bConformToArchetype */);
+				FGraphBuilder::InitMetaSoundPreset(*InReferencedMetaSoundObject, *MetaSoundObject);
+			}
+			else
+			{
+				FGraphBuilder::InitMetaSound(*MetaSoundObject, UKismetSystemLibrary::GetPlatformUserName());
+			}
+
+			return MetaSoundObject;
+		}
+	}
+}
+
+UMetaSoundBaseFactory::UMetaSoundBaseFactory(const FObjectInitializer& ObjectInitializer)
+	: Super(ObjectInitializer)
+{
 	bCreateNew = true;
 	bEditorImport = false;
 	bEditAfterNew = true;
 }
 
-UObject* UMetasoundSourceFactory::FactoryCreateNew(UClass* InClass, UObject* InParent, FName Name, EObjectFlags Flags, UObject* InContext, FFeedbackContext* InFeedbackContext)
+UMetaSoundFactory::UMetaSoundFactory(const FObjectInitializer& ObjectInitializer)
+	: Super(ObjectInitializer)
+{
+	SupportedClass = UMetaSound::StaticClass();
+}
+
+UObject* UMetaSoundFactory::FactoryCreateNew(UClass* InClass, UObject* InParent, FName InName, EObjectFlags InFlags, UObject* InContext, FFeedbackContext* InFeedbackContext)
+{
+	using namespace Metasound::FactoryPrivate;
+	return CreateNewMetaSoundObject<UMetaSound>(InParent, InName, InFlags, ReferencedMetaSoundObject);
+}
+
+UMetaSoundSourceFactory::UMetaSoundSourceFactory(const FObjectInitializer& ObjectInitializer)
+	: Super(ObjectInitializer)
 {
-	using namespace Metasound::Editor;
+	SupportedClass = UMetaSoundSource::StaticClass();
+}
+
+UObject* UMetaSoundSourceFactory::FactoryCreateNew(UClass* InClass, UObject* InParent, FName InName, EObjectFlags InFlags, UObject* InContext, FFeedbackContext* InFeedbackContext)
+{
+	using namespace Metasound::FactoryPrivate;
+	UMetaSoundSource* NewSource = CreateNewMetaSoundObject<UMetaSoundSource>(InParent, InName, InFlags, ReferencedMetaSoundObject);
+
+	// Copy over referenced fields that are specific to sources
+	if (UMetaSoundSource* ReferencedMetaSound = Cast<UMetaSoundSource>(ReferencedMetaSoundObject))
+	{
+		NewSource->OutputFormat = ReferencedMetaSound->OutputFormat;
+	}
 
-	UMetaSoundSource* MetaSoundSource = NewObject<UMetaSoundSource>(InParent, Name, Flags);
-	check(MetaSoundSource);
-	FGraphBuilder::InitMetaSound(*MetaSoundSource, UKismetSystemLibrary::GetPlatformUserName());
-	return MetaSoundSource;
+	return NewSource;
 }
diff --git a/Engine/Plugins/Runtime/Metasound/Source/MetasoundEditor/Private/MetasoundNodeDetailCustomization.cpp b/Engine/Plugins/Runtime/Metasound/Source/MetasoundEditor/Private/MetasoundNodeDetailCustomization.cpp
--- a/Engine/Plugins/Runtime/Metasound/Source/MetasoundEditor/Private/MetasoundNodeDetailCustomization.cpp
+++ b/Engine/Plugins/Runtime/Metasound/Source/MetasoundEditor/Private/MetasoundNodeDetailCustomization.cpp
@@ -652,7 +652,7 @@ namespace Metasound
 			AddDataTypeSelector(DetailLayout, VariableCustomizationPrivate::DataTypeNameText, GraphVariable, bIsRequired);
 
 			CategoryBuilder.AddCustomRow(LOCTEXT("InputPrivate", "Private"))
-			.Visibility(TAttribute<EVisibility>(this, &FMetasoundInputDetailCustomization::ExposePrivateVisibility))
+			.Visibility(TAttribute<EVisibility>(EVisibility::Hidden))
 			.NameContent()
 			[
 				SNew(STextBlock)
@@ -932,7 +932,7 @@ namespace Metasound
 			AddDataTypeSelector(DetailLayout, VariableCustomizationPrivate::DataTypeNameText, GraphVariable, bIsRequired);
 
 			CategoryBuilder.AddCustomRow(LOCTEXT("OutputPrivate", "Private"))
-			.Visibility(TAttribute<EVisibility>(this, &FMetasoundOutputDetailCustomization::ExposePrivateVisibility))
+			.Visibility(TAttribute<EVisibility>(EVisibility::Hidden))
 			.NameContent()
 			[
 				SNew(STextBlock)
diff --git a/Engine/Plugins/Runtime/Metasound/Source/MetasoundEditor/Private/MetasoundNodeDetailCustomization.h b/Engine/Plugins/Runtime/Metasound/Source/MetasoundEditor/Private/MetasoundNodeDetailCustomization.h
--- a/Engine/Plugins/Runtime/Metasound/Source/MetasoundEditor/Private/MetasoundNodeDetailCustomization.h
+++ b/Engine/Plugins/Runtime/Metasound/Source/MetasoundEditor/Private/MetasoundNodeDetailCustomization.h
@@ -167,8 +167,7 @@ namespace Metasound
 			{
 				if (GraphVariable.IsValid())
 				{
-					Frontend::FConstNodeHandle NodeHandle = GraphVariable->GetConstNodeHandle();
-					return NodeHandle->IsRequired();
+					return GraphVariable->IsRequired();
 				}
 
 				return true;
@@ -240,16 +239,6 @@ namespace Metasound
 					NodeHandle->SetNodeStyle(NodeStyle);
 				}
 			}
-
-			EVisibility ExposePrivateVisibility() const
-			{
-				if (IsRequired())
-				{
-					return EVisibility::Collapsed;
-				}
-
-				return EVisibility::Visible;
-			}
 		};
 
 		class FMetasoundInputDetailCustomization : public TMetasoundVariableDetailCustomization<UMetasoundEditorGraphInput>
diff --git a/Engine/Plugins/Runtime/Metasound/Source/MetasoundEditor/Public/MetasoundEditorGraph.h b/Engine/Plugins/Runtime/Metasound/Source/MetasoundEditor/Public/MetasoundEditorGraph.h
--- a/Engine/Plugins/Runtime/Metasound/Source/MetasoundEditor/Public/MetasoundEditorGraph.h
+++ b/Engine/Plugins/Runtime/Metasound/Source/MetasoundEditor/Public/MetasoundEditorGraph.h
@@ -61,6 +61,7 @@ public:
 	Metasound::Frontend::FConstNodeHandle GetConstNodeHandle() const;
 
 	bool CanRename(const FText& InNewName, FText& OutError) const;
+	bool IsRequired() const;
 	TArray<UMetasoundEditorGraphNode*> GetNodes() const;
 	void SetDataType(FName InNewType);
 	void SetDisplayName(const FText& InNewName);
@@ -183,6 +184,6 @@ public:
 
 	bool RemoveVariable(UMetasoundEditorGraphVariable& InVariable);
 
-	friend class UMetasoundFactory;
-	friend class UMetasoundSourceFactory;
+	friend class UMetaSoundFactory;
+	friend class UMetaSoundSourceFactory;
 };
diff --git a/Engine/Plugins/Runtime/Metasound/Source/MetasoundEditor/Public/MetasoundEditorGraphInputNodes.h b/Engine/Plugins/Runtime/Metasound/Source/MetasoundEditor/Public/MetasoundEditorGraphInputNodes.h
--- a/Engine/Plugins/Runtime/Metasound/Source/MetasoundEditor/Public/MetasoundEditorGraphInputNodes.h
+++ b/Engine/Plugins/Runtime/Metasound/Source/MetasoundEditor/Public/MetasoundEditorGraphInputNodes.h
@@ -452,7 +452,7 @@ public:
 
 	virtual void SetFromLiteral(const FMetasoundFrontendLiteral& InLiteral) override
 	{
-// 		Metasound::Editor::InputPrivate::ConvertLiteral<UObject*>(InLiteral, Default.Object);
+		ensure(InLiteral.TryGet(Default.Object));
 	}
 
 	virtual void UpdatePreviewInstance(const Metasound::FVertexKey& InParameterName, TScriptInterface<IAudioParameterInterface>& InParameterInterface) const override
@@ -489,7 +489,10 @@ public:
 
 	virtual void SetFromLiteral(const FMetasoundFrontendLiteral& InLiteral) override
 	{
-// 		Metasound::Editor::InputPrivate::ConvertLiteralToArray<UObject*, FMetasoundEditorGraphInputObjectRef>(InLiteral, Default);
+		TArray<UObject*> ObjectArray;
+		ensure(InLiteral.TryGet(ObjectArray));
+
+		Algo::Transform(ObjectArray, Default, [](UObject* InValue) { return FMetasoundEditorGraphInputObjectRef { InValue }; });
 	}
 
 	virtual void UpdatePreviewInstance(const Metasound::FVertexKey& InParameterName, TScriptInterface<IAudioParameterInterface>& InParameterInterface) const override
diff --git a/Engine/Plugins/Runtime/Metasound/Source/MetasoundEditor/Public/MetasoundFactory.h b/Engine/Plugins/Runtime/Metasound/Source/MetasoundEditor/Public/MetasoundFactory.h
--- a/Engine/Plugins/Runtime/Metasound/Source/MetasoundEditor/Public/MetasoundFactory.h
+++ b/Engine/Plugins/Runtime/Metasound/Source/MetasoundEditor/Public/MetasoundFactory.h
@@ -8,23 +8,36 @@
 
 #include "MetasoundFactory.generated.h"
 
-// TODO: Re-enable and potentially rename once composition is supported
-// UCLASS(hidecategories=Object, MinimalAPI)
-// class UMetasoundFactory : public UFactory
-// {
-// 	GENERATED_UCLASS_BODY()
-// 
-// 	//~ Begin UFactory Interface
-// 	virtual UObject* FactoryCreateNew(UClass* InClass, UObject* InParent, FName Name, EObjectFlags Flags, UObject* InContext, FFeedbackContext* InFeedbackContext) override;
-// 	//~ Begin UFactory Interface
-// };
+class UMetaSound;
+class UMetaSoundSource;
+
+UCLASS(abstract)
+class UMetaSoundBaseFactory : public UFactory
+{
+	GENERATED_UCLASS_BODY()
+
+	// Set to initialize MetaSound type as Preset,
+	// using the provided MetaSound as a Reference
+	UPROPERTY(Transient)
+	TObjectPtr<UObject> ReferencedMetaSoundObject;
+};
+
+UCLASS(hidecategories=Object, MinimalAPI)
+class UMetaSoundFactory : public UMetaSoundBaseFactory
+{
+	GENERATED_UCLASS_BODY()
+ 
+	//~ Begin UFactory Interface
+	virtual UObject* FactoryCreateNew(UClass* InClass, UObject* InParent, FName InName, EObjectFlags InFlags, UObject* InContext, FFeedbackContext* InFeedbackContext) override;
+	//~ Begin UFactory Interface
+ };
 
 UCLASS(hidecategories = Object, MinimalAPI)
-class UMetasoundSourceFactory : public UFactory
+class UMetaSoundSourceFactory : public UMetaSoundBaseFactory
 {
 	GENERATED_UCLASS_BODY()
 
 	//~ Begin UFactory Interface
-	virtual UObject* FactoryCreateNew(UClass* InClass, UObject* InParent, FName Name, EObjectFlags Flags, UObject* InContext, FFeedbackContext* InFeedbackContext) override;
+	virtual UObject* FactoryCreateNew(UClass* InClass, UObject* InParent, FName InName, EObjectFlags InFlags, UObject* InContext, FFeedbackContext* InFeedbackContext) override;
 	//~ Begin UFactory Interface
 };
diff --git a/Engine/Plugins/Runtime/Metasound/Source/MetasoundEngine/Private/Metasound.cpp b/Engine/Plugins/Runtime/Metasound/Source/MetasoundEngine/Private/Metasound.cpp
--- a/Engine/Plugins/Runtime/Metasound/Source/MetasoundEngine/Private/Metasound.cpp
+++ b/Engine/Plugins/Runtime/Metasound/Source/MetasoundEngine/Private/Metasound.cpp
@@ -1,17 +1,62 @@
 // Copyright Epic Games, Inc. All Rights Reserved.
 #include "Metasound.h"
 
+#include "AssetRegistryModule.h"
 #include "CoreMinimal.h"
-#include "HAL/FileManager.h"
-#include "StructSerializer.h"
-#include "UObject/UnrealType.h"
+#include "Internationalization/Text.h"
+#include "MetasoundAssetBase.h"
+#include "MetasoundAudioFormats.h"
+#include "MetasoundEngineEnvironment.h"
+#include "MetasoundFrontendController.h"
+#include "MetasoundFrontendQuery.h"
+#include "MetasoundFrontendQuerySteps.h"
+#include "MetasoundFrontendSearchEngine.h"
+#include "MetasoundGenerator.h"
+#include "MetasoundInstanceTransmitter.h"
+#include "MetasoundLog.h"
+#include "MetasoundOperatorSettings.h"
+#include "MetasoundPrimitives.h"
+#include "MetasoundReceiveNode.h"
+#include "MetasoundTrigger.h"
+#include "MetasoundEnvironment.h"
+#include "UObject/ObjectSaveContext.h"
 
+#if WITH_EDITORONLY_DATA
+#include "EdGraph/EdGraph.h"
+#endif // WITH_EDITORONLY_DATA
+
+#define LOCTEXT_NAMESPACE "MetaSound"
 UMetaSound::UMetaSound(const FObjectInitializer& ObjectInitializer)
 	: Super(ObjectInitializer)
-	, FMetasoundAssetBase(FMetasoundFrontendArchetype())
+	, FMetasoundAssetBase()
+{
+}
+
+#if WITH_EDITOR
+void UMetaSound::PreSave(FObjectPreSaveContext SaveContext)
 {
+	Super::PreSave(SaveContext);
+	UpdateAssetTags(AssetTags);
+	RegisterGraphWithFrontend();
 }
 
+void UMetaSound::PostEditUndo()
+{
+	Super::PostEditUndo();
+
+	if (Graph)
+	{
+		Graph->Synchronize();
+	}
+}
+
+void UMetaSound::PostEditChangeProperty(FPropertyChangedEvent& InEvent)
+{
+	Super::PostEditChangeProperty(InEvent);
+}
+
+#endif // WITHEDITOR
+
 #if WITH_EDITORONLY_DATA
 UEdGraph* UMetaSound::GetGraph()
 {
@@ -35,42 +80,79 @@ const UEdGraph& UMetaSound::GetGraphChecked() const
 	return *Graph;
 }
 
-void UMetaSound::SetGraph(UEdGraph* InGraph)
+FText UMetaSound::GetDisplayName() const
 {
-	Graph = InGraph;
+	FString TypeName = UMetaSound::StaticClass()->GetName();
+	return FMetasoundAssetBase::GetDisplayName(MoveTemp(TypeName));
 }
 #endif // WITH_EDITORONLY_DATA
 
-// Returns document object responsible for serializing asset
-Metasound::Frontend::FDocumentAccessPtr UMetaSound::GetDocument()
+void UMetaSound::ConvertFromPreset()
 {
 	using namespace Metasound::Frontend;
-	return MakeAccessPtr<FDocumentAccessPtr>(MetasoundDocument.AccessPoint, MetasoundDocument);
+	FGraphHandle GraphHandle = GetRootGraphHandle();
+	FMetasoundFrontendGraphStyle Style = GraphHandle->GetGraphStyle();
+	Style.bIsGraphEditable = true;
+	GraphHandle->SetGraphStyle(Style);
 }
 
-// Returns document object responsible for serializing asset
-Metasound::Frontend::FConstDocumentAccessPtr UMetaSound::GetDocument() const
+const FMetasoundFrontendArchetype& UMetaSound::GetArchetype() const
 {
-	using namespace Metasound::Frontend;
-	return Metasound::Frontend::MakeAccessPtr<FConstDocumentAccessPtr>(MetasoundDocument.AccessPoint, MetasoundDocument);
+	return GetBaseArchetype();
 }
 
-const TArray<FMetasoundFrontendArchetype>& UMetaSound::GetPreferredMetasoundArchetypes() const
+Metasound::FOperatorSettings UMetaSound::GetOperatorSettings(Metasound::FSampleRate InSampleRate) const
 {
-	// Not preferred archetypes for a basic UMetaSound.
-	static const TArray<FMetasoundFrontendArchetype> Preferred;
-	return Preferred;
+	const float BlockRate = FMath::Clamp(Metasound::ConsoleVariables::BlockRate, 1.0f, 1000.0f);
+	return Metasound::FOperatorSettings(InSampleRate, BlockRate);
 }
 
-bool UMetaSound::IsMetasoundArchetypeSupported(const FMetasoundFrontendArchetype& InArchetype) const
+Metasound::FSendAddress UMetaSound::CreateSendAddress(uint64 InInstanceID, const FString& InVertexName, const FName& InDataTypeName) const
 {
-	// All archetypes are supported.
-	return true;
+	using namespace Metasound;
+
+	FSendAddress Address;
+
+	Address.Subsystem = GetSubsystemNameForSendScope(ETransmissionScope::Global);
+	Address.ChannelName = FName(FString::Printf(TEXT("%d:%s:%s"), InInstanceID, *InVertexName, *InDataTypeName.ToString()));
+
+	return Address;
 }
 
-const FMetasoundFrontendArchetype& UMetaSound::GetPreferredMetasoundArchetype(const FMetasoundFrontendDocument& InDocument) const
+const TArray<FMetasoundFrontendArchetype>& UMetaSound::GetPreferredArchetypes() const
 {
-	// Prefer to keep original archetype.
-	return InDocument.Archetype;
+	static const TArray<FMetasoundFrontendArchetype> Preferred
+	{
+		GetBaseArchetype()
+	};
+
+	return Preferred;
 }
 
+const FString& UMetaSound::GetAudioDeviceHandleVariableName()
+{
+	static const FString AudioDeviceHandleVarName = TEXT("AudioDeviceHandle");
+	return AudioDeviceHandleVarName;
+}
+
+const FMetasoundFrontendArchetype& UMetaSound::GetBaseArchetype()
+{
+	auto CreateBaseArchetype = []() -> FMetasoundFrontendArchetype
+	{
+		FMetasoundFrontendArchetype Archetype;
+		
+		FMetasoundFrontendEnvironmentVariable AudioDeviceHandle;
+		AudioDeviceHandle.Name = GetAudioDeviceHandleVariableName();
+		AudioDeviceHandle.Metadata.DisplayName = FText::FromString(AudioDeviceHandle.Name);
+		AudioDeviceHandle.Metadata.Description = LOCTEXT("AudioDeviceHandleToolTip", "Audio device handle");
+
+		Archetype.Interface.Environment.Add(AudioDeviceHandle);
+
+		return Archetype;
+	};
+
+	static const FMetasoundFrontendArchetype BaseArchetype = CreateBaseArchetype();
+
+	return BaseArchetype;
+}
+#undef LOCTEXT_NAMESPACE // MetaSound
diff --git a/Engine/Plugins/Runtime/Metasound/Source/MetasoundEngine/Private/MetasoundEngineModule.cpp b/Engine/Plugins/Runtime/Metasound/Source/MetasoundEngine/Private/MetasoundEngineModule.cpp
--- a/Engine/Plugins/Runtime/Metasound/Source/MetasoundEngine/Private/MetasoundEngineModule.cpp
+++ b/Engine/Plugins/Runtime/Metasound/Source/MetasoundEngine/Private/MetasoundEngineModule.cpp
@@ -21,12 +21,13 @@ class FMetasoundEngineModule : public IMetasoundEngineModule
 		Metasound::IMetasoundUObjectRegistry::RegisterUClassArchetype<UMetaSound>(TEXT(""));
 
 		// Register preferred archetypes
+		Metasound::IMetasoundUObjectRegistry::RegisterUClassPreferredArchetypes<UMetaSound>();
 		Metasound::IMetasoundUObjectRegistry::RegisterUClassPreferredArchetypes<UMetaSoundSource>();
 
-		// flush node regsitration queue
+		// flush node registration queue
 		FMetasoundFrontendRegistryContainer::Get()->RegisterPendingNodes();
 
-		FModuleManager::Get().LoadModuleChecked(TEXT("AudioCodecEngine"));
+		FModuleManager::Get().LoadModuleChecked("AudioCodecEngine");
 
 		UE_LOG(LogMetasoundEngine, Log, TEXT("Metasound Engine Initialized"));
 	}
diff --git a/Engine/Plugins/Runtime/Metasound/Source/MetasoundEngine/Private/MetasoundSource.cpp b/Engine/Plugins/Runtime/Metasound/Source/MetasoundEngine/Private/MetasoundSource.cpp
--- a/Engine/Plugins/Runtime/Metasound/Source/MetasoundEngine/Private/MetasoundSource.cpp
+++ b/Engine/Plugins/Runtime/Metasound/Source/MetasoundEngine/Private/MetasoundSource.cpp
@@ -7,6 +7,7 @@
 #include "MetasoundAssetBase.h"
 #include "MetasoundAudioFormats.h"
 #include "MetasoundEngineEnvironment.h"
+#include "MetasoundEnvironment.h"
 #include "MetasoundFrontendController.h"
 #include "MetasoundFrontendQuery.h"
 #include "MetasoundFrontendQuerySteps.h"
@@ -19,7 +20,8 @@
 #include "MetasoundPrimitives.h"
 #include "MetasoundReceiveNode.h"
 #include "MetasoundTrigger.h"
-#include "MetasoundEnvironment.h"
+#include "MetasoundUObjectRegistry.h"
+#include "UObject/ObjectSaveContext.h"
 
 #if WITH_EDITORONLY_DATA
 #include "EdGraph/EdGraph.h"
@@ -28,10 +30,9 @@
 #define LOCTEXT_NAMESPACE "MetaSoundSource"
 
 
-static float MetaSoundBlockRateCVar = 100.f;
 FAutoConsoleVariableRef CVarMetaSoundBlockRate(
 	TEXT("au.MetaSound.BlockRate"),
-	MetaSoundBlockRateCVar,
+	Metasound::ConsoleVariables::BlockRate,
 	TEXT("Sets block rate (blocks per second) of MetaSounds.\n")
 	TEXT("Default: 100.0f"),
 	ECVF_Default);
@@ -40,7 +41,7 @@ static const FName MetasoundSourceArchetypeName = "Metasound Source";
 
 UMetaSoundSource::UMetaSoundSource(const FObjectInitializer& ObjectInitializer)
 	: Super(ObjectInitializer)
-	, FMetasoundAssetBase(GetMonoSourceArchetype())
+	, FMetasoundAssetBase()
 {
 	bRequiresStopFade = true;
 	NumChannels = 1;
@@ -53,6 +54,20 @@ UMetaSoundSource::UMetaSoundSource(const FObjectInitializer& ObjectInitializer)
 }
 
 #if WITH_EDITOR
+void UMetaSoundSource::PreSave(FObjectPreSaveContext SaveContext)
+{
+	Super::PreSave(SaveContext);
+
+	// TODO: Move these to be run anytime the interface changes or
+	// the node is versioned
+	UpdateAssetTags(AssetTags);
+	RegisterGraphWithFrontend();
+	// TODO: Post re-register, determine if interface changed and
+	// test/figure out when we should run updates on referencing
+	// assets (presets only and on breaking changes? pop-up dialog?
+	// Only done through an action user can run and just let other
+	// be broken? Fail on cook if so?)
+}
 
 void UMetaSoundSource::PostEditUndo()
 {
@@ -66,31 +81,82 @@ void UMetaSoundSource::PostEditUndo()
 
 void UMetaSoundSource::PostEditChangeProperty(FPropertyChangedEvent& InEvent)
 {
+	using namespace Metasound;
+	using namespace Metasound::Frontend;
+
 	Super::PostEditChangeProperty(InEvent);
 
-	if (InEvent.GetPropertyName() == GET_MEMBER_NAME_CHECKED(UMetaSoundSource, OutputFormat)) 
+	if (InEvent.GetPropertyName() == GET_MEMBER_NAME_CHECKED(UMetaSoundSource, OutputFormat))
 	{
+		ConvertFromPreset();
+
 		switch (OutputFormat)
 		{
-			case EMetasoundSourceAudioFormat::Stereo:
+			case EMetasoundSourceAudioFormat::Mono:
+			{
+				NumChannels = 1;
+			}
+			break;
 
+			case EMetasoundSourceAudioFormat::Stereo:
+			{
 				NumChannels = 2;
-				ensure(SetMetasoundArchetype(GetStereoSourceArchetype()));
-
-				break;
+			}
+			break;
 
-			case EMetasoundSourceAudioFormat::Mono:
 			default:
+			{
+				static_assert(static_cast<int32>(EMetasoundSourceAudioFormat::COUNT) == 2, "Possible missing format switch case coverage.");
+			}
+			break;
+		}
 
-				NumChannels = 1;
-				ensure(SetMetasoundArchetype(GetMonoSourceArchetype()));
-
-				break;
+		if (ensure(IsArchetypeSupported(GetArchetype())))
+		{
+			ConformDocumentToArchetype();
 		}
 	}
 }
+#endif // WITH_EDITOR
+
+#if WITH_EDITORONLY_DATA
+UEdGraph* UMetaSoundSource::GetGraph()
+{
+	return Graph;
+}
+
+const UEdGraph* UMetaSoundSource::GetGraph() const
+{
+	return Graph;
+}
+
+UEdGraph& UMetaSoundSource::GetGraphChecked()
+{
+	check(Graph);
+	return *Graph;
+}
+
+const UEdGraph& UMetaSoundSource::GetGraphChecked() const
+{
+	check(Graph);
+	return *Graph;
+}
 
-#endif // WITHEDITOR
+FText UMetaSoundSource::GetDisplayName() const
+{
+	FString TypeName = UMetaSoundSource::StaticClass()->GetName();
+	return FMetasoundAssetBase::GetDisplayName(MoveTemp(TypeName));
+}
+#endif // WITH_EDITORONLY_DATA
+
+void UMetaSoundSource::ConvertFromPreset()
+{
+	using namespace Metasound::Frontend;
+	FGraphHandle GraphHandle = GetRootGraphHandle();
+	FMetasoundFrontendGraphStyle Style = GraphHandle->GetGraphStyle();
+	Style.bIsGraphEditable = true;
+	GraphHandle->SetGraphStyle(Style);
+}
 
 bool UMetaSoundSource::IsPlayable() const
 {
@@ -103,6 +169,29 @@ bool UMetaSoundSource::SupportsSubtitles() const
 	return Super::SupportsSubtitles();
 }
 
+const FMetasoundFrontendArchetype& UMetaSoundSource::GetArchetype() const
+{
+	switch (OutputFormat)
+	{
+		case EMetasoundSourceAudioFormat::Mono:
+		{
+			return GetMonoSourceArchetype();
+		}
+
+		case EMetasoundSourceAudioFormat::Stereo:
+		{
+			return GetStereoSourceArchetype();
+		}
+
+		default:
+		{
+			static_assert(static_cast<int32>(EMetasoundSourceAudioFormat::COUNT) == 2, "Possible missing format switch case coverage.");
+		}
+	}
+
+	return GetBaseArchetype();
+}
+
 float UMetaSoundSource::GetDuration()
 {
 	// eh? this is kind of a weird field anyways.
@@ -149,9 +238,16 @@ ISoundGeneratorPtr UMetaSoundSource::CreateSoundGenerator(const FSoundGeneratorI
 	}
 
 	// Check version and attempt to version up if out-of-date
-	if (Frontend::FVersionDocument().Transform(GetDocumentHandle()))
+	Frontend::FDocumentHandle Document = GetDocumentHandle();
+	FName AssetName = GetFName();
+	FString AssetPath = GetPathName();
+	if (Frontend::FVersionDocument(AssetName, AssetPath).Transform(Document))
 	{
-		UE_LOG(LogMetaSound, Display, TEXT("Dynamically applied versioning for UMetaSoundSource [Name:%s].  Resaving the asset will avoid this requirement at runtime."), *GetName());
+		const FMetasoundFrontendDocumentMetadata& Metadata = Document->GetMetadata();
+		UE_LOG(LogMetaSound, Display,
+			TEXT("Dynamically versioned UMetaSoundSource asset [Name:%s] to %s (Resave to avoid update at runtime)."),
+			*GetName(),
+			*Metadata.Version.Number.ToString());
 	}
 
 	// Inject receive nodes for unused transmittable inputs. Perform edits on a copy
@@ -182,7 +278,7 @@ TUniquePtr<IAudioInstanceTransmitter> UMetaSoundSource::CreateInstanceTransmitte
 {
 	Metasound::FMetasoundInstanceTransmitter::FInitParams InitParams(GetOperatorSettings(InParams.SampleRate), InParams.InstanceID);
 
-	for (const FSendInfoAndVertexName& InfoAndName : GetSendInfos(InParams.InstanceID))
+	for (const FSendInfoAndVertexName& InfoAndName : FMetasoundAssetBase::GetSendInfos(InParams.InstanceID))
 	{
 		InitParams.Infos.Add(InfoAndName.SendInfo);
 	}
@@ -190,222 +286,13 @@ TUniquePtr<IAudioInstanceTransmitter> UMetaSoundSource::CreateInstanceTransmitte
 	return MakeUnique<Metasound::FMetasoundInstanceTransmitter>(InitParams);
 }
 
-bool UMetaSoundSource::GetReceiveNodeMetadataForDataType(const FName& InTypeName, FMetasoundFrontendClassMetadata& OutMetadata) const
-{
-	using namespace Metasound;
-
-	const FNodeClassName ClassName = FReceiveNodeNames::GetClassNameForDataType(InTypeName);
-	TArray<FMetasoundFrontendClass> ReceiverNodeClasses = Frontend::ISearchEngine::Get().FindClassesWithName(ClassName, true /* bInSortByVersion */);
-
-	if (ReceiverNodeClasses.IsEmpty())
-	{
-		return false;
-	}
-
-	OutMetadata = ReceiverNodeClasses[0].Metadata;
-	return true;
-}
-
-Metasound::Frontend::FNodeHandle UMetaSoundSource::AddInputPinForSendAddress(const Metasound::FMetasoundInstanceTransmitter::FSendInfo& InSendInfo, Metasound::Frontend::FGraphHandle InGraph) const
-{
-	FMetasoundFrontendClassInput Description;
-	FGuid VertexID = FGuid::NewGuid();
-
-	Description.Name = InSendInfo.Address.ChannelName.ToString();
-	Description.TypeName = Metasound::GetMetasoundDataTypeName<Metasound::FSendAddress>();
-	Description.Metadata.Description = FText::GetEmpty();
-	Description.VertexID = VertexID;
-	Description.DefaultLiteral.Set(InSendInfo.Address.ChannelName.ToString());
-	
-	return InGraph->AddInputVertex(Description);
-}
-
-bool UMetaSoundSource::CopyDocumentAndInjectReceiveNodes(uint64 InInstanceID, const FMetasoundFrontendDocument& InSourceDoc, FMetasoundFrontendDocument& OutDestDoc) const
-{
-	using namespace Metasound;
-	using namespace Metasound::Frontend;
-
-	OutDestDoc = InSourceDoc;
-
-	FDocumentHandle Document = IDocumentController::CreateDocumentHandle(MakeAccessPtr<FDocumentAccessPtr>(OutDestDoc.AccessPoint, OutDestDoc));
-	FGraphHandle RootGraph = Document->GetRootGraph();
-
-	TArray<FSendInfoAndVertexName> SendInfoAndVertexes = GetSendInfos(InInstanceID);
-
-	// Inject receive nodes for each transmittable input
-	for (const FSendInfoAndVertexName& InfoAndVertexName : SendInfoAndVertexes)
-	{
-
-		// Add receive node to graph
-		FMetasoundFrontendClassMetadata ReceiveNodeMetadata;
-		bool bSuccess = GetReceiveNodeMetadataForDataType(InfoAndVertexName.SendInfo.TypeName, ReceiveNodeMetadata);
-		if (!bSuccess)
-		{
-			// TODO: log warning
-			continue;
-		}
-		FNodeHandle ReceiveNode = RootGraph->AddNode(ReceiveNodeMetadata);
-
-		// Add receive node address to graph
-		FNodeHandle AddressNode = UMetaSoundSource::AddInputPinForSendAddress(InfoAndVertexName.SendInfo, RootGraph);
-		TArray<FOutputHandle> AddressNodeOutputs = AddressNode->GetOutputs();
-		if (AddressNodeOutputs.Num() != 1)
-		{
-			// TODO: log warning
-			continue;
-		}
-
-		FOutputHandle AddressOutput = AddressNodeOutputs[0];
-		TArray<FInputHandle> ReceiveAddressInput = ReceiveNode->GetInputsWithVertexName(Metasound::FReceiveNodeNames::GetAddressInputName());
-		if (ReceiveAddressInput.Num() != 1)
-		{
-			// TODO: log error
-			continue;
-		}
-
-		ensure(ReceiveAddressInput[0]->Connect(*AddressOutput));
-
-
-
-		// Swap input node connections with receive node connections
-		FNodeHandle InputNode = RootGraph->GetInputNodeWithName(InfoAndVertexName.VertexName);
-		if (!ensure(InputNode->GetOutputs().Num() == 1))
-		{
-			// TODO: handle input node with varying number of outputs or varying output types.
-			continue;
-		}
-
-		FOutputHandle InputNodeOutput = InputNode->GetOutputs()[0];
-
-		if (ensure(ReceiveNode->IsValid()))
-		{
-			TArray<FOutputHandle> ReceiveNodeOutputs = ReceiveNode->GetOutputs();
-			if (!ensure(ReceiveNodeOutputs.Num() == 1))
-			{
-				// TODO: handle array outputs and receive nodes of varying formats.
-				continue;
-			}
-
-			TArray<FInputHandle> ReceiveDefaultInputs = ReceiveNode->GetInputsWithVertexName(Metasound::FReceiveNodeNames::GetDefaultDataInputName());
-			if (ensure(ReceiveDefaultInputs.Num() == 1))
-			{
-				FOutputHandle ReceiverNodeOutput = ReceiveNodeOutputs[0];
-				for (FInputHandle NodeInput : InputNodeOutput->GetConnectedInputs())
-				{
-					// Swap connections to receiver node
-					ensure(InputNodeOutput->Disconnect(*NodeInput));
-					ensure(ReceiverNodeOutput->Connect(*NodeInput));
-				}
-
-				ReceiveDefaultInputs[0]->Connect(*InputNodeOutput);
-			}
-		}
-	}
-
-	return true;
-}
-
-TArray<FString> UMetaSoundSource::GetTransmittableInputVertexNames() const
-{
-	using namespace Metasound;
-
-	// Unused inputs are all input vertices that are not in the archetype.
-	TArray<FString> ArchetypeInputVertexNames;
-	for (const FMetasoundFrontendClassVertex& Vertex : GetMetasoundArchetype().Interface.Inputs)
-	{
-		ArchetypeInputVertexNames.Add(Vertex.Name);
-	}
-
-	Frontend::FConstGraphHandle RootGraph = GetRootGraphHandle();
-	TArray<FString> GraphInputVertexNames = RootGraph->GetInputVertexNames();
-
-	// Filter graph inputs by archetype inputs.
-	GraphInputVertexNames = GraphInputVertexNames.FilterByPredicate([&](const FString& InName) { return !ArchetypeInputVertexNames.Contains(InName); });
-
-	auto IsDataTypeTransmittable = [&](const FString& InVertexName)
-	{
-		Frontend::FConstClassInputAccessPtr ClassInputPtr = RootGraph->FindClassInputWithName(InVertexName);
-		if (const FMetasoundFrontendClassInput* ClassInput = ClassInputPtr.Get())
-		{
-			FDataTypeRegistryInfo TypeInfo;
-			if (Frontend::GetTraitsForDataType(ClassInput->TypeName, TypeInfo))
-			{
-				if (TypeInfo.bIsTransmittable)
-				{
-					// TODO: Currently values set directly on node pins are represented
-					// as input nodes in the graph. They should not be used for transmission
-					// as the number of input nodes increases quickly as more nodes
-					// are added to a graph. Connecting these input nodes to the 
-					// transmission system is relatively expensive. These undesirable input nodes are
-					// filtered out by ignoring input nodes which are not "Visible". 
-					Frontend::FConstNodeHandle InputNode = RootGraph->GetNodeWithID(ClassInput->NodeID);
-					if (InputNode->IsValid())
-					{
-						return true;
-					}
-				}
-			}
-		}
-		return false;
-	};
-
-	GraphInputVertexNames = GraphInputVertexNames.FilterByPredicate(IsDataTypeTransmittable);
-
-	return GraphInputVertexNames;
-}
-
 Metasound::FOperatorSettings UMetaSoundSource::GetOperatorSettings(Metasound::FSampleRate InSampleRate) const
 {
-	// Metasound graph gets evaluated 100 times per second by default.
-	float BlockRate = FMath::Clamp(MetaSoundBlockRateCVar, 1.0f, 1000.0f); 
+	const float BlockRate = FMath::Clamp(Metasound::ConsoleVariables::BlockRate, 1.0f, 1000.0f);
 	return Metasound::FOperatorSettings(InSampleRate, BlockRate);
 }
 
-Metasound::FSendAddress UMetaSoundSource::CreateSendAddress(uint64 InInstanceID, const FString& InVertexName, const FName& InDataTypeName) const
-{
-	using namespace Metasound;
-
-	FSendAddress Address;
-
-	Address.Subsystem = GetSubsystemNameForSendScope(ETransmissionScope::Global);
-	Address.ChannelName = FName(FString::Printf(TEXT("%d:%s:%s"), InInstanceID, *InVertexName, *InDataTypeName.ToString()));
-
-	return Address;
-}
-
-TArray<UMetaSoundSource::FSendInfoAndVertexName> UMetaSoundSource::GetSendInfos(uint64 InInstanceID) const
-{
-	using FSendInfo = Metasound::FMetasoundInstanceTransmitter::FSendInfo;
-	using namespace Metasound::Frontend;
-
-	TArray<FSendInfoAndVertexName> SendInfos;
-
-	FConstGraphHandle RootGraph = GetRootGraphHandle();
-
-	TArray<FString> SendVertices = GetTransmittableInputVertexNames();
-	for (const FString& VertexName : SendVertices)
-	{
-		FConstNodeHandle InputNode = RootGraph->GetInputNodeWithName(VertexName);
-		for (FConstInputHandle InputHandle : InputNode->GetConstInputs())
-		{
-			FSendInfoAndVertexName Info;
-
-			// TODO: incorporate VertexID into address. But need to ensure that VertexID
-			// will be maintained after injecting Receive nodes. 
-			Info.SendInfo.Address = CreateSendAddress(InInstanceID, InputHandle->GetName(), InputHandle->GetDataType());
-			Info.SendInfo.ParameterName = FName(InputHandle->GetDisplayName().ToString()); // TODO: display name hack. Need to have naming consistent in editor for inputs
-			//Info.SendInfo.ParameterName = FName(*InputHandle->GetName()); // TODO: this is the expected parameter name.
-			Info.SendInfo.TypeName = InputHandle->GetDataType();
-			Info.VertexName = VertexName;
-			
-			SendInfos.Add(Info);
-		}
-	}
-
-	return SendInfos;
-}
-
-const TArray<FMetasoundFrontendArchetype>& UMetaSoundSource::GetPreferredMetasoundArchetypes() const 
+const TArray<FMetasoundFrontendArchetype>& UMetaSoundSource::GetPreferredArchetypes() const 
 {
 	static const TArray<FMetasoundFrontendArchetype> Preferred({GetMonoSourceArchetype(), GetStereoSourceArchetype()});
 
@@ -536,5 +423,4 @@ const FMetasoundFrontendArchetype& UMetaSoundSource::GetStereoSourceArchetype()
 
 	return StereoArchetype;
 }
-
 #undef LOCTEXT_NAMESPACE // MetaSoundSource
diff --git a/Engine/Plugins/Runtime/Metasound/Source/MetasoundEngine/Private/MetasoundUObjectRegistry.cpp b/Engine/Plugins/Runtime/Metasound/Source/MetasoundEngine/Private/MetasoundUObjectRegistry.cpp
--- a/Engine/Plugins/Runtime/Metasound/Source/MetasoundEngine/Private/MetasoundUObjectRegistry.cpp
+++ b/Engine/Plugins/Runtime/Metasound/Source/MetasoundEngine/Private/MetasoundUObjectRegistry.cpp
@@ -47,7 +47,7 @@ namespace Metasound
 				return Classes;
 			}
 
-			UObject* NewObject(UClass* InClass, const FMetasoundFrontendDocument& InDocument, const FString& InPath) const override
+			UObject* NewObject(UClass* InClass, const FMetasoundFrontendDocument& InDocument, const FMetasoundFrontendArchetype& InArchetype, const FString& InPath) const override
 			{
 				auto IsChildClassOfRegisteredClass = [&](const IMetasoundUObjectRegistryEntry* Entry)
 				{
@@ -58,7 +58,7 @@ namespace Metasound
 
 				for (const IMetasoundUObjectRegistryEntry* Entry : EntriesForClass)
 				{
-					if (Entry->GetArchetypeName() == InDocument.Archetype.Name)
+					if (Entry->GetArchetypeName() == InArchetype.Name)
 					{
 						return NewObject(*Entry, InDocument, InPath);
 					}
@@ -67,7 +67,6 @@ namespace Metasound
 				return nullptr;
 			}
 
-
 			bool IsRegisteredClass(UObject* InObject) const override
 			{
 				return (nullptr != GetEntryByUObject(InObject));
@@ -120,6 +119,12 @@ namespace Metasound
 				if (ensure(nullptr != NewAssetBase))
 				{
 					NewAssetBase->SetDocument(InDocument);
+
+					const FMetasoundFrontendArchetype& Archetype = NewAssetBase->GetArchetype();
+					if (ensure(NewAssetBase->IsArchetypeSupported(Archetype)))
+					{
+						NewAssetBase->ConformDocumentToArchetype();
+					}
 				}
 
 #if WITH_EDITOR
diff --git a/Engine/Plugins/Runtime/Metasound/Source/MetasoundEngine/Public/Metasound.h b/Engine/Plugins/Runtime/Metasound/Source/MetasoundEngine/Public/Metasound.h
--- a/Engine/Plugins/Runtime/Metasound/Source/MetasoundEngine/Public/Metasound.h
+++ b/Engine/Plugins/Runtime/Metasound/Source/MetasoundEngine/Public/Metasound.h
@@ -2,27 +2,20 @@
 #pragma once
 
 #include "CoreMinimal.h"
+#include "EdGraph/EdGraph.h"
 #include "MetasoundAssetBase.h"
 #include "MetasoundFrontend.h"
 #include "MetasoundFrontendDocument.h"
-#include "MetasoundFrontendDocumentAccessPtr.h"
-
-
-#if WITH_EDITORONLY_DATA
-#include "EdGraph/EdGraph.h"
-#include "Misc/AssertionMacros.h"
-#endif // WITH_EDITORONLY_DATA
+#include "MetasoundOperatorSettings.h"
+#include "MetasoundRouter.h"
+#include "MetasoundSource.h"
 
 #include "Metasound.generated.h"
 
-// Forward Declarations
-class FEditPropertyChain;
-struct FPropertyChangedEvent;
-
 
 /**
  * This asset type is used for Metasound assets that can only be used as nodes in other Metasound graphs.
- * Because of this, they can have any inputs or outputs they need.
+ * Because of this, they contain no required inputs or outputs.
  */
 UCLASS(hidecategories = object, BlueprintType)
 class METASOUNDENGINE_API UMetaSound : public UObject, public FMetasoundAssetBase
@@ -31,17 +24,14 @@ class METASOUNDENGINE_API UMetaSound : public UObject, public FMetasoundAssetBas
 
 protected:
 	UPROPERTY(EditAnywhere, Category = CustomView)
-	FMetasoundFrontendDocument MetasoundDocument;
-
-	// Returns document object responsible for serializing asset
-	Metasound::Frontend::FDocumentAccessPtr GetDocument() override;
-
-	// Returns document object responsible for serializing asset
-	Metasound::Frontend::FConstDocumentAccessPtr GetDocument() const override;
+	FMetasoundFrontendDocument RootMetaSoundDocument;
 
 #if WITH_EDITORONLY_DATA
+	UPROPERTY(AssetRegistrySearchable)
+	FMetasoundFrontendClassAssetTags AssetTags;
+
 	UPROPERTY()
-	UEdGraph* Graph;
+	UMetasoundEditorGraphBase* Graph;
 #endif // WITH_EDITORONLY_DATA
 
 public:
@@ -52,12 +42,15 @@ public:
 	// while allowing editor to reference directly)
 	static FName GetDocumentPropertyName()
 	{
-		return GET_MEMBER_NAME_CHECKED(UMetaSound, MetasoundDocument);
+		return GET_MEMBER_NAME_CHECKED(UMetaSound, RootMetaSoundDocument);
 	}
 
+	// Name to display in editors
+	virtual FText GetDisplayName() const override;
+
 	// Returns the graph associated with this Metasound. Graph is required to be referenced on
 	// Metasound UObject for editor serialization purposes.
-	// @return Editor graph associated with UMetaSound.
+	// @return Editor graph associated with UMetaSoundSource.
 	virtual UEdGraph* GetGraph() override;
 	virtual const UEdGraph* GetGraph() const override;
 	virtual UEdGraph& GetGraphChecked() override;
@@ -65,9 +58,26 @@ public:
 
 	// Sets the graph associated with this Metasound. Graph is required to be referenced on
 	// Metasound UObject for editor serialization purposes.
-	// @param Editor graph associated with UMetaSound.
-	virtual void SetGraph(UEdGraph* InGraph) override;
-#endif // WITH_EDITORONLY_DATA
+	// @param Editor graph associated with UMetaSoundSource.
+	virtual void SetGraph(UEdGraph* InGraph) override
+	{
+		Graph = CastChecked<UMetasoundEditorGraphBase>(InGraph);
+	}
+#endif // #if WITH_EDITORONLY_DATA
+
+#if WITH_EDITOR
+	virtual void PreSave(FObjectPreSaveContext SaveContext) override;
+	virtual void PostEditUndo() override;
+	virtual void PostEditChangeProperty(FPropertyChangedEvent& InEvent) override;
+#endif // WITH_EDITOR
+
+	virtual const FMetasoundFrontendArchetype& GetArchetype() const override;
+
+	// If set to be a preset, converts to a full-access MetaSound,
+	// removing edit restrictions and excluding it from automatic
+	// interface versioning.
+	UFUNCTION(Category = Metasound, meta = (CallInEditor = "true"))
+	void ConvertFromPreset();
 
 	UObject* GetOwningAsset() override
 	{
@@ -79,10 +89,35 @@ public:
 		return this;
 	}
 
-	const TArray<FMetasoundFrontendArchetype>& GetPreferredMetasoundArchetypes() const override;
+	TUniquePtr<IAudioInstanceTransmitter> CreateInstanceTransmitter(const FAudioInstanceTransmitterInitParams& InParams) const;
+
+	// Get the most up to date archetype for metasound sources.
+	const TArray<FMetasoundFrontendArchetype>& GetPreferredArchetypes() const override;
+
+protected:
+
+	Metasound::Frontend::FDocumentAccessPtr GetDocument() override
+	{
+		using namespace Metasound::Frontend;
+		// Return document using FAccessPoint to inform the TAccessPtr when the 
+		// object is no longer valid.
+		return MakeAccessPtr<FDocumentAccessPtr>(RootMetaSoundDocument.AccessPoint, RootMetaSoundDocument);
+	}
 
-	bool IsMetasoundArchetypeSupported(const FMetasoundFrontendArchetype& InArchetype) const override;
+	Metasound::Frontend::FConstDocumentAccessPtr GetDocument() const override
+	{
+		using namespace Metasound::Frontend;
+		// Return document using FAccessPoint to inform the TAccessPtr when the 
+		// object is no longer valid.
+		return MakeAccessPtr<FConstDocumentAccessPtr>(RootMetaSoundDocument.AccessPoint, RootMetaSoundDocument);
+	}
 
-	const FMetasoundFrontendArchetype& GetPreferredMetasoundArchetype(const FMetasoundFrontendDocument& InDocument) const override;
+	bool CopyDocumentAndInjectReceiveNodes(uint64 InInstanceID, const FMetasoundFrontendDocument& InSourceDoc, FMetasoundFrontendDocument& OutDestDoc) const;	
 
+private:
+	TArray<FString> GetTransmittableInputVertexNames() const;
+	Metasound::FOperatorSettings GetOperatorSettings(Metasound::FSampleRate InSampleRate) const;
+	Metasound::FSendAddress CreateSendAddress(uint64 InInstanceID, const FString& InVertexName, const FName& InDataTypeName) const;
+	static const FString& GetAudioDeviceHandleVariableName();
+	static const FMetasoundFrontendArchetype& GetBaseArchetype();
 };
diff --git a/Engine/Plugins/Runtime/Metasound/Source/MetasoundEngine/Public/MetasoundSource.h b/Engine/Plugins/Runtime/Metasound/Source/MetasoundEngine/Public/MetasoundSource.h
--- a/Engine/Plugins/Runtime/Metasound/Source/MetasoundEngine/Public/MetasoundSource.h
+++ b/Engine/Plugins/Runtime/Metasound/Source/MetasoundEngine/Public/MetasoundSource.h
@@ -13,12 +13,26 @@
 
 #include "MetasoundSource.generated.h"
 
+
+namespace Metasound
+{
+	namespace ConsoleVariables
+	{
+		static float BlockRate = 100.f;
+	}
+} // namespace Metasound
+
 /** Declares the output audio format of the UMetaSoundSource */
 UENUM()
 enum class EMetasoundSourceAudioFormat : uint8
 {
-	Mono,  //< Mono audio output.
-	Stereo //< Stereo audio output.
+	// Mono audio output
+	Mono,
+
+	// Stereo audio output
+	Stereo,
+
+	COUNT UMETA(Hidden)
 };
 
 UCLASS()
@@ -46,6 +60,9 @@ protected:
 	FMetasoundFrontendDocument RootMetasoundDocument;
 
 #if WITH_EDITORONLY_DATA
+	UPROPERTY(AssetRegistrySearchable)
+	FMetasoundFrontendClassAssetTags AssetTags;
+
 	UPROPERTY()
 	UMetasoundEditorGraphBase* Graph;
 #endif // WITH_EDITORONLY_DATA
@@ -57,7 +74,7 @@ public:
 	UPROPERTY(EditAnywhere, BlueprintReadOnly, Category = Metasound)
 	EMetasoundSourceAudioFormat OutputFormat;
 
-#if WITH_EDITOR
+#if WITH_EDITORONLY_DATA
 	// Returns document name (for editor purposes, and avoids making document public for edit
 	// while allowing editor to reference directly)
 	static FName GetDocumentPropertyName()
@@ -65,32 +82,16 @@ public:
 		return GET_MEMBER_NAME_CHECKED(UMetaSoundSource, RootMetasoundDocument);
 	}
 
-	virtual void PostEditUndo() override;
+	// Name to display in editors
+	virtual FText GetDisplayName() const override;
 
 	// Returns the graph associated with this Metasound. Graph is required to be referenced on
 	// Metasound UObject for editor serialization purposes.
 	// @return Editor graph associated with UMetaSoundSource.
-	virtual UEdGraph* GetGraph() override
-	{
-		return Graph;
-	}
-
-	virtual const UEdGraph* GetGraph() const override
-	{
-		return Graph;
-	}
-
-	virtual UEdGraph& GetGraphChecked() override
-	{
-		check(Graph);
-		return *Graph;
-	}
-
-	virtual const UEdGraph& GetGraphChecked() const override
-	{
-		check(Graph);
-		return *Graph;
-	}
+	virtual UEdGraph* GetGraph() override;
+	virtual const UEdGraph* GetGraph() const override;
+	virtual UEdGraph& GetGraphChecked() override;
+	virtual const UEdGraph& GetGraphChecked() const override;
 
 	// Sets the graph associated with this Metasound. Graph is required to be referenced on
 	// Metasound UObject for editor serialization purposes.
@@ -99,6 +100,11 @@ public:
 	{
 		Graph = CastChecked<UMetasoundEditorGraphBase>(InGraph);
 	}
+#endif // #if WITH_EDITORONLY_DATA
+
+#if WITH_EDITOR
+	virtual void PreSave(FObjectPreSaveContext SaveContext) override;
+	virtual void PostEditUndo() override;
 
 	virtual bool GetRedrawThumbnail() const override
 	{
@@ -114,10 +120,17 @@ public:
 		return false;
 	}
 
-	void PostEditChangeProperty(FPropertyChangedEvent& InEvent);
-
+	virtual void PostEditChangeProperty(FPropertyChangedEvent& InEvent) override;
 #endif // WITH_EDITOR
 
+	virtual const FMetasoundFrontendArchetype& GetArchetype() const override;
+
+	// If set to be a preset, converts to a full-access MetaSound,
+	// removing edit restrictions and excluding it from automatic
+	// interface versioning.
+	UFUNCTION(Category = Metasound, meta = (CallInEditor = "true"))
+	void ConvertFromPreset();
+
 	UObject* GetOwningAsset() override
 	{
 		return this;
@@ -128,17 +141,16 @@ public:
 		return this;
 	}
 
-	bool IsPlayable() const override;
-	bool SupportsSubtitles() const override;
-	float GetDuration() override;
-	ISoundGeneratorPtr CreateSoundGenerator(const FSoundGeneratorInitParams& InParams) override;
-	TUniquePtr<IAudioInstanceTransmitter> CreateInstanceTransmitter(const FAudioInstanceTransmitterInitParams& InParams) const override;
+	virtual bool IsPlayable() const override;
+	virtual bool SupportsSubtitles() const override;
+	virtual float GetDuration() override;
+	virtual ISoundGeneratorPtr CreateSoundGenerator(const FSoundGeneratorInitParams& InParams) override;
+	virtual TUniquePtr<IAudioInstanceTransmitter> CreateInstanceTransmitter(const FAudioInstanceTransmitterInitParams& InParams) const override;
 
 	// Get the most up to date archetype for metasound sources.
-	const TArray<FMetasoundFrontendArchetype>& GetPreferredMetasoundArchetypes() const override;
+	const TArray<FMetasoundFrontendArchetype>& GetPreferredArchetypes() const override;
 
 protected:
-
 	Metasound::Frontend::FDocumentAccessPtr GetDocument() override
 	{
 		using namespace Metasound::Frontend;
@@ -156,23 +168,7 @@ protected:
 	}
 
 private:
-	struct FSendInfoAndVertexName
-	{
-		Metasound::FMetasoundInstanceTransmitter::FSendInfo SendInfo;
-		FString VertexName;
-	};
-
-	bool GetReceiveNodeMetadataForDataType(const FName& InTypeName, FMetasoundFrontendClassMetadata& OutMetadata) const;
-	Metasound::Frontend::FNodeHandle AddInputPinForSendAddress(const Metasound::FMetasoundInstanceTransmitter::FSendInfo& InSendInfo, Metasound::Frontend::FGraphHandle InGraph) const;
-	bool CopyDocumentAndInjectReceiveNodes(uint64 InInstanceID, const FMetasoundFrontendDocument& InSourceDoc, FMetasoundFrontendDocument& OutDestDoc) const;	
-	TArray<FString> GetTransmittableInputVertexNames() const;
 	Metasound::FOperatorSettings GetOperatorSettings(Metasound::FSampleRate InSampleRate) const;
-	Metasound::FSendAddress CreateSendAddress(uint64 InInstanceID, const FString& InVertexName, const FName& InDataTypeName) const;
-	
-
-
-	TArray<FSendInfoAndVertexName> GetSendInfos(uint64 InInstanceID) const;
-
 
 	static const FString& GetOnPlayInputName();
 	static const FString& GetAudioOutputName();
diff --git a/Engine/Plugins/Runtime/Metasound/Source/MetasoundEngine/Public/MetasoundUObjectRegistry.h b/Engine/Plugins/Runtime/Metasound/Source/MetasoundEngine/Public/MetasoundUObjectRegistry.h
--- a/Engine/Plugins/Runtime/Metasound/Source/MetasoundEngine/Public/MetasoundUObjectRegistry.h
+++ b/Engine/Plugins/Runtime/Metasound/Source/MetasoundEngine/Public/MetasoundUObjectRegistry.h
@@ -145,7 +145,7 @@ namespace Metasound
 			{
 				static_assert(std::is_base_of<FMetasoundAssetBase, UClassType>::value, "UClass must be derived from FMetasoundAssetBase");
 
-				const TArray<FMetasoundFrontendArchetype>& PreferredArchetypes = GetDefault<UClassType>()->GetPreferredMetasoundArchetypes();
+				const TArray<FMetasoundFrontendArchetype>& PreferredArchetypes = GetDefault<UClassType>()->GetPreferredArchetypes();
 
 				for (const FMetasoundFrontendArchetype& Arch : PreferredArchetypes)
 				{
@@ -171,15 +171,16 @@ namespace Metasound
 			/** Returns all UClasses registered to the archetype name. */
 			virtual TArray<UClass*> GetUClassesForArchetype(const FName& InArchetypeName) const = 0;
 
-			/** Creates a new object from a metasound doucment.
+			/** Creates a new object from a metasound document.
 			 *
 			 * @param InClass - A registered UClass to create.
-			 * @param InDocucument - The FMetasoundFrontendDocument to use when creating the class.
+			 * @param InDocument - The FMetasoundFrontendDocument to use when creating the class.
+			 * @param InArchetype - The FMetasoundFrontendArchetype to use when creating the class.
 			 * @param InPath - If in editor, the created asset will be stored at this content path.
 			 *
 			 * @return A new object. A nullptr on error.
 			 */
-			virtual UObject* NewObject(UClass* InClass, const FMetasoundFrontendDocument& InDocument, const FString& InPath) const = 0;
+			virtual UObject* NewObject(UClass* InClass, const FMetasoundFrontendDocument& InDocument, const FMetasoundFrontendArchetype& InArchetype, const FString& InPath) const = 0;
 
 			/** Returns true if the InObject is of a class or child class which is registered. */
 			virtual bool IsRegisteredClass(UObject* InObject) const = 0;
diff --git a/Engine/Plugins/Runtime/Metasound/Source/MetasoundFrontend/Private/MetasoundAssetBase.cpp b/Engine/Plugins/Runtime/Metasound/Source/MetasoundFrontend/Private/MetasoundAssetBase.cpp
--- a/Engine/Plugins/Runtime/Metasound/Source/MetasoundFrontend/Private/MetasoundAssetBase.cpp
+++ b/Engine/Plugins/Runtime/Metasound/Source/MetasoundFrontend/Private/MetasoundAssetBase.cpp
@@ -4,23 +4,81 @@
 
 #include "Algo/AnyOf.h"
 #include "HAL/FileManager.h"
+#include "Internationalization/Text.h"
 #include "IStructSerializerBackend.h"
 #include "MetasoundArchetype.h"
 #include "MetasoundFrontendController.h"
+#include "MetasoundFrontendGraph.h"
+#include "MetasoundFrontendSearchEngine.h"
 #include "MetasoundFrontendTransform.h"
 #include "MetasoundJsonBackend.h"
 #include "MetasoundLog.h"
+#include "MetasoundReceiveNode.h"
 #include "StructSerializer.h"
 
+#define LOCTEXT_NAMESPACE "MetaSound"
+
 const FString FMetasoundAssetBase::FileExtension(TEXT(".metasound"));
 
-FMetasoundAssetBase::FMetasoundAssetBase()
+void FMetasoundAssetBase::RegisterGraphWithFrontend()
 {
-}
+	using namespace Metasound;
+	using namespace Metasound::Frontend;
 
-FMetasoundAssetBase::FMetasoundAssetBase(const FMetasoundFrontendArchetype& InArchetype)
-:	Archetype(InArchetype)
-{
+	// TODO: unregister prior version.
+	
+	FConstDocumentAccessPtr DocumentPtr = GetDocument();
+	FString AssetName;
+	if (const UObject* OwningAsset = GetOwningAsset())
+	{
+		AssetName = OwningAsset->GetName();
+	}
+
+
+	FCreateMetasoundNodeFunction CreateNode = [Name=AssetName, DocumentPtr=DocumentPtr](const Metasound::FNodeInitData& InInitData) -> TUniquePtr<INode>
+	{
+		check(IsInGameThread());
+
+		if (const FMetasoundFrontendDocument* Document = DocumentPtr.Get())
+		{
+			return FFrontendGraphBuilder().CreateGraph(*Document);
+		}
+		else
+		{
+			UE_LOG(LogMetaSound, Error, TEXT("Failed to create node for MetaSoundSource [Name:%s]"), *Name);
+			return TUniquePtr<INode>(nullptr);
+		}
+	};
+
+	FCreateMetasoundFrontendClassFunction CreateFrontendClass = [Name=AssetName, DocumentPtr=DocumentPtr]() -> FMetasoundFrontendClass
+	{
+		check(IsInGameThread());
+
+		FMetasoundFrontendClass FrontendClass;
+		if (const FMetasoundFrontendDocument* Document = DocumentPtr.Get())
+		{
+			FrontendClass = Document->RootGraph;
+			FrontendClass.Metadata.Type = EMetasoundFrontendClassType::External;
+		}
+		else
+		{
+			UE_LOG(LogMetaSound, Error, TEXT("Failed to create frontend class for MetaSoundSource [Name:%s]"), *Name);
+		}
+
+		return FrontendClass;
+	};
+
+	if (IsValidNodeRegistryKey(RegistryKey))
+	{
+		ensure(FMetasoundFrontendRegistryContainer::Get()->UnregisterExternalNode(RegistryKey));
+	}
+
+	RegistryKey = FMetasoundFrontendRegistryContainer::Get()->RegisterExternalNode(MoveTemp(CreateNode), MoveTemp(CreateFrontendClass));
+
+	if (!IsValidNodeRegistryKey(RegistryKey))
+	{
+		UE_LOG(LogMetaSound, Error, TEXT("Failed to register node for MetaSoundSource [Name:%s]"), *AssetName);
+	}
 }
 
 void FMetasoundAssetBase::SetMetadata(FMetasoundFrontendClassMetadata& InMetadata)
@@ -37,79 +95,292 @@ void FMetasoundAssetBase::SetMetadata(FMetasoundFrontendClassMetadata& InMetadat
 	MarkMetasoundDocumentDirty();
 }
 
-const FMetasoundFrontendArchetype& FMetasoundAssetBase::GetMetasoundArchetype() const
+bool FMetasoundAssetBase::IsArchetypeSupported(const FMetasoundFrontendArchetype& InArchetype) const
 {
-	return Archetype;
+	auto IsEqualArchetype = [&](const FMetasoundFrontendArchetype& SupportedArchetype)
+	{
+		return Metasound::Frontend::IsEqualArchetype(InArchetype, SupportedArchetype);
+	};
+
+	return Algo::AnyOf(GetPreferredArchetypes(), IsEqualArchetype);
 }
 
-bool FMetasoundAssetBase::SetMetasoundArchetype(const FMetasoundFrontendArchetype& InArchetype)
+const FMetasoundFrontendArchetype& FMetasoundAssetBase::GetPreferredArchetypes(const FMetasoundFrontendDocument& InDocument, const FMetasoundFrontendArchetype& InDefaultArchetype) const
 {
-	if (IsMetasoundArchetypeSupported(InArchetype))
+	// Default to archetype provided in case it is supported.
+	if (IsArchetypeSupported(InDefaultArchetype))
 	{
-		Archetype = InArchetype;
-		ConformDocumentToMetasoundArchetype();
+		return InDefaultArchetype;
+	}
 
-		return true;
+	// If existing archetype is not supported, get the most similar that still supports the documents environment.
+	const FMetasoundFrontendArchetype* SimilarArchetype = Metasound::Frontend::FindMostSimilarArchetypeSupportingEnvironment(InDocument, GetPreferredArchetypes());
+
+	if (nullptr != SimilarArchetype)
+	{
+		return *SimilarArchetype;
 	}
 
-	// Archetype was not set because it is not supported for this asset base. 
-	return false;
+	// Nothing found. Return the existing archetype for the FMetasoundAssetBase.
+	return GetArchetype();
 }
 
-bool FMetasoundAssetBase::IsMetasoundArchetypeSupported(const FMetasoundFrontendArchetype& InArchetype) const
+void FMetasoundAssetBase::SetDocument(const FMetasoundFrontendDocument& InDocument)
 {
-	auto IsEqualArchetype = [&](const FMetasoundFrontendArchetype& SupportedArchetype)
+	FMetasoundFrontendDocument& Document = GetDocumentChecked();
+	Document = InDocument;
+}
+
+void FMetasoundAssetBase::ConformDocumentToArchetype()
+{
+	FMetasoundFrontendDocument& Document = GetDocumentChecked();
+	const FMetasoundFrontendArchetype& Archetype = GetArchetype();
+	Metasound::Frontend::FMatchRootGraphToArchetype Transform(Archetype);
+	Transform.Transform(GetDocumentHandle());
+
+	MarkMetasoundDocumentDirty();
+}
+
+bool FMetasoundAssetBase::CopyDocumentAndInjectReceiveNodes(uint64 InInstanceID, const FMetasoundFrontendDocument& InSourceDoc, FMetasoundFrontendDocument& OutDestDoc) const
+{
+	using namespace Metasound;
+	using namespace Metasound::Frontend;
+
+	OutDestDoc = InSourceDoc;
+
+	FDocumentHandle Document = IDocumentController::CreateDocumentHandle(MakeAccessPtr<FDocumentAccessPtr>(OutDestDoc.AccessPoint, OutDestDoc));
+	FGraphHandle RootGraph = Document->GetRootGraph();
+
+	TArray<FSendInfoAndVertexName> SendInfoAndVertexes = GetSendInfos(InInstanceID);
+
+	// Inject receive nodes for each transmittable input
+	for (const FSendInfoAndVertexName& InfoAndVertexName : SendInfoAndVertexes)
 	{
-		return Metasound::Frontend::IsEqualArchetype(InArchetype, SupportedArchetype);
-	};
 
-	return Algo::AnyOf(GetPreferredMetasoundArchetypes(), IsEqualArchetype);
+		// Add receive node to graph
+		FMetasoundFrontendClassMetadata ReceiveNodeMetadata;
+		bool bSuccess = GetReceiveNodeMetadataForDataType(InfoAndVertexName.SendInfo.TypeName, ReceiveNodeMetadata);
+		if (!bSuccess)
+		{
+			// TODO: log warning
+			continue;
+		}
+		FNodeHandle ReceiveNode = RootGraph->AddNode(ReceiveNodeMetadata);
+
+		// Add receive node address to graph
+		FNodeHandle AddressNode = AddInputPinForSendAddress(InfoAndVertexName.SendInfo, RootGraph);
+		TArray<FOutputHandle> AddressNodeOutputs = AddressNode->GetOutputs();
+		if (AddressNodeOutputs.Num() != 1)
+		{
+			// TODO: log warning
+			continue;
+		}
+
+		FOutputHandle AddressOutput = AddressNodeOutputs[0];
+		TArray<FInputHandle> ReceiveAddressInput = ReceiveNode->GetInputsWithVertexName(Metasound::FReceiveNodeNames::GetAddressInputName());
+		if (ReceiveAddressInput.Num() != 1)
+		{
+			// TODO: log error
+			continue;
+		}
+
+		ensure(ReceiveAddressInput[0]->Connect(*AddressOutput));
+
+		// Swap input node connections with receive node connections
+		FNodeHandle InputNode = RootGraph->GetInputNodeWithName(InfoAndVertexName.VertexName);
+		if (!ensure(InputNode->GetOutputs().Num() == 1))
+		{
+			// TODO: handle input node with varying number of outputs or varying output types.
+			continue;
+		}
+
+		FOutputHandle InputNodeOutput = InputNode->GetOutputs()[0];
+
+		if (ensure(ReceiveNode->IsValid()))
+		{
+			TArray<FOutputHandle> ReceiveNodeOutputs = ReceiveNode->GetOutputs();
+			if (!ensure(ReceiveNodeOutputs.Num() == 1))
+			{
+				// TODO: handle array outputs and receive nodes of varying formats.
+				continue;
+			}
+
+			TArray<FInputHandle> ReceiveDefaultInputs = ReceiveNode->GetInputsWithVertexName(Metasound::FReceiveNodeNames::GetDefaultDataInputName());
+			if (ensure(ReceiveDefaultInputs.Num() == 1))
+			{
+				FOutputHandle ReceiverNodeOutput = ReceiveNodeOutputs[0];
+				for (FInputHandle NodeInput : InputNodeOutput->GetConnectedInputs())
+				{
+					// Swap connections to receiver node
+					ensure(InputNodeOutput->Disconnect(*NodeInput));
+					ensure(ReceiverNodeOutput->Connect(*NodeInput));
+				}
+
+				ReceiveDefaultInputs[0]->Connect(*InputNodeOutput);
+			}
+		}
+	}
+
+	return true;
 }
 
-const FMetasoundFrontendArchetype& FMetasoundAssetBase::GetPreferredMetasoundArchetype(const FMetasoundFrontendDocument& InDocument) const
+Metasound::FSendAddress FMetasoundAssetBase::CreateSendAddress(uint64 InInstanceID, const FString& InVertexName, const FName& InDataTypeName) const
 {
-	// Default to archetype provided in case it is supported. 
-	if (IsMetasoundArchetypeSupported(InDocument.Archetype))
+	using namespace Metasound;
+
+	FSendAddress Address;
+
+	Address.Subsystem = GetSubsystemNameForSendScope(ETransmissionScope::Global);
+	Address.ChannelName = FName(FString::Printf(TEXT("%d:%s:%s"), InInstanceID, *InVertexName, *InDataTypeName.ToString()));
+
+	return Address;
+}
+
+TArray<FMetasoundAssetBase::FSendInfoAndVertexName> FMetasoundAssetBase::GetSendInfos(uint64 InInstanceID) const
+{
+	using FSendInfo = Metasound::FMetasoundInstanceTransmitter::FSendInfo;
+	using namespace Metasound::Frontend;
+
+	TArray<FSendInfoAndVertexName> SendInfos;
+
+	FConstGraphHandle RootGraph = GetRootGraphHandle();
+
+	TArray<FString> SendVertices = GetTransmittableInputVertexNames();
+	for (const FString& VertexName : SendVertices)
 	{
-		return InDocument.Archetype;
+		FConstNodeHandle InputNode = RootGraph->GetInputNodeWithName(VertexName);
+		for (FConstInputHandle InputHandle : InputNode->GetConstInputs())
+		{
+			FSendInfoAndVertexName Info;
+
+			// TODO: incorporate VertexID into address. But need to ensure that VertexID
+			// will be maintained after injecting Receive nodes. 
+			Info.SendInfo.Address = CreateSendAddress(InInstanceID, InputHandle->GetName(), InputHandle->GetDataType());
+			Info.SendInfo.ParameterName = FName(InputHandle->GetDisplayName().ToString()); // TODO: display name hack. Need to have naming consistent in editor for inputs
+			//Info.SendInfo.ParameterName = FName(*InputHandle->GetName()); // TODO: this is the expected parameter name.
+			Info.SendInfo.TypeName = InputHandle->GetDataType();
+			Info.VertexName = VertexName;
+
+			SendInfos.Add(Info);
+		}
 	}
 
-	// If existing archetype is not supported, get the most similar that still supports the documents environment.
-	const FMetasoundFrontendArchetype* SimilarArchetype = Metasound::Frontend::FindMostSimilarArchetypeSupportingEnvironment(InDocument, GetPreferredMetasoundArchetypes());
+	return SendInfos;
+}
 
-	if (nullptr != SimilarArchetype)
+TArray<FString> FMetasoundAssetBase::GetTransmittableInputVertexNames() const
+{
+	using namespace Metasound;
+
+	// Unused inputs are all input vertices that are not in the archetype.
+	TArray<FString> ArchetypeInputVertexNames;
+	for (const FMetasoundFrontendClassVertex& Vertex : GetArchetype().Interface.Inputs)
 	{
-		return *SimilarArchetype;
+		ArchetypeInputVertexNames.Add(Vertex.Name);
 	}
 
-	// Nothing found. Return the existing archetype for the FMetasoundAssetBase.
-	return GetMetasoundArchetype();
+	Frontend::FConstGraphHandle RootGraph = GetRootGraphHandle();
+	TArray<FString> GraphInputVertexNames = RootGraph->GetInputVertexNames();
+
+	// Filter graph inputs by archetype inputs.
+	GraphInputVertexNames = GraphInputVertexNames.FilterByPredicate([&](const FString& InName) { return !ArchetypeInputVertexNames.Contains(InName); });
+
+	auto IsDataTypeTransmittable = [&](const FString& InVertexName)
+	{
+		Frontend::FConstClassInputAccessPtr ClassInputPtr = RootGraph->FindClassInputWithName(InVertexName);
+		if (const FMetasoundFrontendClassInput* ClassInput = ClassInputPtr.Get())
+		{
+			FDataTypeRegistryInfo TypeInfo;
+			if (Frontend::GetTraitsForDataType(ClassInput->TypeName, TypeInfo))
+			{
+				if (TypeInfo.bIsTransmittable)
+				{
+					// TODO: Currently values set directly on node pins are represented
+					// as input nodes in the graph. They should not be used for transmission
+					// as the number of input nodes increases quickly as more nodes
+					// are added to a graph. Connecting these input nodes to the 
+					// transmission system is relatively expensive. These undesirable input nodes are
+					// filtered out by ignoring input nodes which are not "Visible". 
+					Frontend::FConstNodeHandle InputNode = RootGraph->GetNodeWithID(ClassInput->NodeID);
+					if (InputNode->IsValid())
+					{
+						return true;
+					}
+				}
+			}
+		}
+		return false;
+	};
+
+	GraphInputVertexNames = GraphInputVertexNames.FilterByPredicate(IsDataTypeTransmittable);
+
+	return GraphInputVertexNames;
 }
 
-void FMetasoundAssetBase::SetDocument(const FMetasoundFrontendDocument& InDocument, bool bForceUpdateArchetype)
+Metasound::Frontend::FNodeHandle FMetasoundAssetBase::AddInputPinForSendAddress(const Metasound::FMetasoundInstanceTransmitter::FSendInfo& InSendInfo, Metasound::Frontend::FGraphHandle InGraph) const
 {
-	FMetasoundFrontendDocument& Document = GetDocumentChecked();
-	Document = InDocument;
+	FMetasoundFrontendClassInput Description;
+	FGuid VertexID = FGuid::NewGuid();
+
+	Description.Name = InSendInfo.Address.ChannelName.ToString();
+	Description.TypeName = Metasound::GetMetasoundDataTypeName<Metasound::FSendAddress>();
+	Description.Metadata.Description = FText::GetEmpty();
+	Description.VertexID = VertexID;
+	Description.DefaultLiteral.Set(InSendInfo.Address.ChannelName.ToString());
+
+	return InGraph->AddInputVertex(Description);
+}
+
+#if WITH_EDITORONLY_DATA
+FText FMetasoundAssetBase::GetDisplayName(FString&& InTypeName) const
+{
+	using namespace Metasound::Frontend;
 
-	FMetasoundFrontendArchetype NewArch = InDocument.Archetype;
+	FConstGraphHandle GraphHandle = GetRootGraphHandle();
+	const bool bIsPreset = !GraphHandle->GetGraphStyle().bIsGraphEditable;
 
-	if (bForceUpdateArchetype || (!IsMetasoundArchetypeSupported(NewArch)))
+	if (!bIsPreset)
 	{
-		NewArch = GetPreferredMetasoundArchetype(Document);
+		return FText::FromString(MoveTemp(InTypeName));
 	}
 
-	ensure(SetMetasoundArchetype(NewArch));
+	return FText::Format(LOCTEXT("PresetDisplayNameFormat", "{0} (Preset)"), FText::FromString(MoveTemp(InTypeName)));
 }
+#endif // WITH_EDITORONLY_DATA
 
-void FMetasoundAssetBase::ConformDocumentToMetasoundArchetype()
+bool FMetasoundAssetBase::GetReceiveNodeMetadataForDataType(const FName& InTypeName, FMetasoundFrontendClassMetadata& OutMetadata) const
 {
-	FMetasoundFrontendDocument& Document = GetDocumentChecked();
-	Document.Archetype = GetMetasoundArchetype();
+	using namespace Metasound;
+	using namespace Metasound::Frontend;
+
+	const FNodeClassName ClassName = FReceiveNodeNames::GetClassNameForDataType(InTypeName);
+	TArray<FMetasoundFrontendClass> ReceiverNodeClasses = ISearchEngine::Get().FindClassesWithName(ClassName, true /* bInSortByVersion */);
+
+	if (ReceiverNodeClasses.IsEmpty())
+	{
+		return false;
+	}
+
+	OutMetadata = ReceiverNodeClasses[0].Metadata;
+	return true;
+}
+
+void FMetasoundAssetBase::UpdateAssetTags(FMetasoundFrontendClassAssetTags& OutTags)
+{
+	using namespace Metasound::Frontend;
+
+	FGraphHandle GraphHandle = GetRootGraphHandle();
+
+	OutTags.ID = GraphHandle->GetClassID();
+
+	OutTags.Name = GraphHandle->GetGraphMetadata().ClassName.Name;
+	OutTags.Namespace = GraphHandle->GetGraphMetadata().ClassName.Namespace;
+	OutTags.Variant = GraphHandle->GetGraphMetadata().ClassName.Variant;
+
+	OutTags.MajorVersion = GraphHandle->GetGraphMetadata().Version.Major;
+	OutTags.MinorVersion = GraphHandle->GetGraphMetadata().Version.Minor;
 
-	Metasound::Frontend::FMatchRootGraphToArchetype Transform;
-	Transform.Transform(GetDocumentHandle());
 
-	MarkMetasoundDocumentDirty();
 }
 
 bool FMetasoundAssetBase::MarkMetasoundDocumentDirty() const
@@ -195,4 +466,4 @@ const FMetasoundFrontendDocument& FMetasoundAssetBase::GetDocumentChecked() cons
 	check(DocAccessPtr.IsValid());
 	return *DocAccessPtr;
 }
-
+#undef LOCTEXT_NAMESPACE // "MetaSound"
diff --git a/Engine/Plugins/Runtime/Metasound/Source/MetasoundFrontend/Private/MetasoundFrontend.cpp b/Engine/Plugins/Runtime/Metasound/Source/MetasoundFrontend/Private/MetasoundFrontend.cpp
--- a/Engine/Plugins/Runtime/Metasound/Source/MetasoundFrontend/Private/MetasoundFrontend.cpp
+++ b/Engine/Plugins/Runtime/Metasound/Source/MetasoundFrontend/Private/MetasoundFrontend.cpp
@@ -164,7 +164,7 @@ namespace Metasound
 			if (ensure(nullptr != Registry))
 			{
 				bool bSuccess = Registry->FindFrontendClassFromRegistered(InNodeInfo, OutClass);
-				ensureAlwaysMsgf(bSuccess, TEXT("Cannot generate description of unregistered node [NodeClassName:%s]"), *InNodeInfo.LookupKey.NodeClassFullName.ToString());
+				ensureAlwaysMsgf(bSuccess, TEXT("Cannot generate description of unregistered node [RegistryKey:%s]"), *InNodeInfo.LookupKey);
 			}
 
 			return OutClass;
diff --git a/Engine/Plugins/Runtime/Metasound/Source/MetasoundFrontend/Private/MetasoundFrontendInvalidController.h b/Engine/Plugins/Runtime/Metasound/Source/MetasoundFrontend/Private/MetasoundFrontendInvalidController.h
--- a/Engine/Plugins/Runtime/Metasound/Source/MetasoundFrontend/Private/MetasoundFrontendInvalidController.h
+++ b/Engine/Plugins/Runtime/Metasound/Source/MetasoundFrontend/Private/MetasoundFrontendInvalidController.h
@@ -170,7 +170,7 @@ namespace Metasound
 
 			const FText& GetDescription() const override { return MetasoundFrontendInvalidControllerPrivate::GetInvalid<FText>(); }
 
-			bool IsRequired() const override { return false; }
+			bool IsRequired(const FMetasoundFrontendArchetype& InArchetype) const override { return false; }
 
 			TSharedRef<IGraphController> AsGraph() override;
 			TSharedRef<const IGraphController> AsGraph() const override;
@@ -236,6 +236,9 @@ namespace Metasound
 			TArray<TSharedRef<const INodeController>> GetConstOutputNodes() const override { return TArray<TSharedRef<const INodeController>>(); }
 			TArray<TSharedRef<const INodeController>> GetConstInputNodes() const override { return TArray<TSharedRef<const INodeController>>(); }
 
+			const FMetasoundFrontendGraphStyle& GetGraphStyle() const override { return MetasoundFrontendInvalidControllerPrivate::GetInvalid<FMetasoundFrontendGraphStyle>(); }
+			void SetGraphStyle(const FMetasoundFrontendGraphStyle& InStyle) override { }
+
 			void IterateConstNodes(TUniqueFunction<void(FConstNodeHandle)> InFunction, EMetasoundFrontendClassType InClassType /* = EMetasoundFrontendClassType::Invalid */) const override { }
 			void IterateNodes(TUniqueFunction<void(FNodeHandle)> InFunction, EMetasoundFrontendClassType InClassType /* = EMetasoundFrontendClassType::Invalid */) override { }
 
@@ -296,6 +299,8 @@ namespace Metasound
 			// Returns the metadata for the current graph, including the name, description and author.
 			const FMetasoundFrontendClassMetadata& GetGraphMetadata() const override { return MetasoundFrontendInvalidControllerPrivate::GetInvalid<FMetasoundFrontendClassMetadata>(); }
 
+			void SetGraphMetadata(const FMetasoundFrontendClassMetadata& InMetadata) { }
+
 			TSharedRef<INodeController> CreateEmptySubgraph(const FMetasoundFrontendClassMetadata& InInfo) override { return FInvalidNodeController::GetInvalid(); }
 
 			TUniquePtr<IOperator> BuildOperator(const FOperatorSettings& InSettings, const FMetasoundEnvironment& InEnvironment, TArray<IOperatorBuilder::FBuildErrorPtr>& OutBuildErrors) const override
@@ -328,9 +333,6 @@ namespace Metasound
 
 				bool IsValid() const override { return false; }
 
-				const TArray<FMetasoundFrontendClassVertex>& GetRequiredInputs() const override { const static TArray<FMetasoundFrontendClassVertex> Empty; return Empty; }
-				const TArray<FMetasoundFrontendClassVertex>& GetRequiredOutputs() const override { const static TArray<FMetasoundFrontendClassVertex> Empty; return Empty; }
-
 				TArray<FMetasoundFrontendClass> GetDependencies() const override { return TArray<FMetasoundFrontendClass>(); }
 				TArray<FMetasoundFrontendGraphClass> GetSubgraphs() const override { return TArray<FMetasoundFrontendGraphClass>(); }
 				TArray<FMetasoundFrontendClass> GetClasses() const override { return TArray<FMetasoundFrontendClass>(); }
diff --git a/Engine/Plugins/Runtime/Metasound/Source/MetasoundFrontend/Private/MetasoundFrontendLiteral.cpp b/Engine/Plugins/Runtime/Metasound/Source/MetasoundFrontend/Private/MetasoundFrontendLiteral.cpp
--- a/Engine/Plugins/Runtime/Metasound/Source/MetasoundFrontend/Private/MetasoundFrontendLiteral.cpp
+++ b/Engine/Plugins/Runtime/Metasound/Source/MetasoundFrontend/Private/MetasoundFrontendLiteral.cpp
@@ -452,6 +452,32 @@ void FMetasoundFrontendLiteral::SetType(EMetasoundFrontendLiteralType InType)
 	}
 }
 
+bool FMetasoundFrontendLiteral::TryGet(UObject*& OutValue) const
+{
+	if (Type == EMetasoundFrontendLiteralType::UObject)
+	{
+		if (ensure(!AsUObject.IsEmpty()))
+		{
+			OutValue = AsUObject[0];
+			return true;
+		}
+	}
+
+	return false;
+}
+
+bool FMetasoundFrontendLiteral::TryGet(TArray<UObject*>& OutValue) const
+{
+	if (Type == EMetasoundFrontendLiteralType::UObject)
+	{
+		OutValue = AsUObject;
+		return true;
+	}
+
+	return false;
+}
+
+
 Metasound::FLiteral FMetasoundFrontendLiteral::ToLiteral(const FName& InMetasoundDataTypeName) const
 {
 	using namespace Metasound;
diff --git a/Engine/Plugins/Runtime/Metasound/Source/MetasoundFrontend/Private/MetasoundFrontendRegistries.cpp b/Engine/Plugins/Runtime/Metasound/Source/MetasoundFrontend/Private/MetasoundFrontendRegistries.cpp
--- a/Engine/Plugins/Runtime/Metasound/Source/MetasoundFrontend/Private/MetasoundFrontendRegistries.cpp
+++ b/Engine/Plugins/Runtime/Metasound/Source/MetasoundFrontend/Private/MetasoundFrontendRegistries.cpp
@@ -20,9 +20,44 @@ namespace Metasound
 	{
 		namespace MetasoundFrontendRegistryPrivate
 		{
+			const FString& GetClassTypeString(EMetasoundFrontendClassType InType)
+			{
+
+				static const FString InputType(TEXT("Input"));
+				static const FString OutputType(TEXT("Output"));
+				static const FString ExternalType(TEXT("External"));
+				static const FString VariableType(TEXT("Variable"));
+				static const FString GraphType(TEXT("Graph"));
+				static const FString InvalidType(TEXT("Invalid"));
+
+				switch (InType)
+				{
+					case EMetasoundFrontendClassType::Input:
+						return InputType;
+
+					case EMetasoundFrontendClassType::Output:
+						return OutputType;
+
+					case EMetasoundFrontendClassType::External:
+						return ExternalType;
+
+					case EMetasoundFrontendClassType::Variable:
+						return VariableType;
+
+					case EMetasoundFrontendClassType::Graph:
+						return GraphType;
+
+					default:
+						static_assert(static_cast<uint8>(EMetasoundFrontendClassType::Invalid) == 5, "Missing EMetasoundFrontendClassType case coverage");
+						return InvalidType;
+				}
+			}
+
 			// All registry keys should be created through this function to ensure consistency.
-			Frontend::FNodeRegistryKey GetRegistryKey(const FName& InClassName, int32 InMajorVersion)
+			Frontend::FNodeRegistryKey GetRegistryKey(EMetasoundFrontendClassType InType, const FString& InFullClassName, int32 InMajorVersion, int32 InMinorVersion)
 			{
+				FString RegistryKey = FString::Format(TEXT("{0}_{1}_{2}.{3}"), {*GetClassTypeString(InType), *InFullClassName, InMajorVersion, InMinorVersion});
+				/*
 				Frontend::FNodeRegistryKey Key;
 
 				Key.NodeClassFullName = InClassName;
@@ -32,6 +67,8 @@ namespace Metasound
 				Key.NodeHash = HashCombine(Key.NodeHash, FCrc::TypeCrc32(InMajorVersion));
 
 				return Key;
+				*/
+				return RegistryKey;
 			}
 
 			// Return the compatible literal with the most descriptive type.
@@ -169,7 +206,8 @@ namespace Metasound
 				 *
 				 * @return True on success.
 				 */
-				bool RegisterExternalNode(Metasound::FCreateMetasoundNodeFunction&& InCreateNode, Metasound::FCreateMetasoundFrontendClassFunction&& InCreateDescription) override;
+				FNodeRegistryKey RegisterExternalNode(Metasound::FCreateMetasoundNodeFunction&& InCreateNode, Metasound::FCreateMetasoundFrontendClassFunction&& InCreateDescription) override;
+				bool UnregisterExternalNode(const FNodeRegistryKey& InKey) override;
 				bool RegisterConversionNode(const FConverterNodeRegistryKey& InNodeKey, const FConverterNodeInfo& InNodeInfo) override;
 				bool IsNodeRegistered(const FNodeRegistryKey& InKey) const override;
 
@@ -571,7 +609,7 @@ namespace Metasound
 				return false;
 			}
 
-			bool FRegistryContainerImpl::RegisterExternalNode(Metasound::FCreateMetasoundNodeFunction&& InCreateNode, Metasound::FCreateMetasoundFrontendClassFunction&& InCreateDescription)
+			FNodeRegistryKey FRegistryContainerImpl::RegisterExternalNode(Metasound::FCreateMetasoundNodeFunction&& InCreateNode, Metasound::FCreateMetasoundFrontendClassFunction&& InCreateDescription)
 			{
 				FNodeRegistryElement RegistryElement = FNodeRegistryElement(MoveTemp(InCreateNode), MoveTemp(InCreateDescription));
 
@@ -582,10 +620,10 @@ namespace Metasound
 					// check to see if an identical node was already registered, and log
 					ensureAlwaysMsgf(
 						!ExternalNodeRegistry.Contains(Key),
-						TEXT("Node with identical name '%s' & Major Version already registered. "
+						TEXT("Node with registry key already registered. "
 							"The previously registered node will be overwritten. "
 							"This can happen if two classes share the same name or if METASOUND_REGISTER_NODE is defined in a public header."),
-						*Key.NodeClassFullName.ToString());
+						*Key);
 
 					// Store update to newly registered node in history so nodes
 					// can be queried by transaction ID
@@ -593,9 +631,24 @@ namespace Metasound
 					RegistryTransactionHistory.Add(MakeAddNodeRegistrationTransaction(ClassInfo));
 
 					// Store registry elements in map so nodes can be queried using registry key.
-					ExternalNodeRegistry.Add(MoveTemp(Key), MoveTemp(RegistryElement));
+					ExternalNodeRegistry.Add(Key, MoveTemp(RegistryElement));
+				}
 
-					return true;
+				return Key;
+			}
+
+			bool FRegistryContainerImpl::UnregisterExternalNode(const FNodeRegistryKey& InKey)
+			{
+				if (IsValidNodeRegistryKey(InKey))
+				{
+					if (ExternalNodeRegistry.Contains(InKey))
+					{
+						FNodeClassInfo ClassInfo = { EMetasoundFrontendClassType::External, InKey };
+						RegistryTransactionHistory.Add(MakeRemoveNodeRegistrationTransaction(ClassInfo));
+
+						ExternalNodeRegistry.Remove(InKey);
+						return true;
+					}
 				}
 
 				return false;
@@ -825,6 +878,11 @@ namespace Metasound
 				static_assert(static_cast<uint32>(EMetasoundFrontendClassType::Invalid) == 5, "Possible missing switch case coverage for EMetasoundFrontendClassType");
 			}
 		}
+
+		bool IsValidNodeRegistryKey(const FNodeRegistryKey& InKey)
+		{
+			return !InKey.IsEmpty();
+		}
 	}
 }
 
@@ -865,12 +923,12 @@ bool FMetasoundFrontendRegistryContainer::GetRegistryKey(const Metasound::Fronte
 
 Metasound::Frontend::FNodeRegistryKey FMetasoundFrontendRegistryContainer::GetRegistryKey(const FNodeClassMetadata& InNodeMetadata)
 {
-	return Metasound::Frontend::MetasoundFrontendRegistryPrivate::GetRegistryKey(InNodeMetadata.ClassName.GetFullName(), InNodeMetadata.MajorVersion);
+	return Metasound::Frontend::MetasoundFrontendRegistryPrivate::GetRegistryKey(EMetasoundFrontendClassType::External, InNodeMetadata.ClassName.GetFullName().ToString(), InNodeMetadata.MajorVersion, InNodeMetadata.MinorVersion);
 }
 
 Metasound::Frontend::FNodeRegistryKey FMetasoundFrontendRegistryContainer::GetRegistryKey(const FMetasoundFrontendClassMetadata& InNodeMetadata)
 {
-	return Metasound::Frontend::MetasoundFrontendRegistryPrivate::GetRegistryKey(InNodeMetadata.ClassName.GetFullName(), InNodeMetadata.Version.Major);
+	return Metasound::Frontend::MetasoundFrontendRegistryPrivate::GetRegistryKey(InNodeMetadata.Type, InNodeMetadata.ClassName.GetFullName().ToString(), InNodeMetadata.Version.Major, InNodeMetadata.Version.Minor);
 }
 
 bool FMetasoundFrontendRegistryContainer::GetFrontendClassFromRegistered(const FMetasoundFrontendClassMetadata& InMetadata, FMetasoundFrontendClass& OutClass)
diff --git a/Engine/Plugins/Runtime/Metasound/Source/MetasoundFrontend/Private/MetasoundFrontendStandardController.cpp b/Engine/Plugins/Runtime/Metasound/Source/MetasoundFrontend/Private/MetasoundFrontendStandardController.cpp
--- a/Engine/Plugins/Runtime/Metasound/Source/MetasoundFrontend/Private/MetasoundFrontendStandardController.cpp
+++ b/Engine/Plugins/Runtime/Metasound/Source/MetasoundFrontend/Private/MetasoundFrontendStandardController.cpp
@@ -639,13 +639,13 @@ namespace Metasound
 
 			if (ConverterInputs.Num() < 1)
 			{
-				UE_LOG(LogMetaSound, Warning, TEXT("Converter node [Name: %s] does not support preferred input vertex [Vertex: %s]"), *InConverterInfo.NodeKey.NodeClassFullName.ToString(), *InConverterInfo.PreferredConverterInputPin);
+				UE_LOG(LogMetaSound, Warning, TEXT("Converter node [Name: %s] does not support preferred input vertex [Vertex: %s]"), *ConverterNode->GetNodeName(), *InConverterInfo.PreferredConverterInputPin);
 				return false;
 			}
 
 			if (ConverterOutputs.Num() < 1)
 			{
-				UE_LOG(LogMetaSound, Warning, TEXT("Converter node [Name: %s] does not support preferred output vertex [Vertex: %s]"), *InConverterInfo.NodeKey.NodeClassFullName.ToString(), *InConverterInfo.PreferredConverterOutputPin);
+				UE_LOG(LogMetaSound, Warning, TEXT("Converter node [Name: %s] does not support preferred output vertex [Vertex: %s]"), *ConverterNode->GetNodeName(), *InConverterInfo.PreferredConverterOutputPin);
 				return false;
 			}
 
@@ -1402,7 +1402,7 @@ namespace Metasound
 			return FInvalidInputController::GetInvalid();
 		}
 
-		bool FBaseNodeController::IsRequired() const
+		bool FBaseNodeController::IsRequired(const FMetasoundFrontendArchetype& InArchetype) const
 		{
 			return false;
 		}
@@ -1776,12 +1776,12 @@ namespace Metasound
 			return OutputDisplayTitle;
 		}
 
-		bool FOutputNodeController::IsRequired() const
+		bool FOutputNodeController::IsRequired(const FMetasoundFrontendArchetype& InArchetype) const
 		{
 			if (NodePtr.IsValid() && OwningGraph->IsValid())
 			{
 				const FString& Name = NodePtr->Name;
-				const TArray<FMetasoundFrontendClassVertex>& RequiredOutputs = OwningGraph->GetOwningDocument()->GetRequiredOutputs();
+				const TArray<FMetasoundFrontendClassVertex>& RequiredOutputs = InArchetype.Interface.Outputs;
 				for (const FMetasoundFrontendClassVertex& OutputVertex : RequiredOutputs)
 				{
 					if (OutputVertex.Name == Name)
@@ -1930,12 +1930,12 @@ namespace Metasound
 			return InputDisplayTitle;
 		}
 
-		bool FInputNodeController::IsRequired() const
+		bool FInputNodeController::IsRequired(const FMetasoundFrontendArchetype& InArchetype) const
 		{
 			if (NodePtr.IsValid() && OwningGraph->IsValid())
 			{
 				const FString& Name = NodePtr->Name;
-				const TArray<FMetasoundFrontendClassVertex>& RequiredInputs = OwningGraph->GetOwningDocument()->GetRequiredInputs();
+				const TArray<FMetasoundFrontendClassVertex>& RequiredInputs = InArchetype.Interface.Inputs;
 				for (const FMetasoundFrontendClassVertex& InputVertex : RequiredInputs)
 				{
 					if (InputVertex.Name == Name)
@@ -2139,6 +2139,24 @@ namespace Metasound
 			return GetNodeByPredicate(IsNodeWithSameID);
 		}
 
+		const FMetasoundFrontendGraphStyle& FGraphController::GetGraphStyle() const
+		{
+			if (GraphClassPtr.IsValid())
+			{
+				return GraphClassPtr->Graph.Style;
+			}
+
+			return FrontendControllerIntrinsics::GetInvalidValueConstRef<FMetasoundFrontendGraphStyle>();
+		}
+
+		void FGraphController::SetGraphStyle(const FMetasoundFrontendGraphStyle& InStyle)
+		{
+			if (GraphClassPtr.IsValid())
+			{
+				GraphClassPtr->Graph.Style = InStyle;
+			}
+		}
+
 		TArray<FNodeHandle> FGraphController::GetOutputNodes()
 		{
 			auto IsOutputNode = [](const FMetasoundFrontendClass& NodeClass, const FMetasoundFrontendNode& Node) 
@@ -2675,6 +2693,14 @@ namespace Metasound
 			return FrontendControllerIntrinsics::GetInvalidValueConstRef<FMetasoundFrontendClassMetadata>();
 		}
 
+		void FGraphController::SetGraphMetadata(const FMetasoundFrontendClassMetadata& InMetadata)
+		{
+			if (GraphClassPtr.IsValid())
+			{
+				GraphClassPtr->Metadata = InMetadata;
+			}
+		}
+
 		FNodeHandle FGraphController::CreateEmptySubgraph(const FMetasoundFrontendClassMetadata& InInfo)
 		{
 			if (IsValid())
@@ -3298,28 +3324,6 @@ namespace Metasound
 			return DocumentPtr.IsValid();
 		}
 
-		const TArray<FMetasoundFrontendClassVertex>& FDocumentController::GetRequiredInputs() const
-		{
-			if (DocumentPtr.IsValid())
-			{
-				return DocumentPtr->Archetype.Interface.Inputs;
-			}
-
-			static const TArray<FMetasoundFrontendClassVertex> Empty;
-			return Empty;
-		}
-
-		const TArray<FMetasoundFrontendClassVertex>& FDocumentController::GetRequiredOutputs() const
-		{
-			if (DocumentPtr.IsValid())
-			{
-				return DocumentPtr->Archetype.Interface.Outputs;
-			}
-
-			static const TArray<FMetasoundFrontendClassVertex> Empty;
-			return Empty;
-		}
-
 		TArray<FMetasoundFrontendClass> FDocumentController::GetDependencies() const
 		{
 			if (DocumentPtr.IsValid())
diff --git a/Engine/Plugins/Runtime/Metasound/Source/MetasoundFrontend/Private/MetasoundFrontendStandardController.h b/Engine/Plugins/Runtime/Metasound/Source/MetasoundFrontend/Private/MetasoundFrontendStandardController.h
--- a/Engine/Plugins/Runtime/Metasound/Source/MetasoundFrontend/Private/MetasoundFrontendStandardController.h
+++ b/Engine/Plugins/Runtime/Metasound/Source/MetasoundFrontend/Private/MetasoundFrontendStandardController.h
@@ -419,7 +419,7 @@ namespace Metasound
 			TArray<FOutputHandle> GetOutputsWithVertexName(const FString& InName) override;
 			TArray<FConstOutputHandle> GetConstOutputsWithVertexName(const FString& InName) const override;
 
-			bool IsRequired() const override;
+			bool IsRequired(const FMetasoundFrontendArchetype& InArchetype) const override;
 
 			/** Returns an input with the given id. 
 			 *
@@ -580,7 +580,7 @@ namespace Metasound
 			void SetDescription(const FText& InDescription) override;
 			void SetDisplayName(const FText& InText) override;
 			const FText& GetDisplayTitle() const override;
-			bool IsRequired() const override;
+			bool IsRequired(const FMetasoundFrontendArchetype& InArchetype) const override;
 
 		protected:
 
@@ -637,7 +637,7 @@ namespace Metasound
 			const FText& GetDescription() const override;
 			const FText& GetDisplayName() const override;
 			const FText& GetDisplayTitle() const override;
-			bool IsRequired() const override;
+			bool IsRequired(const FMetasoundFrontendArchetype& InArchetype) const override;
 			bool IsValid() const override;
 			void SetDescription(const FText& InDescription) override;
 			void SetDisplayName(const FText& InText) override;
@@ -725,6 +725,9 @@ namespace Metasound
 			TArray<FNodeHandle> GetOutputNodes() override;
 			TArray<FConstNodeHandle> GetConstOutputNodes() const override;
 
+			const FMetasoundFrontendGraphStyle& GetGraphStyle() const override;
+			void SetGraphStyle(const FMetasoundFrontendGraphStyle& InStyle) override;
+
 			bool ContainsInputVertexWithName(const FString& InName) const override;
 			bool ContainsOutputVertexWithName(const FString& InName) const override;
 
@@ -788,6 +791,8 @@ namespace Metasound
 			// Returns the metadata for the current graph, including the name, description and author.
 			const FMetasoundFrontendClassMetadata& GetGraphMetadata() const override;
 
+			void SetGraphMetadata(const FMetasoundFrontendClassMetadata& InClassMetadata) override;
+
 			FNodeHandle CreateEmptySubgraph(const FMetasoundFrontendClassMetadata& InInfo) override;
 
 			TUniquePtr<IOperator> BuildOperator(const FOperatorSettings& InSettings, const FMetasoundEnvironment& InEnvironment, TArray<IOperatorBuilder::FBuildErrorPtr>& OutBuildErrors) const override;
@@ -896,9 +901,6 @@ namespace Metasound
 
 			bool IsValid() const override;
 
-			const TArray<FMetasoundFrontendClassVertex>& GetRequiredInputs() const override;
-			const TArray<FMetasoundFrontendClassVertex>& GetRequiredOutputs() const override;
-
 			TArray<FMetasoundFrontendClass> GetDependencies() const override;
 			TArray<FMetasoundFrontendGraphClass> GetSubgraphs() const override;
 			TArray<FMetasoundFrontendClass> GetClasses() const override;
diff --git a/Engine/Plugins/Runtime/Metasound/Source/MetasoundFrontend/Private/MetasoundFrontendTransform.cpp b/Engine/Plugins/Runtime/Metasound/Source/MetasoundFrontend/Private/MetasoundFrontendTransform.cpp
--- a/Engine/Plugins/Runtime/Metasound/Source/MetasoundFrontend/Private/MetasoundFrontendTransform.cpp
+++ b/Engine/Plugins/Runtime/Metasound/Source/MetasoundFrontend/Private/MetasoundFrontendTransform.cpp
@@ -2,7 +2,7 @@
 #include "MetasoundFrontendTransform.h"
 
 #include "MetasoundFrontendDocument.h"
-
+#include "MetasoundLog.h"
 
 namespace Metasound
 {
@@ -18,8 +18,8 @@ namespace Metasound
 			bool bDidEdit = false;
 			FGraphHandle Graph = InDocument->GetRootGraph();
 
-			TArray<FMetasoundFrontendClassVertex> RequiredInputs = InDocument->GetRequiredInputs();
-			TArray<FMetasoundFrontendClassVertex> RequiredOutputs = InDocument->GetRequiredOutputs();
+			const TArray<FMetasoundFrontendClassVertex>& RequiredInputs = Archetype.Interface.Inputs;
+			const TArray<FMetasoundFrontendClassVertex>& RequiredOutputs = Archetype.Interface.Outputs;
 
 			// Go through each input and add or swap if something is missing.
 			for (const FMetasoundFrontendClassVertex& RequiredInput : RequiredInputs)
@@ -102,6 +102,7 @@ namespace Metasound
 					NewMetadata.Version.Number = GetTargetVersion();
 					InDocument->SetMetadata(NewMetadata);
 
+					UE_LOG(LogMetaSound, Display, TEXT("Versioned MetaSound Document to %s"), *Metadata.Version.Number.ToString());
 					return true;
 				}
 		};
@@ -112,7 +113,7 @@ namespace Metasound
 		public:
 			FMetasoundFrontendVersionNumber GetTargetVersion() const override
 			{
-				return FMetasoundFrontendVersionNumber { 1, 1 };
+				return { 1, 1 };
 			}
 
 			void TransformInternal(FDocumentHandle InDocument) const override
@@ -157,6 +158,43 @@ namespace Metasound
 			}
 		};
 
+		/** Versions document from 1.1 to 1.2. */
+		class FVersionDocument_1_2 : public FVersionDocumentTransform
+		{
+		private:
+			const FName Name;
+			const FString& Path;
+
+		public:
+			FVersionDocument_1_2(const FName InName, const FString& InPath)
+				: Name(InName)
+				, Path(InPath)
+			{
+			}
+
+			FMetasoundFrontendVersionNumber GetTargetVersion() const override
+			{
+				return FVersionDocument::GetMaxVersion();
+			}
+
+			void TransformInternal(FDocumentHandle InDocument) const override
+			{
+				FGraphHandle GraphHandle = InDocument->GetRootGraph();
+				FMetasoundFrontendClassMetadata Metadata = GraphHandle->GetGraphMetadata();
+
+				Metadata.ClassName = FMetasoundFrontendClassName { "GraphAsset", Name, *Path };
+				Metadata.DisplayName = FText::FromString(Name.ToString());
+
+				GraphHandle->SetGraphMetadata(Metadata);
+			}
+		};
+
+		FVersionDocument::FVersionDocument(FName InName, const FString& InPath)
+			: Name(InName)
+			, Path(InPath)
+		{
+		}
+
 		bool FVersionDocument::Transform(FDocumentHandle InDocument) const
 		{
 			if (!ensure(InDocument->IsValid()))
@@ -168,7 +206,7 @@ namespace Metasound
 
 			// Add additional transforms here after defining them above, example below.
 			bWasUpdated |= FVersionDocument_1_1().Transform(InDocument);
-
+			bWasUpdated |= FVersionDocument_1_2(Name, Path).Transform(InDocument);
 			return bWasUpdated;
 		}
 	} // namespace Frontend
diff --git a/Engine/Plugins/Runtime/Metasound/Source/MetasoundFrontend/Public/MetasoundAssetBase.h b/Engine/Plugins/Runtime/Metasound/Source/MetasoundFrontend/Public/MetasoundAssetBase.h
--- a/Engine/Plugins/Runtime/Metasound/Source/MetasoundFrontend/Public/MetasoundAssetBase.h
+++ b/Engine/Plugins/Runtime/Metasound/Source/MetasoundFrontend/Public/MetasoundAssetBase.h
@@ -8,6 +8,7 @@
 #include "MetasoundFrontendController.h"
 #include "MetasoundFrontendDocumentAccessPtr.h"
 #include "MetasoundGraph.h"
+#include "MetasoundInstanceTransmitter.h"
 #include "MetasoundLog.h"
 #include "UObject/WeakObjectPtrTemplates.h"
 
@@ -19,41 +20,17 @@ class UEdGraph;
  * Metasound assets.  It provides consistent access to FMetasoundFrontendDocuments, control
  * over the FMetasoundFrontendArchetype of the FMetasoundFrontendDocument.  It also enables the UObject
  * to be utilized by a host of other engine tools built to support Metasounds.
- *
- * Example Usage:
- *
- * class UMyMetasoundClass : public UObject, public FMetasoundAssetBase
- * {
- *
- *     	UPROPERTY()
- *     	FMetasoundFrontendDocument MetasoundDocument;
- *
- * #if WITH_EDITORONLY_DATA
- * 		UPROPERTY()
- * 		UEdGraph* Graph;
- * #endif // WITH_EDITORONLY_DATA
- *
- * 	public:
- * 		// FMetasoundAssetBase Interface
- * 		...
- *      // End FMetasoundAssetBase Interface
- * };
- *
  */
 class METASOUNDFRONTEND_API FMetasoundAssetBase
 {
 public:
 	static const FString FileExtension;
 
-	FMetasoundAssetBase();
-
-	// Construct an FMetasoundAssetBase with a default Metasound Archetype.
-	// @param InDefaultArchetype - Default archetype for a Metasound Document in this class.
-	FMetasoundAssetBase(const FMetasoundFrontendArchetype& InDefaultArchetype);
-
+	FMetasoundAssetBase() = default;
 	virtual ~FMetasoundAssetBase() = default;
 
 #if WITH_EDITORONLY_DATA
+	virtual FText GetDisplayName() const = 0;
 
 	// Returns the graph associated with this Metasound. Graph is required to be referenced on
 	// Metasound UObject for editor serialization purposes.
@@ -67,32 +44,36 @@ public:
 	// Metasound UObject for editor serialization purposes.
 	// @param Editor graph associated with this metasound object.
 	virtual void SetGraph(UEdGraph* InGraph) = 0;
+
 #endif // WITH_EDITORONLY_DATA
 
+	// Registers the root graph of the given asset with the MetaSound Frontend.
+	void RegisterGraphWithFrontend();
+
+	bool CopyDocumentAndInjectReceiveNodes(uint64 InInstanceID, const FMetasoundFrontendDocument& InSourceDoc, FMetasoundFrontendDocument& OutDestDoc) const;
+
 	// Sets/overwrites the root class metadata
 	virtual void SetMetadata(FMetasoundFrontendClassMetadata& InMetadata);
 
 	// Returns  a description of the required inputs and outputs for this metasound UClass.
-	const FMetasoundFrontendArchetype& GetMetasoundArchetype() const;
+	virtual const FMetasoundFrontendArchetype& GetArchetype() const = 0;
 
-	// Sets/overwrites the document archetype.
-	// @param InArchetype - The desired archetype .
-	//
-	// @return True on success. False if archetype is not supported by this object. 
-	bool SetMetasoundArchetype(const FMetasoundFrontendArchetype& InArchetype);
+	// Returns true if the archetype is supported by this object.
+	virtual bool IsArchetypeSupported(const FMetasoundFrontendArchetype& InArchetype) const;
 
 	// Returns an array of archetypes preferred for this class.
-	virtual const TArray<FMetasoundFrontendArchetype>& GetPreferredMetasoundArchetypes() const = 0;
-
-	// Returns true if the archetype is supported by this object.
-	virtual bool IsMetasoundArchetypeSupported(const FMetasoundFrontendArchetype& InArchetype) const;
+	virtual const TArray<FMetasoundFrontendArchetype>& GetPreferredArchetypes() const = 0;
 
 	// Returns the preferred archetype for the given document.
-	virtual const FMetasoundFrontendArchetype& GetPreferredMetasoundArchetype(const FMetasoundFrontendDocument& InDocument) const;
+	virtual const FMetasoundFrontendArchetype& GetPreferredArchetypes(const FMetasoundFrontendDocument& InDocument, const FMetasoundFrontendArchetype& InDefaultArchetype) const;
 
 	// Returns the root class metadata
 	FMetasoundFrontendClassMetadata GetMetadata();
 
+	// TODO:
+	//virtual void OnMetaSoundDependencyAdded(const FMetasoundFrontendClassMetadata& InMetadata) = 0;
+	//virtual void OnMetaSoundDependencyRemoved(const FMetasoundFrontendClassMetadata& InMetadata) = 0;
+
 	// Imports data from a JSON string directly
 	bool ImportFromJSON(const FString& InJSON);
 
@@ -112,19 +93,31 @@ public:
 	// the FMetasoundAssetBase be while queried for a new one using `GetPreferredArchetype`. If `bForceUpdateArchetype`
 	// is true, `GetPreferredArchetype` will be used whether or not the provided document's archetype
 	// is supported. 
-	void SetDocument(const FMetasoundFrontendDocument& InDocument, bool bForceUpdateArchetype=false);
+	void SetDocument(const FMetasoundFrontendDocument& InDocument);
 
-	
 	FMetasoundFrontendDocument& GetDocumentChecked();
 	const FMetasoundFrontendDocument& GetDocumentChecked() const;
 
 	// This must be called on UObject::PostLoad, as well as in this asset's UFactory, to fix up the root document based on the most recent version of the archetype.
-	void ConformDocumentToMetasoundArchetype();
+	void ConformDocumentToArchetype();
 
 	// Calls the outermost package and marks it dirty. 
 	bool MarkMetasoundDocumentDirty() const;
 
 protected:
+	struct FSendInfoAndVertexName
+	{
+		Metasound::FMetasoundInstanceTransmitter::FSendInfo SendInfo;
+		FString VertexName;
+	};
+
+	TArray<FSendInfoAndVertexName> GetSendInfos(uint64 InInstanceID) const;
+
+#if WITH_EDITORONLY_DATA
+	FText GetDisplayName(FString&& InTypeName) const;
+#endif // WITH_EDITORONLY_DATA
+
+	void UpdateAssetTags(FMetasoundFrontendClassAssetTags& OutTags);
 
 	// Returns an access pointer to the document.
 	virtual Metasound::Frontend::FDocumentAccessPtr GetDocument() = 0;
@@ -138,7 +131,12 @@ protected:
 	// Returns the owning asset responsible for transactions applied to metasound
 	virtual const UObject* GetOwningAsset() const = 0;
 
+
 private:
+	bool GetReceiveNodeMetadataForDataType(const FName& InTypeName, FMetasoundFrontendClassMetadata& OutMetadata) const;
+	TArray<FString> GetTransmittableInputVertexNames() const;
+	Metasound::FSendAddress CreateSendAddress(uint64 InInstanceID, const FString& InVertexName, const FName& InDataTypeName) const;
+	Metasound::Frontend::FNodeHandle AddInputPinForSendAddress(const Metasound::FMetasoundInstanceTransmitter::FSendInfo& InSendInfo, Metasound::Frontend::FGraphHandle InGraph) const;
 
-	FMetasoundFrontendArchetype Archetype;
+	Metasound::Frontend::FNodeRegistryKey RegistryKey;
 };
diff --git a/Engine/Plugins/Runtime/Metasound/Source/MetasoundFrontend/Public/MetasoundFrontendController.h b/Engine/Plugins/Runtime/Metasound/Source/MetasoundFrontend/Public/MetasoundFrontendController.h
--- a/Engine/Plugins/Runtime/Metasound/Source/MetasoundFrontend/Public/MetasoundFrontendController.h
+++ b/Engine/Plugins/Runtime/Metasound/Source/MetasoundFrontend/Public/MetasoundFrontendController.h
@@ -342,7 +342,7 @@ namespace Metasound
 			virtual TArray<FConstOutputHandle> GetConstOutputsWithVertexName(const FString& InName) const = 0;
 
 			/** Returns true if node is required to satisfy the document archetype. */
-			virtual bool IsRequired() const = 0;
+			virtual bool IsRequired(const FMetasoundFrontendArchetype& InArchetype) const = 0;
 
 			/** Returns an input with the given id.
 			 *
@@ -446,6 +446,9 @@ namespace Metasound
 			/** Return the metadata for the current graph. */
 			virtual const FMetasoundFrontendClassMetadata& GetGraphMetadata() const = 0;
 
+			/** Sets the metadata for the current graph. */
+			virtual void SetGraphMetadata(const FMetasoundFrontendClassMetadata& InMetadata) = 0;
+
 			/** Return the display name of the graph. */
 			virtual const FText& GetDisplayName() const = 0;
 
@@ -467,6 +470,12 @@ namespace Metasound
 			/** Returns all output nodes in the graph. */
 			virtual TArray<FNodeHandle> GetOutputNodes() = 0;
 
+			// Returns graph style.
+			virtual const FMetasoundFrontendGraphStyle& GetGraphStyle() const = 0;
+
+			// Sets graph style.
+			virtual void SetGraphStyle(const FMetasoundFrontendGraphStyle& InStyle) = 0;
+
 			/** Returns all output nodes in the graph. */
 			virtual TArray<FConstNodeHandle> GetConstOutputNodes() const = 0;
 
@@ -696,12 +705,6 @@ namespace Metasound
 			/** Returns true if the controller is in a valid state. */
 			virtual bool IsValid() const = 0;
 
-			/** Returns an array of inputs which describe the required inputs needed to satisfy the document archetype. */
-			virtual const TArray<FMetasoundFrontendClassVertex>& GetRequiredInputs() const = 0;
-
-			/** Returns an array of outputs which describe the required outputs needed to satisfy the document archetype. */
-			virtual const TArray<FMetasoundFrontendClassVertex>& GetRequiredOutputs() const = 0;
-
 			// TODO: add info on environment variables. 
 			// TODO: consider find/add subgraph
 			// TODO: perhaps functions returning access pointers could be removed from main interface and only exist in FDocumentController.
diff --git a/Engine/Plugins/Runtime/Metasound/Source/MetasoundFrontend/Public/MetasoundFrontendDocument.h b/Engine/Plugins/Runtime/Metasound/Source/MetasoundFrontend/Public/MetasoundFrontendDocument.h
--- a/Engine/Plugins/Runtime/Metasound/Source/MetasoundFrontend/Public/MetasoundFrontendDocument.h
+++ b/Engine/Plugins/Runtime/Metasound/Source/MetasoundFrontend/Public/MetasoundFrontendDocument.h
@@ -565,22 +565,25 @@ struct FMetasoundFrontendInterfaceStyle
 	{
 		TArray<HandleType> SortedHandles = InHandles;
 
-		TMap<FGuid, int32> HandleIDToSortIndex;
-		for (int32 i = 0; i < DefaultSortOrder.Num(); ++i)
+		if (ensure(SortedHandles.Num() == DefaultSortOrder.Num()))
 		{
-			if (InHandles.IsValidIndex(i))
+			TMap<FGuid, int32> HandleIDToSortIndex;
+			for (int32 i = 0; i < DefaultSortOrder.Num(); ++i)
 			{
-				const int32 SortIndex = DefaultSortOrder[i];
-				HandleIDToSortIndex.Add(InHandles[i]->GetID(), SortIndex);
+				if (InHandles.IsValidIndex(i))
+				{
+					const int32 SortIndex = DefaultSortOrder[i];
+					HandleIDToSortIndex.Add(InHandles[i]->GetID(), SortIndex);
+				}
 			}
-		}
 
-		SortedHandles.Sort([&](const HandleType& HandleA, const HandleType& HandleB)
-		{
-			const FGuid HandleAID = HandleA->GetID();
-			const FGuid HandleBID = HandleB->GetID();
-			return HandleIDToSortIndex[HandleAID] < HandleIDToSortIndex[HandleBID];
-		});
+			SortedHandles.Sort([&](const HandleType& HandleA, const HandleType& HandleB)
+			{
+				const FGuid HandleAID = HandleA->GetID();
+				const FGuid HandleBID = HandleB->GetID();
+				return HandleIDToSortIndex[HandleAID] < HandleIDToSortIndex[HandleBID];
+			});
+		}
 
 		return SortedHandles;
 	}
@@ -650,7 +653,6 @@ struct METASOUNDFRONTEND_API FMetasoundFrontendClassName
 		return { Namespace, Name, Variant };
 	}
 
-
 	// Return string version of full name.
 	FString ToString() const;
 
@@ -716,6 +718,30 @@ struct FMetasoundFrontendEditorData
 	TArray<uint8> Data;
 };
 
+USTRUCT()
+struct METASOUNDFRONTEND_API FMetasoundFrontendClassAssetTags
+{
+	GENERATED_BODY()
+
+	UPROPERTY()
+	FGuid ID = Metasound::FrontendInvalidID;
+
+	UPROPERTY()
+	FName Namespace;
+
+	UPROPERTY()
+	FName Name;
+
+	UPROPERTY()
+	FName Variant;
+
+	UPROPERTY()
+	int32 MajorVersion = -1;
+
+	UPROPERTY()
+	int32 MinorVersion = -1;
+};
+
 USTRUCT()
 struct METASOUNDFRONTEND_API FMetasoundFrontendClass
 {
@@ -819,9 +845,6 @@ struct METASOUNDFRONTEND_API FMetasoundFrontendDocument
 	UPROPERTY()
 	FMetasoundFrontendEditorData EditorData;
 
-	UPROPERTY()
-	FMetasoundFrontendArchetype Archetype;
-
 	UPROPERTY()
 	TArray<FMetasoundFrontendClass> Dependencies;
 };
diff --git a/Engine/Plugins/Runtime/Metasound/Source/MetasoundFrontend/Public/MetasoundFrontendLiteral.h b/Engine/Plugins/Runtime/Metasound/Source/MetasoundFrontend/Public/MetasoundFrontendLiteral.h
--- a/Engine/Plugins/Runtime/Metasound/Source/MetasoundFrontend/Public/MetasoundFrontendLiteral.h
+++ b/Engine/Plugins/Runtime/Metasound/Source/MetasoundFrontend/Public/MetasoundFrontendLiteral.h
@@ -83,6 +83,9 @@ public:
 	// Returns true if the literal is in a valid state (Type != EMetasoundFrontendLiteralType::Invalid)
 	bool IsValid() const;
 
+	bool TryGet(UObject*& OutValue) const;
+	bool TryGet(TArray<UObject*>& OutValue) const;
+
 	// Sets the literal to the given type and value to default;
 	void SetType(EMetasoundFrontendLiteralType InType);
 
diff --git a/Engine/Plugins/Runtime/Metasound/Source/MetasoundFrontend/Public/MetasoundFrontendRegistries.h b/Engine/Plugins/Runtime/Metasound/Source/MetasoundFrontend/Public/MetasoundFrontendRegistries.h
--- a/Engine/Plugins/Runtime/Metasound/Source/MetasoundFrontend/Public/MetasoundFrontendRegistries.h
+++ b/Engine/Plugins/Runtime/Metasound/Source/MetasoundFrontend/Public/MetasoundFrontendRegistries.h
@@ -136,12 +136,13 @@ namespace Metasound
 
 	namespace Frontend
 	{
-
+		/* TODO: Node registry key needs to use full version
 		struct METASOUNDFRONTEND_API FNodeRegistryKey
 		{
 			// The class name for the node.
 			FName NodeClassFullName;
 
+
 			// A hash generated from the input types and output types for this node.
 			uint32 NodeHash = 0;
 
@@ -155,6 +156,11 @@ namespace Metasound
 				return InKey.NodeHash;
 			}
 		};
+		*/
+
+		using FNodeRegistryKey = FString;
+
+		METASOUNDFRONTEND_API bool IsValidNodeRegistryKey(const FNodeRegistryKey& InKey);
 
 		// Struct with the basics of a node class' information,
 		// used to look up that node from our node browser functions,
@@ -347,7 +353,8 @@ public:
 	 *
 	 * @return True on success.
 	 */
-	virtual bool RegisterExternalNode(Metasound::FCreateMetasoundNodeFunction&& InCreateNode, Metasound::FCreateMetasoundFrontendClassFunction&& InCreateDescription) = 0;
+	virtual FNodeRegistryKey RegisterExternalNode(Metasound::FCreateMetasoundNodeFunction&& InCreateNode, Metasound::FCreateMetasoundFrontendClassFunction&& InCreateDescription) = 0;
+	virtual bool UnregisterExternalNode(const FNodeRegistryKey& InKey) = 0;
 	virtual bool RegisterConversionNode(const FConverterNodeRegistryKey& InNodeKey, const FConverterNodeInfo& InNodeInfo) = 0;
 	virtual bool IsNodeRegistered(const FNodeRegistryKey& InKey) const = 0;
 
diff --git a/Engine/Plugins/Runtime/Metasound/Source/MetasoundFrontend/Public/MetasoundFrontendTransform.h b/Engine/Plugins/Runtime/Metasound/Source/MetasoundFrontend/Public/MetasoundFrontendTransform.h
--- a/Engine/Plugins/Runtime/Metasound/Source/MetasoundFrontend/Public/MetasoundFrontendTransform.h
+++ b/Engine/Plugins/Runtime/Metasound/Source/MetasoundFrontend/Public/MetasoundFrontendTransform.h
@@ -1,8 +1,9 @@
 // Copyright Epic Games, Inc. All Rights Reserved.
-
 #pragma once
 
 #include "MetasoundFrontendController.h"
+#include "MetasoundFrontendDocument.h"
+
 
 namespace Metasound
 {
@@ -20,14 +21,32 @@ namespace Metasound
 		class METASOUNDFRONTEND_API FMatchRootGraphToArchetype : public IDocumentTransform
 		{
 			public:
+				FMatchRootGraphToArchetype(const FMetasoundFrontendArchetype& InArchetype)
+					: Archetype(InArchetype)
+				{
+				}
+
 				bool Transform(FDocumentHandle InDocument) const override;
+
+			private:
+				FMetasoundFrontendArchetype Archetype;
 		};
 
 		/** Base class for versioning a document. */
 		class METASOUNDFRONTEND_API FVersionDocument : public IDocumentTransform
 		{
+			const FName Name;
+			const FString& Path;
+
 			public:
-				bool Transform(FDocumentHandle InDocument) const;
+				static FMetasoundFrontendVersionNumber GetMaxVersion()
+				{
+					return FMetasoundFrontendVersionNumber { 1, 2 };
+				}
+
+				FVersionDocument(FName InName, const FString& InPath);
+
+				bool Transform(FDocumentHandle InDocument) const override;
 		};
 	}
 }
diff --git a/Engine/Plugins/Runtime/Metasound/Source/MetasoundFrontend/Public/MetasoundNodeRegistrationMacro.h b/Engine/Plugins/Runtime/Metasound/Source/MetasoundFrontend/Public/MetasoundNodeRegistrationMacro.h
--- a/Engine/Plugins/Runtime/Metasound/Source/MetasoundFrontend/Public/MetasoundNodeRegistrationMacro.h
+++ b/Engine/Plugins/Runtime/Metasound/Source/MetasoundFrontend/Public/MetasoundNodeRegistrationMacro.h
@@ -40,7 +40,7 @@ namespace Metasound
 
 		bAlreadyRegisteredThisDataType = true;
 
-		bool bSuccessfullyRegisteredNode = FMetasoundFrontendRegistryContainer::Get()->RegisterExternalNode(
+		Frontend::FNodeRegistryKey Key = FMetasoundFrontendRegistryContainer::Get()->RegisterExternalNode(
 			[](const Metasound::FNodeInitData& InInitData) -> TUniquePtr<Metasound::INode>
 			{
 				return TUniquePtr<Metasound::INode>(new FNodeType(InInitData));
@@ -50,7 +50,8 @@ namespace Metasound
 				return Metasound::Frontend::GenerateClassDescription(InMetadata);
 			}
 		);
-
+			
+		const bool bSuccessfullyRegisteredNode = Frontend::IsValidNodeRegistryKey(Key);
 		ensureAlwaysMsgf(bSuccessfullyRegisteredNode, TEXT("Registering node class failed. Please check the logs."));
 		return bSuccessfullyRegisteredNode;
 	}
