Disable MetaSound Graph Registration for now

diff --git a/Engine/Plugins/Runtime/Metasound/Source/MetasoundEngine/Private/MetasoundSource.cpp b/Engine/Plugins/Runtime/Metasound/Source/MetasoundEngine/Private/MetasoundSource.cpp
--- a/Engine/Plugins/Runtime/Metasound/Source/MetasoundEngine/Private/MetasoundSource.cpp
+++ b/Engine/Plugins/Runtime/Metasound/Source/MetasoundEngine/Private/MetasoundSource.cpp
@@ -61,7 +61,7 @@ void UMetaSoundSource::PreSave(FObjectPreSaveContext SaveContext)
 	// TODO: Move these to be run anytime the interface changes or
 	// the node is versioned
 	UpdateAssetTags(AssetTags);
-	RegisterGraphWithFrontend();
+	// RegisterGraphWithFrontend();
 	// TODO: Post re-register, determine if interface changed and
 	// test/figure out when we should run updates on referencing
 	// assets (presets only and on breaking changes? pop-up dialog?
