PR #8124: Fix cooking when MetaSound plugin is enabled (Contributed by [REDACTED_PERSON])

diff --git a/Engine/Plugins/Runtime/Metasound/Metasound.uplugin b/Engine/Plugins/Runtime/Metasound/Metasound.uplugin
--- a/Engine/Plugins/Runtime/Metasound/Metasound.uplugin
+++ b/Engine/Plugins/Runtime/Metasound/Metasound.uplugin
@@ -43,7 +43,7 @@
 		{
 			"Name": "MetasoundEditor",
 			"Type": "UncookedOnly",
-			"LoadingPhase": "PostEngineInit"
+			"LoadingPhase": "PostDefault"
 		}
 	],
 	"Plugins": [
