Adding dedup logic for subgraph node registration

diff --git a/Engine/Plugins/Runtime/Metasound/Source/MetasoundFrontend/Private/MetasoundAssetBase.cpp b/Engine/Plugins/Runtime/Metasound/Source/MetasoundFrontend/Private/MetasoundAssetBase.cpp
--- a/Engine/Plugins/Runtime/Metasound/Source/MetasoundFrontend/Private/MetasoundAssetBase.cpp
+++ b/Engine/Plugins/Runtime/Metasound/Source/MetasoundFrontend/Private/MetasoundAssetBase.cpp
@@ -25,8 +25,6 @@ void FMetasoundAssetBase::RegisterGraphWithFrontend()
 	using namespace Metasound;
 	using namespace Metasound::Frontend;
 
-	// TODO: unregister prior version.
-	
 	FConstDocumentAccessPtr DocumentPtr = GetDocument();
 	FString AssetName;
 	if (const UObject* OwningAsset = GetOwningAsset())
@@ -35,6 +33,7 @@ void FMetasoundAssetBase::RegisterGraphWithFrontend()
 	}
 
 
+	// Function for creating a core INode.
 	FCreateMetasoundNodeFunction CreateNode = [Name=AssetName, DocumentPtr=DocumentPtr](const Metasound::FNodeInitData& InInitData) -> TUniquePtr<INode>
 	{
 		check(IsInGameThread());
@@ -50,6 +49,7 @@ void FMetasoundAssetBase::RegisterGraphWithFrontend()
 		}
 	};
 
+	// Function for creating a frontend class.
 	FCreateMetasoundFrontendClassFunction CreateFrontendClass = [Name=AssetName, DocumentPtr=DocumentPtr]() -> FMetasoundFrontendClass
 	{
 		check(IsInGameThread());
@@ -70,6 +70,7 @@ void FMetasoundAssetBase::RegisterGraphWithFrontend()
 
 	if (IsValidNodeRegistryKey(RegistryKey))
 	{
+		// Unregister prior version if it exists.
 		ensure(FMetasoundFrontendRegistryContainer::Get()->UnregisterExternalNode(RegistryKey));
 	}
 
diff --git a/Engine/Plugins/Runtime/Metasound/Source/MetasoundFrontend/Private/MetasoundFrontend.cpp b/Engine/Plugins/Runtime/Metasound/Source/MetasoundFrontend/Private/MetasoundFrontend.cpp
--- a/Engine/Plugins/Runtime/Metasound/Source/MetasoundFrontend/Private/MetasoundFrontend.cpp
+++ b/Engine/Plugins/Runtime/Metasound/Source/MetasoundFrontend/Private/MetasoundFrontend.cpp
@@ -42,12 +42,12 @@ namespace Metasound
 			}
 		}
 
-		TArray<FNodeClassInfo> GetNodeClassesRegisteredSince(FRegistryTransactionID InTransactionID, FRegistryTransactionID* OutCurrentTransactionID)
+		TArray<const IRegistryTransaction*> GetRegistryTransactionsSince(FRegistryTransactionID InTransactionID, FRegistryTransactionID* OutCurrentTransactionID)
 		{
 			FMetasoundFrontendRegistryContainer* Registry = FMetasoundFrontendRegistryContainer::Get();
 			if (ensure(nullptr != Registry))
 			{
-				return Registry->GetNodeClassesRegisteredSince(InTransactionID, OutCurrentTransactionID);
+				return Registry->GetRegistryTransactionsSince(InTransactionID, OutCurrentTransactionID);
 			}
 			else
 			{
@@ -55,7 +55,7 @@ namespace Metasound
 				{
 					*OutCurrentTransactionID = GetOriginRegistryTransactionID();
 				}
-				return TArray<FNodeClassInfo>();
+				return TArray<const IRegistryTransaction*>();
 			}
 		}
 
diff --git a/Engine/Plugins/Runtime/Metasound/Source/MetasoundFrontend/Private/MetasoundFrontendQuery.cpp b/Engine/Plugins/Runtime/Metasound/Source/MetasoundFrontend/Private/MetasoundFrontendQuery.cpp
--- a/Engine/Plugins/Runtime/Metasound/Source/MetasoundFrontend/Private/MetasoundFrontendQuery.cpp
+++ b/Engine/Plugins/Runtime/Metasound/Source/MetasoundFrontend/Private/MetasoundFrontendQuery.cpp
@@ -32,6 +32,24 @@ namespace Metasound
 			FStreamFunction Func;
 		};
 
+		// Wrapper for step defined by function
+		struct FTransformFunctionFrontendQueryStep: IFrontendQueryTransformStep
+		{
+			using FTransformFunction = FFrontendQueryStep::FTransformFunction;
+
+			FTransformFunctionFrontendQueryStep(FTransformFunction InFunc)
+			:	Func(InFunc)
+			{
+			}
+
+			void Transform(FFrontendQueryEntry::FValue& InValue) const override
+			{
+				Func(InValue);
+			}
+
+			FTransformFunction Func;
+		};
+
 		// Wrapper for step defined by function
 		struct FMapFunctionFrontendQueryStep: IFrontendQueryMapStep
 		{
@@ -265,6 +283,44 @@ namespace Metasound
 			}
 		};
 
+		struct FTransformStepExecuter : TStepExecuter<IFrontendQueryTransformStep>
+		{
+			using TStepExecuter<IFrontendQueryTransformStep>::TStepExecuter;
+
+			virtual EResultModificationState Increment(FFrontendQuerySelection& InOutResult) override
+			{
+				EResultModificationState ResultState = EResultModificationState::Unmodified;
+
+				if (Step.IsValid())
+				{
+					if (InOutResult.GetSelection().Num() > 0)
+					{
+						ResultState = EResultModificationState::Modified;
+					}
+
+					for (FFrontendQueryEntry* Entry : InOutResult.GetSelection())
+					{
+						if (nullptr != Entry)
+						{
+							Step->Transform(Entry->Value);
+						}
+					}
+				}
+
+				return ResultState;
+			}
+
+			virtual EResultModificationState Merge(const FFrontendQuerySelection& InIncremental, FFrontendQuerySelection& InOutResult) override
+			{
+				return Append(InIncremental, InOutResult);
+			}
+
+			virtual EResultModificationState Execute(FFrontendQuerySelection& InOutResult) override
+			{
+				return Increment(InOutResult);
+			}
+		};
+
 		struct FMapStepExecuter : TStepExecuter<IFrontendQueryMapStep>
 		{
 			using TStepExecuter<IFrontendQueryMapStep>::TStepExecuter;
@@ -308,10 +364,18 @@ namespace Metasound
 			using FKey = FFrontendQueryEntry::FKey;
 
 		private:
+			static const FKey InvalidKey;
 
-			static FFrontendQueryEntry::FKey GetKey(const FFrontendQueryEntry* InEntry)
+			static const FFrontendQueryEntry::FKey& GetKey(const FFrontendQueryEntry* InEntry)
 			{
-				return (nullptr != InEntry) ? InEntry->Key : FFrontendQueryEntry::InvalidKey;
+				if (nullptr != InEntry)
+				{
+					return InEntry->Key;
+				}
+				else
+				{
+					return InvalidKey;
+				}
 			}
 
 			// Reduce entries.
@@ -446,6 +510,8 @@ namespace Metasound
 			}
 		};
 
+		const FReduceStepExecuter::FKey FReduceStepExecuter::InvalidKey;
+
 		struct FFilterStepExecuter : TStepExecuter<IFrontendQueryFilterStep>
 		{
 			using TStepExecuter<IFrontendQueryFilterStep>::TStepExecuter;
@@ -647,6 +713,105 @@ namespace Metasound
 		};
 	}
 
+	FFrontendQueryKey::FFrontendQueryKey()
+	: Key(TInPlaceType<FFrontendQueryKey::FInvalid>())
+	, Hash(INDEX_NONE)
+	{}
+
+	FFrontendQueryKey::FFrontendQueryKey(int32 InKey)
+	: Key(TInPlaceType<int32>(), InKey)
+	, Hash(GetTypeHash(InKey))
+	{
+	}
+
+	FFrontendQueryKey::FFrontendQueryKey(const FString& InKey)
+	: Key(TInPlaceType<FString>(), InKey)
+	, Hash(GetTypeHash(InKey))
+	{
+	}
+
+	FFrontendQueryKey::FFrontendQueryKey(const FName& InKey)
+	: Key(TInPlaceType<FName>(), InKey)
+	, Hash(GetTypeHash(InKey))
+	{
+	}
+
+	bool FFrontendQueryKey::IsValid() const
+	{
+		return Key.GetIndex() != FKeyType::IndexOfType<FFrontendQueryKey::FInvalid>();
+	}
+
+	bool operator==(const FFrontendQueryKey& InLHS, const FFrontendQueryKey& InRHS)
+	{
+		if (InLHS.Hash == InRHS.Hash)
+		{
+			if (InLHS.IsValid() && InRHS.IsValid())
+			{
+				if (InLHS.Key.GetIndex() == InRHS.Key.GetIndex())
+				{
+					switch(InLHS.Key.GetIndex())
+					{
+						case FFrontendQueryKey::FKeyType::IndexOfType<int32>():
+							return InLHS.Key.Get<int32>() == InRHS.Key.Get<int32>();
+
+						case FFrontendQueryKey::FKeyType::IndexOfType<FString>():
+							return InLHS.Key.Get<FString>() == InRHS.Key.Get<FString>();
+
+						case FFrontendQueryKey::FKeyType::IndexOfType<FName>():
+							return InLHS.Key.Get<FName>() == InRHS.Key.Get<FName>();
+
+						default:
+							// Unhandled case type.
+							checkNoEntry();
+					}
+				}
+			}
+		}
+
+		return false;
+	}
+
+	bool operator!=(const FFrontendQueryKey& InLHS, const FFrontendQueryKey& InRHS)
+	{
+		return !(InLHS == InRHS);
+	}
+
+	bool operator<(const FFrontendQueryKey& InLHS, const FFrontendQueryKey& InRHS)
+	{
+		if (InLHS.Hash != InRHS.Hash)
+		{
+			return InLHS.Hash < InRHS.Hash;
+		}
+
+		if (InLHS.Key.GetIndex() != InRHS.Key.GetIndex())
+		{
+			return InLHS.Key.GetIndex() < InRHS.Key.GetIndex();
+		}
+
+		switch(InLHS.Key.GetIndex())
+		{
+			case FFrontendQueryKey::FKeyType::IndexOfType<int32>():
+				return InLHS.Key.Get<int32>() < InRHS.Key.Get<int32>();
+
+			case FFrontendQueryKey::FKeyType::IndexOfType<FString>():
+				return InLHS.Key.Get<FString>() < InRHS.Key.Get<FString>();
+
+			case FFrontendQueryKey::FKeyType::IndexOfType<FName>():
+				return InLHS.Key.Get<FName>().FastLess(InRHS.Key.Get<FName>());
+
+			default:
+				// Unhandled case type.
+				checkNoEntry();
+		}
+
+		return false;
+	}
+
+	uint32 GetTypeHash(const FFrontendQueryKey& InKey)
+	{
+		return InKey.Hash;
+	}
+
 	FFrontendQuerySelection::FFrontendQuerySelection(const FFrontendQuerySelection& InOther)
 	{
 		AppendToStorageAndSelection(InOther);
@@ -847,6 +1012,11 @@ namespace Metasound
 	{
 	}
 
+	FFrontendQueryStep::FFrontendQueryStep(FTransformFunction&& InFunc)
+	:	StepExecuter(MakeUnique<FrontendQueryPrivate::FTransformStepExecuter>(MakeUnique<FrontendQueryPrivate::FTransformFunctionFrontendQueryStep>(MoveTemp(InFunc))))
+	{
+	}
+
 	FFrontendQueryStep::FFrontendQueryStep(FMapFunction&& InFunc)
 	:	StepExecuter(MakeUnique<FrontendQueryPrivate::FMapStepExecuter>(MakeUnique<FrontendQueryPrivate::FMapFunctionFrontendQueryStep>(MoveTemp(InFunc))))
 	{
@@ -887,6 +1057,11 @@ namespace Metasound
 	{
 	}
 	
+	FFrontendQueryStep::FFrontendQueryStep(TUniquePtr<IFrontendQueryTransformStep>&& InStep)
+	:	StepExecuter(MakeUnique<FrontendQueryPrivate::FTransformStepExecuter>(MoveTemp(InStep)))
+	{
+	}
+
 	FFrontendQueryStep::FFrontendQueryStep(TUniquePtr<IFrontendQueryMapStep>&& InStep)
 	:	StepExecuter(MakeUnique<FrontendQueryPrivate::FMapStepExecuter>(MoveTemp(InStep)))
 	{
@@ -981,6 +1156,11 @@ namespace Metasound
 		return AddFunctionStep(MoveTemp(InFunc));
 	}
 
+	FFrontendQuery& FFrontendQuery::AddTransformLambdaStep(FTransformFunction&& InFunc)
+	{
+		return AddFunctionStep(MoveTemp(InFunc));
+	}
+
 	FFrontendQuery& FFrontendQuery::AddMapLambdaStep(FMapFunction&& InFunc)
 	{
 		return AddFunctionStep(MoveTemp(InFunc));
diff --git a/Engine/Plugins/Runtime/Metasound/Source/MetasoundFrontend/Private/MetasoundFrontendQuerySteps.cpp b/Engine/Plugins/Runtime/Metasound/Source/MetasoundFrontend/Private/MetasoundFrontendQuerySteps.cpp
--- a/Engine/Plugins/Runtime/Metasound/Source/MetasoundFrontend/Private/MetasoundFrontendQuerySteps.cpp
+++ b/Engine/Plugins/Runtime/Metasound/Source/MetasoundFrontend/Private/MetasoundFrontendQuerySteps.cpp
@@ -9,7 +9,6 @@
 
 namespace Metasound
 {
-
 	FNodeClassRegistrationEvents::FNodeClassRegistrationEvents()
 	: CurrentTransactionID(Frontend::GetOriginRegistryTransactionID())
 	{
@@ -17,12 +16,17 @@ namespace Metasound
 
 	void FNodeClassRegistrationEvents::Stream(TArray<FFrontendQueryEntry>& OutEntries)
 	{
-		const TArray<Frontend::FNodeClassInfo> ClassInfos = Frontend::GetNodeClassesRegisteredSince(CurrentTransactionID, &CurrentTransactionID);
+		using namespace Frontend;
+
+		const TArray<const IRegistryTransaction*> Transactions = GetRegistryTransactionsSince(CurrentTransactionID, &CurrentTransactionID);
 
-		for (const Frontend::FNodeClassInfo& ClassInfo : ClassInfos)
+		for (const IRegistryTransaction* Transaction : Transactions)
 		{
-			FFrontendQueryEntry::FValue Value(TInPlaceType<FMetasoundFrontendClass>(), Frontend::GenerateClassDescription(ClassInfo));
-			OutEntries.Emplace(MoveTemp(Value));
+			if (const FNodeClassInfo* Info = Transaction->GetNodeClassInfo())
+			{
+				FFrontendQueryEntry::FValue Value(TInPlaceType<const IRegistryTransaction*>(), Transaction);
+				OutEntries.Emplace(MoveTemp(Value));
+			}
 		}
 	}
 
@@ -31,6 +35,84 @@ namespace Metasound
 		CurrentTransactionID = Frontend::GetOriginRegistryTransactionID();
 	}
 
+	FFrontendQueryEntry::FKey FMapRegistrationEventsToNodeRegistryKeys::Map(const FFrontendQueryEntry& InEntry) const 
+	{
+		using namespace Frontend;
+
+		FNodeRegistryKey RegistryKey;
+
+		if (ensure(InEntry.Value.IsType<const IRegistryTransaction*>()))
+		{
+			if (const IRegistryTransaction* Transaction = InEntry.Value.Get<const IRegistryTransaction*>())
+			{
+				const FNodeClassInfo* ClassInfo = Transaction->GetNodeClassInfo();
+
+				if (nullptr != ClassInfo)
+				{
+					RegistryKey = ClassInfo->LookupKey;
+				}
+			}
+		}
+
+		return RegistryKey;
+	}
+
+	void FReduceRegistrationEventsToCurrentStatus::Reduce(FFrontendQueryEntry::FKey InKey, TArrayView<FFrontendQueryEntry * const>& InEntries, FReduceOutputView& OutResult) const
+	{
+		using namespace Frontend;
+
+		int32 State = 0;
+		FFrontendQueryEntry* FinalEntry = nullptr;
+
+		for (FFrontendQueryEntry* Entry : InEntries)
+		{
+			if (ensure(Entry->Value.IsType<const IRegistryTransaction*>()))
+			{
+				if (const IRegistryTransaction* Transaction = Entry->Value.Get<const IRegistryTransaction*>())
+				{
+					switch (Transaction->GetTransactionType())
+					{
+						case ETransactionType::Add:
+							State++;
+							FinalEntry = Entry;
+							break;
+
+						case ETransactionType::Remove:
+							State--;
+							break;
+					}
+				}
+			}
+		}
+
+		if ((nullptr != FinalEntry) && (State > 0))
+		{
+			OutResult.Add(*FinalEntry);
+		}
+	}
+
+	void FTransformRegistrationEventsToClasses::Transform(FFrontendQueryEntry::FValue& InValue) const
+	{
+		using namespace Frontend;
+
+		FMetasoundFrontendClass FrontendClass;
+
+		if (ensure(InValue.IsType<const IRegistryTransaction*>()))
+		{
+			if (const IRegistryTransaction* Transaction = InValue.Get<const IRegistryTransaction*>())
+			{
+				const FNodeClassInfo* ClassInfo = Transaction->GetNodeClassInfo();
+				if (nullptr != ClassInfo)
+				{
+					bool bSuccess = FMetasoundFrontendRegistryContainer::Get()->FindFrontendClassFromRegistered(*ClassInfo, FrontendClass);
+					check(bSuccess);
+				}
+			}
+		}
+
+		InValue.Set<FMetasoundFrontendClass>(MoveTemp(FrontendClass));
+	}
+
 	FFilterClassesByInputVertexDataType::FFilterClassesByInputVertexDataType(const FName& InTypeName)
 	:	InputVertexTypeName(InTypeName)
 	{
@@ -83,11 +165,10 @@ namespace Metasound
 		return InEntry.Value.Get<FMetasoundFrontendClass>().ID == ClassID;
 	}
 
-	FFrontendQueryEntry::FKey FMapClassNameToMajorVersion::Map(const FFrontendQueryEntry& InEntry) const
+	FFrontendQueryEntry::FKey FMapToFullClassName::Map(const FFrontendQueryEntry& InEntry) const
 	{
 		const FMetasoundFrontendClass& FrontendClass = InEntry.Value.Get<FMetasoundFrontendClass>();
-		const uint32 HashKey = GetTypeHash(FrontendClass.Metadata.ClassName.GetFullName());
-		return static_cast<FFrontendQueryEntry::FKey>(HashKey);
+		return FFrontendQueryEntry::FKey(FrontendClass.Metadata.ClassName.GetFullName());
 	}
 
 	void FReduceClassesToHighestVersion::Reduce(FFrontendQueryEntry::FKey InKey, TArrayView<FFrontendQueryEntry * const>& InEntries, FReduceOutputView& OutResult) const
diff --git a/Engine/Plugins/Runtime/Metasound/Source/MetasoundFrontend/Private/MetasoundFrontendRegistries.cpp b/Engine/Plugins/Runtime/Metasound/Source/MetasoundFrontend/Private/MetasoundFrontendRegistries.cpp
--- a/Engine/Plugins/Runtime/Metasound/Source/MetasoundFrontend/Private/MetasoundFrontendRegistries.cpp
+++ b/Engine/Plugins/Runtime/Metasound/Source/MetasoundFrontend/Private/MetasoundFrontendRegistries.cpp
@@ -22,7 +22,6 @@ namespace Metasound
 		{
 			const FString& GetClassTypeString(EMetasoundFrontendClassType InType)
 			{
-
 				static const FString InputType(TEXT("Input"));
 				static const FString OutputType(TEXT("Output"));
 				static const FString ExternalType(TEXT("External"));
@@ -213,7 +212,8 @@ namespace Metasound
 
 				// Get all available nodes
 				TArray<Frontend::FNodeClassInfo> GetAllAvailableNodeClasses(Frontend::FRegistryTransactionID* OutCurrentRegistryTransactionID) const override;
-				TArray<Metasound::Frontend::FNodeClassInfo> GetNodeClassesRegisteredSince(Metasound::Frontend::FRegistryTransactionID InSince, Metasound::Frontend::FRegistryTransactionID* OutCurrentRegistryTransactionID) const override;
+
+				TArray<const Metasound::Frontend::IRegistryTransaction*> GetRegistryTransactionsSince(Metasound::Frontend::FRegistryTransactionID InSince, Metasound::Frontend::FRegistryTransactionID* OutCurrentRegistryTransactionID) const override;
 
 				// Return any data types that can be used as a metasound input type or output type.
 				TArray<FName> GetAllValidDataTypes() override;
@@ -339,21 +339,9 @@ namespace Metasound
 				return OutClasses;
 			}
 
-			TArray<Metasound::Frontend::FNodeClassInfo> FRegistryContainerImpl::GetNodeClassesRegisteredSince(Metasound::Frontend::FRegistryTransactionID InSince, Metasound::Frontend::FRegistryTransactionID* OutCurrentRegistryTransactionID) const 
+			TArray<const Metasound::Frontend::IRegistryTransaction*> FRegistryContainerImpl::GetRegistryTransactionsSince(Metasound::Frontend::FRegistryTransactionID InSince, Metasound::Frontend::FRegistryTransactionID* OutCurrentRegistryTransactionID) const 
 			{
-				TArray<FNodeClassInfo> OutClasses;
-
-				TArray<const IRegistryTransaction*> Transactions = RegistryTransactionHistory.GetTransactions(InSince, OutCurrentRegistryTransactionID);
-
-				for (const IRegistryTransaction* Transaction : Transactions)
-				{
-					if (const FNodeClassInfo* Info = Transaction->GetNodeClassInfo())
-					{
-						OutClasses.Add(*Info);
-					}
-				}
-
-				return OutClasses;
+				return RegistryTransactionHistory.GetTransactions(InSince, OutCurrentRegistryTransactionID);
 			}
 
 			TUniquePtr<Metasound::INode> FRegistryContainerImpl::ConstructInputNode(const FName& InInputType, Metasound::FInputNodeConstructorParams&& InParams)
@@ -628,7 +616,7 @@ namespace Metasound
 					// Store update to newly registered node in history so nodes
 					// can be queried by transaction ID
 					FNodeClassInfo ClassInfo = { EMetasoundFrontendClassType::External, Key };
-					RegistryTransactionHistory.Add(MakeAddNodeRegistrationTransaction(ClassInfo));
+					RegistryTransactionHistory.Add(MakeAddNodeRegistryTransaction(ClassInfo));
 
 					// Store registry elements in map so nodes can be queried using registry key.
 					ExternalNodeRegistry.Add(Key, MoveTemp(RegistryElement));
@@ -644,7 +632,7 @@ namespace Metasound
 					if (ExternalNodeRegistry.Contains(InKey))
 					{
 						FNodeClassInfo ClassInfo = { EMetasoundFrontendClassType::External, InKey };
-						RegistryTransactionHistory.Add(MakeRemoveNodeRegistrationTransaction(ClassInfo));
+						RegistryTransactionHistory.Add(MakeRemoveNodeRegistryTransaction(ClassInfo));
 
 						ExternalNodeRegistry.Remove(InKey);
 						return true;
@@ -877,14 +865,14 @@ namespace Metasound
 				}
 				static_assert(static_cast<uint32>(EMetasoundFrontendClassType::Invalid) == 5, "Possible missing switch case coverage for EMetasoundFrontendClassType");
 			}
-		}
+		} // namespace MetasoundFrontendRegistriesPrivate
 
 		bool IsValidNodeRegistryKey(const FNodeRegistryKey& InKey)
 		{
 			return !InKey.IsEmpty();
 		}
-	}
-}
+	} // namespace Frontend
+} // namespace Metasound
 
 FMetasoundFrontendRegistryContainer* FMetasoundFrontendRegistryContainer::LazySingleton = nullptr;
 
diff --git a/Engine/Plugins/Runtime/Metasound/Source/MetasoundFrontend/Private/MetasoundFrontendRegistryTransaction.cpp b/Engine/Plugins/Runtime/Metasound/Source/MetasoundFrontend/Private/MetasoundFrontendRegistryTransaction.cpp
--- a/Engine/Plugins/Runtime/Metasound/Source/MetasoundFrontend/Private/MetasoundFrontendRegistryTransaction.cpp
+++ b/Engine/Plugins/Runtime/Metasound/Source/MetasoundFrontend/Private/MetasoundFrontendRegistryTransaction.cpp
@@ -92,15 +92,14 @@ namespace Metasound
 
 				return Result;
 			}
-			
 		}
 
-		FRegistryTransactionPtr MakeAddNodeRegistrationTransaction(const FNodeClassInfo& InInfo)
+		FRegistryTransactionPtr MakeAddNodeRegistryTransaction(const FNodeClassInfo& InInfo)
 		{
-			class FNodeRegistrationTransaction : public IRegistryTransaction
+			class FNodeRegistryTransaction : public IRegistryTransaction
 			{
 			public:
-				FNodeRegistrationTransaction(const FNodeClassInfo& InNodeClassInfo)
+				FNodeRegistryTransaction(const FNodeClassInfo& InNodeClassInfo)
 				: NodeClassInfo(InNodeClassInfo)
 				{
 				}
@@ -112,7 +111,7 @@ namespace Metasound
 
 				TUniquePtr<IRegistryTransaction> Clone() const override
 				{
-					return MakeUnique<FNodeRegistrationTransaction>(*this);
+					return MakeUnique<FNodeRegistryTransaction>(*this);
 				}
 
 				const FNodeClassInfo* GetNodeClassInfo() const 
@@ -125,15 +124,15 @@ namespace Metasound
 				FNodeClassInfo NodeClassInfo;
 			};
 
-			return MakeUnique<FNodeRegistrationTransaction>(InInfo);
+			return MakeUnique<FNodeRegistryTransaction>(InInfo);
 		}
 
-		FRegistryTransactionPtr MakeRemoveNodeRegistrationTransaction(const FNodeClassInfo& InInfo)
+		FRegistryTransactionPtr MakeRemoveNodeRegistryTransaction(const FNodeClassInfo& InInfo)
 		{
-			class FNodeRegistrationTransaction : public IRegistryTransaction
+			class FNodeRegistryTransaction : public IRegistryTransaction
 			{
 			public:
-				FNodeRegistrationTransaction(const FNodeClassInfo& InNodeClassInfo)
+				FNodeRegistryTransaction(const FNodeClassInfo& InNodeClassInfo)
 				: NodeClassInfo(InNodeClassInfo)
 				{
 				}
@@ -145,7 +144,7 @@ namespace Metasound
 
 				TUniquePtr<IRegistryTransaction> Clone() const override
 				{
-					return MakeUnique<FNodeRegistrationTransaction>(*this);
+					return MakeUnique<FNodeRegistryTransaction>(*this);
 				}
 
 				const FNodeClassInfo* GetNodeClassInfo() const 
@@ -158,7 +157,7 @@ namespace Metasound
 				FNodeClassInfo NodeClassInfo;
 			};
 
-			return MakeUnique<FNodeRegistrationTransaction>(InInfo);
+			return MakeUnique<FNodeRegistryTransaction>(InInfo);
 		}
 	}
 }
diff --git a/Engine/Plugins/Runtime/Metasound/Source/MetasoundFrontend/Private/MetasoundFrontendRegistryTransaction.h b/Engine/Plugins/Runtime/Metasound/Source/MetasoundFrontend/Private/MetasoundFrontendRegistryTransaction.h
--- a/Engine/Plugins/Runtime/Metasound/Source/MetasoundFrontend/Private/MetasoundFrontendRegistryTransaction.h
+++ b/Engine/Plugins/Runtime/Metasound/Source/MetasoundFrontend/Private/MetasoundFrontendRegistryTransaction.h
@@ -4,45 +4,20 @@
 
 #include "CoreMinimal.h"
 #include "MetasoundFrontend.h"
+#include "MetasoundFrontendRegistries.h"
 
 namespace Metasound
 {
 	namespace Frontend
 	{
-		using FRegistryTransactionID = int32;
-
 		/** Returns an ID representing the beginning of the transaction history. */
 		FRegistryTransactionID GetOriginRegistryTransactionID();
 
-		/** Describes the type of transaction. */
-		enum class ETransactionType : uint8
-		{
-			Add,     //< Something was added to the registry.
-			Remove   //< Something was removed from the registry.
-		};
-
-		/** Interface for a registry transaction. */
-		class IRegistryTransaction
-		{
-		public:
-			virtual ~IRegistryTransaction() = default;
-			virtual TUniquePtr<IRegistryTransaction> Clone() const = 0;
-
-			/** Returns the type of transaction */
-			virtual ETransactionType GetTransactionType() const = 0;
-
-			/** If a FNodeClassInfo added during the transaction, this will return
-			 * a non-null pointer to the FNodeClassInfo. */
-			virtual const FNodeClassInfo* GetNodeClassInfo() const = 0;
-		};
-
-		using FRegistryTransactionPtr = TUniquePtr<IRegistryTransaction>;
-
 		/** Create a node registration transaction. */
-		FRegistryTransactionPtr MakeAddNodeRegistrationTransaction(const FNodeClassInfo& InInfo);
+		FRegistryTransactionPtr MakeAddNodeRegistryTransaction(const FNodeClassInfo& InInfo);
 
 		/** Create a node unregistration transaction. */
-		FRegistryTransactionPtr MakeRemoveNodeRegistrationTransaction(const FNodeClassInfo& InInfo);
+		FRegistryTransactionPtr MakeRemoveNodeRegistryTransaction(const FNodeClassInfo& InInfo);
 
 		/** Maintains a history of IRegistryTransactions. Calls are threadsafe (excluding
 		 * the constructor and destructor.)
diff --git a/Engine/Plugins/Runtime/Metasound/Source/MetasoundFrontend/Private/MetasoundFrontendSearchEngine.cpp b/Engine/Plugins/Runtime/Metasound/Source/MetasoundFrontend/Private/MetasoundFrontendSearchEngine.cpp
--- a/Engine/Plugins/Runtime/Metasound/Source/MetasoundFrontend/Private/MetasoundFrontendSearchEngine.cpp
+++ b/Engine/Plugins/Runtime/Metasound/Source/MetasoundFrontend/Private/MetasoundFrontendSearchEngine.cpp
@@ -64,11 +64,17 @@ namespace Metasound
 			{
 				TUniquePtr<FFrontendQuery> NewQuery = MakeUnique<FFrontendQuery>();
 
-				NewQuery->AddStep<FNodeClassRegistrationEvents>();
+				// TODO: this query will be slow to update when new nodes are registered. 
+				// Consider reworking call sites, or creating a reduce function
+				// which performs the majority of the query logic.
+				NewQuery->AddStep<FNodeClassRegistrationEvents>()
+					.AddStep<FMapRegistrationEventsToNodeRegistryKeys>()
+					.AddStep<FReduceRegistrationEventsToCurrentStatus>()
+					.AddStep<FTransformRegistrationEventsToClasses>();
 
 				if (!bInIncludeDeprecated)
 				{
-					NewQuery->AddStep<FMapClassNameToMajorVersion>()
+					NewQuery->AddStep<FMapToFullClassName>()
 						.AddStep<FReduceClassesToHighestVersion>();
 				}
 
@@ -105,6 +111,9 @@ namespace Metasound
 			{
 				TUniquePtr<FFrontendQuery> NewQuery = MakeUnique<FFrontendQuery>();
 				NewQuery->AddStep<FNodeClassRegistrationEvents>()
+					.AddStep<FMapRegistrationEventsToNodeRegistryKeys>()
+					.AddStep<FReduceRegistrationEventsToCurrentStatus>()
+					.AddStep<FTransformRegistrationEventsToClasses>()
 					.AddStep<FFilterClassesByClassName>(InName);
 
 				if (bInSortByVersion)
@@ -139,9 +148,14 @@ namespace Metasound
 			{
 				TUniquePtr<FFrontendQuery> NewQuery = MakeUnique<FFrontendQuery>();
 
+				// TODO: Create index of class names and major version to avoid duplicating
+				// this query for each class name.
 				NewQuery->AddStep<FNodeClassRegistrationEvents>()
+					.AddStep<FMapRegistrationEventsToNodeRegistryKeys>()
+					.AddStep<FReduceRegistrationEventsToCurrentStatus>()
+					.AddStep<FTransformRegistrationEventsToClasses>()
 					.AddStep<FFilterClassesByClassName>(InName)
-					.AddStep<FMapClassNameToMajorVersion>()
+					.AddStep<FMapToFullClassName>()
 					.AddStep<FReduceClassesToHighestVersion>();
 
 				Query = AddQuery(QueryName, MoveTemp(NewQuery));
@@ -174,9 +188,14 @@ namespace Metasound
 			{
 				TUniquePtr<FFrontendQuery> NewQuery = MakeUnique<FFrontendQuery>();
 
+				// TODO: Create index of class names to avoid duplicating
+				// this query for each class name.
 				NewQuery->AddStep<FNodeClassRegistrationEvents>()
+					.AddStep<FMapRegistrationEventsToNodeRegistryKeys>()
+					.AddStep<FReduceRegistrationEventsToCurrentStatus>()
+					.AddStep<FTransformRegistrationEventsToClasses>()
 					.AddStep<FFilterClassesByClassName>(InName)
-					.AddStep<FMapClassNameToMajorVersion>()
+					.AddStep<FMapToFullClassName>()
 					.AddStep<FReduceClassesToMajorVersion>(InMajorVersion);
 
 				Query = AddQuery(QueryName, MoveTemp(NewQuery));
diff --git a/Engine/Plugins/Runtime/Metasound/Source/MetasoundFrontend/Public/MetasoundFrontend.h b/Engine/Plugins/Runtime/Metasound/Source/MetasoundFrontend/Public/MetasoundFrontend.h
--- a/Engine/Plugins/Runtime/Metasound/Source/MetasoundFrontend/Public/MetasoundFrontend.h
+++ b/Engine/Plugins/Runtime/Metasound/Source/MetasoundFrontend/Public/MetasoundFrontend.h
@@ -15,6 +15,7 @@ namespace Metasound
 {	
 	namespace Frontend
 	{
+		class IRegistryTransaction;
 		using FRegistryTransactionID = int32;
 
 		METASOUNDFRONTEND_API FRegistryTransactionID GetCurrentRegistryTransactionID();
@@ -22,7 +23,14 @@ namespace Metasound
 		// Get all available nodes of any type.
 		METASOUNDFRONTEND_API TArray<FNodeClassInfo> GetAllAvailableNodeClasses(FRegistryTransactionID* OutCurrentTransactionID=nullptr);
 
-		METASOUNDFRONTEND_API TArray<FNodeClassInfo> GetNodeClassesRegisteredSince(FRegistryTransactionID InTransactionID, FRegistryTransactionID* OutCurrentTransactionID=nullptr);
+		/** Return all registry transactions that have occurred since the given transaction ID. 
+		 *
+		 * @param InTransactionID - Initial transaction ID. 
+		 * @param OutCurrentTransactionID - If not null, this will be set to the current transaction ID.
+		 *
+		 * @return An array of IRegistryTransactions that occurred in the range (InTransactionID, OutCurrentTransactionID].
+		 */
+		METASOUNDFRONTEND_API TArray<const IRegistryTransaction*> GetRegistryTransactionsSince(FRegistryTransactionID InTransactionID, FRegistryTransactionID* OutCurrentTransactionID);
 
 		/** Returns all metadata (name, description, author, what to say if it's missing) for a given node.
 		 *
diff --git a/Engine/Plugins/Runtime/Metasound/Source/MetasoundFrontend/Public/MetasoundFrontendQuery.h b/Engine/Plugins/Runtime/Metasound/Source/MetasoundFrontend/Public/MetasoundFrontendQuery.h
--- a/Engine/Plugins/Runtime/Metasound/Source/MetasoundFrontend/Public/MetasoundFrontendQuery.h
+++ b/Engine/Plugins/Runtime/Metasound/Source/MetasoundFrontend/Public/MetasoundFrontendQuery.h
@@ -5,21 +5,50 @@
 #include "CoreMinimal.h"
 #include "Math/NumericLimits.h"
 #include "MetasoundFrontendDocument.h"
+#include "MetasoundFrontendRegistries.h"
 #include "Misc/TVariant.h"
 
 namespace Metasound
 {
+	/** FFronendQueryKey allows entries to be partitioned by their key. A key
+	 * can be created by default constructor, int32, FString or FName.
+	 */
+	struct METASOUNDFRONTEND_API FFrontendQueryKey
+	{
+		FFrontendQueryKey();
+		FFrontendQueryKey(int32 InKey);
+		FFrontendQueryKey(const FString& InKey);
+		FFrontendQueryKey(const FName& InKey);
+
+		FFrontendQueryKey(const FFrontendQueryKey&) = default;
+		FFrontendQueryKey& operator=(const FFrontendQueryKey&) = default;
+		FFrontendQueryKey(FFrontendQueryKey&&) = default;
+		FFrontendQueryKey& operator=(FFrontendQueryKey&&) = default;
+
+		friend bool operator==(const FFrontendQueryKey& InLHS, const FFrontendQueryKey& InRHS);
+		friend bool operator!=(const FFrontendQueryKey& InLHS, const FFrontendQueryKey& InRHS);
+		friend bool operator<(const FFrontendQueryKey& InLHS, const FFrontendQueryKey& InRHS);
+		friend uint32 GetTypeHash(const FFrontendQueryKey& InKey);
+		
+		bool IsValid() const;
+
+	private:
+		struct FInvalid {};
+
+		using FKeyType = TVariant<FInvalid, int32, FString, FName>;
+		FKeyType Key;
+		uint32 Hash;
+	};
+
 	/** FFrontendQueryEntry represents one value in the query. It contains a key,
 	 * value and score. 
 	 */
-	struct FFrontendQueryEntry
+	struct METASOUNDFRONTEND_API FFrontendQueryEntry
 	{
-		using FValue = TVariant<FMetasoundFrontendDocument, FMetasoundFrontendClass, FMetasoundFrontendNode>;
-		using FKey = int32;
-
-		static constexpr FKey InvalidKey = TNumericLimits<int32>::Max();
+		using FValue = TVariant<const Frontend::IRegistryTransaction*, FMetasoundFrontendDocument, FMetasoundFrontendClass, FMetasoundFrontendNode>;
+		using FKey = FFrontendQueryKey;
 
-		FKey Key = InvalidKey;
+		FKey Key;
 		FValue Value;
 		float Score = 0.f;
 
@@ -182,6 +211,14 @@ namespace Metasound
 			virtual void Stream(TArray<FFrontendQueryEntry>& OutEntries) = 0;
 	};
 
+	/** Interface for a query step which transforms an entry's value. */
+	class IFrontendQueryTransformStep : public IFrontendQueryStep
+	{
+		public:
+			virtual ~IFrontendQueryTransformStep() = default;
+			virtual void Transform(FFrontendQueryEntry::FValue& InValue) const = 0;
+	};
+
 	/** Interface for a query step which maps entries to keys. */
 	class IFrontendQueryMapStep : public IFrontendQueryStep
 	{
@@ -274,6 +311,7 @@ namespace Metasound
 		using FReduceOutputView = FFrontendQuerySelection::FReduceOutputView;
 
 		using FStreamFunction = TUniqueFunction<void (TArray<FFrontendQueryEntry>&)>;
+		using FTransformFunction = TFunction<void (FFrontendQueryEntry::FValue&)>;
 		using FMapFunction = TFunction<FFrontendQueryEntry::FKey (const FFrontendQueryEntry&)>;
 		using FReduceFunction = TFunction<void (FFrontendQueryEntry::FKey, TArrayView<FFrontendQueryEntry * const>, FReduceOutputView& )>;
 		using FFilterFunction = TFunction<bool (const FFrontendQueryEntry&)>;
@@ -283,6 +321,7 @@ namespace Metasound
 
 		/** Create query step using TFunction or lambda. */
 		FFrontendQueryStep(FStreamFunction&& InFunc);
+		FFrontendQueryStep(FTransformFunction&& InFunc);
 		FFrontendQueryStep(FMapFunction&& InFunc);
 		FFrontendQueryStep(FReduceFunction&& InFunc);
 		FFrontendQueryStep(FFilterFunction&& InFilt);
@@ -293,6 +332,7 @@ namespace Metasound
 		/** Create a query step using a IFrotnedQueryStep */
 		FFrontendQueryStep(TUniquePtr<IFrontendQuerySource>&& InSource);
 		FFrontendQueryStep(TUniquePtr<IFrontendQueryStreamStep>&& InStep);
+		FFrontendQueryStep(TUniquePtr<IFrontendQueryTransformStep>&& InStep);
 		FFrontendQueryStep(TUniquePtr<IFrontendQueryMapStep>&& InStep);
 		FFrontendQueryStep(TUniquePtr<IFrontendQueryReduceStep>&& InStep);
 		FFrontendQueryStep(TUniquePtr<IFrontendQueryFilterStep>&& InStep);
@@ -333,6 +373,7 @@ namespace Metasound
 
 		using EResultModificationState = FFrontendQueryStep::EResultModificationState;
 		using FStreamFunction = FFrontendQueryStep::FStreamFunction;
+		using FTransformFunction = FFrontendQueryStep::FTransformFunction;
 		using FMapFunction = FFrontendQueryStep::FMapFunction;
 		using FReduceFunction = FFrontendQueryStep::FReduceFunction;
 		using FFilterFunction = FFrontendQueryStep::FFilterFunction;
@@ -362,6 +403,7 @@ namespace Metasound
 		}
 
 		FFrontendQuery& AddStreamLambdaStep(FStreamFunction&& InFunc);
+		FFrontendQuery& AddTransformLambdaStep(FTransformFunction&& InFunc);
 		FFrontendQuery& AddMapLambdaStep(FMapFunction&& InFunc);
 		FFrontendQuery& AddReduceLambdaStep(FReduceFunction&& InFunc);
 		FFrontendQuery& AddFilterLambdaStep(FFilterFunction&& InFunc);
diff --git a/Engine/Plugins/Runtime/Metasound/Source/MetasoundFrontend/Public/MetasoundFrontendQuerySteps.h b/Engine/Plugins/Runtime/Metasound/Source/MetasoundFrontend/Public/MetasoundFrontendQuerySteps.h
--- a/Engine/Plugins/Runtime/Metasound/Source/MetasoundFrontend/Public/MetasoundFrontendQuerySteps.h
+++ b/Engine/Plugins/Runtime/Metasound/Source/MetasoundFrontend/Public/MetasoundFrontendQuerySteps.h
@@ -9,7 +9,6 @@
 
 namespace Metasound
 {
-
 	/** Streams node classes that have been newly registered or unregistered since last call to Stream()
 	 */
 	class METASOUNDFRONTEND_API FNodeClassRegistrationEvents : public IFrontendQuerySource
@@ -23,6 +22,29 @@ namespace Metasound
 		Frontend::FRegistryTransactionID CurrentTransactionID;
 	};
 
+	/** Partitions node registration events by their node registration keys. */
+	class METASOUNDFRONTEND_API FMapRegistrationEventsToNodeRegistryKeys : public IFrontendQueryMapStep
+	{
+	public:
+		FFrontendQueryEntry::FKey Map(const FFrontendQueryEntry& InEntry) const override;
+	};
+
+	/** Reduces registration events mapped to the same key by inspecting their add/remove state in
+	 * order to determine their final state. If an item has been added more than it has been removed,
+	 * then it is added to the output. Otherwise, it is omitted. */
+	class METASOUNDFRONTEND_API FReduceRegistrationEventsToCurrentStatus: public IFrontendQueryReduceStep
+	{
+	public:
+		void Reduce(FFrontendQueryEntry::FKey InKey, TArrayView<FFrontendQueryEntry * const>& InEntries, FReduceOutputView& OutResult) const override;
+	};
+
+	/** Transforms a registration event into a FMetasoundFrontendClass. */
+	class METASOUNDFRONTEND_API FTransformRegistrationEventsToClasses : public IFrontendQueryTransformStep
+	{
+	public:
+		virtual void Transform(FFrontendQueryEntry::FValue& InValue) const override;
+	};
+
 	class METASOUNDFRONTEND_API FFilterClassesByInputVertexDataType : public IFrontendQueryFilterStep
 	{
 	public: 
@@ -79,7 +101,7 @@ namespace Metasound
 		FGuid ClassID;
 	};
 
-	class METASOUNDFRONTEND_API FMapClassNameToMajorVersion : public IFrontendQueryMapStep
+	class METASOUNDFRONTEND_API FMapToFullClassName : public IFrontendQueryMapStep
 	{
 	public:
 		FFrontendQueryEntry::FKey Map(const FFrontendQueryEntry& InEntry) const override;
diff --git a/Engine/Plugins/Runtime/Metasound/Source/MetasoundFrontend/Public/MetasoundFrontendRegistries.h b/Engine/Plugins/Runtime/Metasound/Source/MetasoundFrontend/Public/MetasoundFrontendRegistries.h
--- a/Engine/Plugins/Runtime/Metasound/Source/MetasoundFrontend/Public/MetasoundFrontendRegistries.h
+++ b/Engine/Plugins/Runtime/Metasound/Source/MetasoundFrontend/Public/MetasoundFrontendRegistries.h
@@ -14,10 +14,6 @@
 
 namespace Metasound
 {
-	namespace Frontend
-	{
-		using FRegistryTransactionID = int32;
-	}
 
 	// Base implementation for NodeConstructorCallbacks
 	struct FDefaultNodeConstructorParams
@@ -52,29 +48,6 @@ namespace Metasound
 	using FVariableNodeConstructorParams = FDefaultLiteralNodeConstructorParams;
 	using FInputNodeConstructorParams = FDefaultLiteralNodeConstructorParams;
 
-	/*
-	class IMetasoundDataTypeRegistryEntry
-	{
-	public:
-		virtual TUniquePtr<Metasound::INode> CreateInputNode(FInputNodeConstructorParams&&) const = 0;
-		virtual TUniquePtr<Metasound::INode> CreateOutputNode(const FOutputNodeConstructorParams&) const = 0;
-	};
-
-	class IMetasoundProxyRegistryEntry
-	{
-	public:
-		// This function is used to create a proxy from a datatype's base uclass.
-		virtual Audio::IProxyDataPtr CreateAudioProxy(UObject*) const = 0;
-	};
-
-	class IMetasoundNodeRegistryEntry
-	{
-	public:
-		virtual TUniquePtr<Metasound::INode> CreateNode(const Metasound::FNodeInitData&) const = 0;
-		virtual FMetasoundFrontendClass CreateMetasoundFrontendClass() const = 0;
-	};
-	*/
-
 	using FCreateInputNodeFunction = TFunction<TUniquePtr<Metasound::INode>(::Metasound::FInputNodeConstructorParams&&)>;
 	using FCreateVariableNodeFunction = TFunction<TUniquePtr<Metasound::INode>(::Metasound::FVariableNodeConstructorParams&&)>;
 	using FCreateOutputNodeFunction = TFunction<TUniquePtr<Metasound::INode>(::Metasound::FOutputNodeConstructorParams&&)>;
@@ -136,30 +109,11 @@ namespace Metasound
 
 	namespace Frontend
 	{
-		/* TODO: Node registry key needs to use full version
-		struct METASOUNDFRONTEND_API FNodeRegistryKey
-		{
-			// The class name for the node.
-			FName NodeClassFullName;
-
-
-			// A hash generated from the input types and output types for this node.
-			uint32 NodeHash = 0;
-
-			FORCEINLINE bool operator==(const FNodeRegistryKey& Other) const
-			{
-				return NodeHash == Other.NodeHash && NodeClassFullName == Other.NodeClassFullName;
-			}
-
-			friend uint32 GetTypeHash(const Metasound::Frontend::FNodeRegistryKey& InKey)
-			{
-				return InKey.NodeHash;
-			}
-		};
-		*/
-
 		using FNodeRegistryKey = FString;
 
+		// Returns true if the registry key is a valid key. 
+		//
+		// This does *not* connote that the registry key exists in the registry.
 		METASOUNDFRONTEND_API bool IsValidNodeRegistryKey(const FNodeRegistryKey& InKey);
 
 		// Struct with the basics of a node class' information,
@@ -275,7 +229,34 @@ namespace Metasound
 				return {};
 			}
 		};
-	}
+
+		using FRegistryTransactionID = int32;
+
+		/** Describes the type of transaction. */
+		enum class ETransactionType : uint8
+		{
+			Add,     //< Something was added to the registry.
+			Remove,  //< Something was removed from the registry.
+			Invalid
+		};
+
+		/** Interface for a registry transaction. */
+		class IRegistryTransaction
+		{
+		public:
+			virtual ~IRegistryTransaction() = default;
+			virtual TUniquePtr<IRegistryTransaction> Clone() const = 0;
+
+			/** Returns the type of transaction */
+			virtual ETransactionType GetTransactionType() const = 0;
+
+			/** If a FNodeClassInfo added during the transaction, this will return
+			 * a non-null pointer to the FNodeClassInfo. */
+			virtual const FNodeClassInfo* GetNodeClassInfo() const = 0;
+		};
+
+		using FRegistryTransactionPtr = TUniquePtr<IRegistryTransaction>;
+	} // namespace Frontend
 
 	struct FDataTypeConstructorCallbacks
 	{
@@ -299,7 +280,7 @@ namespace Metasound
 
 		FCreateDataChannelFunction CreateDataChannel;
 	};
-} // namespace Frontend
+} // namespace Metasound
 
 /**
  * Singleton registry for all types and nodes.
@@ -361,8 +342,8 @@ public:
 	/** Return all node classes registered. */
 	virtual TArray<Metasound::Frontend::FNodeClassInfo> GetAllAvailableNodeClasses(Metasound::Frontend::FRegistryTransactionID* OutCurrentRegistryTransactionID) const = 0;
 
-	/** Return node classes registered since a given transaction id. */
-	virtual TArray<Metasound::Frontend::FNodeClassInfo> GetNodeClassesRegisteredSince(Metasound::Frontend::FRegistryTransactionID InSince, Metasound::Frontend::FRegistryTransactionID* OutCurrentRegistryTransactionID) const = 0;
+	/** Return registry transactions since a given transaction id. */
+	virtual TArray<const Metasound::Frontend::IRegistryTransaction*> GetRegistryTransactionsSince(Metasound::Frontend::FRegistryTransactionID InSince, Metasound::Frontend::FRegistryTransactionID* OutCurrentRegistryTransactionID) const = 0;
 
 	// Return any data types that can be used as a metasound input type or output type.
 	virtual TArray<FName> GetAllValidDataTypes() = 0;
