CIS fix for unintended recursive function call when GetTypeHash(int32) not visible to translation unit.
#rnx

diff --git a/Engine/Plugins/Runtime/Metasound/Source/MetasoundFrontend/Private/MetasoundFrontendQuery.cpp b/Engine/Plugins/Runtime/Metasound/Source/MetasoundFrontend/Private/MetasoundFrontendQuery.cpp
--- a/Engine/Plugins/Runtime/Metasound/Source/MetasoundFrontend/Private/MetasoundFrontendQuery.cpp
+++ b/Engine/Plugins/Runtime/Metasound/Source/MetasoundFrontend/Private/MetasoundFrontendQuery.cpp
@@ -7,6 +7,7 @@
 #include "Algo/Transform.h"
 #include "CoreMinimal.h"
 #include "MetasoundFrontendDocument.h"
+#include "Templates/TypeHash.h"
 
 namespace Metasound
 {
diff --git a/Engine/Plugins/Runtime/Metasound/Source/MetasoundFrontend/Private/MetasoundFrontendQuerySteps.cpp b/Engine/Plugins/Runtime/Metasound/Source/MetasoundFrontend/Private/MetasoundFrontendQuerySteps.cpp
--- a/Engine/Plugins/Runtime/Metasound/Source/MetasoundFrontend/Private/MetasoundFrontendQuerySteps.cpp
+++ b/Engine/Plugins/Runtime/Metasound/Source/MetasoundFrontend/Private/MetasoundFrontendQuerySteps.cpp
@@ -54,7 +54,7 @@ namespace Metasound
 			}
 		}
 
-		return RegistryKey;
+		return FFrontendQueryEntry::FKey(RegistryKey);
 	}
 
 	void FReduceRegistrationEventsToCurrentStatus::Reduce(FFrontendQueryEntry::FKey InKey, TArrayView<FFrontendQueryEntry * const>& InEntries, FReduceOutputView& OutResult) const
diff --git a/Engine/Plugins/Runtime/Metasound/Source/MetasoundFrontend/Public/MetasoundFrontendQuery.h b/Engine/Plugins/Runtime/Metasound/Source/MetasoundFrontend/Public/MetasoundFrontendQuery.h
--- a/Engine/Plugins/Runtime/Metasound/Source/MetasoundFrontend/Public/MetasoundFrontendQuery.h
+++ b/Engine/Plugins/Runtime/Metasound/Source/MetasoundFrontend/Public/MetasoundFrontendQuery.h
@@ -16,9 +16,9 @@ namespace Metasound
 	struct METASOUNDFRONTEND_API FFrontendQueryKey
 	{
 		FFrontendQueryKey();
-		FFrontendQueryKey(int32 InKey);
-		FFrontendQueryKey(const FString& InKey);
-		FFrontendQueryKey(const FName& InKey);
+		explicit FFrontendQueryKey(int32 InKey);
+		explicit FFrontendQueryKey(const FString& InKey);
+		explicit FFrontendQueryKey(const FName& InKey);
 
 		FFrontendQueryKey(const FFrontendQueryKey&) = default;
 		FFrontendQueryKey& operator=(const FFrontendQueryKey&) = default;
