Flanger DSP object and corresponding Metasound node

diff --git a/Engine/Plugins/Runtime/Metasound/Source/MetasoundStandardNodes/Private/MetasoundFlangerNode.cpp b/Engine/Plugins/Runtime/Metasound/Source/MetasoundStandardNodes/Private/MetasoundFlangerNode.cpp
new file mode 100644
--- /dev/null
+++ b/Engine/Plugins/Runtime/Metasound/Source/MetasoundStandardNodes/Private/MetasoundFlangerNode.cpp
@@ -0,0 +1,196 @@
+// Copyright Epic Games, Inc. All Rights Reserved.
+
+#include "Internationalization/Text.h"
+#include "MetasoundFacade.h"
+#include "MetasoundExecutableOperator.h"
+#include "MetasoundNodeRegistrationMacro.h"
+#include "MetasoundPrimitives.h"
+#include "MetasoundStandardNodesNames.h"
+#include "MetasoundStandardNodesCategories.h"
+#include "DSP/BufferVectorOperations.h"
+#include "DSP/Flanger.h"
+#include "MetasoundParamHelper.h"
+
+#define LOCTEXT_NAMESPACE "MetasoundStandardNodes_FlangerNode"
+
+namespace Metasound
+{
+	namespace FlangerVertexNames
+	{
+		METASOUND_PARAM(InputAudio, "In Audio", "The audio input.");
+		METASOUND_PARAM(InputModulationRate, "Modulation Rate", "The LFO frequency (rate) that varies the delay time, in Hz. Clamped at blockrate.");
+		METASOUND_PARAM(InputModulationDepth, "Modulation Depth", "The LFO amplitude (strength) that scales the delay time.");
+		METASOUND_PARAM(InputCenterDelay, "Center Delay", "The center delay amount (in milliseconds).");
+		METASOUND_PARAM(InputMixLevel, "Mix Level", "Balance between original and delayed signal (Should be between 0 and 1.0; 0.5 is equal amounts of each and > 0.5 is more delayed signal than non-delayed signal).");
+		
+		METASOUND_PARAM(OutputAudio, "Out Audio", "Output audio with flanger effect applied.");
+	}
+
+	class FFlangerOperator : public TExecutableOperator<FFlangerOperator>
+	{
+	public:
+
+		static const FNodeClassMetadata& GetNodeInfo()
+		{
+			auto InitNodeInfo = []() -> FNodeClassMetadata
+			{
+				const FName OperatorName = TEXT("Flanger");
+				const FText NodeDisplayName = LOCTEXT("Metasound_FlangerNodeDisplayName", "Flanger");
+				const FText NodeDescription = LOCTEXT("Metasound_FlangerNodeDescription", "Applies a flanger effect to input audio.");
+
+				FNodeClassMetadata Info;
+				Info.ClassName = { Metasound::StandardNodes::Namespace, OperatorName, TEXT("") };
+				Info.MajorVersion = 1;
+				Info.MinorVersion = 0;
+				Info.DisplayName = NodeDisplayName;
+				Info.Description = NodeDescription;
+				Info.Author = PluginAuthor;
+				Info.PromptIfMissing = PluginNodeMissingPrompt;
+				Info.DefaultInterface = GetVertexInterface();
+
+				return Info;
+			};
+
+			static const FNodeClassMetadata Info = InitNodeInfo();
+			return Info;
+		}
+
+		static const FVertexInterface& GetVertexInterface()
+		{
+			using namespace FlangerVertexNames;
+
+			static const FVertexInterface Interface(
+				FInputVertexInterface(
+					TInputDataVertexModel<FAudioBuffer>(METASOUND_GET_PARAM_NAME_AND_TT(InputAudio)), 
+					TInputDataVertexModel<float>(METASOUND_GET_PARAM_NAME_AND_TT(InputModulationRate), 0.5f),
+					TInputDataVertexModel<float>(METASOUND_GET_PARAM_NAME_AND_TT(InputModulationDepth), 0.5f),
+					TInputDataVertexModel<float>(METASOUND_GET_PARAM_NAME_AND_TT(InputCenterDelay), 0.5f),
+					TInputDataVertexModel<float>(METASOUND_GET_PARAM_NAME_AND_TT(InputMixLevel), 0.5f)
+				),
+				FOutputVertexInterface(
+					TOutputDataVertexModel<FAudioBuffer>(METASOUND_GET_PARAM_NAME_AND_TT(OutputAudio))
+				)
+			);
+
+			return Interface;
+		}
+
+		static TUniquePtr<IOperator> CreateOperator(const FCreateOperatorParams& InParams, FBuildErrorArray& OutErrors)
+		{
+			using namespace FlangerVertexNames;
+
+			const FDataReferenceCollection& InputCollection = InParams.InputDataReferences;
+			const FInputVertexInterface& InputInterface = GetVertexInterface().GetInputInterface();
+
+			FAudioBufferReadRef AudioInput = InputCollection.GetDataReadReferenceOrConstruct<FAudioBuffer>(METASOUND_GET_PARAM_NAME(InputAudio), InParams.OperatorSettings);
+			FFloatReadRef ModulationRate = InParams.InputDataReferences.GetDataReadReferenceOrConstructWithVertexDefault<float>(InputInterface, METASOUND_GET_PARAM_NAME(InputModulationRate), InParams.OperatorSettings);
+			FFloatReadRef ModulationDepth = InParams.InputDataReferences.GetDataReadReferenceOrConstructWithVertexDefault<float>(InputInterface, METASOUND_GET_PARAM_NAME(InputModulationDepth), InParams.OperatorSettings);
+			FFloatReadRef CenterDelay = InParams.InputDataReferences.GetDataReadReferenceOrConstructWithVertexDefault<float>(InputInterface, METASOUND_GET_PARAM_NAME(InputCenterDelay), InParams.OperatorSettings);
+			FFloatReadRef MixLevel = InParams.InputDataReferences.GetDataReadReferenceOrConstructWithVertexDefault<float>(InputInterface, METASOUND_GET_PARAM_NAME(InputMixLevel), InParams.OperatorSettings);
+
+			return MakeUnique<FFlangerOperator>(InParams.OperatorSettings, AudioInput, ModulationRate, ModulationDepth, CenterDelay, MixLevel);
+		}
+
+		FFlangerOperator(const FOperatorSettings& InSettings,
+			const FAudioBufferReadRef& InAudioInput,
+			const FFloatReadRef& InputModulationRate,
+			const FFloatReadRef& InputModulationDepth,
+			const FFloatReadRef& InputCenterDelay,
+			const FFloatReadRef& InMixLevel)
+
+			: AudioIn(InAudioInput)
+			, ModulationRate(InputModulationRate)
+			, ModulationDepth(InputModulationDepth)
+			, CenterDelay(InputCenterDelay)
+			, MixLevel(InMixLevel)
+			, AudioOut(FAudioBufferWriteRef::CreateNew(InSettings))
+			, Flanger()
+		{
+			Flanger.Init(InSettings.GetSampleRate());
+		}
+
+		FDataReferenceCollection GetInputs() const
+		{
+			using namespace FlangerVertexNames;
+
+			FDataReferenceCollection InputDataReferences;
+			InputDataReferences.AddDataReadReference(METASOUND_GET_PARAM_NAME(InputAudio), AudioIn);
+			InputDataReferences.AddDataReadReference(METASOUND_GET_PARAM_NAME(InputModulationRate), ModulationRate);
+			InputDataReferences.AddDataReadReference(METASOUND_GET_PARAM_NAME(InputModulationDepth), ModulationDepth);
+			InputDataReferences.AddDataReadReference(METASOUND_GET_PARAM_NAME(InputCenterDelay), CenterDelay);
+			InputDataReferences.AddDataReadReference(METASOUND_GET_PARAM_NAME(InputMixLevel), MixLevel);
+
+			return InputDataReferences;
+		}
+
+		virtual FDataReferenceCollection GetOutputs() const
+		{
+			using namespace FlangerVertexNames;
+
+			FDataReferenceCollection OutputDataReferences;
+			OutputDataReferences.AddDataReadReference(METASOUND_GET_PARAM_NAME(OutputAudio), AudioOut);
+			return OutputDataReferences;
+		}
+
+		void Execute()
+		{
+			// Update flanger parameters if necessary 
+			Flanger.SetModulationRate(*ModulationRate);
+			Flanger.SetCenterDelay(*CenterDelay);
+			Flanger.SetModulationDepth(*ModulationDepth);
+			Flanger.SetMixLevel(*MixLevel);
+
+			const float* InputAudio = AudioIn->GetData();
+			float* OutputAudio = AudioOut->GetData();
+			int32 NumFrames = AudioIn->Num();
+
+			// Copy input audio into aligned buffer
+			AlignedAudioIn.Reset();
+			AlignedAudioIn.Append(InputAudio, NumFrames);
+
+			Flanger.ProcessAudio(AlignedAudioIn, NumFrames, AlignedAudioOut);
+			FMemory::Memcpy(OutputAudio, AlignedAudioOut.GetData(), sizeof(float) * NumFrames);
+		}
+
+	private:
+		// The input audio buffer
+		FAudioBufferReadRef AudioIn;
+
+		// LFO parameters
+		FFloatReadRef ModulationRate;
+		FFloatReadRef ModulationDepth;
+		FFloatReadRef CenterDelay;
+
+		// Balance between original and delayed signal 
+		// (Should be between 0 and 1.0; 
+		// 0.5 is equal amounts of each and 
+		// > 0.5 is more delayed signal than non-delayed signal)
+		FFloatReadRef MixLevel;
+
+		// Audio output
+		FAudioBufferWriteRef AudioOut;
+
+		// Flanger DSP object
+		Audio::FFlanger Flanger;
+
+		// Input/output aligned buffers
+		Audio::FAlignedFloatBuffer AlignedAudioIn;
+		Audio::FAlignedFloatBuffer AlignedAudioOut;
+	};
+
+	class METASOUNDSTANDARDNODES_API FFlangerNode : public FNodeFacade
+	{
+	public:
+		/**
+		 * Constructor used by the Metasound Frontend.
+		 */
+		FFlangerNode(const FNodeInitData& InitData)
+			: FNodeFacade(InitData.InstanceName, InitData.InstanceID, TFacadeOperatorClass<FFlangerOperator>())
+		{
+		}
+	};
+
+	METASOUND_REGISTER_NODE(FFlangerNode)
+}
+
+#undef LOCTEXT_NAMESPACE
diff --git a/Engine/Source/Runtime/SignalProcessing/Private/Flanger.cpp b/Engine/Source/Runtime/SignalProcessing/Private/Flanger.cpp
--- a/Engine/Source/Runtime/SignalProcessing/Private/Flanger.cpp
+++ b/Engine/Source/Runtime/SignalProcessing/Private/Flanger.cpp
@@ -4,4 +4,83 @@
 
 namespace Audio
 {
+	const float FFlanger::MaxDelaySec = 5.0f;
+	const float FFlanger::MaxModulationRate = 20.0f;
+	const float FFlanger::MaxCenterDelay = 20.0f;
+
+	FFlanger::FFlanger()
+	{
+	}
+
+	FFlanger::~FFlanger()
+	{
+	}
+
+	void FFlanger::Init(const float InSampleRate)
+	{
+		SampleRate = InSampleRate;
+		DelayBuffer.Init(InSampleRate, MaxDelaySec);
+		DelayedSignalLevel = FMath::Clamp(MixLevel, 0.0f, 1.0f);
+		NonDelayedSignalLevel = 1.0f - DelayedSignalLevel;
+	}
+
+	void FFlanger::SetModulationRate(const float InModulationRate)
+	{
+		float ClampedInModulationRate = FMath::Clamp(InModulationRate, 0.0f, MaxModulationRate);
+		if (!FMath::IsNearlyEqual(ClampedInModulationRate, ModulationRate))
+		{
+			ModulationRate = ClampedInModulationRate;
+		}
+	}
+
+	// since this will be clamped based on CenterDelayMsec to avoid clipping, 
+	// make sure SetCenterDelay is called before this function if necessary
+	void FFlanger::SetModulationDepth(const float InModulationDepth)
+	{
+		float ClampedInModulationDepth = FMath::Clamp(InModulationDepth, 0.0f, CenterDelayMsec);
+		if (!FMath::IsNearlyEqual(ClampedInModulationDepth, ModulationDepth))
+		{
+			ModulationDepth = ClampedInModulationDepth;
+		}
+	}
+
+	void FFlanger::SetCenterDelay(const float InCenterDelay)
+	{
+		float ClampedInCenterDelay = FMath::Clamp(InCenterDelay, 0.0f, MaxCenterDelay);
+		if (!FMath::IsNearlyEqual(ClampedInCenterDelay, CenterDelayMsec))
+		{
+			CenterDelayMsec = ClampedInCenterDelay;
+		}
+	}
+
+	void FFlanger::SetMixLevel(const float InMixLevel)
+	{
+		float ClampedInMixLevel = FMath::Clamp(InMixLevel, 0.0f, 1.0f);
+		if (!FMath::IsNearlyEqual(ClampedInMixLevel, MixLevel))
+		{
+			MixLevel = ClampedInMixLevel;
+			DelayedSignalLevel = MixLevel;
+			NonDelayedSignalLevel = 1.0f - DelayedSignalLevel;
+		}
+	}
+
+	void FFlanger::ProcessAudio(const FAlignedFloatBuffer& InBuffer, const int32 InNumSamples, FAlignedFloatBuffer& OutBuffer)
+	{
+		// LFO sample for delay is generated at buffer, not sample, rate 
+		LFO.GenerateBuffer(SampleRate / InNumSamples, ModulationRate, &DelaySample, 1);
+		DelayBuffer.SetEasedDelayMsec(ModulationDepth * DelaySample + CenterDelayMsec, true);
+
+		// initialize scratch buffer if necessary and get delay samples (non SIMD)
+		if (ScratchBuffer.Num() != InNumSamples)
+		{
+			ScratchBuffer.Reset();
+			ScratchBuffer.AddUninitialized(InBuffer.Num());
+		}
+		for (int32 SampleIndex = 0; SampleIndex < InNumSamples; ++SampleIndex)
+		{
+			ScratchBuffer[SampleIndex] = DelayBuffer.ProcessAudioSample(InBuffer[SampleIndex]);
+		}
+
+		BufferWeightedSumFast(ScratchBuffer, DelayedSignalLevel, InBuffer, NonDelayedSignalLevel, OutBuffer);
+	}
 }
diff --git a/Engine/Source/Runtime/SignalProcessing/Public/DSP/Flanger.h b/Engine/Source/Runtime/SignalProcessing/Public/DSP/Flanger.h
--- a/Engine/Source/Runtime/SignalProcessing/Public/DSP/Flanger.h
+++ b/Engine/Source/Runtime/SignalProcessing/Public/DSP/Flanger.h
@@ -3,8 +3,57 @@
 #pragma once
 
 #include "CoreMinimal.h"
+#include "DSP/BufferVectorOperations.h"
 #include "DSP/Delay.h"
 
 namespace Audio
 {
+	class SIGNALPROCESSING_API FFlanger
+	{
+	public: 
+		FFlanger();
+		~FFlanger();
+
+		void Init(const float InSampleRate);
+
+		void SetModulationRate(const float InModulationRate);
+		void SetModulationDepth(const float InModulationDepth);
+		void SetCenterDelay(const float InCenterDelay);
+		void SetMixLevel(const float InMixLevel);
+
+		void ProcessAudio(const FAlignedFloatBuffer& InBuffer, const int32 InNumSamples, FAlignedFloatBuffer& OutBuffer);
+	protected: 
+		static const float MaxDelaySec;
+		static const float MaxModulationRate;
+		static const float MaxCenterDelay;
+
+		// LFO parameters
+		// LFO frequency 
+		float ModulationRate = 0.5f;
+		// Modulation depth is clamped to CenterDelayMsec to avoid clipping
+		float ModulationDepth = 0.5f;
+		float CenterDelayMsec = 0.5f;
+
+		// Balance between original and delayed signal 
+		// (Should be between 0 and 1.0; 
+		// 0.5 is equal amounts of each and 
+		// > 0.5 is more delayed signal than non-delayed signal)
+		float MixLevel = 0.5f;
+		float DelayedSignalLevel = 0.5f;
+		float NonDelayedSignalLevel = 0.5f;
+
+		// Internal delay buffer 
+		FDelay DelayBuffer;
+		// Scratch buffer used for accumulating delay samples per block
+		FAlignedFloatBuffer ScratchBuffer;
+
+		// Delay sample generated from LFO per block
+		float DelaySample = 0.0f;
+
+		// Internal LFO for delay amount
+		FSinOsc2DRotation LFO;
+
+		// The audio sample rate, Init() must be called to initialize this
+		float SampleRate = 48000.0f;
+	};
 }
