CIS fix for GetTypeHash scope
#rnx

diff --git a/Engine/Plugins/Runtime/Metasound/Source/MetasoundFrontend/Private/MetasoundFrontendQuery.cpp b/Engine/Plugins/Runtime/Metasound/Source/MetasoundFrontend/Private/MetasoundFrontendQuery.cpp
--- a/Engine/Plugins/Runtime/Metasound/Source/MetasoundFrontend/Private/MetasoundFrontendQuery.cpp
+++ b/Engine/Plugins/Runtime/Metasound/Source/MetasoundFrontend/Private/MetasoundFrontendQuery.cpp
@@ -721,19 +721,19 @@ namespace Metasound
 
 	FFrontendQueryKey::FFrontendQueryKey(int32 InKey)
 	: Key(TInPlaceType<int32>(), InKey)
-	, Hash(GetTypeHash(InKey))
+	, Hash(::GetTypeHash(InKey))
 	{
 	}
 
 	FFrontendQueryKey::FFrontendQueryKey(const FString& InKey)
 	: Key(TInPlaceType<FString>(), InKey)
-	, Hash(GetTypeHash(InKey))
+	, Hash(::GetTypeHash(InKey))
 	{
 	}
 
 	FFrontendQueryKey::FFrontendQueryKey(const FName& InKey)
 	: Key(TInPlaceType<FName>(), InKey)
-	, Hash(GetTypeHash(InKey))
+	, Hash(::GetTypeHash(InKey))
 	{
 	}
 
