Add Mid-Side Encode/Decode nodes to Metasounds

diff --git a/Engine/Plugins/Runtime/Metasound/Source/MetasoundStandardNodes/Private/MetasoundMidSideNodes.cpp b/Engine/Plugins/Runtime/Metasound/Source/MetasoundStandardNodes/Private/MetasoundMidSideNodes.cpp
new file mode 100644
--- /dev/null
+++ b/Engine/Plugins/Runtime/Metasound/Source/MetasoundStandardNodes/Private/MetasoundMidSideNodes.cpp
@@ -0,0 +1,478 @@
+// Copyright Epic Games, Inc. All Rights Reserved.
+
+#include "Internationalization/Text.h"
+#include "MetasoundExecutableOperator.h"
+#include "MetasoundEnumRegistrationMacro.h"
+#include "MetasoundNodeRegistrationMacro.h"
+#include "MetasoundPrimitives.h"
+#include "MetasoundStandardNodesNames.h"
+#include "MetasoundAudioBuffer.h"
+#include "DSP/Dsp.h"
+#include "MetasoundStandardNodesCategories.h"
+#include "MetasoundFacade.h"
+#include "MetasoundParamHelper.h"
+#include "DSP/FloatArrayMath.h"
+
+#define LOCTEXT_NAMESPACE "MetasoundStandardNodes_MetasoundMidSideNodes"
+
+namespace Metasound
+{
+	/* Mid-Side Encoder */
+	namespace MidSideEncodeVertexNames
+	{
+		METASOUND_PARAM(InputAudioLeft, "In Left", "The left audio channel to convert.");
+		METASOUND_PARAM(InputAudioRight, "In Right", "The right audio channel to convert.");
+		METASOUND_PARAM(InputSpreadAmount, "Spread Amount", "Amount of Stereo Spread. 0.0 is no spread, 0.5 is the original signal, and 1.0 is full wide.");
+		METASOUND_PARAM(InputEqualPower, "Equal Power", "Whether an equal power relationship between the mid and side signals should be maintained.");
+
+		METASOUND_PARAM(OutputAudioMid, "Out Mid", "The Mid content from the audio signal.");
+		METASOUND_PARAM(OutputAudioSide, "Out Side", "the Side content from the audio signal.");
+	}
+
+	// Operator Class
+	class FMidSideEncodeOperator : public TExecutableOperator<FMidSideEncodeOperator>
+	{
+	public:
+
+		FMidSideEncodeOperator(const FOperatorSettings& InSettings,
+			const FAudioBufferReadRef& InLeftAudioInput,
+			const FAudioBufferReadRef& InRightAudioInput,
+			const FFloatReadRef& InSpreadAmount,
+			const FBoolReadRef& InEqualPower)
+			: AudioInputLeft(InLeftAudioInput)
+			, AudioInputRight(InRightAudioInput)
+			, SpreadAmount(InSpreadAmount)
+			, bEqualPower(InEqualPower)
+			, AudioOutputMid(FAudioBufferWriteRef::CreateNew(InSettings))
+			, AudioOutputSide(FAudioBufferWriteRef::CreateNew(InSettings))
+			, MidScale(0.0f)
+			, SideScale(0.0f)
+			, PrevMidScale(0.0f)
+			, PrevSideScale(0.0f)
+			, PrevSpreadAmount(*SpreadAmount)
+		{
+		}
+
+		static const FNodeClassMetadata& GetNodeInfo()
+		{
+			auto CreateNodeClassMetadata = []() -> FNodeClassMetadata
+			{
+				FVertexInterface NodeInterface = DeclareVertexInterface();
+
+				FNodeClassMetadata Metadata
+				{
+					FNodeClassName{StandardNodes::Namespace, TEXT("Mid-Side Decode"), StandardNodes::AudioVariant},
+					1, // Major Version
+					0, // Minor Version
+					LOCTEXT("MidSideEncodeDisplayName", "Mid-Side Encode"),
+					LOCTEXT("MidSideEncodeDesc", "Converts a stereo audio signal from Left and Right to Mid and Side channels."),
+					PluginAuthor,
+					PluginNodeMissingPrompt,
+					NodeInterface,
+					{StandardNodes::Spatalization},
+					{TEXT("MidSideEncode")},
+					FNodeDisplayStyle{}
+				};
+
+				return Metadata;
+			};
+
+			static const FNodeClassMetadata Metadata = CreateNodeClassMetadata();
+			return Metadata;
+		}
+
+		static const FVertexInterface& DeclareVertexInterface()
+		{
+			using namespace MidSideEncodeVertexNames;
+
+			static const FVertexInterface Interface(
+				FInputVertexInterface(
+					TInputDataVertexModel<FAudioBuffer>(METASOUND_GET_PARAM_NAME_AND_TT(InputAudioLeft)),
+					TInputDataVertexModel<FAudioBuffer>(METASOUND_GET_PARAM_NAME_AND_TT(InputAudioRight)),
+					TInputDataVertexModel<float>(METASOUND_GET_PARAM_NAME_AND_TT(InputSpreadAmount), 0.5f),
+					TInputDataVertexModel<bool>(METASOUND_GET_PARAM_NAME_AND_TT(InputEqualPower), false)
+				),
+				FOutputVertexInterface(
+					TOutputDataVertexModel<FAudioBuffer>(METASOUND_GET_PARAM_NAME_AND_TT(OutputAudioMid)),
+					TOutputDataVertexModel<FAudioBuffer>(METASOUND_GET_PARAM_NAME_AND_TT(OutputAudioSide))
+				)
+			);
+
+			return Interface;
+		}
+
+		virtual FDataReferenceCollection GetInputs() const override
+		{
+			using namespace MidSideEncodeVertexNames;
+
+			FDataReferenceCollection InputDataReferences;
+			
+			InputDataReferences.AddDataReadReference(METASOUND_GET_PARAM_NAME(InputAudioLeft), AudioInputLeft);
+			InputDataReferences.AddDataReadReference(METASOUND_GET_PARAM_NAME(InputAudioRight), AudioInputRight);
+			InputDataReferences.AddDataReadReference(METASOUND_GET_PARAM_NAME(InputSpreadAmount), SpreadAmount);
+			InputDataReferences.AddDataReadReference(METASOUND_GET_PARAM_NAME(InputEqualPower), bEqualPower);
+			
+			return InputDataReferences;
+
+		}
+
+		virtual FDataReferenceCollection GetOutputs() const override
+		{
+			using namespace MidSideEncodeVertexNames;
+			
+			FDataReferenceCollection OutputDataReferences;
+
+			OutputDataReferences.AddDataReadReference(METASOUND_GET_PARAM_NAME(OutputAudioMid), AudioOutputMid);
+			OutputDataReferences.AddDataReadReference(METASOUND_GET_PARAM_NAME(OutputAudioSide), AudioOutputSide);
+			
+			return OutputDataReferences;
+		}
+		
+		static TUniquePtr<IOperator> CreateOperator(const FCreateOperatorParams& InParams, FBuildErrorArray& OutErrors)
+		{
+			using namespace MidSideEncodeVertexNames;
+			const FDataReferenceCollection& Inputs = InParams.InputDataReferences;
+			const FInputVertexInterface& InputInterface = DeclareVertexInterface().GetInputInterface();
+			
+			FAudioBufferReadRef LeftAudioIn = Inputs.GetDataReadReferenceOrConstruct<FAudioBuffer>(METASOUND_GET_PARAM_NAME(InputAudioLeft), InParams.OperatorSettings);
+			FAudioBufferReadRef RightAudioIn = Inputs.GetDataReadReferenceOrConstruct<FAudioBuffer>(METASOUND_GET_PARAM_NAME(InputAudioRight), InParams.OperatorSettings);
+			FFloatReadRef SpreadAmountIn = Inputs.GetDataReadReferenceOrConstructWithVertexDefault<float>(InputInterface, METASOUND_GET_PARAM_NAME(InputSpreadAmount), InParams.OperatorSettings);
+			FBoolReadRef bEqualPowerIn = Inputs.GetDataReadReferenceOrConstructWithVertexDefault<bool>(InputInterface, METASOUND_GET_PARAM_NAME(InputEqualPower), InParams.OperatorSettings);
+
+			return MakeUnique<FMidSideEncodeOperator>(InParams.OperatorSettings, LeftAudioIn, RightAudioIn, SpreadAmountIn, bEqualPowerIn);
+		}
+
+		void Execute()
+		{
+			/* Update internal variables, if necessary */
+			
+			const float ClampedSpread = FMath::Clamp(*SpreadAmount, 0.0f, 1.0f);
+			bool bNeedsUpdate = FMath::IsNearlyEqual(ClampedSpread, PrevSpreadAmount);
+			if (bNeedsUpdate)
+			{
+				// Convert to radians between 0.0 and PI/2
+				const float SpreadScale = ClampedSpread * 0.5f * PI;
+
+				// Compute equal power relationship between Mid and Side
+				FMath::SinCos(&SideScale, &MidScale, SpreadScale);
+
+				// Adjust gain so 0.5f Spread results in a 1.0f to 1.0f gain ratio between Mid and Side
+				const float NormalizingFactor = UE_SQRT_2;
+				MidScale *= NormalizingFactor;
+				SideScale *= NormalizingFactor;
+
+				// Clamp values if not Equal Power
+				if (!*bEqualPower)
+				{
+					MidScale = FMath::Clamp(MidScale, 0.0f, 1.0f);
+					SideScale = FMath::Clamp(SideScale, 0.0f, 1.0f);
+				}
+
+			}
+
+			PrevSpreadAmount = ClampedSpread;
+
+			/* Generate Mid-Side Output */
+
+			// Encode the signal to to MS
+			Audio::EncodeMidSide(*AudioInputLeft, *AudioInputRight, *AudioOutputMid, *AudioOutputSide);			
+			
+			// We should now be in MS mode, now we can apply our gain scalars
+
+			// SideScale never changes quickly without MidScale also changing quickly
+			if (FMath::IsNearlyEqual(PrevMidScale, MidScale))
+			{
+				Audio::ArrayMultiplyByConstantInPlace(*AudioOutputMid, MidScale);
+				Audio::ArrayMultiplyByConstantInPlace(*AudioOutputSide, SideScale);
+			}
+			else
+			{
+				Audio::FadeBufferFast(*AudioOutputMid, PrevMidScale, MidScale);
+				Audio::FadeBufferFast(*AudioOutputSide, PrevSideScale, SideScale);
+				
+				PrevMidScale = MidScale;
+				PrevSideScale = SideScale;
+			}
+			
+
+		}
+
+	private:
+
+		// The input audio buffer
+		FAudioBufferReadRef AudioInputLeft;
+		FAudioBufferReadRef AudioInputRight;
+
+		// Stereo spread amount
+		FFloatReadRef SpreadAmount;
+
+		// Whether an equal power relationship is maintained between mid and side channels
+		FBoolReadRef bEqualPower;
+
+		// Output audio buffer
+		FAudioBufferWriteRef AudioOutputMid;
+		FAudioBufferWriteRef AudioOutputSide;
+
+		// Internal variables used to calculate Mid and side channel gain
+		float MidScale;
+		float SideScale;
+		
+		// Storing previous scales for interpolated gain
+		float PrevMidScale;
+		float PrevSideScale;
+
+		// Previous Spread variable; need to update MidScale and SideScale if this changes between Execute() calls
+		float PrevSpreadAmount;
+
+	};
+
+	// Node Class
+	class FMidSideEncodeNode : public FNodeFacade
+	{
+	public:
+		FMidSideEncodeNode(const FNodeInitData& InitData)
+			: FNodeFacade(InitData.InstanceName, InitData.InstanceID, TFacadeOperatorClass<FMidSideEncodeOperator>())
+		{
+		}
+	};
+
+	// Register node
+	METASOUND_REGISTER_NODE(FMidSideEncodeNode)
+
+
+	/* Mid-Side Decoder */
+
+	namespace MidSideDecodeVertexNames
+	{
+		METASOUND_PARAM(InputAudioMid, "In Mid", "The mid audio channel to convert.");
+		METASOUND_PARAM(InputAudioSide, "In Side", "The side audio channel to convert.");
+		METASOUND_PARAM(InputSpreadAmount, "Spread Amount", "Amount of Stereo Spread. 0.0 is no spread, 0.5 is the original signal, and 1.0 is full wide.");
+		METASOUND_PARAM(InputEqualPower, "Equal Power", "Whether an equal power relationship between the mid and side signals should be maintained.");
+	
+		METASOUND_PARAM(OutputAudioLeft, "Out Left", "The left audio channel which has been processed.");
+		METASOUND_PARAM(OutputAudioRight, "Out Right", "the Right audio channel which has been processed.");
+	}
+
+	// Operator Class
+	class FMidSideDecodeOperator : public TExecutableOperator<FMidSideDecodeOperator>
+	{
+	public:
+
+		FMidSideDecodeOperator(const FOperatorSettings& InSettings,
+			const FAudioBufferReadRef& InLeftAudioInput,
+			const FAudioBufferReadRef& InRightAudioInput,
+			const FFloatReadRef& InSpreadAmount,
+			const FBoolReadRef& InEqualPower)
+			: AudioInputMid(InLeftAudioInput)
+			, AudioInputSide(InRightAudioInput)
+			, SpreadAmount(InSpreadAmount)
+			, bEqualPower(InEqualPower)
+			, AudioOutputLeft(FAudioBufferWriteRef::CreateNew(InSettings))
+			, AudioOutputRight(FAudioBufferWriteRef::CreateNew(InSettings))
+			, MidScale(0.0f)
+			, SideScale(0.0f)
+			, PrevMidScale(0.0f)
+			, PrevSideScale(0.0f)
+			, PrevSpreadAmount(*SpreadAmount)
+		{
+		}
+
+		static const FNodeClassMetadata& GetNodeInfo()
+		{
+			auto CreateNodeClassMetadata = []() -> FNodeClassMetadata
+			{
+				FVertexInterface NodeInterface = DeclareVertexInterface();
+
+				FNodeClassMetadata Metadata
+				{
+					FNodeClassName{StandardNodes::Namespace, TEXT("Mid-Side Encode"), StandardNodes::AudioVariant},
+					1, // Major Version
+					0, // Minor Version
+					LOCTEXT("MidSideDecodeDisplayName", "Mid-Side Decode"),
+					LOCTEXT("MidSideDecodeDesc", "Decodes a stereo signal with Mid and Side channels to Left and Right channels."),
+					PluginAuthor,
+					PluginNodeMissingPrompt,
+					NodeInterface,
+					{StandardNodes::Spatalization},
+					{TEXT("MidSideDecode")},
+					FNodeDisplayStyle{}
+				};
+
+				return Metadata;
+			};
+
+			static const FNodeClassMetadata Metadata = CreateNodeClassMetadata();
+			return Metadata;
+		}
+
+		static const FVertexInterface& DeclareVertexInterface()
+		{
+			using namespace MidSideDecodeVertexNames;
+
+			static const FVertexInterface Interface(
+				FInputVertexInterface(
+					TInputDataVertexModel<FAudioBuffer>(METASOUND_GET_PARAM_NAME_AND_TT(InputAudioMid)),
+					TInputDataVertexModel<FAudioBuffer>(METASOUND_GET_PARAM_NAME_AND_TT(InputAudioSide)),
+					TInputDataVertexModel<float>(METASOUND_GET_PARAM_NAME_AND_TT(InputSpreadAmount), 0.5f),
+					TInputDataVertexModel<bool>(METASOUND_GET_PARAM_NAME_AND_TT(InputEqualPower), false)
+				),
+				FOutputVertexInterface(
+					TOutputDataVertexModel<FAudioBuffer>(METASOUND_GET_PARAM_NAME_AND_TT(OutputAudioLeft)),
+					TOutputDataVertexModel<FAudioBuffer>(METASOUND_GET_PARAM_NAME_AND_TT(OutputAudioRight))
+				)
+			);
+
+			return Interface;
+		}
+
+		virtual FDataReferenceCollection GetInputs() const override
+		{
+			using namespace MidSideDecodeVertexNames;
+
+			FDataReferenceCollection InputDataReferences;
+
+			InputDataReferences.AddDataReadReference(METASOUND_GET_PARAM_NAME(InputAudioMid), AudioInputMid);
+			InputDataReferences.AddDataReadReference(METASOUND_GET_PARAM_NAME(InputAudioSide), AudioInputSide);
+			InputDataReferences.AddDataReadReference(METASOUND_GET_PARAM_NAME(InputSpreadAmount), SpreadAmount);
+			InputDataReferences.AddDataReadReference(METASOUND_GET_PARAM_NAME(InputEqualPower), bEqualPower);
+
+			return InputDataReferences;
+
+		}
+
+		virtual FDataReferenceCollection GetOutputs() const override
+		{
+			using namespace MidSideDecodeVertexNames;
+
+			FDataReferenceCollection OutputDataReferences;
+
+			OutputDataReferences.AddDataReadReference(METASOUND_GET_PARAM_NAME(OutputAudioLeft), AudioOutputLeft);
+			OutputDataReferences.AddDataReadReference(METASOUND_GET_PARAM_NAME(OutputAudioRight), AudioOutputRight);
+
+			return OutputDataReferences;
+		}
+
+		static TUniquePtr<IOperator> CreateOperator(const FCreateOperatorParams& InParams, FBuildErrorArray& OutErrors)
+		{
+			using namespace MidSideDecodeVertexNames;
+			const FDataReferenceCollection& Inputs = InParams.InputDataReferences;
+			const FInputVertexInterface& InputInterface = DeclareVertexInterface().GetInputInterface();
+
+			FAudioBufferReadRef MidAudioIn = Inputs.GetDataReadReferenceOrConstruct<FAudioBuffer>(METASOUND_GET_PARAM_NAME(InputAudioMid), InParams.OperatorSettings);
+			FAudioBufferReadRef SideAudioIn = Inputs.GetDataReadReferenceOrConstruct<FAudioBuffer>(METASOUND_GET_PARAM_NAME(InputAudioSide), InParams.OperatorSettings);
+			FFloatReadRef SpreadAmountIn = Inputs.GetDataReadReferenceOrConstructWithVertexDefault<float>(InputInterface, METASOUND_GET_PARAM_NAME(InputSpreadAmount), InParams.OperatorSettings);
+			FBoolReadRef bEqualPowerIn = Inputs.GetDataReadReferenceOrConstructWithVertexDefault<bool>(InputInterface, METASOUND_GET_PARAM_NAME(InputEqualPower), InParams.OperatorSettings);
+
+			return MakeUnique<FMidSideDecodeOperator>(InParams.OperatorSettings, MidAudioIn, SideAudioIn, SpreadAmountIn, bEqualPowerIn);
+		}
+
+		void Execute()
+		{
+			/* Update internal variables, if necessary */
+			const float ClampedSpread = FMath::Clamp(*SpreadAmount, 0.0f, 1.0f);
+			bool bNeedsUpdate = FMath::IsNearlyEqual(ClampedSpread, PrevSpreadAmount);
+			if (bNeedsUpdate)
+			{
+				// Convert to radians between 0.0 and PI/2
+				const float SpreadScale = ClampedSpread * 0.5f * PI;
+
+				// Compute equal power relationship between Mid and Side
+				FMath::SinCos(&SideScale, &MidScale, SpreadScale);
+
+
+				// Adjust gain so 0.5f Spread results in a 1.0f to 1.0f gain ratio between Mid and Side
+				const float NormalizingFactor = UE_SQRT_2;
+				MidScale *= NormalizingFactor;
+				SideScale *= NormalizingFactor;
+
+				// Clamp values if not Equal Power
+				if (!*bEqualPower)
+				{
+					MidScale = FMath::Clamp(MidScale, 0.0f, 1.0f);
+					SideScale = FMath::Clamp(SideScale, 0.0f, 1.0f);
+				}
+			}
+
+			PrevSpreadAmount = ClampedSpread;
+
+			/* Generate Mid-Side Output */
+
+			// We need to make a copy of the input data so we can scale it before decoding
+			const float* MidInput = AudioInputMid->GetData();
+			const float* SideInput = AudioInputSide->GetData();
+
+			const int32 NumFrames = AudioInputMid->Num();
+
+			MidInBuffer.Reset();
+			MidInBuffer.Append(MidInput, NumFrames);
+			SideInBuffer.Reset();
+			SideInBuffer.Append(SideInput, NumFrames);
+
+
+			// SideScale never changes quickly without MidScale also changing quickly
+			if (FMath::IsNearlyEqual(PrevMidScale, MidScale))
+			{
+				Audio::ArrayMultiplyByConstantInPlace(MidInBuffer, MidScale);
+				Audio::ArrayMultiplyByConstantInPlace(SideInBuffer, SideScale);
+			}
+			else
+			{
+				Audio::FadeBufferFast(MidInBuffer, PrevMidScale, MidScale);
+				Audio::FadeBufferFast(SideInBuffer, PrevSideScale, SideScale);
+
+				PrevMidScale = MidScale;
+				PrevSideScale = SideScale;
+			}
+
+			// Convert to L/R Signal
+			Audio::DecodeMidSide(MidInBuffer, SideInBuffer, *AudioOutputLeft, *AudioOutputRight);
+		}
+
+	private:
+
+		// The input audio buffer
+		FAudioBufferReadRef AudioInputMid;
+		FAudioBufferReadRef AudioInputSide;
+
+		// Stereo spread amount
+		FFloatReadRef SpreadAmount;
+
+		// Whether an equal power relationship is maintained between mid and side channels
+		FBoolReadRef bEqualPower;
+
+		// Output audio buffer
+		FAudioBufferWriteRef AudioOutputLeft;
+		FAudioBufferWriteRef AudioOutputRight;
+
+		// Internal variables used to calculate Mid and side channel gain
+		float MidScale;
+		float SideScale;
+
+		// Storing previous scales for interpolated gain
+		float PrevMidScale;
+		float PrevSideScale;
+
+		// Previous Spread variable; need to update MidScale and SideScale if this changes between Execute() calls
+		float PrevSpreadAmount;
+
+		// Reusable buffers to save performance during Execute().
+		Audio::FAlignedFloatBuffer MidInBuffer;
+		Audio::FAlignedFloatBuffer SideInBuffer;
+	};
+
+	// Node Class
+	class FMidSideDecodeNode : public FNodeFacade
+	{
+	public:
+		FMidSideDecodeNode(const FNodeInitData& InitData)
+			: FNodeFacade(InitData.InstanceName, InitData.InstanceID, TFacadeOperatorClass<FMidSideDecodeOperator>())
+		{
+		}
+	};
+
+	// Register node
+	METASOUND_REGISTER_NODE(FMidSideDecodeNode)
+
+}
+
+#undef LOCTEXT_NAMESPACE
diff --git a/Engine/Source/Runtime/SignalProcessing/Private/Dsp.cpp b/Engine/Source/Runtime/SignalProcessing/Private/Dsp.cpp
new file mode 100644
--- /dev/null
+++ b/Engine/Source/Runtime/SignalProcessing/Private/Dsp.cpp
@@ -0,0 +1,65 @@
+// Copyright Epic Games, Inc. All Rights Reserved.
+
+#include "DSP/Dsp.h"
+
+namespace Audio
+{
+	void EncodeMidSide(
+		const FAlignedFloatBuffer& InLeftChannel, 
+		const FAlignedFloatBuffer& InRightChannel, 
+		FAlignedFloatBuffer& OutMidChannel, 
+		FAlignedFloatBuffer& OutSideChannel)
+	{
+		const int32 LeftInNumElements = InLeftChannel.Num();
+		const int32 RightInNumElements = InRightChannel.Num();
+		const bool bInputChannelsEqual = LeftInNumElements == RightInNumElements;
+		const bool bOutputChannelsBigEnough = (OutMidChannel.Num() >= LeftInNumElements) && (OutSideChannel.Num() >= LeftInNumElements);
+
+		if (!bInputChannelsEqual || !bOutputChannelsBigEnough)
+		{
+			return;
+		}
+
+		OutMidChannel.Reset(LeftInNumElements);
+		OutSideChannel.Reset(LeftInNumElements);
+		OutMidChannel.AddUninitialized(LeftInNumElements);
+		OutSideChannel.AddUninitialized(LeftInNumElements);
+
+		ArraySubtract(InLeftChannel, InRightChannel, OutSideChannel);
+
+		//Output
+		SumBuffers(InLeftChannel, InRightChannel, OutMidChannel);
+
+	}
+
+	void DecodeMidSide(
+		const FAlignedFloatBuffer& InMidChannel,
+		const FAlignedFloatBuffer& InSideChannel,
+		FAlignedFloatBuffer& OutLeftChannel,
+		FAlignedFloatBuffer& OutRightChannel)
+	{
+		const int32 LeftInNumElements = InMidChannel.Num();
+		const int32 RightInNumElements = InSideChannel.Num();
+		const bool bInputChannelsEqual = LeftInNumElements == RightInNumElements;
+		const bool bOutputChannelsBigEnough = (OutLeftChannel.Num() >= LeftInNumElements) && (OutRightChannel.Num() >= LeftInNumElements);
+
+		if (!bInputChannelsEqual || !bOutputChannelsBigEnough)
+		{
+			return;
+		}
+
+		OutLeftChannel.Reset(LeftInNumElements);
+		OutRightChannel.Reset(LeftInNumElements);
+		OutLeftChannel.AddUninitialized(InMidChannel.Num());
+		OutRightChannel.AddUninitialized(InSideChannel.Num());
+
+		// Calculate the average value between the two signals at each sample
+		const float AverageScale = 0.5f;
+		BufferWeightedSumFast(InMidChannel, AverageScale, InSideChannel, AverageScale, OutLeftChannel);
+
+		ArraySubtract(InMidChannel, InSideChannel, OutRightChannel);
+		MultiplyBufferByConstantInPlace(OutRightChannel, AverageScale);
+
+	}
+
+}
\ No newline at end of file
diff --git a/Engine/Source/Runtime/SignalProcessing/Private/FloatArrayMath.cpp b/Engine/Source/Runtime/SignalProcessing/Private/FloatArrayMath.cpp
--- a/Engine/Source/Runtime/SignalProcessing/Private/FloatArrayMath.cpp
+++ b/Engine/Source/Runtime/SignalProcessing/Private/FloatArrayMath.cpp
@@ -660,16 +660,13 @@ namespace Audio
 	{
 		const int32 Num = InMinuend.Num();
 
-		checkf(Num == InSubtrahend.Num(), TEXT("InMinuend and InSubtrahend must have equal Num elements (%d vs %d)"), Num, InSubtrahend.Num());
-
-		OutArray.Reset(Num);
+		checkf(Num == InSubtrahend.Num() && Num == OutArray.Num(), TEXT("InMinuend, InSubtrahend, and OutArray must have equal Num elements (%d vs %d vs %d)"), Num, InSubtrahend.Num(), OutArray.Num());
+		
 
 		if (Num < 1)
 		{
 			return;
 		}
-
-		OutArray.AddUninitialized(Num);
 		
 		const float* MinuendPtr = InMinuend.GetData();
 		const float* SubtrahendPtr = InSubtrahend.GetData();
@@ -681,6 +678,11 @@ namespace Audio
 		}
 	}
 
+	void ArraySubtract(const FAlignedFloatBuffer& InMinuend, const FAlignedFloatBuffer& InSubtrahend, FAlignedFloatBuffer& OutArray)
+	{
+		BufferSubtractFast(InMinuend, InSubtrahend, OutArray);
+	}
+
 	void ArraySqrtInPlace(TArrayView<float> InValues)
 	{
 		const int32 Num = InValues.Num();
diff --git a/Engine/Source/Runtime/SignalProcessing/Public/DSP/Dsp.h b/Engine/Source/Runtime/SignalProcessing/Public/DSP/Dsp.h
--- a/Engine/Source/Runtime/SignalProcessing/Public/DSP/Dsp.h
+++ b/Engine/Source/Runtime/SignalProcessing/Public/DSP/Dsp.h
@@ -7,6 +7,7 @@
 #include "Templates/IsFloatingPoint.h"
 #include "Templates/IsIntegral.h"
 #include "Templates/IsSigned.h"
+#include "FloatArrayMath.h"
 
 // Macros which can be enabled to cause DSP sample checking
 #if 0
@@ -371,7 +372,7 @@ namespace Audio
 		OutRight = FMath::Clamp(FastSin(RightPhase), 0.0f, 1.0f);
 	}
  
-	// This function encodes a stereo Left/Right signal into a stereo Mid/Side signal 
+	// Encodes a stereo Left/Right signal into a stereo Mid/Side signal 
 	static FORCEINLINE void EncodeMidSide(float& LeftChannel, float& RightChannel)
 	{
 		const float Temp = (LeftChannel - RightChannel);
@@ -380,7 +381,15 @@ namespace Audio
 		RightChannel = Temp;
 	}
 
-	// This function decodes a stereo Mid/Side signal into a stereo Left/Right signal
+	// Encodes a stereo Left/Right signal into a stereo Mid/Side signal
+	SIGNALPROCESSING_API void EncodeMidSide(
+		const FAlignedFloatBuffer& InLeftChannel,
+		const FAlignedFloatBuffer& InRightChannel,
+		FAlignedFloatBuffer& OutMidChannel,
+		FAlignedFloatBuffer& OutSideChannel);
+	
+
+	// Decodes a stereo Mid/Side signal into a stereo Left/Right signal
 	static FORCEINLINE void DecodeMidSide(float& MidChannel, float& SideChannel)
 	{
 		const float Temp = (MidChannel - SideChannel) * 0.5f;
@@ -389,6 +398,13 @@ namespace Audio
 		SideChannel = Temp;
 	}
 
+	// Decodes a stereo Mid/Side signal into a stereo Left/Right signal
+	SIGNALPROCESSING_API void DecodeMidSide(
+		const FAlignedFloatBuffer& InMidChannel,
+		const FAlignedFloatBuffer& InSideChannel,
+		FAlignedFloatBuffer& OutLeftChannel,
+		FAlignedFloatBuffer& OutRightChannel);
+
 	// Helper function to get bandwidth from Q
 	static FORCEINLINE float GetBandwidthFromQ(const float InQ)
 	{
diff --git a/Engine/Source/Runtime/SignalProcessing/Public/DSP/FloatArrayMath.h b/Engine/Source/Runtime/SignalProcessing/Public/DSP/FloatArrayMath.h
--- a/Engine/Source/Runtime/SignalProcessing/Public/DSP/FloatArrayMath.h
+++ b/Engine/Source/Runtime/SignalProcessing/Public/DSP/FloatArrayMath.h
@@ -135,6 +135,7 @@ namespace Audio
 	 *  OutArray is the array which holds the result.
 	 */
 	SIGNALPROCESSING_API void ArraySubtract(TArrayView<const float> InMinuend, TArrayView<const float> InSubtrahend, TArray<float>& OutArray);
+	SIGNALPROCESSING_API void ArraySubtract(const FAlignedFloatBuffer& InMinuend, const FAlignedFloatBuffer& InSubtrahend, FAlignedFloatBuffer& OutArray);
 
 	/** Subtract value from each element in InValues */
 	SIGNALPROCESSING_API void ArraySubtractByConstantInPlace(TArrayView<float> InValues, float InSubtrahend);
