Fixing MetaSound use of FNodeClassInfo and FNodeRegsitryKey

diff --git a/Engine/Plugins/Runtime/Metasound/Source/MetasoundEditor/Private/MetasoundEditorGraphNode.cpp b/Engine/Plugins/Runtime/Metasound/Source/MetasoundEditor/Private/MetasoundEditorGraphNode.cpp
--- a/Engine/Plugins/Runtime/Metasound/Source/MetasoundEditor/Private/MetasoundEditorGraphNode.cpp
+++ b/Engine/Plugins/Runtime/Metasound/Source/MetasoundEditor/Private/MetasoundEditorGraphNode.cpp
@@ -277,7 +277,7 @@ FString UMetasoundEditorGraphNode::GetPinMetaData(FName InPinName, FName InKey)
 				return FString();
 			}
 
-			Metasound::FDataTypeRegistryInfo DataTypeInfo;
+			Metasound::Frontend::FDataTypeRegistryInfo DataTypeInfo;
 			if (!ensure(Registry->GetInfoForDataType(Handle->GetDataType(), DataTypeInfo)))
 			{
 				return FString();
diff --git a/Engine/Plugins/Runtime/Metasound/Source/MetasoundEditor/Private/MetasoundEditorModule.h b/Engine/Plugins/Runtime/Metasound/Source/MetasoundEditor/Private/MetasoundEditorModule.h
--- a/Engine/Plugins/Runtime/Metasound/Source/MetasoundEditor/Private/MetasoundEditorModule.h
+++ b/Engine/Plugins/Runtime/Metasound/Source/MetasoundEditor/Private/MetasoundEditorModule.h
@@ -20,6 +20,8 @@ namespace Metasound
 {
 	namespace Editor
 	{
+		using FDataTypeRegistryInfo = Frontend::FDataTypeRegistryInfo;
+
 		struct FEditorDataType
 		{
 			FEdGraphPinType PinType;
diff --git a/Engine/Plugins/Runtime/Metasound/Source/MetasoundFrontend/Private/MetasoundAssetBase.cpp b/Engine/Plugins/Runtime/Metasound/Source/MetasoundFrontend/Private/MetasoundAssetBase.cpp
--- a/Engine/Plugins/Runtime/Metasound/Source/MetasoundFrontend/Private/MetasoundAssetBase.cpp
+++ b/Engine/Plugins/Runtime/Metasound/Source/MetasoundFrontend/Private/MetasoundAssetBase.cpp
@@ -32,49 +32,71 @@ void FMetasoundAssetBase::RegisterGraphWithFrontend()
 		AssetName = OwningAsset->GetName();
 	}
 
-
-	// Function for creating a core INode.
-	FCreateMetasoundNodeFunction CreateNode = [Name=AssetName, DocumentPtr=DocumentPtr](const Metasound::FNodeInitData& InInitData) -> TUniquePtr<INode>
+	// Registers node by copying document. Updates to document require re-registration.
+	class FNodeRegistryEntry : public INodeRegistryEntry
 	{
-		check(IsInGameThread());
+	public:
+		FNodeRegistryEntry(const FString& InName, const FMetasoundFrontendDocument& InDocument)
+		: Name(InName)
+		, Document(InDocument)
+		{
+			// Copy frontend class to preserve original document.
+			FrontendClass = Document.RootGraph;
+			FrontendClass.Metadata.Type = EMetasoundFrontendClassType::External;
+			ClassInfo = FNodeClassInfo{FrontendClass};
+		}
+
+		virtual ~FNodeRegistryEntry() = default;
 
-		if (const FMetasoundFrontendDocument* Document = DocumentPtr.Get())
+		virtual const FNodeClassInfo& GetClassInfo() const override
 		{
-			return FFrontendGraphBuilder().CreateGraph(*Document);
+			return ClassInfo;
 		}
-		else
+
+		virtual TUniquePtr<INode> CreateNode(FDefaultNodeConstructorParams&&) const override
 		{
-			UE_LOG(LogMetaSound, Error, TEXT("Failed to create node for MetaSoundSource [Name:%s]"), *Name);
-			return TUniquePtr<INode>(nullptr);
+			return nullptr;
 		}
-	};
 
-	// Function for creating a frontend class.
-	FCreateMetasoundFrontendClassFunction CreateFrontendClass = [Name=AssetName, DocumentPtr=DocumentPtr]() -> FMetasoundFrontendClass
-	{
-		check(IsInGameThread());
+		virtual TUniquePtr<INode> CreateNode(FDefaultLiteralNodeConstructorParams&&) const override
+		{
+			return nullptr;
+		}
 
-		FMetasoundFrontendClass FrontendClass;
-		if (const FMetasoundFrontendDocument* Document = DocumentPtr.Get())
+		virtual TUniquePtr<INode> CreateNode(const FNodeInitData&) const override
 		{
-			FrontendClass = Document->RootGraph;
-			FrontendClass.Metadata.Type = EMetasoundFrontendClassType::External;
+			return FFrontendGraphBuilder().CreateGraph(Document);
 		}
-		else
+
+		virtual const FMetasoundFrontendClass& GetFrontendClass() const override
+		{
+			return FrontendClass;
+		}
+
+		virtual TUniquePtr<INodeRegistryEntry> Clone() const override
 		{
-			UE_LOG(LogMetaSound, Error, TEXT("Failed to create frontend class for MetaSoundSource [Name:%s]"), *Name);
+			return MakeUnique<FNodeRegistryEntry>(Name, Document);
 		}
 
-		return FrontendClass;
+	private:
+		
+		FString Name;
+		FMetasoundFrontendDocument Document;
+		FMetasoundFrontendClass FrontendClass;
+		FNodeClassInfo ClassInfo;
 	};
 
+
 	if (IsValidNodeRegistryKey(RegistryKey))
 	{
 		// Unregister prior version if it exists.
-		ensure(FMetasoundFrontendRegistryContainer::Get()->UnregisterExternalNode(RegistryKey));
+		ensure(FMetasoundFrontendRegistryContainer::Get()->UnregisterNode(RegistryKey));
+	}
+	
+	if (const FMetasoundFrontendDocument* Doc = DocumentPtr.Get())
+	{
+		RegistryKey = FMetasoundFrontendRegistryContainer::Get()->RegisterNode(MakeUnique<FNodeRegistryEntry>(AssetName, *Doc));
 	}
-
-	RegistryKey = FMetasoundFrontendRegistryContainer::Get()->RegisterExternalNode(MoveTemp(CreateNode), MoveTemp(CreateFrontendClass));
 
 	if (!IsValidNodeRegistryKey(RegistryKey))
 	{
@@ -291,7 +313,7 @@ TArray<FString> FMetasoundAssetBase::GetTransmittableInputVertexNames() const
 		Frontend::FConstClassInputAccessPtr ClassInputPtr = RootGraph->FindClassInputWithName(InVertexName);
 		if (const FMetasoundFrontendClassInput* ClassInput = ClassInputPtr.Get())
 		{
-			FDataTypeRegistryInfo TypeInfo;
+			Frontend::FDataTypeRegistryInfo TypeInfo;
 			if (Frontend::GetTraitsForDataType(ClassInput->TypeName, TypeInfo))
 			{
 				if (TypeInfo.bIsTransmittable)
diff --git a/Engine/Plugins/Runtime/Metasound/Source/MetasoundFrontend/Private/MetasoundFrontend.cpp b/Engine/Plugins/Runtime/Metasound/Source/MetasoundFrontend/Private/MetasoundFrontend.cpp
--- a/Engine/Plugins/Runtime/Metasound/Source/MetasoundFrontend/Private/MetasoundFrontend.cpp
+++ b/Engine/Plugins/Runtime/Metasound/Source/MetasoundFrontend/Private/MetasoundFrontend.cpp
@@ -25,23 +25,6 @@ namespace Metasound
 {
 	namespace Frontend
 	{
-		TArray<FNodeClassInfo> GetAllAvailableNodeClasses(FRegistryTransactionID* OutCurrentRegistryTransactionID)
-		{
-			FMetasoundFrontendRegistryContainer* Registry = FMetasoundFrontendRegistryContainer::Get();
-			if (ensure(nullptr != Registry))
-			{
-				return Registry->GetAllAvailableNodeClasses(OutCurrentRegistryTransactionID);
-			}
-			else
-			{
-				if (nullptr != OutCurrentRegistryTransactionID)
-				{
-					*OutCurrentRegistryTransactionID = GetOriginRegistryTransactionID();
-				}
-				return TArray<FNodeClassInfo>();
-			}
-		}
-
 		TArray<const IRegistryTransaction*> GetRegistryTransactionsSince(FRegistryTransactionID InTransactionID, FRegistryTransactionID* OutCurrentTransactionID)
 		{
 			FMetasoundFrontendRegistryContainer* Registry = FMetasoundFrontendRegistryContainer::Get();
@@ -60,9 +43,9 @@ namespace Metasound
 		}
 
 		// gets all metadata (name, description, author, what to say if it's missing) for a given node.
-		FMetasoundFrontendClassMetadata GenerateClassMetadata(const FNodeClassInfo& InInfo)
+		FMetasoundFrontendClassMetadata GenerateClassMetadata(const FNodeRegistryKey& InKey)
 		{
-			return GenerateClassDescription(InInfo).Metadata;
+			return GenerateClassDescription(InKey).Metadata;
 		}
 
 		FMetasoundFrontendClass GenerateClassDescription(const FNodeClassMetadata& InNodeMetadata, EMetasoundFrontendClassType ClassType)
@@ -155,7 +138,7 @@ namespace Metasound
 			return ClassDescription;
 		}
 
-		FMetasoundFrontendClass GenerateClassDescription(const FNodeClassInfo& InNodeInfo)
+		FMetasoundFrontendClass GenerateClassDescription(const FNodeRegistryKey& InKey)
 		{
 			FMetasoundFrontendClass OutClass;
 
@@ -163,8 +146,8 @@ namespace Metasound
 
 			if (ensure(nullptr != Registry))
 			{
-				bool bSuccess = Registry->FindFrontendClassFromRegistered(InNodeInfo, OutClass);
-				ensureAlwaysMsgf(bSuccess, TEXT("Cannot generate description of unregistered node [RegistryKey:%s]"), *InNodeInfo.LookupKey);
+				bool bSuccess = Registry->FindFrontendClassFromRegistered(InKey, OutClass);
+				ensureAlwaysMsgf(bSuccess, TEXT("Cannot generate description of unregistered node [RegistryKey:%s]"), *InKey);
 			}
 
 			return OutClass;
diff --git a/Engine/Plugins/Runtime/Metasound/Source/MetasoundFrontend/Private/MetasoundFrontendBaseClasses.cpp b/Engine/Plugins/Runtime/Metasound/Source/MetasoundFrontend/Private/MetasoundFrontendBaseClasses.cpp
deleted file mode 100644
diff --git a/Engine/Plugins/Runtime/Metasound/Source/MetasoundFrontend/Private/MetasoundFrontendController.cpp b/Engine/Plugins/Runtime/Metasound/Source/MetasoundFrontend/Private/MetasoundFrontendController.cpp
--- a/Engine/Plugins/Runtime/Metasound/Source/MetasoundFrontend/Private/MetasoundFrontendController.cpp
+++ b/Engine/Plugins/Runtime/Metasound/Source/MetasoundFrontend/Private/MetasoundFrontendController.cpp
@@ -2,7 +2,6 @@
 
 #include "MetasoundFrontendController.h"
 
-#include "MetasoundFrontendBaseClasses.h"
 #include "MetasoundFrontendInvalidController.h"
 #include "MetasoundFrontendStandardController.h"
 
diff --git a/Engine/Plugins/Runtime/Metasound/Source/MetasoundFrontend/Private/MetasoundFrontendDocumentAccessPtr.cpp b/Engine/Plugins/Runtime/Metasound/Source/MetasoundFrontend/Private/MetasoundFrontendDocumentAccessPtr.cpp
--- a/Engine/Plugins/Runtime/Metasound/Source/MetasoundFrontend/Private/MetasoundFrontendDocumentAccessPtr.cpp
+++ b/Engine/Plugins/Runtime/Metasound/Source/MetasoundFrontend/Private/MetasoundFrontendDocumentAccessPtr.cpp
@@ -25,11 +25,12 @@ namespace Metasound
 
 			bool IsMatchingMetasoundClass(const FNodeClassInfo& InNodeClass, const FMetasoundFrontendClassMetadata& InMetadata) 
 			{
-				if (InNodeClass.NodeType == InMetadata.Type)
-				{
-					return InNodeClass.LookupKey == FRegistry::GetRegistryKey(InMetadata);
-				}
-				return false;
+				return FRegistry::GetRegistryKey(InNodeClass) == FRegistry::GetRegistryKey(InMetadata);
+			}
+
+			bool IsMatchingMetasoundClass(const FNodeRegistryKey& InKey, const FMetasoundFrontendClassMetadata& InMetadata) 
+			{
+				return InKey == FRegistry::GetRegistryKey(InMetadata);
 			}
 		}
 
@@ -513,6 +514,22 @@ namespace Metasound
 			return GetMemberAccessPtr<FClassAccessPtr>(FindClassWithInfo);
 		}
 
+		FClassAccessPtr FDocumentAccessPtr::GetClassWithRegistryKey(const FNodeRegistryKey& InKey)
+		{
+			auto FindClassWithRegistryKey = [Key=InKey](FMetasoundFrontendDocument& InDoc)
+			{
+				auto IsClassWithRegistryKey = [=](const FMetasoundFrontendClass& InClass) 
+				{ 
+					using namespace MetasoundFrontendDocumentAccessPtrPrivate;
+					return IsMatchingMetasoundClass(Key, InClass.Metadata);
+				};
+
+				return InDoc.Dependencies.FindByPredicate(IsClassWithRegistryKey);
+			};
+
+			return GetMemberAccessPtr<FClassAccessPtr>(FindClassWithRegistryKey);
+		}
+
 		FConstGraphClassAccessPtr FDocumentAccessPtr::GetRootGraph() const
 		{
 			return GetMemberAccessPtr<FConstGraphClassAccessPtr>([](const FMetasoundFrontendDocument& InDoc) { return &InDoc.RootGraph; });
@@ -595,6 +612,22 @@ namespace Metasound
 			return GetMemberAccessPtr<FConstClassAccessPtr>(FindClassWithInfo);
 		}
 
+		FConstClassAccessPtr FDocumentAccessPtr::GetClassWithRegistryKey(const FNodeRegistryKey& InKey) const
+		{
+			auto FindClassWithRegistryKey = [Key=InKey](const FMetasoundFrontendDocument& InDoc)
+			{
+				auto IsClassWithRegistryKey = [=](const FMetasoundFrontendClass& InClass) 
+				{ 
+					using namespace MetasoundFrontendDocumentAccessPtrPrivate;
+					return IsMatchingMetasoundClass(Key, InClass.Metadata);
+				};
+
+				return InDoc.Dependencies.FindByPredicate(IsClassWithRegistryKey);
+			};
+
+			return GetMemberAccessPtr<FConstClassAccessPtr>(FindClassWithRegistryKey);
+		}
+
 		FConstGraphClassAccessPtr FConstDocumentAccessPtr::GetRootGraph() const
 		{
 			return GetMemberAccessPtr<FConstGraphClassAccessPtr>([](const FMetasoundFrontendDocument& InDoc) { return &InDoc.RootGraph; });
@@ -676,5 +709,22 @@ namespace Metasound
 
 			return GetMemberAccessPtr<FConstClassAccessPtr>(FindClassWithInfo);
 		}
+
+
+		FConstClassAccessPtr FConstDocumentAccessPtr::GetClassWithRegistryKey(const FNodeRegistryKey& InKey) const
+		{
+			auto FindClassWithRegistryKey = [Key=InKey](const FMetasoundFrontendDocument& InDoc)
+			{
+				auto IsClassWithRegistryKey = [=](const FMetasoundFrontendClass& InClass) 
+				{ 
+					using namespace MetasoundFrontendDocumentAccessPtrPrivate;
+					return IsMatchingMetasoundClass(Key, InClass.Metadata);
+				};
+
+				return InDoc.Dependencies.FindByPredicate(IsClassWithRegistryKey);
+			};
+
+			return GetMemberAccessPtr<FConstClassAccessPtr>(FindClassWithRegistryKey);
+		}
 	}
 }
diff --git a/Engine/Plugins/Runtime/Metasound/Source/MetasoundFrontend/Private/MetasoundFrontendGraph.cpp b/Engine/Plugins/Runtime/Metasound/Source/MetasoundFrontend/Private/MetasoundFrontendGraph.cpp
--- a/Engine/Plugins/Runtime/Metasound/Source/MetasoundFrontend/Private/MetasoundFrontendGraph.cpp
+++ b/Engine/Plugins/Runtime/Metasound/Source/MetasoundFrontend/Private/MetasoundFrontendGraph.cpp
@@ -183,7 +183,7 @@ namespace Metasound
 					};
 
 					// TODO: create input node using external class definition
-					return FMetasoundFrontendRegistryContainer::Get()->ConstructInputNode(InputVertex.TypeName, MoveTemp(InitParams));
+					return FMetasoundFrontendRegistryContainer::Get()->CreateInputNode(InputVertex.TypeName, MoveTemp(InitParams));
 				}
 				else
 				{
@@ -225,7 +225,7 @@ namespace Metasound
 			}
 
 			// TODO: create output node using external class definition
-			return FMetasoundFrontendRegistryContainer::Get()->ConstructOutputNode(OutputVertex.TypeName, MoveTemp(InitParams));
+			return FMetasoundFrontendRegistryContainer::Get()->CreateOutputNode(OutputVertex.TypeName, MoveTemp(InitParams));
 		}
 		return TUniquePtr<INode>(nullptr);
 	}
@@ -249,7 +249,7 @@ namespace Metasound
 
 		Metasound::Frontend::FNodeRegistryKey Key = FMetasoundFrontendRegistryContainer::GetRegistryKey(InClass.Metadata);
 
-		return FMetasoundFrontendRegistryContainer::Get()->ConstructExternalNode(Key, InitData);
+		return FMetasoundFrontendRegistryContainer::Get()->CreateNode(Key, InitData);
 	}
 
 	const FMetasoundFrontendClassInput* FFrontendGraphBuilder::FindClassInputForInputNode(const FMetasoundFrontendGraphClass& InOwningGraph, const FMetasoundFrontendNode& InInputNode, int32& OutClassInputIndex)
@@ -530,7 +530,7 @@ namespace Metasound
 			// 1. Construct and add the default variable to the graph
 			{
 				FVariableNodeConstructorParams InitParams = VariableData.InitParams.Clone();
-				TUniquePtr<const INode> DefaultVariable = FMetasoundFrontendRegistryContainer::Get()->ConstructVariableNode(VariableData.TypeName, MoveTemp(InitParams));
+				TUniquePtr<const INode> DefaultVariable = FMetasoundFrontendRegistryContainer::Get()->CreateVariableNode(VariableData.TypeName, MoveTemp(InitParams));
 				InGraphContext.Graph->AddNode(VariableNodeID, TSharedPtr<const INode>(DefaultVariable.Release()));
 			}
 
diff --git a/Engine/Plugins/Runtime/Metasound/Source/MetasoundFrontend/Private/MetasoundFrontendInvalidController.h b/Engine/Plugins/Runtime/Metasound/Source/MetasoundFrontend/Private/MetasoundFrontendInvalidController.h
--- a/Engine/Plugins/Runtime/Metasound/Source/MetasoundFrontend/Private/MetasoundFrontendInvalidController.h
+++ b/Engine/Plugins/Runtime/Metasound/Source/MetasoundFrontend/Private/MetasoundFrontendInvalidController.h
@@ -5,7 +5,6 @@
 
 #include "CoreMinimal.h"
 #include "MetasoundBuilderInterface.h"
-#include "MetasoundFrontendBaseClasses.h"
 #include "MetasoundFrontendController.h"
 #include "MetasoundFrontendDocument.h"
 #include "MetasoundGraph.h"
@@ -287,7 +286,6 @@ namespace Metasound
 			// @returns false if the input name couldn't be found.
 			bool ClearLiteralForInput(const FString& InInputName, FGuid InVertexID) override { return false; }
 
-			TSharedRef<INodeController> AddNode(const FNodeClassInfo& InNodeClass) override { return FInvalidNodeController::GetInvalid(); }
 			TSharedRef<INodeController> AddNode(const FNodeRegistryKey& InNodeClass) override { return FInvalidNodeController::GetInvalid(); }
 			TSharedRef<INodeController> AddNode(const FMetasoundFrontendClassMetadata& InNodeClass) override { return FInvalidNodeController::GetInvalid(); }
 			TSharedRef<INodeController> AddDuplicateNode(const INodeController& InNode) override { return FInvalidNodeController::GetInvalid(); }
@@ -341,8 +339,8 @@ namespace Metasound
 				FConstGraphClassAccessPtr FindSubgraphWithID(FGuid InClassID) const override { return FConstGraphClassAccessPtr(); }
 				FConstClassAccessPtr FindClassWithID(FGuid InClassID) const override { return FConstClassAccessPtr(); }
 
-				FConstClassAccessPtr FindClass(const FNodeClassInfo& InNodeClass) const override { return FConstClassAccessPtr(); }
-				FConstClassAccessPtr FindOrAddClass(const FNodeClassInfo& InNodeClass) override { return FConstClassAccessPtr(); }
+				FConstClassAccessPtr FindClass(const FNodeRegistryKey& InKey) const override { return FConstClassAccessPtr(); }
+				FConstClassAccessPtr FindOrAddClass(const FNodeRegistryKey& InKey) override { return FConstClassAccessPtr(); }
 				FConstClassAccessPtr FindClass(const FMetasoundFrontendClassMetadata& InMetadata) const override{ return FConstClassAccessPtr(); }
 				FConstClassAccessPtr FindOrAddClass(const FMetasoundFrontendClassMetadata& InMetadata) override{ return FConstClassAccessPtr(); }
 				FGraphHandle AddDuplicateSubgraph(const IGraphController& InGraph) override { return FInvalidGraphController::GetInvalid(); }
diff --git a/Engine/Plugins/Runtime/Metasound/Source/MetasoundFrontend/Private/MetasoundFrontendLiteral.cpp b/Engine/Plugins/Runtime/Metasound/Source/MetasoundFrontend/Private/MetasoundFrontendLiteral.cpp
--- a/Engine/Plugins/Runtime/Metasound/Source/MetasoundFrontend/Private/MetasoundFrontendLiteral.cpp
+++ b/Engine/Plugins/Runtime/Metasound/Source/MetasoundFrontend/Private/MetasoundFrontendLiteral.cpp
@@ -502,7 +502,7 @@ Metasound::FLiteral FMetasoundFrontendLiteral::ToLiteral(const FName& InMetasoun
 			FMetasoundFrontendRegistryContainer* Registry = FMetasoundFrontendRegistryContainer::Get();
 			if (ensure(Registry))
 			{
-				Literal = Registry->GenerateLiteralForUObject(InMetasoundDataTypeName, AsUObject[0]);
+				Literal = Registry->CreateLiteralFromUObject(InMetasoundDataTypeName, AsUObject[0]);
 			}
 		}
 	}
@@ -513,7 +513,7 @@ Metasound::FLiteral FMetasoundFrontendLiteral::ToLiteral(const FName& InMetasoun
 		FMetasoundFrontendRegistryContainer* Registry = FMetasoundFrontendRegistryContainer::Get();
 		if (ensure(Registry))
 		{
-			Literal = Registry->GenerateLiteralForUObjectArray(InMetasoundDataTypeName, AsUObject);
+			Literal = Registry->CreateLiteralFromUObjectArray(InMetasoundDataTypeName, AsUObject);
 		}
 	}
 	else 
diff --git a/Engine/Plugins/Runtime/Metasound/Source/MetasoundFrontend/Private/MetasoundFrontendQuerySteps.cpp b/Engine/Plugins/Runtime/Metasound/Source/MetasoundFrontend/Private/MetasoundFrontendQuerySteps.cpp
--- a/Engine/Plugins/Runtime/Metasound/Source/MetasoundFrontend/Private/MetasoundFrontendQuerySteps.cpp
+++ b/Engine/Plugins/Runtime/Metasound/Source/MetasoundFrontend/Private/MetasoundFrontendQuerySteps.cpp
@@ -45,11 +45,10 @@ namespace Metasound
 		{
 			if (const IRegistryTransaction* Transaction = InEntry.Value.Get<const IRegistryTransaction*>())
 			{
-				const FNodeClassInfo* ClassInfo = Transaction->GetNodeClassInfo();
-
-				if (nullptr != ClassInfo)
+				const FNodeRegistryKey* Key = Transaction->GetNodeRegistryKey();
+				if (nullptr != Key)
 				{
-					RegistryKey = ClassInfo->LookupKey;
+					RegistryKey = *Key;
 				}
 			}
 		}
@@ -101,10 +100,10 @@ namespace Metasound
 		{
 			if (const IRegistryTransaction* Transaction = InValue.Get<const IRegistryTransaction*>())
 			{
-				const FNodeClassInfo* ClassInfo = Transaction->GetNodeClassInfo();
-				if (nullptr != ClassInfo)
+				const FNodeRegistryKey* Key = Transaction->GetNodeRegistryKey();
+				if (nullptr != Key)
 				{
-					bool bSuccess = FMetasoundFrontendRegistryContainer::Get()->FindFrontendClassFromRegistered(*ClassInfo, FrontendClass);
+					bool bSuccess = FMetasoundFrontendRegistryContainer::Get()->FindFrontendClassFromRegistered(*Key, FrontendClass);
 					check(bSuccess);
 				}
 			}
diff --git a/Engine/Plugins/Runtime/Metasound/Source/MetasoundFrontend/Private/MetasoundFrontendRegistries.cpp b/Engine/Plugins/Runtime/Metasound/Source/MetasoundFrontend/Private/MetasoundFrontendRegistries.cpp
--- a/Engine/Plugins/Runtime/Metasound/Source/MetasoundFrontend/Private/MetasoundFrontendRegistries.cpp
+++ b/Engine/Plugins/Runtime/Metasound/Source/MetasoundFrontend/Private/MetasoundFrontendRegistries.cpp
@@ -2,6 +2,7 @@
 
 #include "MetasoundFrontendRegistries.h"
 
+#include "Algo/ForEach.h"
 #include "CoreMinimal.h"
 #include "MetasoundFrontendRegistryTransaction.h"
 #include "MetasoundLog.h"
@@ -56,17 +57,6 @@ namespace Metasound
 			Frontend::FNodeRegistryKey GetRegistryKey(EMetasoundFrontendClassType InType, const FString& InFullClassName, int32 InMajorVersion, int32 InMinorVersion)
 			{
 				FString RegistryKey = FString::Format(TEXT("{0}_{1}_{2}.{3}"), {*GetClassTypeString(InType), *InFullClassName, InMajorVersion, InMinorVersion});
-				/*
-				Frontend::FNodeRegistryKey Key;
-
-				Key.NodeClassFullName = InClassName;
-
-				// NodeHash is hash of node name and major version.
-				Key.NodeHash = FCrc::StrCrc32(*Key.NodeClassFullName.ToString());
-				Key.NodeHash = HashCombine(Key.NodeHash, FCrc::TypeCrc32(InMajorVersion));
-
-				return Key;
-				*/
 				return RegistryKey;
 			}
 
@@ -168,6 +158,195 @@ namespace Metasound
 				}
 			}
 
+			// Node registry entry for input nodes created from a data type registry entry.
+			class FInputNodeRegistryEntry : public INodeRegistryEntry
+			{
+			public:
+				FInputNodeRegistryEntry() = delete;
+
+				FInputNodeRegistryEntry(TUniquePtr<IDataTypeRegistryEntry>&& InDataTypeEntry)
+				: DataTypeEntry(MoveTemp(InDataTypeEntry))
+				{
+					if (DataTypeEntry.IsValid())
+					{
+						FrontendClass = DataTypeEntry->GetFrontendInputClass();
+						ClassInfo = FNodeClassInfo(FrontendClass);
+					}
+				}
+
+				virtual ~FInputNodeRegistryEntry() = default;
+
+				virtual const FNodeClassInfo& GetClassInfo() const override
+				{
+					return ClassInfo;
+				}
+
+				virtual TUniquePtr<INode> CreateNode(FDefaultNodeConstructorParams&&) const override
+				{
+					return nullptr;
+				}
+
+				virtual TUniquePtr<INode> CreateNode(FDefaultLiteralNodeConstructorParams&& InParams) const override
+				{
+					if (DataTypeEntry.IsValid())
+					{
+						return DataTypeEntry->CreateInputNode(MoveTemp(InParams));
+					}
+					return nullptr;
+				}
+
+				virtual TUniquePtr<INode> CreateNode(const FNodeInitData&) const override
+				{
+					return nullptr;
+				}
+
+				virtual const FMetasoundFrontendClass& GetFrontendClass() const override
+				{
+					return FrontendClass;
+				}
+
+				virtual TUniquePtr<INodeRegistryEntry> Clone() const override
+				{
+					if (DataTypeEntry.IsValid())
+					{
+						return MakeUnique<FInputNodeRegistryEntry>(DataTypeEntry->Clone());
+					}
+					return MakeUnique<FInputNodeRegistryEntry>(TUniquePtr<IDataTypeRegistryEntry>());
+				}
+
+			private:
+				
+				TUniquePtr<IDataTypeRegistryEntry> DataTypeEntry;
+				FNodeClassInfo ClassInfo;
+				FMetasoundFrontendClass FrontendClass;
+			};
+
+			// Node registry entry for output nodes created from a data type registry entry.
+			class FOutputNodeRegistryEntry : public INodeRegistryEntry
+			{
+			public:
+				FOutputNodeRegistryEntry() = delete;
+
+				FOutputNodeRegistryEntry(TUniquePtr<IDataTypeRegistryEntry>&& InDataTypeEntry)
+				: DataTypeEntry(MoveTemp(InDataTypeEntry))
+				{
+					if (DataTypeEntry.IsValid())
+					{
+						FrontendClass = DataTypeEntry->GetFrontendOutputClass();
+						ClassInfo = FNodeClassInfo(FrontendClass);
+					}
+				}
+
+				virtual ~FOutputNodeRegistryEntry() = default;
+
+				virtual const FNodeClassInfo& GetClassInfo() const override
+				{
+					return ClassInfo;
+				}
+
+				virtual TUniquePtr<INode> CreateNode(FDefaultNodeConstructorParams&& InParams) const override
+				{
+					if (DataTypeEntry.IsValid())
+					{
+						return DataTypeEntry->CreateOutputNode(MoveTemp(InParams));
+					}
+					return nullptr;
+				}
+
+				virtual TUniquePtr<INode> CreateNode(FDefaultLiteralNodeConstructorParams&& InParams) const override
+				{
+					return nullptr;
+				}
+
+				virtual TUniquePtr<INode> CreateNode(const FNodeInitData&) const override
+				{
+					return nullptr;
+				}
+
+				virtual const FMetasoundFrontendClass& GetFrontendClass() const override
+				{
+					return FrontendClass;
+				}
+
+				virtual TUniquePtr<INodeRegistryEntry> Clone() const override
+				{
+					if (DataTypeEntry.IsValid())
+					{
+						return MakeUnique<FOutputNodeRegistryEntry>(DataTypeEntry->Clone());
+					}
+					return MakeUnique<FOutputNodeRegistryEntry>(TUniquePtr<IDataTypeRegistryEntry>());
+				}
+
+			private:
+				
+				TUniquePtr<IDataTypeRegistryEntry> DataTypeEntry;
+				FNodeClassInfo ClassInfo;
+				FMetasoundFrontendClass FrontendClass;
+			};
+
+			// Node registry entry for variable nodes created from a data type registry entry.
+			class FVariableNodeRegistryEntry : public INodeRegistryEntry
+			{
+			public:
+				FVariableNodeRegistryEntry() = delete;
+
+				FVariableNodeRegistryEntry(TUniquePtr<IDataTypeRegistryEntry>&& InDataTypeEntry)
+				: DataTypeEntry(MoveTemp(InDataTypeEntry))
+				{
+					if (DataTypeEntry.IsValid())
+					{
+						FrontendClass = DataTypeEntry->GetFrontendVariableClass();
+						ClassInfo = FNodeClassInfo(FrontendClass);
+					}
+				}
+
+				virtual ~FVariableNodeRegistryEntry() = default;
+
+				virtual const FNodeClassInfo& GetClassInfo() const override
+				{
+					return ClassInfo;
+				}
+
+				virtual TUniquePtr<INode> CreateNode(FDefaultNodeConstructorParams&& InParams) const override
+				{
+					return nullptr;
+				}
+
+				virtual TUniquePtr<INode> CreateNode(FDefaultLiteralNodeConstructorParams&& InParams) const override
+				{
+					if (DataTypeEntry.IsValid())
+					{
+						return DataTypeEntry->CreateVariableNode(MoveTemp(InParams));
+					}
+					return nullptr;
+				}
+
+				virtual TUniquePtr<INode> CreateNode(const FNodeInitData&) const override
+				{
+					return nullptr;
+				}
+
+				virtual const FMetasoundFrontendClass& GetFrontendClass() const override
+				{
+					return FrontendClass;
+				}
+
+				virtual TUniquePtr<INodeRegistryEntry> Clone() const override
+				{
+					if (DataTypeEntry.IsValid())
+					{
+						return MakeUnique<FVariableNodeRegistryEntry>(DataTypeEntry->Clone());
+					}
+					return MakeUnique<FVariableNodeRegistryEntry>(TUniquePtr<IDataTypeRegistryEntry>());
+				}
+
+			private:
+				
+				TUniquePtr<IDataTypeRegistryEntry> DataTypeEntry;
+				FNodeClassInfo ClassInfo;
+				FMetasoundFrontendClass FrontendClass;
+			};
+
 			// Registry container private implementation.
 			class FRegistryContainerImpl : public FMetasoundFrontendRegistryContainer
 			{
@@ -178,10 +357,8 @@ namespace Metasound
 				using FConverterNodeInfo = ::Metasound::Frontend::FConverterNodeInfo;
 
 				using FNodeRegistryKey = Metasound::Frontend::FNodeRegistryKey;
-				using FNodeRegistryElement = Metasound::Frontend::FNodeRegistryElement;
 
-				using FDataTypeRegistryInfo = Metasound::FDataTypeRegistryInfo;
-				using FDataTypeConstructorCallbacks = ::Metasound::FDataTypeConstructorCallbacks;
+				using FDataTypeRegistryInfo = Metasound::Frontend::FDataTypeRegistryInfo;
 				using FNodeClassMetadata = Metasound::FNodeClassMetadata;
 				using IEnumDataTypeInterface = Metasound::Frontend::IEnumDataTypeInterface;
 
@@ -205,13 +382,11 @@ namespace Metasound
 				 *
 				 * @return True on success.
 				 */
-				FNodeRegistryKey RegisterExternalNode(Metasound::FCreateMetasoundNodeFunction&& InCreateNode, Metasound::FCreateMetasoundFrontendClassFunction&& InCreateDescription) override;
-				bool UnregisterExternalNode(const FNodeRegistryKey& InKey) override;
-				bool RegisterConversionNode(const FConverterNodeRegistryKey& InNodeKey, const FConverterNodeInfo& InNodeInfo) override;
-				bool IsNodeRegistered(const FNodeRegistryKey& InKey) const override;
+				virtual FNodeRegistryKey RegisterNode(TUniquePtr<Metasound::Frontend::INodeRegistryEntry>&&) override;
+				virtual bool UnregisterNode(const FNodeRegistryKey& InKey) override;
+				virtual bool IsNodeRegistered(const FNodeRegistryKey& InKey) const override;
 
-				// Get all available nodes
-				TArray<Frontend::FNodeClassInfo> GetAllAvailableNodeClasses(Frontend::FRegistryTransactionID* OutCurrentRegistryTransactionID) const override;
+				bool RegisterConversionNode(const FConverterNodeRegistryKey& InNodeKey, const FConverterNodeInfo& InNodeInfo) override;
 
 				TArray<const Metasound::Frontend::IRegistryTransaction*> GetRegistryTransactionsSince(Metasound::Frontend::FRegistryTransactionID InSince, Metasound::Frontend::FRegistryTransactionID* OutCurrentRegistryTransactionID) const override;
 
@@ -221,25 +396,25 @@ namespace Metasound
 				void IterateRegistry(Metasound::FIterateMetasoundFrontendClassFunction InIterFunc, EMetasoundFrontendClassType InClassType = EMetasoundFrontendClassType::Invalid) const override;
 
 				// Find Frontend Document data.
-				bool FindFrontendClassFromRegistered(const FNodeClassInfo& InClassInfo, FMetasoundFrontendClass& OutClass) override;
-				bool FindFrontendClassFromRegistered(const FMetasoundFrontendClassMetadata& InMetadata, FMetasoundFrontendClass& OutClass) override;
-				bool FindInputNodeClassMetadataForDataType(const FName& InDataTypeName, FMetasoundFrontendClassMetadata& OutMetadata) override;
-				bool FindVariableNodeClassMetadataForDataType(const FName& InDataTypeName, FMetasoundFrontendClassMetadata& OutMetadata) override;
-				bool FindOutputNodeClassMetadataForDataType(const FName& InDataTypeName, FMetasoundFrontendClassMetadata& OutMetadata) override;
+				bool FindFrontendClassFromRegistered(const FNodeRegistryKey& InKey, FMetasoundFrontendClass& OutClass) override;
+				virtual bool FindNodeClassInfoFromRegistered(const Metasound::Frontend::FNodeRegistryKey& InKey, FNodeClassInfo& OutInfo) override;
+				bool FindInputNodeRegistryKeyForDataType(const FName& InDataTypeName, FNodeRegistryKey& OutKey) override;
+				bool FindVariableNodeRegistryKeyForDataType(const FName& InDataTypeName, FNodeRegistryKey& OutKey) override;
+				bool FindOutputNodeRegistryKeyForDataType(const FName& InDataTypeName, FNodeRegistryKey& OutKey) override;
 
 				// Create a new instance of a C++ implemented node from the registry.
-				TUniquePtr<Metasound::INode> ConstructInputNode(const FName& InInputType, Metasound::FInputNodeConstructorParams&& InParams) override;
-				TUniquePtr<Metasound::INode> ConstructVariableNode(const FName& InVariableType, FVariableNodeConstructorParams&& InParams) override;
-				TUniquePtr<Metasound::INode> ConstructOutputNode(const FName& InOutputType, Metasound::FOutputNodeConstructorParams&& InParams) override;
-				TUniquePtr<Metasound::INode> ConstructExternalNode(const FNodeRegistryKey& InKey, const Metasound::FNodeInitData& InInitData) override;
+				TUniquePtr<Metasound::INode> CreateInputNode(const FName& InInputType, Metasound::FInputNodeConstructorParams&& InParams) override;
+				TUniquePtr<Metasound::INode> CreateVariableNode(const FName& InVariableType, FVariableNodeConstructorParams&& InParams) override;
+				TUniquePtr<Metasound::INode> CreateOutputNode(const FName& InOutputType, Metasound::FOutputNodeConstructorParams&& InParams) override;
+				virtual TUniquePtr<INode> CreateNode(const FNodeRegistryKey& InKey, FDefaultNodeConstructorParams&&) const override;
+				virtual TUniquePtr<INode> CreateNode(const FNodeRegistryKey& InKey, FDefaultLiteralNodeConstructorParams&&) const override;
+				virtual TUniquePtr<Metasound::INode> CreateNode(const FNodeRegistryKey& InKey, const Metasound::FNodeInitData& InInitData) const override;
 
 				// Returns a list of possible nodes to use to convert from FromDataType to ToDataType.
 				// Returns an empty array if none are available.
 				TArray<FConverterNodeInfo> GetPossibleConverterNodes(const FName& FromDataType, const FName& ToDataType) override;
 
-				bool RegisterDataType(const FDataTypeRegistryInfo& InDataInfo, const FDataTypeConstructorCallbacks& InCallbacks) override;
-
-				bool RegisterEnumDataInterface(FName InDataType, TSharedPtr<IEnumDataTypeInterface>&& InInterface) override;
+				virtual bool RegisterDataType(const FName& InName, TUniquePtr<Metasound::Frontend::IDataTypeRegistryEntry>&&) override;
 
 				// Get the desired kind of literal for a given data type. Returns EConstructorArgType::Invalid if the data type couldn't be found.
 				Metasound::ELiteralType GetDesiredLiteralTypeForDataType(FName InDataType) const override;
@@ -249,8 +424,8 @@ namespace Metasound
 
 				// Handle uobjects and literals
 				UClass* GetLiteralUClassForDataType(FName InDataType) const override;
-				Metasound::FLiteral GenerateLiteralForUObject(const FName& InDataType, UObject* InObject) override;
-				Metasound::FLiteral GenerateLiteralForUObjectArray(const FName& InDataType, TArray<UObject*> InObjectArray) override;
+				Metasound::FLiteral CreateLiteralFromUObject(const FName& InDataType, UObject* InObject) override;
+				Metasound::FLiteral CreateLiteralFromUObjectArray(const FName& InDataType, TArray<UObject*> InObjectArray) override;
 
 				// Get info about a specific data type (what kind of literals we can use, etc.)
 				// @returns false if InDataType wasn't found in the registry. 
@@ -261,6 +436,9 @@ namespace Metasound
 				TSharedPtr<Metasound::IDataChannel, ESPMode::ThreadSafe> CreateDataChannelForDataType(const FName& InDataType, const Metasound::FOperatorSettings& InOperatorSettings) const override;
 			private:
 
+				const IDataTypeRegistryEntry* FindDataTypeEntry(const FName& InDataTypeName) const;
+				const INodeRegistryEntry* FindNodeEntry(const FNodeRegistryKey& InKey) const;
+
 
 				// This buffer is used to enqueue nodes and datatypes to register when the module has been initialized,
 				// in order to avoid bad behavior with ensures, logs, etc. on static initialization.
@@ -272,23 +450,12 @@ namespace Metasound
 				FCriticalSection LazyInitCommandCritSection;
 
 				// Registry in which we keep all information about nodes implemented in C++.
-				TMap<FNodeRegistryKey, FNodeRegistryElement> ExternalNodeRegistry;
+				TMap<FNodeRegistryKey, TUniquePtr<INodeRegistryEntry>> RegisteredNodes;
+				TMap<FName, TUniquePtr<IDataTypeRegistryEntry>> RegisteredDataTypes;
 
 				// Registry in which we keep lists of possible nodes to use to convert between two datatypes
 				TMap<FConverterNodeRegistryKey, FConverterNodeRegistryValue> ConverterNodeRegistry;
 
-				struct FDataTypeRegistryElement
-				{
-					Metasound::FDataTypeConstructorCallbacks Callbacks;
-
-					Metasound::FDataTypeRegistryInfo Info;
-
-					TSharedPtr<const Metasound::Frontend::IEnumDataTypeInterface> EnumInterface;
-				};
-
-				TMap<FName, FDataTypeRegistryElement> DataTypeRegistry;
-				TMap<FNodeRegistryKey, FDataTypeRegistryElement> DataTypeNodeRegistry;
-				
 				FRegistryTransactionHistory RegistryTransactionHistory;
 			};
 
@@ -322,33 +489,18 @@ namespace Metasound
 				return true;
 			}
 
-			TArray<FNodeClassInfo> FRegistryContainerImpl::GetAllAvailableNodeClasses(FRegistryTransactionID* OutCurrentRegistryTransactionID) const 
-			{
-				TArray<FNodeClassInfo> OutClasses;
-
-				TArray<const IRegistryTransaction*> Transactions = RegistryTransactionHistory.GetTransactions(GetOriginRegistryTransactionID(), OutCurrentRegistryTransactionID);
-
-				for (const IRegistryTransaction* Transaction : Transactions)
-				{
-					if (const FNodeClassInfo* Info = Transaction->GetNodeClassInfo())
-					{
-						OutClasses.Add(*Info);
-					}
-				}
-
-				return OutClasses;
-			}
-
 			TArray<const Metasound::Frontend::IRegistryTransaction*> FRegistryContainerImpl::GetRegistryTransactionsSince(Metasound::Frontend::FRegistryTransactionID InSince, Metasound::Frontend::FRegistryTransactionID* OutCurrentRegistryTransactionID) const 
 			{
 				return RegistryTransactionHistory.GetTransactions(InSince, OutCurrentRegistryTransactionID);
 			}
 
-			TUniquePtr<Metasound::INode> FRegistryContainerImpl::ConstructInputNode(const FName& InInputType, Metasound::FInputNodeConstructorParams&& InParams)
+			TUniquePtr<Metasound::INode> FRegistryContainerImpl::CreateInputNode(const FName& InDataType, Metasound::FInputNodeConstructorParams&& InParams)
 			{
-				if (ensureAlwaysMsgf(DataTypeRegistry.Contains(InInputType), TEXT("Couldn't find data type %s!"), *InInputType.ToString()))
+				const IDataTypeRegistryEntry* Entry = FindDataTypeEntry(InDataType);
+
+				if (ensureAlwaysMsgf(nullptr != Entry, TEXT("Could not find data type [Name:%s]"), *InDataType.ToString()))
 				{
-					return DataTypeRegistry[InInputType].Callbacks.CreateInputNode(MoveTemp(InParams));
+					return Entry->CreateInputNode(MoveTemp(InParams));
 				}
 				else
 				{
@@ -356,11 +508,13 @@ namespace Metasound
 				}
 			}
 
-			TUniquePtr<Metasound::INode> FRegistryContainerImpl::ConstructVariableNode(const FName& InVariableType, Metasound::FVariableNodeConstructorParams&& InParams)
+			TUniquePtr<Metasound::INode> FRegistryContainerImpl::CreateVariableNode(const FName& InDataType, Metasound::FVariableNodeConstructorParams&& InParams)
 			{
-				if (ensureAlwaysMsgf(DataTypeRegistry.Contains(InVariableType), TEXT("Couldn't find data type %s!"), *InVariableType.ToString()))
+				const IDataTypeRegistryEntry* Entry = FindDataTypeEntry(InDataType);
+
+				if (ensureAlwaysMsgf(nullptr != Entry, TEXT("Could not find data type [Name:%s]"), *InDataType.ToString()))
 				{
-					return DataTypeRegistry[InVariableType].Callbacks.CreateVariableNode(MoveTemp(InParams));
+					return Entry->CreateVariableNode(MoveTemp(InParams));
 				}
 				else
 				{
@@ -368,11 +522,13 @@ namespace Metasound
 				}
 			}
 
-			TUniquePtr<Metasound::INode> FRegistryContainerImpl::ConstructOutputNode(const FName& InOutputType, Metasound::FOutputNodeConstructorParams&& InParams)
+			TUniquePtr<Metasound::INode> FRegistryContainerImpl::CreateOutputNode(const FName& InDataType, Metasound::FOutputNodeConstructorParams&& InParams)
 			{
-				if (ensureAlwaysMsgf(DataTypeRegistry.Contains(InOutputType), TEXT("Couldn't find data type %s!"), *InOutputType.ToString()))
+				const IDataTypeRegistryEntry* Entry = FindDataTypeEntry(InDataType);
+
+				if (ensureAlwaysMsgf(nullptr != Entry, TEXT("Could not find data type [Name:%s]"), *InDataType.ToString()))
 				{
-					return DataTypeRegistry[InOutputType].Callbacks.CreateOutputNode(MoveTemp(InParams));
+					return Entry->CreateOutputNode(MoveTemp(InParams));
 				}
 				else
 				{
@@ -380,15 +536,18 @@ namespace Metasound
 				}
 			}
 
-			Metasound::FLiteral FRegistryContainerImpl::GenerateLiteralForUObject(const FName& InDataType, UObject* InObject)
+			Metasound::FLiteral FRegistryContainerImpl::CreateLiteralFromUObject(const FName& InDataType, UObject* InObject)
 			{
-				if (ensureAlwaysMsgf(DataTypeRegistry.Contains(InDataType), TEXT("Couldn't find data type %s!"), *InDataType.ToString()))
+				const IDataTypeRegistryEntry* Entry = FindDataTypeEntry(InDataType);
+
+				if (ensureAlwaysMsgf(nullptr != Entry, TEXT("Could not find data type [Name:%s]"), *InDataType.ToString()))
+
 				{
-					if (Audio::IProxyDataPtr ProxyPtr = DataTypeRegistry[InDataType].Callbacks.CreateAudioProxy(InObject))
+					if (Audio::IProxyDataPtr ProxyPtr = Entry->CreateProxy(InObject))
 					{
 						if (InObject)
 						{
-							ensureAlwaysMsgf(ProxyPtr.IsValid(), TEXT("UObject failed to create a valid proxy!"));
+							ensureAlwaysMsgf(ProxyPtr.IsValid(), TEXT("UObject failed to create a valid proxy."));
 						}
 						return Metasound::FLiteral(MoveTemp(ProxyPtr));
 					}
@@ -397,15 +556,18 @@ namespace Metasound
 				return Metasound::FLiteral();
 			}
 
-			Metasound::FLiteral FRegistryContainerImpl::GenerateLiteralForUObjectArray(const FName& InDataType, TArray<UObject*> InObjectArray)
+			Metasound::FLiteral FRegistryContainerImpl::CreateLiteralFromUObjectArray(const FName& InDataType, TArray<UObject*> InObjectArray)
 			{
-				if (ensureAlwaysMsgf(DataTypeRegistry.Contains(InDataType), TEXT("Couldn't find data type %s!"), *InDataType.ToString()))
+				const IDataTypeRegistryEntry* Entry = FindDataTypeEntry(InDataType);
+
+				if (ensureAlwaysMsgf(nullptr != Entry, TEXT("Could not find data type [Name:%s]"), *InDataType.ToString()))
+
 				{
 					TArray<Audio::IProxyDataPtr> ProxyArray;
 
 					for (UObject* InObject : InObjectArray)
 					{
-						Audio::IProxyDataPtr ProxyPtr = DataTypeRegistry[InDataType].Callbacks.CreateAudioProxy(InObject);
+						Audio::IProxyDataPtr ProxyPtr = Entry->CreateProxy(InObject);
 
 						if (InObject)
 						{
@@ -422,17 +584,40 @@ namespace Metasound
 					return Metasound::FLiteral();
 				}
 			}
+			TUniquePtr<Metasound::INode> FRegistryContainerImpl::CreateNode(const FNodeRegistryKey& InKey, FDefaultNodeConstructorParams&& InParams) const
+			{
+				const INodeRegistryEntry* Entry = FindNodeEntry(InKey);
 
-			TUniquePtr<Metasound::INode> FRegistryContainerImpl::ConstructExternalNode(const FNodeRegistryKey& InKey, const Metasound::FNodeInitData& InInitData)
+				if (ensureAlwaysMsgf(nullptr != Entry, TEXT("Could not find node [RegistryKey:%s]"), *InKey))
+				{
+					return Entry->CreateNode(MoveTemp(InParams));
+				}
+
+				return nullptr;
+			}
+
+			TUniquePtr<Metasound::INode> FRegistryContainerImpl::CreateNode(const FNodeRegistryKey& InKey, FDefaultLiteralNodeConstructorParams&& InParams) const
 			{
-				if (!ExternalNodeRegistry.Contains(InKey))
+				const INodeRegistryEntry* Entry = FindNodeEntry(InKey);
+
+				if (ensureAlwaysMsgf(nullptr != Entry, TEXT("Could not find node [RegistryKey:%s]"), *InKey))
 				{
-					return nullptr;
+					return Entry->CreateNode(MoveTemp(InParams));
 				}
-				else
+
+				return nullptr;
+			}
+
+			TUniquePtr<Metasound::INode> FRegistryContainerImpl::CreateNode(const FNodeRegistryKey& InKey, const Metasound::FNodeInitData& InInitData) const
+			{
+				const INodeRegistryEntry* Entry = FindNodeEntry(InKey);
+
+				if (ensureAlwaysMsgf(nullptr != Entry, TEXT("Could not find node [RegistryKey:%s]"), *InKey))
 				{
-					return ExternalNodeRegistry[InKey].CreateNode(InInitData);
+					return Entry->CreateNode(InInitData);
 				}
+
+				return nullptr;
 			}
 
 			TArray<::Metasound::Frontend::FConverterNodeInfo> FRegistryContainerImpl::GetPossibleConverterNodes(const FName& FromDataType, const FName& ToDataType)
@@ -448,166 +633,152 @@ namespace Metasound
 				}
 			}
 
-
-
 			Metasound::ELiteralType FRegistryContainerImpl::GetDesiredLiteralTypeForDataType(FName InDataType) const
 			{
-				if (!DataTypeRegistry.Contains(InDataType))
+				const IDataTypeRegistryEntry* Entry = FindDataTypeEntry(InDataType);
+
+				if (nullptr == Entry)
 				{
 					return Metasound::ELiteralType::Invalid;
 				}
 
-				const FDataTypeRegistryElement& DataTypeInfo = DataTypeRegistry[InDataType];
+				const FDataTypeRegistryInfo& Info = Entry->GetDataTypeInfo();
 
 				// If there's a designated preferred literal type for this datatype, use that.
-				if (DataTypeInfo.Info.PreferredLiteralType != Metasound::ELiteralType::Invalid)
+				if (Info.PreferredLiteralType != Metasound::ELiteralType::Invalid)
 				{
-					return DataTypeInfo.Info.PreferredLiteralType;
+					return Info.PreferredLiteralType;
 				}
 
 				// Otherwise, we opt for the highest precision construction option available.
-				return Metasound::Frontend::MetasoundFrontendRegistryPrivate::GetMostDescriptiveLiteralForDataType(DataTypeInfo.Info);
+				return Metasound::Frontend::MetasoundFrontendRegistryPrivate::GetMostDescriptiveLiteralForDataType(Info);
 			}
 
 			UClass* FRegistryContainerImpl::GetLiteralUClassForDataType(FName InDataType) const
 			{
-				if (!DataTypeRegistry.Contains(InDataType))
-				{
-					ensureAlwaysMsgf(false, TEXT("DataType %s not registered."), *InDataType.ToString());
-					return nullptr;
-				}
-				else
+				const IDataTypeRegistryEntry* Entry = FindDataTypeEntry(InDataType);
+
+				if (ensureAlwaysMsgf(nullptr != Entry, TEXT("Could not find data type [Name:%s]."), *InDataType.ToString()))
 				{
-					return DataTypeRegistry[InDataType].Info.ProxyGeneratorClass;
+					return Entry->GetDataTypeInfo().ProxyGeneratorClass;
 				}
+
+				return nullptr;
 			}
 
 			bool FRegistryContainerImpl::DoesDataTypeSupportLiteralType(FName InDataType, Metasound::ELiteralType InLiteralType) const
 			{
-				if (!DataTypeRegistry.Contains(InDataType))
-				{
-					ensureAlwaysMsgf(false, TEXT("DataType %s not registered."), *InDataType.ToString());
-					return false;
-				}
+				const IDataTypeRegistryEntry* Entry = FindDataTypeEntry(InDataType);
 
-				const FDataTypeRegistryElement& DataTypeInfo = DataTypeRegistry[InDataType];
-				
-				switch (InLiteralType)
+				if (ensureAlwaysMsgf(nullptr != Entry, TEXT("Could not find data type [Name:%s]."), *InDataType.ToString()))
 				{
-					case Metasound::ELiteralType::Boolean:
-					{
-						return DataTypeInfo.Info.bIsBoolParsable;
-					}
-					case Metasound::ELiteralType::BooleanArray:
-					{
-						return DataTypeInfo.Info.bIsBoolArrayParsable;
-					}
+					const FDataTypeRegistryInfo& Info = Entry->GetDataTypeInfo();
 
-					case Metasound::ELiteralType::Integer:
-					{
-						return DataTypeInfo.Info.bIsIntParsable;
-					}
-					case Metasound::ELiteralType::IntegerArray:
+					switch (InLiteralType)
 					{
-						return DataTypeInfo.Info.bIsIntArrayParsable;
-					}
+						case Metasound::ELiteralType::Boolean:
+						{
+							return Info.bIsBoolParsable;
+						}
+						case Metasound::ELiteralType::BooleanArray:
+						{
+							return Info.bIsBoolArrayParsable;
+						}
 
-					case Metasound::ELiteralType::Float:
-					{
-						return DataTypeInfo.Info.bIsFloatParsable;
-					}
-					case Metasound::ELiteralType::FloatArray:
-					{
-						return DataTypeInfo.Info.bIsFloatArrayParsable;
-					}
+						case Metasound::ELiteralType::Integer:
+						{
+							return Info.bIsIntParsable;
+						}
+						case Metasound::ELiteralType::IntegerArray:
+						{
+							return Info.bIsIntArrayParsable;
+						}
 
-					case Metasound::ELiteralType::String:
-					{
-						return DataTypeInfo.Info.bIsStringParsable;
-					}
-					case Metasound::ELiteralType::StringArray:
-					{
-						return DataTypeInfo.Info.bIsStringArrayParsable;
-					}
+						case Metasound::ELiteralType::Float:
+						{
+							return Info.bIsFloatParsable;
+						}
+						case Metasound::ELiteralType::FloatArray:
+						{
+							return Info.bIsFloatArrayParsable;
+						}
 
-					case Metasound::ELiteralType::UObjectProxy:
-					{
-						return DataTypeInfo.Info.bIsProxyParsable;
-					}
-					case Metasound::ELiteralType::UObjectProxyArray:
-					{
-						return DataTypeInfo.Info.bIsProxyArrayParsable;
-					}
+						case Metasound::ELiteralType::String:
+						{
+							return Info.bIsStringParsable;
+						}
+						case Metasound::ELiteralType::StringArray:
+						{
+							return Info.bIsStringArrayParsable;
+						}
 
-					case Metasound::ELiteralType::None:
-					{
-						return DataTypeInfo.Info.bIsDefaultParsable;
-					}
-					case Metasound::ELiteralType::NoneArray:
-					{
-						return DataTypeInfo.Info.bIsDefaultArrayParsable;
-					}
+						case Metasound::ELiteralType::UObjectProxy:
+						{
+							return Info.bIsProxyParsable;
+						}
+						case Metasound::ELiteralType::UObjectProxyArray:
+						{
+							return Info.bIsProxyArrayParsable;
+						}
 
-					case Metasound::ELiteralType::Invalid:
-					default:
-					{
-						static_assert(static_cast<int32>(Metasound::ELiteralType::COUNT) == 13, "Possible missing case coverage for ELiteralType");
-						return false;
+						case Metasound::ELiteralType::None:
+						{
+							return Info.bIsDefaultParsable;
+						}
+						case Metasound::ELiteralType::NoneArray:
+						{
+							return Info.bIsDefaultArrayParsable;
+						}
+
+						case Metasound::ELiteralType::Invalid:
+						default:
+						{
+							static_assert(static_cast<int32>(Metasound::ELiteralType::COUNT) == 13, "Possible missing case coverage for ELiteralType");
+							return false;
+						}
 					}
 				}
+
+				return false;
 			}
 
-			bool FRegistryContainerImpl::RegisterDataType(const ::Metasound::FDataTypeRegistryInfo& InDataInfo, const ::Metasound::FDataTypeConstructorCallbacks& InCallbacks)
+			bool FRegistryContainerImpl::RegisterDataType(const FName& InName, TUniquePtr<IDataTypeRegistryEntry>&& InRegistryEntry)
 			{
-				if (!ensureAlwaysMsgf(!DataTypeRegistry.Contains(InDataInfo.DataTypeName),
-					TEXT("Name collision when trying to register Metasound Data Type %s! DataType must have "
-						"unique name and REGISTER_METASOUND_DATATYPE cannot be called in a public header."),
-						*InDataInfo.DataTypeName.ToString()))
+				if (InRegistryEntry.IsValid())
 				{
-					// todo: capture callstack for previous declaration for non-shipping builds to help clarify who already registered this name for a type.
-					return false;
-				}
-				else
-				{
-					FDataTypeRegistryElement InElement = { InCallbacks, InDataInfo };
-
-					DataTypeRegistry.Add(InDataInfo.DataTypeName, InElement);
+					if (!ensureAlwaysMsgf(!RegisteredDataTypes.Contains(InName),
+						TEXT("Name collision when trying to register Metasound Data Type [Name:%s]. DataType must have "
+							"unique name and REGISTER_METASOUND_DATATYPE cannot be called in a public header."),
+							*InName.ToString()))
+					{
+						return false;
+					}
 
-					Metasound::Frontend::FNodeRegistryKey InputNodeRegistryKey = GetRegistryKey(InCallbacks.CreateFrontendInputClass().Metadata);
-					DataTypeNodeRegistry.Add(InputNodeRegistryKey, InElement);
 
-					Metasound::Frontend::FNodeRegistryKey VariableNodeRegistryKey = GetRegistryKey(InCallbacks.CreateFrontendVariableClass().Metadata);
-					DataTypeNodeRegistry.Add(VariableNodeRegistryKey, InElement);
+					RegisteredDataTypes.Add(InName, InRegistryEntry->Clone());
 
-					Metasound::Frontend::FNodeRegistryKey OutputNodeRegistryKey = GetRegistryKey(InCallbacks.CreateFrontendOutputClass().Metadata);
-					DataTypeNodeRegistry.Add(OutputNodeRegistryKey, InElement);
+					RegisterNode(MakeUnique<FInputNodeRegistryEntry>(InRegistryEntry->Clone()));
+					RegisterNode(MakeUnique<FOutputNodeRegistryEntry>(InRegistryEntry->Clone()));
+					RegisterNode(MakeUnique<FVariableNodeRegistryEntry>(MoveTemp(InRegistryEntry)));
 
-					UE_LOG(LogMetaSound, Display, TEXT("Registered Metasound Datatype %s."), *InDataInfo.DataTypeName.ToString());
+					UE_LOG(LogMetaSound, Verbose, TEXT("Registered Metasound Datatype [Name:%s]."), *InName.ToString());
 					return true;
 				}
-			}
 
-			bool FRegistryContainerImpl::RegisterEnumDataInterface(FName InDataType, TSharedPtr<IEnumDataTypeInterface>&& InInterface)
-			{
-				if (FDataTypeRegistryElement* Element = DataTypeRegistry.Find(InDataType))
-				{
-					Element->EnumInterface = MoveTemp(InInterface);
-					return true;
-				}
 				return false;
 			}
 
-			FNodeRegistryKey FRegistryContainerImpl::RegisterExternalNode(Metasound::FCreateMetasoundNodeFunction&& InCreateNode, Metasound::FCreateMetasoundFrontendClassFunction&& InCreateDescription)
+			FNodeRegistryKey FRegistryContainerImpl::RegisterNode(TUniquePtr<INodeRegistryEntry>&& InEntry)
 			{
-				FNodeRegistryElement RegistryElement = FNodeRegistryElement(MoveTemp(InCreateNode), MoveTemp(InCreateDescription));
-
 				FNodeRegistryKey Key;
 
-				if (ensure(FMetasoundFrontendRegistryContainer::GetRegistryKey(RegistryElement, Key)))
+				if (InEntry.IsValid())
 				{
+					Key = FMetasoundFrontendRegistryContainer::GetRegistryKey(InEntry->GetClassInfo());
+
 					// check to see if an identical node was already registered, and log
 					ensureAlwaysMsgf(
-						!ExternalNodeRegistry.Contains(Key),
+						!RegisteredNodes.Contains(Key),
 						TEXT("Node with registry key already registered. "
 							"The previously registered node will be overwritten. "
 							"This can happen if two classes share the same name or if METASOUND_REGISTER_NODE is defined in a public header."),
@@ -615,26 +786,24 @@ namespace Metasound
 
 					// Store update to newly registered node in history so nodes
 					// can be queried by transaction ID
-					FNodeClassInfo ClassInfo = { EMetasoundFrontendClassType::External, Key };
-					RegistryTransactionHistory.Add(MakeAddNodeRegistryTransaction(ClassInfo));
+					RegistryTransactionHistory.Add(MakeAddNodeRegistryTransaction(Key, InEntry->GetClassInfo()));
 
 					// Store registry elements in map so nodes can be queried using registry key.
-					ExternalNodeRegistry.Add(Key, MoveTemp(RegistryElement));
+					RegisteredNodes.Add(Key, MoveTemp(InEntry));
 				}
 
 				return Key;
 			}
 
-			bool FRegistryContainerImpl::UnregisterExternalNode(const FNodeRegistryKey& InKey)
+			bool FRegistryContainerImpl::UnregisterNode(const FNodeRegistryKey& InKey)
 			{
 				if (IsValidNodeRegistryKey(InKey))
 				{
-					if (ExternalNodeRegistry.Contains(InKey))
+					if (const INodeRegistryEntry* Entry = FindNodeEntry(InKey))
 					{
-						FNodeClassInfo ClassInfo = { EMetasoundFrontendClassType::External, InKey };
-						RegistryTransactionHistory.Add(MakeRemoveNodeRegistryTransaction(ClassInfo));
+						RegistryTransactionHistory.Add(MakeRemoveNodeRegistryTransaction(InKey, Entry->GetClassInfo()));
 
-						ExternalNodeRegistry.Remove(InKey);
+						RegisteredNodes.Remove(InKey);
 						return true;
 					}
 				}
@@ -666,14 +835,14 @@ namespace Metasound
 
 			bool FRegistryContainerImpl::IsNodeRegistered(const FNodeRegistryKey& InKey) const
 			{
-				return ExternalNodeRegistry.Contains(InKey);
+				return RegisteredNodes.Contains(InKey);
 			}
 
 			TArray<FName> FRegistryContainerImpl::GetAllValidDataTypes()
 			{
 				TArray<FName> OutDataTypes;
 
-				for (auto& DataTypeTuple : DataTypeRegistry)
+				for (auto& DataTypeTuple : RegisteredDataTypes)
 				{
 					OutDataTypes.Add(DataTypeTuple.Key);
 				}
@@ -681,106 +850,80 @@ namespace Metasound
 				return OutDataTypes;
 			}
 
-			bool FRegistryContainerImpl::GetInfoForDataType(FName InDataType, Metasound::FDataTypeRegistryInfo& OutInfo)
+			bool FRegistryContainerImpl::GetInfoForDataType(FName InDataType, Metasound::Frontend::FDataTypeRegistryInfo& OutInfo)
 			{
-				if (!DataTypeRegistry.Contains(InDataType))
+				if (const IDataTypeRegistryEntry* Entry = FindDataTypeEntry(InDataType))
 				{
-					return false;
-				}
-				else
-				{
-					OutInfo = DataTypeRegistry[InDataType].Info;
+					OutInfo = Entry->GetDataTypeInfo();
 					return true;
 				}
+				return false;
 			}
 
-
-			TSharedPtr<const Metasound::Frontend::IEnumDataTypeInterface>
-			FRegistryContainerImpl::GetEnumInterfaceForDataType(FName InDataType) const
+			TSharedPtr<const Metasound::Frontend::IEnumDataTypeInterface> FRegistryContainerImpl::GetEnumInterfaceForDataType(FName InDataType) const
 			{
-				if (const FDataTypeRegistryElement* Element = DataTypeRegistry.Find(InDataType))
+				if (const IDataTypeRegistryEntry* Entry = FindDataTypeEntry(InDataType))
 				{
-					return Element->EnumInterface;
+					return Entry->GetEnumInterface();
 				}
 				return nullptr;
 			}
 
 			TSharedPtr<Metasound::IDataChannel, ESPMode::ThreadSafe> FRegistryContainerImpl::CreateDataChannelForDataType(const FName& InDataType, const Metasound::FOperatorSettings& InSettings) const
 			{
-				if (const FDataTypeRegistryElement* Element = DataTypeRegistry.Find(InDataType))
+				if (const IDataTypeRegistryEntry* Entry = FindDataTypeEntry(InDataType))
 				{
-					if (Element->Callbacks.CreateDataChannel)
-					{
-						return Element->Callbacks.CreateDataChannel(InSettings);
-					}
+					return Entry->CreateDataChannel(InSettings);
 				}
 				return nullptr;
 			}
 
-			bool FRegistryContainerImpl::FindFrontendClassFromRegistered(const FNodeClassInfo& InClassInfo, FMetasoundFrontendClass& OutClass)
+			bool FRegistryContainerImpl::FindFrontendClassFromRegistered(const FNodeRegistryKey& InKey, FMetasoundFrontendClass& OutClass)
 			{
-				if (InClassInfo.NodeType == EMetasoundFrontendClassType::External)
+				if (const INodeRegistryEntry* Entry = FindNodeEntry(InKey))
 				{
-					if (ExternalNodeRegistry.Contains(InClassInfo.LookupKey))
-					{
-						OutClass = ExternalNodeRegistry[InClassInfo.LookupKey].CreateFrontendClass();
-						return true;
-					}
-				}
-				else if (InClassInfo.NodeType == EMetasoundFrontendClassType::Input)
-				{
-					if (DataTypeNodeRegistry.Contains(InClassInfo.LookupKey))
-					{
-						OutClass = DataTypeNodeRegistry[InClassInfo.LookupKey].Callbacks.CreateFrontendInputClass();
-						return true;
-					}
-				}
-				else if (InClassInfo.NodeType == EMetasoundFrontendClassType::Output)
-				{
-					if (DataTypeNodeRegistry.Contains(InClassInfo.LookupKey))
-					{
-						OutClass = DataTypeNodeRegistry[InClassInfo.LookupKey].Callbacks.CreateFrontendOutputClass();
-						return true;
-					}
+					OutClass = Entry->GetFrontendClass();
+					return true;
 				}
-
 				return false;
 			}
 
-			bool FRegistryContainerImpl::FindFrontendClassFromRegistered(const FMetasoundFrontendClassMetadata& InMetadata, FMetasoundFrontendClass& OutClass)
+			bool FRegistryContainerImpl::FindNodeClassInfoFromRegistered(const Metasound::Frontend::FNodeRegistryKey& InKey, FNodeClassInfo& OutInfo)
 			{
-				FNodeClassInfo Info;
-				Info.NodeType = InMetadata.Type;
-				Info.LookupKey = GetRegistryKey(InMetadata);
+				if (const INodeRegistryEntry* Entry = FindNodeEntry(InKey))
+				{
+					OutInfo = Entry->GetClassInfo();
+					return true;
+				}
 
-				return FindFrontendClassFromRegistered(Info, OutClass);
+				return false;
 			}
 
-			bool FRegistryContainerImpl::FindInputNodeClassMetadataForDataType(const FName& InDataTypeName, FMetasoundFrontendClassMetadata& OutMetadata)
+			bool FRegistryContainerImpl::FindInputNodeRegistryKeyForDataType(const FName& InDataTypeName, FNodeRegistryKey& OutKey)
 			{
-				if (const FDataTypeRegistryElement* Element = DataTypeRegistry.Find(InDataTypeName))
+				if (const IDataTypeRegistryEntry* Entry = FindDataTypeEntry(InDataTypeName))
 				{
-					OutMetadata = Element->Callbacks.CreateFrontendInputClass().Metadata;
+					OutKey = GetRegistryKey(Entry->GetFrontendInputClass().Metadata);
 					return true;
 				}
 				return false;
 			}
 
-			bool FRegistryContainerImpl::FindVariableNodeClassMetadataForDataType(const FName& InDataTypeName, FMetasoundFrontendClassMetadata& OutMetadata)
+			bool FRegistryContainerImpl::FindVariableNodeRegistryKeyForDataType(const FName& InDataTypeName, FNodeRegistryKey& OutKey)
 			{
-				if (const FDataTypeRegistryElement* Element = DataTypeRegistry.Find(InDataTypeName))
+				if (const IDataTypeRegistryEntry* Entry = FindDataTypeEntry(InDataTypeName))
 				{
-					OutMetadata = Element->Callbacks.CreateFrontendVariableClass().Metadata;
+					OutKey = GetRegistryKey(Entry->GetFrontendVariableClass().Metadata);
 					return true;
 				}
 				return false;
 			}
 
-			bool FRegistryContainerImpl::FindOutputNodeClassMetadataForDataType(const FName& InDataTypeName, FMetasoundFrontendClassMetadata& OutMetadata)
+			bool FRegistryContainerImpl::FindOutputNodeRegistryKeyForDataType(const FName& InDataTypeName, FNodeRegistryKey& OutKey)
 			{
-				if (const FDataTypeRegistryElement* Element = DataTypeRegistry.Find(InDataTypeName))
+				if (const IDataTypeRegistryEntry* Entry = FindDataTypeEntry(InDataTypeName))
 				{
-					OutMetadata = Element->Callbacks.CreateFrontendOutputClass().Metadata;
+					OutKey = GetRegistryKey(Entry->GetFrontendOutputClass().Metadata);
 					return true;
 				}
 				return false;
@@ -788,82 +931,51 @@ namespace Metasound
 
 			void FRegistryContainerImpl::IterateRegistry(Metasound::FIterateMetasoundFrontendClassFunction InIterFunc, EMetasoundFrontendClassType InClassType) const
 			{
-				auto IterateExternalClasses = [=]
+				auto WrappedFunc = [&](const TPair<FNodeRegistryKey, TUniquePtr<INodeRegistryEntry>>& Pair)
 				{
-					for (const TPair<FNodeRegistryKey, FNodeRegistryElement>& Pair : ExternalNodeRegistry)
-					{
-						InIterFunc(Pair.Value.CreateFrontendClass());
-					}
-				};
-
-				auto IterateInputClasses = [=]
-				{
-					for (const TPair<FNodeRegistryKey, FDataTypeRegistryElement>& Pair : DataTypeNodeRegistry)
-					{
-						InIterFunc(Pair.Value.Callbacks.CreateFrontendInputClass());
-					}
+					InIterFunc(Pair.Value->GetFrontendClass());
 				};
 
-				auto IterateOutputClasses = [=]
+				if (EMetasoundFrontendClassType::Invalid == InClassType)
 				{
-					for (const TPair<FNodeRegistryKey, FDataTypeRegistryElement>& Pair : DataTypeNodeRegistry)
-					{
-						InIterFunc(Pair.Value.Callbacks.CreateFrontendOutputClass());
-					}
-				};
-
-				auto IterateVariableClasses = [=]
+					// Iterate through all classes. 
+					Algo::ForEach(RegisteredNodes, WrappedFunc);
+				}
+				else
 				{
-					for (const TPair<FNodeRegistryKey, FDataTypeRegistryElement>& Pair : DataTypeNodeRegistry)
+					// Only call function on classes of certain type.
+					auto IsMatchingClassType = [&](const TPair<FNodeRegistryKey, TUniquePtr<INodeRegistryEntry>>& Pair)
 					{
-						InIterFunc(Pair.Value.Callbacks.CreateFrontendVariableClass());
-					}
-				};
+						return Pair.Value->GetClassInfo().Type == InClassType;
+					};
+					Algo::ForEachIf(RegisteredNodes, IsMatchingClassType, WrappedFunc);
+				}
+			}
 
-				switch (InClassType)
+			const IDataTypeRegistryEntry* FRegistryContainerImpl::FindDataTypeEntry(const FName& InDataTypeName) const
+			{
+				if (const TUniquePtr<IDataTypeRegistryEntry>* Entry = RegisteredDataTypes.Find(InDataTypeName))
 				{
-					case EMetasoundFrontendClassType::External:
-					{
-						IterateExternalClasses();
-					}
-					break;
-
-					// TODO: Implement
-					case EMetasoundFrontendClassType::Graph:
-					{
-						checkNoEntry();
-					}
-					break;
-
-					case EMetasoundFrontendClassType::Input:
-					{
-						IterateInputClasses();
-					}
-					break;
-
-					case EMetasoundFrontendClassType::Output:
+					if (Entry->IsValid())
 					{
-						IterateOutputClasses();
+						return Entry->Get();
 					}
-					break;
+				}
 
-					case EMetasoundFrontendClassType::Variable:
-					{
-						IterateVariableClasses();
-					}
-					break;
+				return nullptr;
+			}
 
-					default:
-					case EMetasoundFrontendClassType::Invalid:
+			const INodeRegistryEntry* FRegistryContainerImpl::FindNodeEntry(const FNodeRegistryKey& InKey) const
+			{
+				if (const TUniquePtr<INodeRegistryEntry>* Entry = RegisteredNodes.Find(InKey))
+				{
+					if (Entry->IsValid())
 					{
-						IterateExternalClasses();
-						IterateVariableClasses();
-						IterateInputClasses();
-						IterateOutputClasses();
+						return Entry->Get();
 					}
-					break;
 				}
-				static_assert(static_cast<uint32>(EMetasoundFrontendClassType::Invalid) == 5, "Possible missing switch case coverage for EMetasoundFrontendClassType");
+
+				return nullptr;
 			}
 		} // namespace MetasoundFrontendRegistriesPrivate
 
@@ -871,6 +983,23 @@ namespace Metasound
 		{
 			return !InKey.IsEmpty();
 		}
+
+		FNodeClassInfo::FNodeClassInfo(const FMetasoundFrontendClassMetadata& InMetadata)
+		: ClassName(InMetadata.ClassName)
+		, Type(InMetadata.Type)
+		, Version(InMetadata.Version)
+		{
+		}
+
+		FNodeClassInfo::FNodeClassInfo(const FMetasoundFrontendClass& InClass)
+		: FNodeClassInfo(InClass.Metadata)
+		{
+		}
+
+		FNodeClassInfo::FNodeClassInfo(const FNodeClassMetadata& InMetadata)
+		: FNodeClassInfo(FMetasoundFrontendClassMetadata(InMetadata))
+		{
+		}
 	} // namespace Frontend
 } // namespace Metasound
 
@@ -895,20 +1024,6 @@ void FMetasoundFrontendRegistryContainer::ShutdownMetasoundFrontend()
 	}
 }
 
-bool FMetasoundFrontendRegistryContainer::GetRegistryKey(const Metasound::Frontend::FNodeRegistryElement& InElement, Metasound::Frontend::FNodeRegistryKey& OutKey)
-{
-	if (InElement.CreateFrontendClass)
-	{
-		FMetasoundFrontendClass FrontendClass = InElement.CreateFrontendClass();
-
-		OutKey = GetRegistryKey(FrontendClass.Metadata);
-
-		return true;
-	}
-
-	return false;
-}
-
 Metasound::Frontend::FNodeRegistryKey FMetasoundFrontendRegistryContainer::GetRegistryKey(const FNodeClassMetadata& InNodeMetadata)
 {
 	return Metasound::Frontend::MetasoundFrontendRegistryPrivate::GetRegistryKey(EMetasoundFrontendClassType::External, InNodeMetadata.ClassName.GetFullName().ToString(), InNodeMetadata.MajorVersion, InNodeMetadata.MinorVersion);
@@ -919,41 +1034,56 @@ Metasound::Frontend::FNodeRegistryKey FMetasoundFrontendRegistryContainer::GetRe
 	return Metasound::Frontend::MetasoundFrontendRegistryPrivate::GetRegistryKey(InNodeMetadata.Type, InNodeMetadata.ClassName.GetFullName().ToString(), InNodeMetadata.Version.Major, InNodeMetadata.Version.Minor);
 }
 
-bool FMetasoundFrontendRegistryContainer::GetFrontendClassFromRegistered(const FMetasoundFrontendClassMetadata& InMetadata, FMetasoundFrontendClass& OutClass)
+Metasound::Frontend::FNodeRegistryKey FMetasoundFrontendRegistryContainer::GetRegistryKey(const FNodeClassInfo& InClassInfo)
+{
+	return Metasound::Frontend::MetasoundFrontendRegistryPrivate::GetRegistryKey(InClassInfo.Type, InClassInfo.ClassName.GetFullName().ToString(), InClassInfo.Version.Major, InClassInfo.Version.Minor);
+}
+
+bool FMetasoundFrontendRegistryContainer::GetFrontendClassFromRegistered(const FNodeRegistryKey& InKey, FMetasoundFrontendClass& OutClass)
 {
 	FMetasoundFrontendRegistryContainer* Registry = FMetasoundFrontendRegistryContainer::Get();
 
 	if (ensure(nullptr != Registry))
 	{
-		return Registry->FindFrontendClassFromRegistered(InMetadata, OutClass);
+		return Registry->FindFrontendClassFromRegistered(InKey, OutClass);
 	}
 
 	return false;
 }
 
-bool FMetasoundFrontendRegistryContainer::GetInputNodeClassMetadataForDataType(const FName& InDataTypeName, FMetasoundFrontendClassMetadata& OutMetadata)
+bool FMetasoundFrontendRegistryContainer::GetNodeClassInfoFromRegistered(const FNodeRegistryKey& InKey, FNodeClassInfo& OutInfo)
+{
+	if (FMetasoundFrontendRegistryContainer* Registry = FMetasoundFrontendRegistryContainer::Get())
+	{
+		return Registry->FindNodeClassInfoFromRegistered(InKey, OutInfo);
+	}
+	return false;
+}
+
+
+bool FMetasoundFrontendRegistryContainer::GetInputNodeRegistryKeyForDataType(const FName& InDataTypeName, FNodeRegistryKey& OutKey)
 {
 	if (FMetasoundFrontendRegistryContainer* Registry = FMetasoundFrontendRegistryContainer::Get())
 	{
-		return Registry->FindInputNodeClassMetadataForDataType(InDataTypeName, OutMetadata);
+		return Registry->FindInputNodeRegistryKeyForDataType(InDataTypeName, OutKey);
 	}
 	return false;
 }
 
-bool FMetasoundFrontendRegistryContainer::GetVariableNodeClassMetadataForDataType(const FName& InDataTypeName, FMetasoundFrontendClassMetadata& OutMetadata)
+bool FMetasoundFrontendRegistryContainer::GetVariableNodeRegistryKeyForDataType(const FName& InDataTypeName, FNodeRegistryKey& OutKey)
 {
 	if (FMetasoundFrontendRegistryContainer* Registry = FMetasoundFrontendRegistryContainer::Get())
 	{
-		return Registry->FindVariableNodeClassMetadataForDataType(InDataTypeName, OutMetadata);
+		return Registry->FindVariableNodeRegistryKeyForDataType(InDataTypeName, OutKey);
 	}
 	return false;
 }
 
-bool FMetasoundFrontendRegistryContainer::GetOutputNodeClassMetadataForDataType(const FName& InDataTypeName, FMetasoundFrontendClassMetadata& OutMetadata)
+bool FMetasoundFrontendRegistryContainer::GetOutputNodeRegistryKeyForDataType(const FName& InDataTypeName, FNodeRegistryKey& OutKey)
 {
 	if (FMetasoundFrontendRegistryContainer* Registry = FMetasoundFrontendRegistryContainer::Get())
 	{
-		return Registry->FindOutputNodeClassMetadataForDataType(InDataTypeName, OutMetadata);
+		return Registry->FindOutputNodeRegistryKeyForDataType(InDataTypeName, OutKey);
 	}
 	return false;
 }
diff --git a/Engine/Plugins/Runtime/Metasound/Source/MetasoundFrontend/Private/MetasoundFrontendRegistryTransaction.cpp b/Engine/Plugins/Runtime/Metasound/Source/MetasoundFrontend/Private/MetasoundFrontendRegistryTransaction.cpp
--- a/Engine/Plugins/Runtime/Metasound/Source/MetasoundFrontend/Private/MetasoundFrontendRegistryTransaction.cpp
+++ b/Engine/Plugins/Runtime/Metasound/Source/MetasoundFrontend/Private/MetasoundFrontendRegistryTransaction.cpp
@@ -94,70 +94,84 @@ namespace Metasound
 			}
 		}
 
-		FRegistryTransactionPtr MakeAddNodeRegistryTransaction(const FNodeClassInfo& InInfo)
+		FRegistryTransactionPtr MakeAddNodeRegistryTransaction(const FNodeRegistryKey& InKey, const FNodeClassInfo& InInfo)
 		{
 			class FNodeRegistryTransaction : public IRegistryTransaction
 			{
 			public:
-				FNodeRegistryTransaction(const FNodeClassInfo& InNodeClassInfo)
+				FNodeRegistryTransaction(const FNodeRegistryKey& InKey, const FNodeClassInfo& InNodeClassInfo)
 				: NodeClassInfo(InNodeClassInfo)
+				, Key(InKey)
 				{
 				}
 
-				ETransactionType GetTransactionType() const override
+				virtual ETransactionType GetTransactionType() const override
 				{
 					return ETransactionType::Add;
 				}
 
-				TUniquePtr<IRegistryTransaction> Clone() const override
+				virtual TUniquePtr<IRegistryTransaction> Clone() const override
 				{
 					return MakeUnique<FNodeRegistryTransaction>(*this);
 				}
 
-				const FNodeClassInfo* GetNodeClassInfo() const 
+				virtual const FNodeClassInfo* GetNodeClassInfo() const override
 				{
 					return &NodeClassInfo;
 				}
 
+				virtual const FNodeRegistryKey* GetNodeRegistryKey() const override
+				{
+					return &Key;
+				}
+
 			private:
 
 				FNodeClassInfo NodeClassInfo;
+				FNodeRegistryKey Key;
 			};
 
-			return MakeUnique<FNodeRegistryTransaction>(InInfo);
+			return MakeUnique<FNodeRegistryTransaction>(InKey, InInfo);
 		}
 
-		FRegistryTransactionPtr MakeRemoveNodeRegistryTransaction(const FNodeClassInfo& InInfo)
+		FRegistryTransactionPtr MakeRemoveNodeRegistryTransaction(const FNodeRegistryKey& InKey, const FNodeClassInfo& InInfo)
 		{
 			class FNodeRegistryTransaction : public IRegistryTransaction
 			{
 			public:
-				FNodeRegistryTransaction(const FNodeClassInfo& InNodeClassInfo)
+				FNodeRegistryTransaction(const FNodeRegistryKey& InKey, const FNodeClassInfo& InNodeClassInfo)
 				: NodeClassInfo(InNodeClassInfo)
+				, Key(InKey)
 				{
 				}
 
-				ETransactionType GetTransactionType() const override
+				virtual ETransactionType GetTransactionType() const override
 				{
 					return ETransactionType::Remove;
 				}
 
-				TUniquePtr<IRegistryTransaction> Clone() const override
+				virtual TUniquePtr<IRegistryTransaction> Clone() const override
 				{
 					return MakeUnique<FNodeRegistryTransaction>(*this);
 				}
 
-				const FNodeClassInfo* GetNodeClassInfo() const 
+				virtual const FNodeClassInfo* GetNodeClassInfo() const override
 				{
 					return &NodeClassInfo;
 				}
 
+				virtual const FNodeRegistryKey* GetNodeRegistryKey() const override
+				{
+					return &Key;
+				}
+
 			private:
 
 				FNodeClassInfo NodeClassInfo;
+				FNodeRegistryKey Key;
 			};
 
-			return MakeUnique<FNodeRegistryTransaction>(InInfo);
+			return MakeUnique<FNodeRegistryTransaction>(InKey, InInfo);
 		}
 	}
 }
diff --git a/Engine/Plugins/Runtime/Metasound/Source/MetasoundFrontend/Private/MetasoundFrontendRegistryTransaction.h b/Engine/Plugins/Runtime/Metasound/Source/MetasoundFrontend/Private/MetasoundFrontendRegistryTransaction.h
--- a/Engine/Plugins/Runtime/Metasound/Source/MetasoundFrontend/Private/MetasoundFrontendRegistryTransaction.h
+++ b/Engine/Plugins/Runtime/Metasound/Source/MetasoundFrontend/Private/MetasoundFrontendRegistryTransaction.h
@@ -14,10 +14,10 @@ namespace Metasound
 		FRegistryTransactionID GetOriginRegistryTransactionID();
 
 		/** Create a node registration transaction. */
-		FRegistryTransactionPtr MakeAddNodeRegistryTransaction(const FNodeClassInfo& InInfo);
+		FRegistryTransactionPtr MakeAddNodeRegistryTransaction(const FNodeRegistryKey& InKey, const FNodeClassInfo& InInfo);
 
 		/** Create a node unregistration transaction. */
-		FRegistryTransactionPtr MakeRemoveNodeRegistryTransaction(const FNodeClassInfo& InInfo);
+		FRegistryTransactionPtr MakeRemoveNodeRegistryTransaction(const FNodeRegistryKey& InKey, const FNodeClassInfo& InInfo);
 
 		/** Maintains a history of IRegistryTransactions. Calls are threadsafe (excluding
 		 * the constructor and destructor.)
diff --git a/Engine/Plugins/Runtime/Metasound/Source/MetasoundFrontend/Private/MetasoundFrontendStandardController.cpp b/Engine/Plugins/Runtime/Metasound/Source/MetasoundFrontend/Private/MetasoundFrontendStandardController.cpp
--- a/Engine/Plugins/Runtime/Metasound/Source/MetasoundFrontend/Private/MetasoundFrontendStandardController.cpp
+++ b/Engine/Plugins/Runtime/Metasound/Source/MetasoundFrontend/Private/MetasoundFrontendStandardController.cpp
@@ -275,7 +275,6 @@ namespace Metasound
 			return Access;
 		}
 
-
 		//
 		// FInputNodeOutputController
 		// 
@@ -409,7 +408,6 @@ namespace Metasound
 			return false;
 		}
 
-
 		//
 		// FBaseInputController
 		// 
@@ -586,7 +584,7 @@ namespace Metasound
 
 			// If data types are not equal, check for converter nodes which could
 			// convert data type.
-			OutConnectability.PossibleConverterNodeClasses = FMetasoundFrontendRegistryContainer::Get()->GetPossibleConverterNodes(InController.GetDataType(), GetDataType());
+			OutConnectability.PossibleConverterNodeClasses = FRegistry::Get()->GetPossibleConverterNodes(InController.GetDataType(), GetDataType());
 
 			if (OutConnectability.PossibleConverterNodeClasses.Num() > 0)
 			{
@@ -632,8 +630,6 @@ namespace Metasound
 			// Generate the converter node.
 			FNodeHandle ConverterNode = Graph->AddNode(InConverterInfo.NodeKey);
 
-			// TODO: may want to rename Point -> Pin. 
-			// TODO: fix naming confusion between "Pin" and "VertexKey". ConverterInfo should use vertexkey. 
 			TArray<FInputHandle> ConverterInputs = ConverterNode->GetInputsWithVertexName(InConverterInfo.PreferredConverterInputPin);
 			TArray<FOutputHandle> ConverterOutputs = ConverterNode->GetOutputsWithVertexName(InConverterInfo.PreferredConverterOutputPin);
 
@@ -894,7 +890,6 @@ namespace Metasound
 			return false;
 		}
 
-
 		//
 		// FBaseNodeController
 		//
@@ -2310,10 +2305,10 @@ namespace Metasound
 				auto IsInputWithSameName = [&](const FMetasoundFrontendClassInput& ExistingDesc) { return ExistingDesc.Name == InClassInput.Name; };
 				if (Algo::NoneOf(GraphClass->Interface.Inputs, IsInputWithSameName))
 				{
-					FMetasoundFrontendClassMetadata ClassMetadata;
-					if (FMetasoundFrontendRegistryContainer::GetInputNodeClassMetadataForDataType(InClassInput.TypeName, ClassMetadata))
+					FNodeRegistryKey Key;
+					if (FRegistry::GetInputNodeRegistryKeyForDataType(InClassInput.TypeName, Key))
 					{
-						if (FConstClassAccessPtr InputClass = OwningDocument->FindOrAddClass(ClassMetadata))
+						if (FConstClassAccessPtr InputClass = OwningDocument->FindOrAddClass(Key))
 						{
 							FMetasoundFrontendNode& Node = GraphClass->Graph.Nodes.Emplace_GetRef(*InputClass);
 
@@ -2384,10 +2379,10 @@ namespace Metasound
 				auto IsOutputWithSameName = [&](const FMetasoundFrontendClassOutput& ExistingDesc) { return ExistingDesc.Name == InClassOutput.Name; };
 				if (Algo::NoneOf(GraphClass->Interface.Outputs, IsOutputWithSameName))
 				{
-					FMetasoundFrontendClassMetadata ClassMetadata;
-					if (FMetasoundFrontendRegistryContainer::GetOutputNodeClassMetadataForDataType(InClassOutput.TypeName, ClassMetadata))
+					FNodeRegistryKey Key;
+					if (FRegistry::GetOutputNodeRegistryKeyForDataType(InClassOutput.TypeName, Key))
 					{
-						if (FConstClassAccessPtr OutputClass = OwningDocument->FindOrAddClass(ClassMetadata))
+						if (FConstClassAccessPtr OutputClass = OwningDocument->FindOrAddClass(Key))
 						{
 							FMetasoundFrontendNode& Node = GraphClass->Graph.Nodes.Emplace_GetRef(*OutputClass);
 
@@ -2459,7 +2454,7 @@ namespace Metasound
 		{
 			if (const FMetasoundFrontendClassInput* Desc = FindInputDescriptionWithName(InInputName))
 			{
-				return FMetasoundFrontendRegistryContainer::Get()->GetDesiredLiteralTypeForDataType(Desc->TypeName);
+				return FRegistry::Get()->GetDesiredLiteralTypeForDataType(Desc->TypeName);
 			}
 			return ELiteralType::Invalid;
 		}
@@ -2469,7 +2464,7 @@ namespace Metasound
 		{
 			if (const FMetasoundFrontendClassInput* Desc = FindInputDescriptionWithName(InInputName))
 			{
-				return FMetasoundFrontendRegistryContainer::Get()->GetLiteralUClassForDataType(Desc->TypeName);
+				return FRegistry::Get()->GetLiteralUClassForDataType(Desc->TypeName);
 			}
 			return nullptr;
 		}
@@ -2567,20 +2562,6 @@ namespace Metasound
 			}
 		}
 
-		FMetasoundFrontendClassMetadata FGraphController::CreateInputClassMetadata(const FMetasoundFrontendClassInput& InClassInput)
-		{
-			FMetasoundFrontendClassMetadata ClassMetadata;
-			FMetasoundFrontendRegistryContainer::GetInputNodeClassMetadataForDataType(InClassInput.TypeName, ClassMetadata);
-			return ClassMetadata;
-		}
-
-		FMetasoundFrontendClassMetadata FGraphController::CreateOutputClassMetadata(const FMetasoundFrontendClassOutput& InClassOutput)
-		{
-			FMetasoundFrontendClassMetadata ClassMetadata;
-			FMetasoundFrontendRegistryContainer::GetOutputNodeClassMetadataForDataType(InClassOutput.TypeName, ClassMetadata);
-			return ClassMetadata;
-		}
-
 		// This can be used to clear the current literal for a given input.
 		// @returns false if the input name couldn't be found.
 		bool FGraphController::ClearLiteralForInput(const FString& InInputName, FGuid InVertexID)
@@ -2593,32 +2574,28 @@ namespace Metasound
 			return false;
 		}
 
-		FNodeHandle FGraphController::AddNode(const FNodeClassInfo& InNodeClass)
+		FNodeHandle FGraphController::AddNode(const FNodeRegistryKey& InKey)
 		{
 			if (IsValid())
 			{
-				if (FConstClassAccessPtr DependencyDescription = OwningDocument->FindOrAddClass(InNodeClass))
+				// Construct a FNodeClassInfo from this lookup key.
+				if (FConstClassAccessPtr DependencyDescription = OwningDocument->FindOrAddClass(InKey))
 				{
 					return AddNode(DependencyDescription);
 				}
+				else
+				{
+					UE_LOG(LogMetaSound, Warning, TEXT("Failed to find node class info with registry key [Key:%s]"), *InKey);
+					return INodeController::GetInvalidHandle();
+				}
 			}
 
 			return FInvalidNodeController::GetInvalid();
 		}
 
-		FNodeHandle FGraphController::AddNode(const FNodeRegistryKey& InNodeClass)
-		{
-			// Construct a FNodeClassInfo from this lookup key.
-			FNodeClassInfo ClassInfo;
-			ClassInfo.LookupKey = InNodeClass;
-			ClassInfo.NodeType = EMetasoundFrontendClassType::External;
-
-			return AddNode(ClassInfo);
-		}
-
 		FNodeHandle FGraphController::AddNode(const FMetasoundFrontendClassMetadata& InClassMetadata)
 		{
-			return AddNode(FMetasoundFrontendRegistryContainer::GetRegistryKey(InClassMetadata));
+			return AddNode(FRegistry::GetRegistryKey(InClassMetadata));
 		}
 
 		FNodeHandle FGraphController::AddDuplicateNode(const INodeController& InNode)
@@ -2732,7 +2709,7 @@ namespace Metasound
 				return TUniquePtr<IOperator>(nullptr);
 			}
 
-			// TODO: Implement inflation step here.
+			// TODO: Implement subgraph inflation step here.
 
 			// TODO: bubble up errors. 
 			TArray<FMetasoundFrontendGraphClass> Subgraphs = OwningDocument->GetSubgraphs();
@@ -3414,10 +3391,10 @@ namespace Metasound
 
 		FGraphHandle FDocumentController::AddDuplicateSubgraph(const IGraphController& InGraph)
 		{
-			// TODO: class IDs are fucked. 
-			// Maybe class IDs should just be the linking info and not a GUID.
-			// ClassIDs are just internal references, so they need to be fixed up
-			// here if swapping documents. 
+			// TODO: class IDs have issues.. 
+			// Currently ClassIDs are just used for internal linking. They need to be fixed up
+			// here if swapping documents. In the future, ClassIDs should be unique and consistent
+			// across documents and platforms.
 
 			FConstDocumentAccess GraphDocumentAccess = GetSharedAccess(*InGraph.GetOwningDocument());
 			const FMetasoundFrontendDocument* OtherDocument = GraphDocumentAccess.ConstDocument.Get();
@@ -3487,36 +3464,40 @@ namespace Metasound
 			return FrontendControllerIntrinsics::GetInvalidValueConstRef<FMetasoundFrontendDocumentMetadata>();
 		}
 
-		FConstClassAccessPtr FDocumentController::FindClass(const FNodeClassInfo& InNodeClass) const
+		FConstClassAccessPtr FDocumentController::FindClass(const FNodeRegistryKey& InKey) const
 		{
-			return DocumentPtr.GetClassWithInfo(InNodeClass);
+			return DocumentPtr.GetClassWithRegistryKey(InKey);
 		}
 
-		FConstClassAccessPtr FDocumentController::FindOrAddClass(const FNodeClassInfo& InNodeClass)
+		FConstClassAccessPtr FDocumentController::FindOrAddClass(const FNodeRegistryKey& InKey)
 		{
-			FConstClassAccessPtr ClassPtr = FindClass(InNodeClass);
+			FConstClassAccessPtr ClassPtr = FindClass(InKey);
 
 			auto AddClass = [=](FMetasoundFrontendClass&& NewClassDescription)
 			{
 				FConstClassAccessPtr NewClassPtr;
 				if (DocumentPtr.IsValid())
 				{
-					// FNodeClassInfo does not contain enough info to add a subgraph.
-					check(EMetasoundFrontendClassType::Graph != InNodeClass.NodeType);
+					// Cannot add a subgraph using this method because dependencies
+					// of external graph are not added in this method.
+					check(EMetasoundFrontendClassType::Graph != NewClassDescription.Metadata.Type);
 					NewClassDescription.ID = FGuid::NewGuid();
 
 					DocumentPtr->Dependencies.Add(MoveTemp(NewClassDescription));
-					NewClassPtr = FindClass(InNodeClass);
+					NewClassPtr = FindClass(InKey);
 				}
 				return NewClassPtr;
 			};
 
-			if (ClassPtr.IsValid())
+			if (const FMetasoundFrontendClass* MetasoundClass = ClassPtr.Get())
 			{
 				// External node classes must match version to return shared definition.
-				if (InNodeClass.NodeType == EMetasoundFrontendClassType::External)
+				if (MetasoundClass->Metadata.Type == EMetasoundFrontendClassType::External)
 				{
-					FMetasoundFrontendClass NewClass = GenerateClassDescription(InNodeClass);
+					// TODO: Assuming we want to recheck classes when they add another
+					// node, this should be replace with a call to synchronize a 
+					// single class.
+					FMetasoundFrontendClass NewClass = GenerateClassDescription(InKey);
 					if (NewClass.Metadata.Version.Major != ClassPtr->Metadata.Version.Major)
 					{
 						return AddClass(MoveTemp(NewClass));
@@ -3526,7 +3507,7 @@ namespace Metasound
 				return ClassPtr;
 			}
 
-			FMetasoundFrontendClass NewClass = GenerateClassDescription(InNodeClass);
+			FMetasoundFrontendClass NewClass = GenerateClassDescription(InKey);
 			return AddClass(MoveTemp(NewClass));
 		}
 
@@ -3561,7 +3542,9 @@ namespace Metasound
 				if ((EMetasoundFrontendClassType::External == InMetadata.Type) || (EMetasoundFrontendClassType::Input == InMetadata.Type) || (EMetasoundFrontendClassType::Output == InMetadata.Type))
 				{
 					FMetasoundFrontendClass NewClass;
-					if (FMetasoundFrontendRegistryContainer::GetFrontendClassFromRegistered(InMetadata, NewClass))
+					FNodeRegistryKey Key = FRegistry::GetRegistryKey(InMetadata);
+
+					if (FRegistry::GetFrontendClassFromRegistered(Key, NewClass))
 					{
 						NewClass.ID = FGuid::NewGuid();
 						DocumentPtr->Dependencies.Add(NewClass);
@@ -3569,8 +3552,8 @@ namespace Metasound
 					else
 					{
 						UE_LOG(LogMetaSound, Error,
-							TEXT("Cannot add external dependency. No Metasound class found with matching metadata Name: "
-								"\"%s\" (%s). Suggested solution \"%s\" by %s."),
+							TEXT("Cannot add external dependency. No Metasound class found with matching registry key [Key:%s, Name:%s, Version:%s]. Suggested solution \"%s\" by %s."),
+							*Key,
 							*InMetadata.ClassName.GetFullName().ToString(),
 							*InMetadata.Version.ToString(),
 							*InMetadata.PromptIfMissing.ToString(),
@@ -3639,18 +3622,24 @@ namespace Metasound
 			}
 			while (NumDependenciesRemovedThisItr > 0);
 
-			FMetasoundFrontendRegistryContainer* Registry = FMetasoundFrontendRegistryContainer::Get();
+			/*
+			FRegistry* Registry = FRegistry::Get();
 			if (ensure(Registry))
 			{
+				// Make sure classes are up-to-date with registered versions of class
 				for (FMetasoundFrontendClass& Class : DocumentPtr->Dependencies)
 				{
 					FMetasoundFrontendClass RegisteredClass;
-					if (Registry->FindFrontendClassFromRegistered(Class.Metadata, RegisteredClass))
+					FNodeRegistryKey Key = FRegistry::GetRegistryKey(Class.Metadata);
+					if (Registry->FindFrontendClassFromRegistered(Key, RegisteredClass))
 					{
-						Class.Metadata = RegisteredClass.Metadata;
+						Class = RegisteredClass;
+						// TODO: update nodes for class.
+						// TODO: perform minor version updating.
 					}
 				}
 			}
+			*/
 		}
 
 		FGraphHandle FDocumentController::GetRootGraph()
@@ -3777,11 +3766,12 @@ namespace Metasound
 
 		bool FDocumentController::IsMatchingMetasoundClass(const FNodeClassInfo& InNodeClass, const FMetasoundFrontendClassMetadata& InMetadata) 
 		{
-			if (InNodeClass.NodeType == InMetadata.Type)
-			{
-				return InNodeClass.LookupKey == FRegistry::GetRegistryKey(InMetadata);
-			}
-			return false;
+			return FRegistry::GetRegistryKey(InNodeClass) == FRegistry::GetRegistryKey(InMetadata);
+		}
+
+		bool FDocumentController::IsMatchingMetasoundClass(const FNodeRegistryKey& InKey, const FMetasoundFrontendClassMetadata& InMetadata) 
+		{
+			return InKey == FRegistry::GetRegistryKey(InMetadata);
 		}
 
 		FDocumentAccess FDocumentController::ShareAccess() 
@@ -3802,7 +3792,6 @@ namespace Metasound
 
 			return Access;
 		}
-
 	}
 }
 #undef LOCTEXT_NAMESPACE
diff --git a/Engine/Plugins/Runtime/Metasound/Source/MetasoundFrontend/Private/MetasoundFrontendStandardController.h b/Engine/Plugins/Runtime/Metasound/Source/MetasoundFrontend/Private/MetasoundFrontendStandardController.h
--- a/Engine/Plugins/Runtime/Metasound/Source/MetasoundFrontend/Private/MetasoundFrontendStandardController.h
+++ b/Engine/Plugins/Runtime/Metasound/Source/MetasoundFrontend/Private/MetasoundFrontendStandardController.h
@@ -4,7 +4,6 @@
 
 #include "CoreMinimal.h"
 #include "MetasoundBuilderInterface.h"
-#include "MetasoundFrontendBaseClasses.h"
 #include "MetasoundFrontendController.h"
 #include "MetasoundFrontendDocument.h"
 #include "MetasoundFrontendInvalidController.h"
@@ -21,6 +20,7 @@ namespace Metasound
 		 */
 		class FBaseOutputController : public IOutputController
 		{
+			using FRegistry = FMetasoundFrontendRegistryContainer;
 		public:
 
 			struct FInitParams
@@ -92,6 +92,7 @@ namespace Metasound
 		 */
 		class FInputNodeOutputController : public FBaseOutputController
 		{
+			using FRegistry = FMetasoundFrontendRegistryContainer;
 		public:
 			struct FInitParams
 			{
@@ -136,6 +137,7 @@ namespace Metasound
 		 */
 		class FOutputNodeOutputController : public FBaseOutputController
 		{
+			using FRegistry = FMetasoundFrontendRegistryContainer;
 		public:
 			struct FInitParams
 			{
@@ -176,6 +178,7 @@ namespace Metasound
 		 */
 		class FBaseInputController : public IInputController 
 		{
+			using FRegistry = FMetasoundFrontendRegistryContainer;
 		public:
 
 			struct FInitParams
@@ -252,6 +255,7 @@ namespace Metasound
 		 */
 		class FOutputNodeInputController : public FBaseInputController 
 		{
+			using FRegistry = FMetasoundFrontendRegistryContainer;
 		public:
 			struct FInitParams
 			{
@@ -293,6 +297,7 @@ namespace Metasound
 		 */
 		class FInputNodeInputController : public FBaseInputController 
 		{
+			using FRegistry = FMetasoundFrontendRegistryContainer;
 		public:
 			struct FInitParams
 			{
@@ -331,6 +336,7 @@ namespace Metasound
 		 */
 		class FBaseNodeController : public INodeController
 		{
+			using FRegistry = FMetasoundFrontendRegistryContainer;
 		public:
 
 			struct FInitParams
@@ -489,6 +495,8 @@ namespace Metasound
 		/** FNodeController represents a external or subgraph node. */
 		class FNodeController : public FBaseNodeController
 		{
+			using FRegistry = FMetasoundFrontendRegistryContainer;
+
 			// Private token only allows members or friends to call constructor.
 			enum EPrivateToken { Token };
 
@@ -539,6 +547,8 @@ namespace Metasound
 		/** FOutputNodeController represents an output node. */
 		class FOutputNodeController: public FBaseNodeController
 		{
+			using FRegistry = FMetasoundFrontendRegistryContainer;
+
 			// Private token only allows members or friends to call constructor.
 			enum EPrivateToken { Token };
 
@@ -600,6 +610,8 @@ namespace Metasound
 		/** FInputNodeController represents an input node. */
 		class FInputNodeController: public FBaseNodeController
 		{
+			using FRegistry = FMetasoundFrontendRegistryContainer;
+
 			// Private token only allows members or friends to call constructor.
 			enum EPrivateToken { Token };
 
@@ -665,6 +677,7 @@ namespace Metasound
 		/** FGraphController represents a Metasound graph class. */
 		class FGraphController : public IGraphController
 		{
+			using FRegistry = FMetasoundFrontendRegistryContainer;
 
 			// Private token only allows members or friends to call constructor.
 			enum EPrivateToken { Token };
@@ -779,7 +792,6 @@ namespace Metasound
 			// @returns false if the input name couldn't be found.
 			bool ClearLiteralForInput(const FString& InInputName, FGuid InVertexID) override;
 
-			FNodeHandle AddNode(const FNodeClassInfo& InNodeClass) override;
 			FNodeHandle AddNode(const FNodeRegistryKey& InNodeClass) override;
 			FNodeHandle AddNode(const FMetasoundFrontendClassMetadata& InClassMetadata) override;
 			FNodeHandle AddDuplicateNode(const INodeController& InNode) override;
@@ -867,9 +879,6 @@ namespace Metasound
 			FClassOutputAccessPtr FindOutputDescriptionWithNodeID(FGuid InNodeID);
 			FConstClassOutputAccessPtr FindOutputDescriptionWithNodeID(FGuid InNodeID) const;
 
-			FMetasoundFrontendClassMetadata CreateInputClassMetadata(const FMetasoundFrontendClassInput& InClassInput);
-			FMetasoundFrontendClassMetadata CreateOutputClassMetadata(const FMetasoundFrontendClassOutput& InClassOutput);
-
 			FGraphClassAccessPtr GraphClassPtr;
 			FDocumentHandle OwningDocument; 
 		};
@@ -909,10 +918,10 @@ namespace Metasound
 			FConstGraphClassAccessPtr FindSubgraphWithID(FGuid InClassID) const override;
 			FConstClassAccessPtr FindClassWithID(FGuid InClassID) const override;
 
-			FConstClassAccessPtr FindClass(const FNodeClassInfo& InNodeClass) const override;
+			FConstClassAccessPtr FindClass(const FNodeRegistryKey& InKey) const override;
 			FConstClassAccessPtr FindClass(const FMetasoundFrontendClassMetadata& InMetadata) const override;
 
-			FConstClassAccessPtr FindOrAddClass(const FNodeClassInfo& InNodeClass) override;
+			FConstClassAccessPtr FindOrAddClass(const FNodeRegistryKey& InKey) override;
 			FConstClassAccessPtr FindOrAddClass(const FMetasoundFrontendClassMetadata& InMetadata) override;
 
 			virtual FGraphHandle AddDuplicateSubgraph(const IGraphController& InGraph) override;
@@ -942,6 +951,7 @@ namespace Metasound
 
 			static bool IsMatchingMetasoundClass(const FMetasoundFrontendClassMetadata& InMetadataA, const FMetasoundFrontendClassMetadata& InMetadataB);
 			static bool IsMatchingMetasoundClass(const FNodeClassInfo& InNodeClass, const FMetasoundFrontendClassMetadata& InMetadata);
+			static bool IsMatchingMetasoundClass(const FNodeRegistryKey& InKey, const FMetasoundFrontendClassMetadata& InMetadata);
 
 		protected:
 			FDocumentAccess ShareAccess() override;
diff --git a/Engine/Plugins/Runtime/Metasound/Source/MetasoundFrontend/Private/MetasoundFrontendTransform.cpp b/Engine/Plugins/Runtime/Metasound/Source/MetasoundFrontend/Private/MetasoundFrontendTransform.cpp
--- a/Engine/Plugins/Runtime/Metasound/Source/MetasoundFrontend/Private/MetasoundFrontendTransform.cpp
+++ b/Engine/Plugins/Runtime/Metasound/Source/MetasoundFrontend/Private/MetasoundFrontendTransform.cpp
@@ -121,14 +121,23 @@ namespace Metasound
 				FGraphHandle GraphHandle = InDocument->GetRootGraph();
 				TArray<FNodeHandle> FrontendNodes = GraphHandle->GetNodes();
 
-				// If was connected to a hidden node, then remove that node and set the literal value accordingly.
+				// Before literals could be stored on node inputs directly, they were stored
+				// by creating hidden input nodes. Update the doc by finding all hidden input
+				// nodes, placing the literal value of the input node directly on the
+				// downstream node's input. Then delete the hidden input node.
 				for (FNodeHandle& NodeHandle : FrontendNodes)
 				{
-					if (NodeHandle->GetNodeStyle().Display.Visibility == EMetasoundFrontendNodeStyleDisplayVisibility::Hidden)
+					const bool bIsHiddenNode = NodeHandle->GetNodeStyle().Display.Visibility == EMetasoundFrontendNodeStyleDisplayVisibility::Hidden;
+					const bool bIsInputNode = EMetasoundFrontendClassType::Input == NodeHandle->GetClassMetadata().Type;
+					const bool bIsHiddenInputNode = bIsHiddenNode && bIsInputNode;
+
+					if (bIsHiddenInputNode)
 					{
+						// Get literal value from input node.
 						const FGuid VertexID = GraphHandle->GetVertexIDForInputVertex(NodeHandle->GetNodeName());
 						const FMetasoundFrontendLiteral DefaultLiteral = GraphHandle->GetDefaultInput(VertexID);
 
+						// Apply literal value to downstream node's inputs.
 						TArray<FOutputHandle> OutputHandles = NodeHandle->GetOutputs();
 						if (ensure(OutputHandles.Num() == 1))
 						{
@@ -150,9 +159,8 @@ namespace Metasound
 									Input->SetLiteral(DefaultLiteral);
 								}
 							}
-
-							GraphHandle->RemoveNode(*NodeHandle);
 						}
+						GraphHandle->RemoveNode(*NodeHandle);
 					}
 				}
 			}
diff --git a/Engine/Plugins/Runtime/Metasound/Source/MetasoundFrontend/Public/MetasoundAssetBase.h b/Engine/Plugins/Runtime/Metasound/Source/MetasoundFrontend/Public/MetasoundAssetBase.h
--- a/Engine/Plugins/Runtime/Metasound/Source/MetasoundFrontend/Public/MetasoundAssetBase.h
+++ b/Engine/Plugins/Runtime/Metasound/Source/MetasoundFrontend/Public/MetasoundAssetBase.h
@@ -4,7 +4,6 @@
 #include "CoreMinimal.h"
 #include "MetasoundAccessPtr.h"
 #include "MetasoundFrontend.h"
-#include "MetasoundFrontendBaseClasses.h"
 #include "MetasoundFrontendController.h"
 #include "MetasoundFrontendDocumentAccessPtr.h"
 #include "MetasoundGraph.h"
diff --git a/Engine/Plugins/Runtime/Metasound/Source/MetasoundFrontend/Public/MetasoundDataTypeRegistrationMacro.h b/Engine/Plugins/Runtime/Metasound/Source/MetasoundFrontend/Public/MetasoundDataTypeRegistrationMacro.h
--- a/Engine/Plugins/Runtime/Metasound/Source/MetasoundFrontend/Public/MetasoundDataTypeRegistrationMacro.h
+++ b/Engine/Plugins/Runtime/Metasound/Source/MetasoundFrontend/Public/MetasoundDataTypeRegistrationMacro.h
@@ -59,7 +59,6 @@ namespace Metasound
 
 	namespace MetasoundDataTypeRegistrationPrivate
 	{
-
 		// Helper utility to test if we can transmit a datatype between a send and a receive node.
 		template <typename TDataType>
 		struct TIsTransmittable
@@ -68,12 +67,6 @@ namespace Metasound
 			static constexpr bool bIsCopyConstructible = std::is_copy_constructible<TDataType>::value;
 			static constexpr bool bIsCopyAssignable = std::is_copy_assignable<TDataType>::value;
 
-			// TODO: audio types were intended to be send/receive nodes but they require 
-			// template specialization.  TIsTransmittable should ask the Send and Receive nodes
-			// if they are transmittable rather than attempting to do the logic for all
-			// types here. 
-			//static constexpr bool bIsAudioDataType = TIsDerivedFrom<TDataType, IAudioDataType>::Value;
-
 			static constexpr bool bCanBeTransmitted = bIsCopyConstructible && bIsCopyAssignable;
 
 		public:
@@ -240,50 +233,121 @@ namespace Metasound
 			AttemptToRegisterConverter<TDataType, FString>();
 		}
 
-		// SFINAE for enum types.
-		template<typename TDataType, typename std::enable_if<TEnumTraits<TDataType>::bIsEnum, bool>::type = true>
-		bool RegisterEnumDataTypeWithFrontend()
+		/** Creates the FDataTypeRegistryInfo for a data type.
+		 * 
+		 * @tparam TDataType - The data type to create info for.
+		 * @tparam PreferredArgType - The preferred constructor argument type to use when creating an instance of the data type.
+		 * @tparam UClassToUse - The preferred UObject class to use when constructing from an Audio::IProxyData.
+		 */
+		template<typename TDataType, ELiteralType PreferredArgType, typename UClassToUse>
+		Frontend::FDataTypeRegistryInfo CreateDataTypeInfo()
 		{
-			using InnerType = typename TEnumTraits<TDataType>::InnerType;
-			using FStringHelper = TEnumStringHelper<InnerType>;
+			Frontend::FDataTypeRegistryInfo RegistryInfo;
+
+			RegistryInfo.DataTypeName = GetMetasoundDataTypeName<TDataType>();
+			RegistryInfo.PreferredLiteralType = PreferredArgType;
+
+			RegistryInfo.bIsDefaultParsable = TIsParsable<TDataType, FLiteral::FNone>::Value;		
+			RegistryInfo.bIsBoolParsable = TIsParsable<TDataType, bool>::Value;
+			RegistryInfo.bIsIntParsable = TIsParsable<TDataType, int32>::Value;
+			RegistryInfo.bIsFloatParsable = TIsParsable<TDataType, float>::Value;
+			RegistryInfo.bIsStringParsable = TIsParsable<TDataType, FString>::Value;
+			RegistryInfo.bIsProxyParsable = TIsParsable<TDataType, const Audio::IProxyDataPtr&>::Value;
+
+			RegistryInfo.bIsDefaultArrayParsable = TIsParsable<TDataType, TArray<FLiteral::FNone>>::Value;		
+			RegistryInfo.bIsBoolArrayParsable = TIsParsable<TDataType, TArray<bool>>::Value;
+			RegistryInfo.bIsIntArrayParsable = TIsParsable<TDataType, TArray<int32>>::Value;
+			RegistryInfo.bIsFloatArrayParsable = TIsParsable<TDataType, TArray<float>>::Value;
+			RegistryInfo.bIsStringArrayParsable = TIsParsable<TDataType, TArray<FString>>::Value;
+			RegistryInfo.bIsProxyArrayParsable = TIsParsable<TDataType, const TArray<Audio::IProxyDataPtr>& >::Value;
 
-			struct FEnumHandler : Metasound::Frontend::IEnumDataTypeInterface
+			RegistryInfo.bIsEnum = TEnumTraits<TDataType>::bIsEnum;
+
+			RegistryInfo.bIsTransmittable = TIsTransmittable<TDataType>::Value;
+			
+			if constexpr (std::is_base_of<UObject, UClassToUse>::value)
 			{
-				FName GetNamespace() const override
-				{
-					return FStringHelper::GetNamespace();
-				}
-				int32 GetDefaultValue() const override
-				{
-					return static_cast<int32>(TEnumTraits<TDataType>::DefaultValue);
-				}
-				const TArray<FGenericInt32Entry>& GetAllEntries() const override
+				RegistryInfo.ProxyGeneratorClass = UClassToUse::StaticClass();
+			}
+			else 
+			{
+				static_assert(std::is_same<UClassToUse, void>::value, "Only UObject derived classes can supply proxy interfaces.");
+				RegistryInfo.ProxyGeneratorClass = nullptr;
+			}
+
+			return RegistryInfo;
+		}
+
+		/** Returns an IEnumDataTypeInterface pointer for the data type. If the 
+		 * data type has no IEnumDataTypeInterface, the returned pointer will be
+		 * invalid.
+		 *
+		 * @tparam TDataType - The data type to create the interface for.
+		 *
+		 * @return A shared pointer to the IEnumDataTypeInterface. If the TDataType
+		 *         does not have an IEnumDataTypeInterface, returns an invalid pointer.
+		 */
+		template<typename TDataType>
+		TSharedPtr<Frontend::IEnumDataTypeInterface> GetEnumDataTypeInterface()
+		{
+			TSharedPtr<Frontend::IEnumDataTypeInterface> EnumInterfacePtr;
+
+			using FEnumTraits = TEnumTraits<TDataType>;
+
+			// Check if data type is an enum.
+			if constexpr (FEnumTraits::bIsEnum)
+			{
+				using InnerType = typename FEnumTraits::InnerType;
+				using FStringHelper = TEnumStringHelper<InnerType>;
+
+				struct FEnumHandler : Metasound::Frontend::IEnumDataTypeInterface
 				{
-					auto BuildIntEntries = []()
+					FName GetNamespace() const override
 					{
-						// Convert to int32 representation 
-						TArray<FGenericInt32Entry> IntEntries;
-						IntEntries.Reserve(FStringHelper::GetAllEntries().Num());
-						for (const TEnumEntry<InnerType>& i : FStringHelper::GetAllEntries())
+						return FStringHelper::GetNamespace();
+					}
+
+					int32 GetDefaultValue() const override
+					{
+						return static_cast<int32>(TEnumTraits<TDataType>::DefaultValue);
+					}
+
+					const TArray<FGenericInt32Entry>& GetAllEntries() const override
+					{
+						auto BuildIntEntries = []()
 						{
-							IntEntries.Emplace(i);
-						}
+							// Convert to int32 representation 
+							TArray<FGenericInt32Entry> IntEntries;
+							IntEntries.Reserve(FStringHelper::GetAllEntries().Num());
+							for (const TEnumEntry<InnerType>& i : FStringHelper::GetAllEntries())
+							{
+								IntEntries.Emplace(i);
+							}
+							return IntEntries;
+						};
+						static const TArray<FGenericInt32Entry> IntEntries = BuildIntEntries();
 						return IntEntries;
-					};
-					static const TArray<FGenericInt32Entry> IntEntries = BuildIntEntries();
-					return IntEntries;
-				}
-			};
+					}
+				};
+
+				EnumInterfacePtr = MakeShared<FEnumHandler>();
+			}
 
-			return FMetasoundFrontendRegistryContainer::Get()->RegisterEnumDataInterface(
-				GetMetasoundDataTypeName<TDataType>(), MakeShared<FEnumHandler>());
+			return EnumInterfacePtr;
 		}
-		
-		// SFINAE stub for non-enum types. 
-		template<typename TDataType, typename std::enable_if<!TEnumTraits<TDataType>::bIsEnum, bool>::type = true>
-		bool RegisterEnumDataTypeWithFrontend() { return false; }
 
-		template<typename TDataType, ELiteralType PreferredArgType = ELiteralType::None, typename UClassToUse = UObject>
+		/** Registers a data type with the MetaSound Frontend. This allows the data type
+		 * to be used in Input and Output nodes by informing the Frontend how to 
+		 * instantiate an instance. 
+		 *
+		 * @tparam TDataType - The data type to register.
+		 * @tparam PreferredArgType - The preferred constructor argument type to use when creating an instance of the data type.
+		 * @tparam UClassToUse - The preferred UObject class to use when constructing from an Audio::IProxyDataPtr. If the type is not
+		 *                       constructible with an Audio::IProxyDataPtr, then this should be void.
+		 *
+		 * @return True on success, false on failure.
+		 */
+		template<typename TDataType, ELiteralType PreferredArgType = ELiteralType::None, typename UClassToUse = void>
 		bool RegisterDataTypeWithFrontendInternal()
 		{
 			// if we reenter this code (because DECLARE_METASOUND_DATA_REFERENCE_TYPES was called twice with the same type),
@@ -297,136 +361,133 @@ namespace Metasound
 
 			bAlreadyRegisteredThisDataType = true;
 
-			// Lambdas that generate our template-instantiated input, variable, and output nodes:
-			FCreateInputNodeFunction InputNodeConstructor = [](FInputNodeConstructorParams&& InParams) -> TUniquePtr<INode>
+			// Define registry entry class for this data type.
+			class FDataTypeRegistryEntry : public Frontend::IDataTypeRegistryEntry
 			{
-				return TUniquePtr<INode>(new TInputNode<TDataType>(InParams.NodeName, InParams.InstanceID, InParams.VertexName, MoveTemp(InParams.InitParam)));
-			};
+			public:
+				FDataTypeRegistryEntry()
+				: Info(CreateDataTypeInfo<TDataType, PreferredArgType, UClassToUse>())
+				, EnumInterface(GetEnumDataTypeInterface<TDataType>())
+				{
+					// Create class info using prototype node
+					// TODO: register nodes with static class info instead of prototype instance.
 
-			FCreateMetasoundFrontendClassFunction CreateFrontendInputClass = []() -> FMetasoundFrontendClass
-			{
-				// Create class info using prototype node
-				// TODO: register input nodes with static class info.
-				static TInputNode<TDataType> Prototype(TEXT(""), FGuid(), TEXT(""), FLiteral());
-				return Metasound::Frontend::GenerateClassDescription(Prototype.GetMetadata(), EMetasoundFrontendClassType::Input);
-			};
+					TInputNode<TDataType> InputPrototype(TEXT(""), FGuid(), TEXT(""), FLiteral());
+					InputClass = Metasound::Frontend::GenerateClassDescription(InputPrototype.GetMetadata(), EMetasoundFrontendClassType::Input);
 
-			FCreateVariableNodeFunction VariableNodeConstructor = [](FVariableNodeConstructorParams&& InParams) -> TUniquePtr<INode>
-			{
-				return TUniquePtr<INode>(new TGetVariableNode<TDataType>(InParams.NodeName, InParams.InstanceID, InParams.VertexName, MoveTemp(InParams.InitParam)));
-			};
+					TOutputNode<TDataType> OutputPrototype(TEXT(""), FGuid(), TEXT(""));
+					OutputClass = Metasound::Frontend::GenerateClassDescription(OutputPrototype.GetMetadata(), EMetasoundFrontendClassType::Output);
 
-			FCreateMetasoundFrontendClassFunction CreateFrontendVariableClass = []() -> FMetasoundFrontendClass
-			{
-				// Create class info using prototype node
-				static TGetVariableNode<TDataType> Prototype(TEXT(""), FGuid(), TEXT(""), FLiteral());
-				return Metasound::Frontend::GenerateClassDescription(Prototype.GetMetadata(), EMetasoundFrontendClassType::Variable);
-			};
+					TGetVariableNode<TDataType> VariablePrototype(TEXT(""), FGuid(), TEXT(""), FLiteral());
+					VariableClass = Metasound::Frontend::GenerateClassDescription(VariablePrototype.GetMetadata(), EMetasoundFrontendClassType::Variable);
+				}
+				virtual ~FDataTypeRegistryEntry() = default;
 
-			FCreateOutputNodeFunction OutputNodeConstructor = [](const FOutputNodeConstructorParams& InParams) -> TUniquePtr<INode>
-			{
-				return TUniquePtr<INode>(new TOutputNode<TDataType>(InParams.NodeName, InParams.InstanceID, InParams.VertexName));
-			};
+				virtual const Frontend::FDataTypeRegistryInfo& GetDataTypeInfo() const override
+				{
+					return Info;
+				}
 
-			FCreateMetasoundFrontendClassFunction CreateFrontendOutputClass = []() -> FMetasoundFrontendClass
-			{
-				// Create class info using prototype node
-				// TODO: register output nodes with static class info.
-				static TOutputNode<TDataType> Prototype(TEXT(""), FGuid(), TEXT(""));
-				return Metasound::Frontend::GenerateClassDescription(Prototype.GetMetadata(), EMetasoundFrontendClassType::Output);
-			};
-			
-			// By default, this function should not be used, unless the preferred arg type is UObjectProxy or UObjectProxyArray, and UClassToUse should be specified.
-			FCreateAudioProxyFunction ProxyGenerator = [](UObject* InObject) -> Audio::IProxyDataPtr
-			{
-				checkNoEntry();
-				return Audio::IProxyDataPtr(nullptr);
-			};
+				virtual TSharedPtr<const Frontend::IEnumDataTypeInterface> GetEnumInterface() const override
+				{
+					return EnumInterface;
+				}
 
-			// If this datatype uses a UObject or UObject array literal, we generate a lambda to build a proxy here:
-			constexpr bool bSpecifiedUClassForProxy = !TIsSame<UClassToUse, UObject>::Value;
-			if (bSpecifiedUClassForProxy)
-			{
-				static_assert(!bSpecifiedUClassForProxy || std::is_base_of<IAudioProxyDataFactory, UClassToUse>::value, "If a Metasound Datatype uses a UObject as a literal, the UClass of that object needs to also derive from Audio::IProxyDataFactory. See USoundWave as an example.");
-				ProxyGenerator = [](UObject* InObject) -> Audio::IProxyDataPtr
+				virtual const FMetasoundFrontendClass& GetFrontendInputClass() const override
 				{
-					if (InObject)
-					{
-						IAudioProxyDataFactory* ObjectAsFactory = CastToAudioProxyDataFactory<UClassToUse>(InObject);
-						if (ensureAlways(ObjectAsFactory))
-						{
-							static FName ProxySubsystemName = TEXT("Metasound");
+					return InputClass;
+				}
 
-							Audio::FProxyDataInitParams ProxyInitParams;
-							ProxyInitParams.NameOfFeatureRequestingProxy = ProxySubsystemName;
+				virtual const FMetasoundFrontendClass& GetFrontendVariableClass() const override
+				{
+					return VariableClass;
+				}
 
-							return ObjectAsFactory->CreateNewProxyData(ProxyInitParams);
-						}
-					}
+				virtual const FMetasoundFrontendClass& GetFrontendOutputClass() const override
+				{
+					return OutputClass;
+				}
 
-					return Audio::IProxyDataPtr(nullptr);
-				};
-			}
+				virtual TUniquePtr<INode> CreateInputNode(FInputNodeConstructorParams&& InParams) const override
+				{
+					return MakeUnique<TInputNode<TDataType>>(InParams.NodeName, InParams.InstanceID, InParams.VertexName, MoveTemp(InParams.InitParam));
+				}
 
-			FCreateDataChannelFunction CreateDataChannelFunc(&FTransmissionDataChannelFactory::CreateDataChannel<TDataType>);
-
-			// Pack all of our various constructor lambdas to a single struct.
-			FDataTypeConstructorCallbacks Callbacks = 
-			{ 
-				MoveTemp(InputNodeConstructor),
-				MoveTemp(CreateFrontendInputClass),
-				MoveTemp(VariableNodeConstructor),
-				MoveTemp(CreateFrontendVariableClass),
-				MoveTemp(OutputNodeConstructor),
-				MoveTemp(CreateFrontendOutputClass),
-				MoveTemp(ProxyGenerator),
-				MoveTemp(CreateDataChannelFunc)
-			};
+				virtual TUniquePtr<INode> CreateOutputNode(FOutputNodeConstructorParams&& InParams) const override
+				{
+					return MakeUnique<TOutputNode<TDataType>>(InParams.NodeName, InParams.InstanceID, InParams.VertexName);
+				}
 
+				virtual TUniquePtr<INode> CreateVariableNode(FVariableNodeConstructorParams&& InParams) const override
+				{
+					return MakeUnique<TGetVariableNode<TDataType>>(InParams.NodeName, InParams.InstanceID, InParams.VertexName, MoveTemp(InParams.InitParam));
+				}
 
-			FDataTypeRegistryInfo RegistryInfo;
+				virtual Audio::IProxyDataPtr CreateProxy(UObject* InObject) const override
+				{
+					// Only attempt to create proxy if the `UClassToUse` is not void.
+					if constexpr (!std::is_same<UClassToUse, void>::value)
+					{
+						static_assert(std::is_base_of<IAudioProxyDataFactory, UClassToUse>::value, "If a Metasound Datatype uses a UObject as a literal, the UClass of that object needs to also derive from Audio::IProxyDataFactory. See USoundWave as an example.");
+						if (InObject)
+						{
+							IAudioProxyDataFactory* ObjectAsFactory = CastToAudioProxyDataFactory<UClassToUse>(InObject);
+							if (ensureAlways(ObjectAsFactory))
+							{
+								static FName ProxySubsystemName = TEXT("Metasound");
 
-			RegistryInfo.DataTypeName = GetMetasoundDataTypeName<TDataType>();
-			RegistryInfo.PreferredLiteralType = PreferredArgType;
+								Audio::FProxyDataInitParams ProxyInitParams;
+								ProxyInitParams.NameOfFeatureRequestingProxy = ProxySubsystemName;
 
-			RegistryInfo.bIsDefaultParsable = TIsParsable<TDataType, FLiteral::FNone>::Value;		
-			RegistryInfo.bIsBoolParsable = TIsParsable<TDataType, bool>::Value;
-			RegistryInfo.bIsIntParsable = TIsParsable<TDataType, int32>::Value;
-			RegistryInfo.bIsFloatParsable = TIsParsable<TDataType, float>::Value;
-			RegistryInfo.bIsStringParsable = TIsParsable<TDataType, FString>::Value;
-			RegistryInfo.bIsProxyParsable = TIsParsable<TDataType, const Audio::IProxyDataPtr&>::Value;
+								return ObjectAsFactory->CreateNewProxyData(ProxyInitParams);
+							}
+						}
+					}
 
-			RegistryInfo.bIsDefaultArrayParsable = TIsParsable<TDataType, TArray<FLiteral::FNone>>::Value;		
-			RegistryInfo.bIsBoolArrayParsable = TIsParsable<TDataType, TArray<bool>>::Value;
-			RegistryInfo.bIsIntArrayParsable = TIsParsable<TDataType, TArray<int32>>::Value;
-			RegistryInfo.bIsFloatArrayParsable = TIsParsable<TDataType, TArray<float>>::Value;
-			RegistryInfo.bIsStringArrayParsable = TIsParsable<TDataType, TArray<FString>>::Value;
+					return Audio::IProxyDataPtr(nullptr);
+				}
 
-			RegistryInfo.bIsProxyArrayParsable = TIsParsable<TDataType, const TArray<Audio::IProxyDataPtr>& >::Value;
-			RegistryInfo.bIsEnum = TEnumTraits<TDataType>::bIsEnum;
+				virtual TSharedPtr<IDataChannel, ESPMode::ThreadSafe> CreateDataChannel(const FOperatorSettings& InOperatorSettings) const override
+				{
+					return FTransmissionDataChannelFactory::CreateDataChannel<TDataType>(InOperatorSettings);
+				}
 
-			RegistryInfo.bIsTransmittable = TIsTransmittable<TDataType>::Value;
-			
-			RegistryInfo.ProxyGeneratorClass = UClassToUse::StaticClass();
+				virtual TUniquePtr<IDataTypeRegistryEntry> Clone() const override
+				{
+					return MakeUnique<FDataTypeRegistryEntry>();
+				}
+
+			private:
+				Frontend::FDataTypeRegistryInfo Info;
+				FMetasoundFrontendClass InputClass;
+				FMetasoundFrontendClass OutputClass;
+				FMetasoundFrontendClass VariableClass;
+				TSharedPtr<Frontend::IEnumDataTypeInterface> EnumInterface;
+			};
 
-			bool bSucceeded = FMetasoundFrontendRegistryContainer::Get()->RegisterDataType(RegistryInfo, MoveTemp(Callbacks));
+			bool bSucceeded = FMetasoundFrontendRegistryContainer::Get()->RegisterDataType(GetMetasoundDataTypeName<TDataType>(), MakeUnique<FDataTypeRegistryEntry>());
 			ensureAlwaysMsgf(bSucceeded, TEXT("Failed to register data type %s in the node registry!"), *GetMetasoundDataTypeString<TDataType>());
 
-			// If its an enum, register its data interface AFTER we've registered the data type.
-			if (TEnumTraits<TDataType>::bIsEnum)
+			if (bSucceeded)
 			{
-				if (!ensure(RegisterEnumDataTypeWithFrontend<TDataType>()))
-				{
-					return false;
-				}
+				RegisterConverterNodes<TDataType>();
+				AttemptToRegisterSendAndReceiveNodes<TDataType>();
 			}
-
-			RegisterConverterNodes<TDataType>();
-			AttemptToRegisterSendAndReceiveNodes<TDataType>();
 			
 			return bSucceeded;
 		}
 
+		/** Registers an array of a data type with the MetaSound Frontend. This allows 
+		 * an array of the data type to be used in Input, Output, Send and Receive 
+		 * nodes by informing the Frontend how to instantiate an instance. 
+		 *
+		 * @tparam TDataType - The data type to register.
+		 * @tparam PreferredArgType - The preferred constructor argument type to use when creating an instance of the data type.
+		 * @tparam UClassToUse - The preferred UObject class to use when constructing from an Audio::IProxyDataPtr. If the type is not
+		 *                       constructible with an Audio::IProxyDataPtr, then this should be void.
+		 *
+		 * @return True on success, false on failure.
+		 */
 		template<typename TDataType, ELiteralType PreferredArgType, typename UClassToUse>
 		bool RegisterDataTypeArrayWithFrontend()
 		{
@@ -444,7 +505,18 @@ namespace Metasound
 		}
 	}
 	
-	template<typename TDataType, ELiteralType PreferredArgType = ELiteralType::None, typename UClassToUse = UObject>
+	/** Registers a data type with the MetaSound Frontend. This allows the data type 
+	 * to be used in Input, Output, Send and Receive  nodes by informing the 
+	 * Frontend how to instantiate an instance. 
+	 *
+	 * @tparam TDataType - The data type to register.
+	 * @tparam PreferredArgType - The preferred constructor argument type to use when creating an instance of the data type.
+	 * @tparam UClassToUse - The preferred UObject class to use when constructing from an Audio::IProxyDataPtr. If the type is not
+	 *                       constructible with an Audio::IProxyDataPtr, then this should be void.
+	 *
+	 * @return True on success, false on failure.
+	 */
+	template<typename TDataType, ELiteralType PreferredArgType = ELiteralType::None, typename UClassToUse = void>
 	bool RegisterDataTypeWithFrontend()
 	{
 		using namespace MetasoundDataTypeRegistrationPrivate;
@@ -461,13 +533,19 @@ namespace Metasound
 	}
 
 
+	/** Registration info for a data type.
+	 *
+	 * @tparam DataType - The data type to be registered. 
+	 */
 	template<typename DataType>
 	struct TMetasoundDataTypeRegistration
 	{
 		static_assert(std::is_same<DataType, typename std::decay<DataType>::type>::value, "DataType and decayed DataType must be the same");
 		
+		// To register a data type, an input node must be able to instantiate it.
 		static constexpr bool bCanRegister = TInputNode<DataType>::bCanRegister;
 
+		// If a data type has been successfully registered, this will be true.
 		static const bool bSuccessfullyRegistered;
 	};
 }
diff --git a/Engine/Plugins/Runtime/Metasound/Source/MetasoundFrontend/Public/MetasoundFrontend.h b/Engine/Plugins/Runtime/Metasound/Source/MetasoundFrontend/Public/MetasoundFrontend.h
--- a/Engine/Plugins/Runtime/Metasound/Source/MetasoundFrontend/Public/MetasoundFrontend.h
+++ b/Engine/Plugins/Runtime/Metasound/Source/MetasoundFrontend/Public/MetasoundFrontend.h
@@ -20,9 +20,6 @@ namespace Metasound
 
 		METASOUNDFRONTEND_API FRegistryTransactionID GetCurrentRegistryTransactionID();
 
-		// Get all available nodes of any type.
-		METASOUNDFRONTEND_API TArray<FNodeClassInfo> GetAllAvailableNodeClasses(FRegistryTransactionID* OutCurrentTransactionID=nullptr);
-
 		/** Return all registry transactions that have occurred since the given transaction ID. 
 		 *
 		 * @param InTransactionID - Initial transaction ID. 
@@ -38,7 +35,7 @@ namespace Metasound
 		 *
 		 * @return Metadata for node.
 		 */
-		METASOUNDFRONTEND_API FMetasoundFrontendClassMetadata GenerateClassMetadata(const FNodeClassInfo& InInfo);
+		METASOUNDFRONTEND_API FMetasoundFrontendClassMetadata GenerateClassMetadata(const FNodeRegistryKey& InKey);
 
 		/** Generates a new FMetasoundFrontendClass from Node Metadata 
 		 *
@@ -54,7 +51,7 @@ namespace Metasound
 		 *
 		 * @return Class description for external node.
 		 */
-		METASOUNDFRONTEND_API FMetasoundFrontendClass GenerateClassDescription(const Metasound::Frontend::FNodeClassInfo& InInfo);
+		METASOUNDFRONTEND_API FMetasoundFrontendClass GenerateClassDescription(const Metasound::Frontend::FNodeRegistryKey& InKey);
 
 		/** Generates a new FMetasoundFrontendClass from Node init data
 		 *
@@ -132,6 +129,5 @@ namespace Metasound
 			// to convert between the input and output.
 			TArray<FConverterNodeInfo> PossibleConverterNodeClasses;
 		};
-
 	} // namespace Frontend
 } // namespace Metasound
diff --git a/Engine/Plugins/Runtime/Metasound/Source/MetasoundFrontend/Public/MetasoundFrontendBaseClasses.h b/Engine/Plugins/Runtime/Metasound/Source/MetasoundFrontend/Public/MetasoundFrontendBaseClasses.h
deleted file mode 100644
diff --git a/Engine/Plugins/Runtime/Metasound/Source/MetasoundFrontend/Public/MetasoundFrontendController.h b/Engine/Plugins/Runtime/Metasound/Source/MetasoundFrontend/Public/MetasoundFrontendController.h
--- a/Engine/Plugins/Runtime/Metasound/Source/MetasoundFrontend/Public/MetasoundFrontendController.h
+++ b/Engine/Plugins/Runtime/Metasound/Source/MetasoundFrontend/Public/MetasoundFrontendController.h
@@ -6,7 +6,6 @@
 #include "MetasoundAccessPtr.h"
 #include "MetasoundBuilderInterface.h"
 #include "MetasoundFrontend.h"
-#include "MetasoundFrontendBaseClasses.h"
 #include "MetasoundFrontendDocument.h"
 #include "MetasoundFrontendDocumentAccessPtr.h"
 #include "MetasoundFrontendRegistries.h"
@@ -274,14 +273,6 @@ namespace Metasound
 			 * @return True on success, false on failure. 
 			 */
 			virtual bool Disconnect() = 0;
-
-			/*
-			virtual bool IsLiteralSet() const = 0;
-			virtual FMetasoundFrontendLiteral GetLiteral() const = 0;
-			virtual bool SetLiteral(const FMetasoundFrontendLiteral& InLiteral) = 0;
-			*/
-
-
 		};
 
 		/* An INodeController provides methods for querying and manipulating a Metasound node. */
@@ -618,21 +609,13 @@ namespace Metasound
 			 */
 			virtual bool ClearLiteralForInput(const FString& InInputName, FGuid InVertexID) = 0;
 
-			/** Add a new node to this graph.
-			 *
-			 * @param InNodeClass - Info for node class.
-			 * 
-			 * @return Node handle for class. On error, an invalid handle is returned. 
-			 */
-			virtual FNodeHandle AddNode(const FNodeClassInfo& InNodeClass) = 0;
-
-			/** Add a new node to this graph.
+			/** Add a new node to this graph from the node registry.
 			 *
-			 * @param InNodeClass - Info for node class.
+			 * @param InKey - Registry key for node.
 			 * 
 			 * @return Node handle for class. On error, an invalid handle is returned. 
 			 */
-			virtual FNodeHandle AddNode(const FNodeRegistryKey& InNodeClass) = 0;
+			virtual FNodeHandle AddNode(const FNodeRegistryKey& InKey) = 0;
 
 			/** Add a new node to this graph.
 			 *
@@ -642,8 +625,8 @@ namespace Metasound
 			 */
 			virtual FNodeHandle AddNode(const FMetasoundFrontendClassMetadata& InClassMetadata) = 0;
 
-			/** Add a new node to this graph by duplicating the input node.
-			 * The new node has the same interface and node class as the input node.
+			/** Add a new node to this graph by duplicating the supplied node.
+			 * The new node has the same interface and node class as the supplied node.
 			 * This method will not duplicate node connections.
 			 *
 			 * @param InNodeController - Node to duplicate.
@@ -726,7 +709,7 @@ namespace Metasound
 			 *
 			 * @return A pointer to the found object, or nullptr if it could not be found.
 			 */
-			virtual FConstClassAccessPtr FindClass(const FNodeClassInfo& InNodeClass) const = 0; // TODO: swap for just using metadata in future. 
+			virtual FConstClassAccessPtr FindClass(const FNodeRegistryKey& InKey) const = 0; 
 
 			/** Attempts to find an existing Metasound class description corresponding
 			 * to a dependency which matches the provided class information. If the
@@ -735,7 +718,7 @@ namespace Metasound
 			 *
 			 * @return A pointer to the object, or nullptr on error.
 			 */
-			virtual FConstClassAccessPtr FindOrAddClass(const FNodeClassInfo& InNodeClass) = 0; // TODO: swap for just using metadata in future. 
+			virtual FConstClassAccessPtr FindOrAddClass(const FNodeRegistryKey& InKey) = 0; 
 
 			/** Returns an existing Metasound class description corresponding to 
 			 * a dependency which matches the provided class information.
diff --git a/Engine/Plugins/Runtime/Metasound/Source/MetasoundFrontend/Public/MetasoundFrontendDocumentAccessPtr.h b/Engine/Plugins/Runtime/Metasound/Source/MetasoundFrontend/Public/MetasoundFrontendDocumentAccessPtr.h
--- a/Engine/Plugins/Runtime/Metasound/Source/MetasoundFrontend/Public/MetasoundFrontendDocumentAccessPtr.h
+++ b/Engine/Plugins/Runtime/Metasound/Source/MetasoundFrontend/Public/MetasoundFrontendDocumentAccessPtr.h
@@ -143,6 +143,7 @@ namespace Metasound
 			FClassAccessPtr GetClassWithID(const FGuid& InID);
 			FClassAccessPtr GetClassWithMetadata(const FMetasoundFrontendClassMetadata& InMetadata);
 			FClassAccessPtr GetClassWithInfo(const FNodeClassInfo& InInfo);
+			FClassAccessPtr GetClassWithRegistryKey(const FNodeRegistryKey& InKey);
 
 			FConstGraphClassAccessPtr GetRootGraph() const;
 			FConstGraphClassAccessPtr GetSubgraphWithID(const FGuid& InID) const;
@@ -150,6 +151,7 @@ namespace Metasound
 			FConstClassAccessPtr GetClassWithID(const FGuid& InID) const;
 			FConstClassAccessPtr GetClassWithMetadata(const FMetasoundFrontendClassMetadata& InMetadata) const;
 			FConstClassAccessPtr GetClassWithInfo(const FNodeClassInfo& InInfo) const;
+			FConstClassAccessPtr GetClassWithRegistryKey(const FNodeRegistryKey& InKey) const;
 		};
 
 		class METASOUNDFRONTEND_API FConstDocumentAccessPtr : public TAccessPtr<const FMetasoundFrontendDocument>
@@ -166,6 +168,7 @@ namespace Metasound
 			FConstClassAccessPtr GetClassWithID(const FGuid& InID) const;
 			FConstClassAccessPtr GetClassWithMetadata(const FMetasoundFrontendClassMetadata& InMetadata) const;
 			FConstClassAccessPtr GetClassWithInfo(const FNodeClassInfo& InInfo) const;
+			FConstClassAccessPtr GetClassWithRegistryKey(const FNodeRegistryKey& InKey) const;
 		};
 	}
 }
diff --git a/Engine/Plugins/Runtime/Metasound/Source/MetasoundFrontend/Public/MetasoundFrontendRegistries.h b/Engine/Plugins/Runtime/Metasound/Source/MetasoundFrontend/Public/MetasoundFrontendRegistries.h
--- a/Engine/Plugins/Runtime/Metasound/Source/MetasoundFrontend/Public/MetasoundFrontendRegistries.h
+++ b/Engine/Plugins/Runtime/Metasound/Source/MetasoundFrontend/Public/MetasoundFrontendRegistries.h
@@ -14,7 +14,6 @@
 
 namespace Metasound
 {
-
 	// Base implementation for NodeConstructorCallbacks
 	struct FDefaultNodeConstructorParams
 	{
@@ -32,6 +31,7 @@ namespace Metasound
 		FString VertexName;
 		FLiteral InitParam = FLiteral::CreateInvalid();
 
+		// TODO: Remove Clone(), move variable instead. 
 		FDefaultLiteralNodeConstructorParams Clone() const
 		{
 			return
@@ -48,64 +48,7 @@ namespace Metasound
 	using FVariableNodeConstructorParams = FDefaultLiteralNodeConstructorParams;
 	using FInputNodeConstructorParams = FDefaultLiteralNodeConstructorParams;
 
-	using FCreateInputNodeFunction = TFunction<TUniquePtr<Metasound::INode>(::Metasound::FInputNodeConstructorParams&&)>;
-	using FCreateVariableNodeFunction = TFunction<TUniquePtr<Metasound::INode>(::Metasound::FVariableNodeConstructorParams&&)>;
-	using FCreateOutputNodeFunction = TFunction<TUniquePtr<Metasound::INode>(::Metasound::FOutputNodeConstructorParams&&)>;
-
-	// This function is used to create a proxy from a datatype's base uclass.
-	using FCreateAudioProxyFunction = TFunction<Audio::IProxyDataPtr(UObject*)> ;
-
-	// Creates a data channel for senders and receivers of this data type.
-	using FCreateDataChannelFunction = TFunction<TSharedPtr<IDataChannel, ESPMode::ThreadSafe>(const FOperatorSettings&)>;
-
-	using FCreateMetasoundNodeFunction = TFunction<TUniquePtr<Metasound::INode>(const Metasound::FNodeInitData&)>;
-	using FCreateMetasoundFrontendClassFunction = TFunction<FMetasoundFrontendClass()>;
-	using FIterateMetasoundFrontendClassFunction = TFunction<void(const FMetasoundFrontendClass&)>;
-
-	// Various elements that we pass to the frontend registry based on templated type traits.
-	struct FDataTypeRegistryInfo
-	{
-		// The name of the data type itself.
-		FName DataTypeName;
-
-		// What type we should default to using for literals.
-		ELiteralType PreferredLiteralType = ELiteralType::Invalid;
-
-		// This indicates the type can only be constructed with FOperatorSettings or the default constructor.
-		bool bIsDefaultParsable = false;
-
-		// These bools signify what basic literal types can be use to describe this data type.
-		bool bIsBoolParsable = false;
-		bool bIsIntParsable = false;
-		bool bIsFloatParsable = false;
-		bool bIsStringParsable = false;
-		bool bIsProxyParsable = false;
-		bool bIsDefaultArrayParsable = false;
-		bool bIsBoolArrayParsable = false;
-		bool bIsIntArrayParsable = false;
-		bool bIsFloatArrayParsable = false;
-		bool bIsStringArrayParsable = false;
-		bool bIsProxyArrayParsable = false;
-
-		// Is an TEnum wrapped enum
-		bool bIsEnum = false;
-
-		// Determines whether the type can be used with send/receive transmitters
-		bool bIsTransmittable = false;
-
-		// If this datatype was registered with a specific UClass to use to filter with, that will be used here:
-		UClass* ProxyGeneratorClass = nullptr;
-
-		FORCEINLINE bool IsArrayType() const
-		{
-			return bIsDefaultArrayParsable
-			|| bIsBoolArrayParsable
-			|| bIsIntArrayParsable
-			|| bIsFloatArrayParsable
-			|| bIsStringArrayParsable
-			|| bIsProxyArrayParsable;
-		}
-	};
+	using FIterateMetasoundFrontendClassFunction = TFunctionRef<void(const FMetasoundFrontendClass&)>;
 
 	namespace Frontend
 	{
@@ -116,80 +59,117 @@ namespace Metasound
 		// This does *not* connote that the registry key exists in the registry.
 		METASOUNDFRONTEND_API bool IsValidNodeRegistryKey(const FNodeRegistryKey& InKey);
 
-		// Struct with the basics of a node class' information,
-		// used to look up that node from our node browser functions,
-		// and also used in FGraphHandle::AddNewNode.
+		/** FNodeClassInfo contains a minimal set of information needed to find
+		 * and query node classes. 
+		 */
 		struct METASOUNDFRONTEND_API FNodeClassInfo
 		{
-			// The type for this node.
-			EMetasoundFrontendClassType NodeType = EMetasoundFrontendClassType::Invalid;
-
-			// The lookup key used for the internal node registry.
-			FNodeRegistryKey LookupKey;
-		};
-
-
-		struct METASOUNDFRONTEND_API FNodeRegistryElement
-		{
-			// This lambda can be used to get an INodeBase for this specific node class.
-			FCreateMetasoundNodeFunction CreateNode;
-
-			FCreateMetasoundFrontendClassFunction CreateFrontendClass;
-
-			FNodeRegistryElement(FCreateMetasoundNodeFunction&& InCreateNodeFunction, FCreateMetasoundFrontendClassFunction&& InCreateDescriptionFunction)
-				: CreateNode(MoveTemp(InCreateNodeFunction))
-				, CreateFrontendClass(InCreateDescriptionFunction)
-			{
-			}
+			FMetasoundFrontendClassName ClassName;
+			EMetasoundFrontendClassType Type = EMetasoundFrontendClassType::Invalid;
+			FMetasoundFrontendVersionNumber Version;
+
+			FNodeClassInfo() = default;
+			FNodeClassInfo(const FMetasoundFrontendClassMetadata& InMetadata);
+			FNodeClassInfo(const FMetasoundFrontendClass& InClass);
+			FNodeClassInfo(const FNodeClassMetadata& InMetadata);
 		};
 
-		struct METASOUNDFRONTEND_API FConverterNodeRegistryKey
+		/** INodeRegistryEntry declares the interface for a node registry entry.
+		 * Each node class in the registry must satisfy this interface. 
+		 */
+		class INodeRegistryEntry
 		{
-			// The datatype one would like to convert from.
-			FName FromDataType;
-
-			// The datatype one would like to convert to.
-			FName ToDataType;
-
-			FORCEINLINE bool operator==(const FConverterNodeRegistryKey& Other) const
-			{
-				return FromDataType == Other.FromDataType && ToDataType == Other.ToDataType;
-			}
-
-			friend uint32 GetTypeHash(const ::Metasound::Frontend::FConverterNodeRegistryKey& InKey)
-			{
-				return HashCombine(GetTypeHash(InKey.FromDataType), GetTypeHash(InKey.ToDataType));
-			}
+		public:
+			virtual ~INodeRegistryEntry() = default;
+
+			/** Return FNodeClassInfo for the node class.
+			 *
+			 * Implementations of method should avoid any expensive operations 
+			 * (e.g. loading from disk, allocating memory) as this method is called
+			 * frequently when querying nodes.
+			 */
+			virtual const FNodeClassInfo& GetClassInfo() const = 0;
+
+			/** Create a node given FDefaultNodeConstructorParams.
+			 *
+			 * If a node can be created with FDefaultNodeConstructorParams, this function
+			 * should return a valid node pointer.
+			 */
+			virtual TUniquePtr<INode> CreateNode(FDefaultNodeConstructorParams&&) const = 0;
+
+			/** Create a node given FDefaultLiteralNodeConstructorParams.
+			 *
+			 * If a node can be created with FDefaultLiteralNodeConstructorParams, this function
+			 * should return a valid node pointer.
+			 */
+			virtual TUniquePtr<INode> CreateNode(FDefaultLiteralNodeConstructorParams&&) const = 0;
+
+			/** Create a node given FNodeInitData.
+			 *
+			 * If a node can be created with FNodeInitData, this function
+			 * should return a valid node pointer.
+			 */
+			virtual TUniquePtr<INode> CreateNode(const FNodeInitData&) const = 0;
+
+			/** Return a FMetasoundFrontendClass which describes the node. */
+			virtual const FMetasoundFrontendClass& GetFrontendClass() const = 0;
+
+			/** Clone this registry entry. */
+			virtual TUniquePtr<INodeRegistryEntry> Clone() const = 0;
 		};
 
-		struct METASOUNDFRONTEND_API FConverterNodeInfo
+		/** FDataTypeRegsitryInfo contains runtime inspectable behavior of a registered
+		 * MetaSound data type.
+		 */
+		struct FDataTypeRegistryInfo
 		{
-			// If this node has multiple input pins, we use this to designate which pin should be used.
-			FVertexKey PreferredConverterInputPin;
-
-			// If this node has multiple output pins, we use this to designate which pin should be used.
-			FVertexKey PreferredConverterOutputPin;
-
-			// The key for this node in the node registry.
-			FNodeRegistryKey NodeKey;
-
-			FORCEINLINE bool operator==(const FConverterNodeInfo& Other) const
+			// The name of the data type.
+			FName DataTypeName;
+
+			// The preferred constructor argument type for creating instances of the data type.
+			ELiteralType PreferredLiteralType = ELiteralType::Invalid;
+
+			// Constructor argument support in TDataTypeLiteralFactory<TDataType>;
+			bool bIsDefaultParsable = false;
+			bool bIsBoolParsable = false;
+			bool bIsIntParsable = false;
+			bool bIsFloatParsable = false;
+			bool bIsStringParsable = false;
+			bool bIsProxyParsable = false;
+			bool bIsDefaultArrayParsable = false;
+			bool bIsBoolArrayParsable = false;
+			bool bIsIntArrayParsable = false;
+			bool bIsFloatArrayParsable = false;
+			bool bIsStringArrayParsable = false;
+			bool bIsProxyArrayParsable = false;
+
+			// Is an TEnum wrapped enum
+			bool bIsEnum = false;
+
+			// Determines whether the type can be used with send/receive transmitters
+			bool bIsTransmittable = false;
+
+			// If this datatype was registered with a specific UClass to use to filter with, that will be used here:
+			UClass* ProxyGeneratorClass = nullptr;
+
+			FORCEINLINE bool IsArrayType() const
 			{
-				return NodeKey == Other.NodeKey;
+				return bIsDefaultArrayParsable
+				|| bIsBoolArrayParsable
+				|| bIsIntArrayParsable
+				|| bIsFloatArrayParsable
+				|| bIsStringArrayParsable
+				|| bIsProxyArrayParsable;
 			}
 		};
 
-		struct METASOUNDFRONTEND_API FConverterNodeRegistryValue
-		{
-			// A list of nodes that can perform a conversion between the two datatypes described in the FConverterNodeRegistryKey for this map element.
-			TArray<FConverterNodeInfo> PotentialConverterNodes;
-		};
-
+		/** Interface for metadata of a registered MetaSound enum type. */
 		struct METASOUNDFRONTEND_API IEnumDataTypeInterface
 		{
 			using FGenericInt32Entry = TEnumEntry<int32>;
 
 			virtual ~IEnumDataTypeInterface() = default;
+
 			virtual const TArray<FGenericInt32Entry>& GetAllEntries() const = 0;
 			virtual FName GetNamespace() const = 0;
 			virtual int32 GetDefaultValue() const = 0;
@@ -230,11 +210,99 @@ namespace Metasound
 			}
 		};
 
+		/** Registry entry interface for a MetaSound data type. */
+		class IDataTypeRegistryEntry
+		{
+		public:
+			virtual ~IDataTypeRegistryEntry() = default;
+
+			/** Return the FDataTypeRegistryInfo for the data type */
+			virtual const FDataTypeRegistryInfo& GetDataTypeInfo() const = 0;
+
+			/** Return the enum interface for the data type. If the data type does
+			 * not support an enum interface, the returned pointer should be invalid.
+			 */
+			virtual TSharedPtr<const IEnumDataTypeInterface> GetEnumInterface() const = 0;
+
+			/** Return an FMetasoundFrontenClass representing an input node of the data type. */
+			virtual const FMetasoundFrontendClass& GetFrontendInputClass() const = 0;
+
+			/** Return an FMetasoundFrontenClass representing a variable node of the data type. */
+			virtual const FMetasoundFrontendClass& GetFrontendVariableClass() const = 0;
+
+			/** Return an FMetasoundFrontenClass representing an output node of the data type. */
+			virtual const FMetasoundFrontendClass& GetFrontendOutputClass() const = 0;
+
+			/** Create an input node */
+			virtual TUniquePtr<INode> CreateInputNode(FInputNodeConstructorParams&&) const = 0;
+
+			/** Create an output node */
+			virtual TUniquePtr<INode> CreateOutputNode(FOutputNodeConstructorParams&&) const = 0;
+			
+			/** Create a variable node */
+			virtual TUniquePtr<INode> CreateVariableNode(FVariableNodeConstructorParams&&) const = 0;
+
+			/** Create a proxy from a UObject. If this data type does not support
+			 * UObject proxies, return a nullptr. */
+			virtual Audio::IProxyDataPtr CreateProxy(UObject*) const = 0;
+
+			/** Create a data channel for transmission. If this data type does not
+			 * support transmission, return a nullptr. */
+			virtual TSharedPtr<IDataChannel, ESPMode::ThreadSafe> CreateDataChannel(const FOperatorSettings&) const = 0;
+
+			/** Clone this registry entry. */
+			virtual TUniquePtr<IDataTypeRegistryEntry> Clone() const = 0;
+		};
+
+		struct METASOUNDFRONTEND_API FConverterNodeRegistryKey
+		{
+			// The datatype one would like to convert from.
+			FName FromDataType;
+
+			// The datatype one would like to convert to.
+			FName ToDataType;
+
+			FORCEINLINE bool operator==(const FConverterNodeRegistryKey& Other) const
+			{
+				return FromDataType == Other.FromDataType && ToDataType == Other.ToDataType;
+			}
+
+			friend uint32 GetTypeHash(const ::Metasound::Frontend::FConverterNodeRegistryKey& InKey)
+			{
+				return HashCombine(GetTypeHash(InKey.FromDataType), GetTypeHash(InKey.ToDataType));
+			}
+		};
+
+		struct METASOUNDFRONTEND_API FConverterNodeInfo
+		{
+			// If this node has multiple input pins, we use this to designate which pin should be used.
+			FVertexKey PreferredConverterInputPin;
+
+			// If this node has multiple output pins, we use this to designate which pin should be used.
+			FVertexKey PreferredConverterOutputPin;
+
+			// The key for this node in the node registry.
+			FNodeRegistryKey NodeKey;
+
+			FORCEINLINE bool operator==(const FConverterNodeInfo& Other) const
+			{
+				return NodeKey == Other.NodeKey;
+			}
+		};
+
+		struct METASOUNDFRONTEND_API FConverterNodeRegistryValue
+		{
+			// A list of nodes that can perform a conversion between the two datatypes described in the FConverterNodeRegistryKey for this map element.
+			TArray<FConverterNodeInfo> PotentialConverterNodes;
+		};
+
+
 		using FRegistryTransactionID = int32;
 
 		/** Describes the type of transaction. */
 		enum class ETransactionType : uint8
 		{
+			// TODO: Rename to be more explicit (eg. NodeRegistration, NodeUnregistration)
 			Add,     //< Something was added to the registry.
 			Remove,  //< Something was removed from the registry.
 			Invalid
@@ -253,33 +321,15 @@ namespace Metasound
 			/** If a FNodeClassInfo added during the transaction, this will return
 			 * a non-null pointer to the FNodeClassInfo. */
 			virtual const FNodeClassInfo* GetNodeClassInfo() const = 0;
+			
+			/** If a node was registered or unregistered during this transaction,
+			 * this returns a valid FNodeRegistryKey. */
+			virtual const FNodeRegistryKey* GetNodeRegistryKey() const = 0;
 		};
 
 		using FRegistryTransactionPtr = TUniquePtr<IRegistryTransaction>;
-	} // namespace Frontend
-
-	struct FDataTypeConstructorCallbacks
-	{
-		// Constructs a TInputNode<> with the corresponding datatype.
-		FCreateInputNodeFunction CreateInputNode;
-
-		FCreateMetasoundFrontendClassFunction CreateFrontendInputClass;
-
-		// Constructs a TGetVariableNode<> with the corresponding datatype.
-		FCreateVariableNodeFunction CreateVariableNode;
-
-		FCreateMetasoundFrontendClassFunction CreateFrontendVariableClass;
-
-		// Constructs a TOutputNode<> with the corresponding datatype.
-		FCreateOutputNodeFunction CreateOutputNode;
-
-		FCreateMetasoundFrontendClassFunction CreateFrontendOutputClass;
 
-		// For datatypes that use a UObject literal or a UObject literal array, this lambda generates a literal from the corresponding UObject.
-		FCreateAudioProxyFunction CreateAudioProxy;
-
-		FCreateDataChannelFunction CreateDataChannel;
-	};
+	} // namespace Frontend
 } // namespace Metasound
 
 /**
@@ -289,29 +339,29 @@ class METASOUNDFRONTEND_API FMetasoundFrontendRegistryContainer
 {
 
 public:
+	using FNodeClassInfo = Metasound::Frontend::FNodeClassInfo;
 	using FConverterNodeRegistryKey = ::Metasound::Frontend::FConverterNodeRegistryKey;
 	using FConverterNodeRegistryValue = ::Metasound::Frontend::FConverterNodeRegistryValue;
 	using FConverterNodeInfo = ::Metasound::Frontend::FConverterNodeInfo;
 
 	using FNodeRegistryKey = Metasound::Frontend::FNodeRegistryKey;
-	using FNodeRegistryElement = Metasound::Frontend::FNodeRegistryElement;
-
-	using FDataTypeRegistryInfo = Metasound::FDataTypeRegistryInfo;
-	using FDataTypeConstructorCallbacks = ::Metasound::FDataTypeConstructorCallbacks;
 	using FNodeClassMetadata = Metasound::FNodeClassMetadata;
-	using IEnumDataTypeInterface = Metasound::Frontend::IEnumDataTypeInterface;
 
+	using FDataTypeRegistryInfo = Metasound::Frontend::FDataTypeRegistryInfo;
+	using IEnumDataTypeInterface = Metasound::Frontend::IEnumDataTypeInterface;
 
 	static FMetasoundFrontendRegistryContainer* Get();
 	static void ShutdownMetasoundFrontend();
 
 	static FNodeRegistryKey GetRegistryKey(const FNodeClassMetadata& InNodeMetadata);
 	static FNodeRegistryKey GetRegistryKey(const FMetasoundFrontendClassMetadata& InNodeMetadata);
-	static bool GetRegistryKey(const FNodeRegistryElement& InElement, FNodeRegistryKey& OutKey);
-	static bool GetFrontendClassFromRegistered(const FMetasoundFrontendClassMetadata& InMetadata, FMetasoundFrontendClass& OutClass);
-	static bool GetInputNodeClassMetadataForDataType(const FName& InDataTypeName, FMetasoundFrontendClassMetadata& OutMetadata);
-	static bool GetVariableNodeClassMetadataForDataType(const FName& InDataTypeName, FMetasoundFrontendClassMetadata& OutMetadata);
-	static bool GetOutputNodeClassMetadataForDataType(const FName& InDataTypeName, FMetasoundFrontendClassMetadata& OutMetadata);
+	static FNodeRegistryKey GetRegistryKey(const FNodeClassInfo& ClassInfo);
+
+	static bool GetFrontendClassFromRegistered(const FNodeRegistryKey& InKey, FMetasoundFrontendClass& OutClass);
+	static bool GetNodeClassInfoFromRegistered(const FNodeRegistryKey& InKey, FNodeClassInfo& OutInfo);
+	static bool GetInputNodeRegistryKeyForDataType(const FName& InDataTypeName, FNodeRegistryKey& OutKey);
+	static bool GetVariableNodeRegistryKeyForDataType(const FName& InDataTypeName, FNodeRegistryKey& OutKey);
+	static bool GetOutputNodeRegistryKeyForDataType(const FName& InDataTypeName, FNodeRegistryKey& OutKey);
 
 
 	FMetasoundFrontendRegistryContainer() = default;
@@ -327,54 +377,30 @@ public:
 	// This is called on module startup. This invokes any registration commands enqueued by our registration macros.
 	virtual void RegisterPendingNodes() = 0;
 
-	/** Register external node with the frontend.
-	 *
-	 * @param InCreateNode - Function for creating node from FNodeInitData.
-	 * @param InCreateDescription - Function for creating a FMetasoundFrontendClass.
-	 *
-	 * @return True on success.
-	 */
-	virtual FNodeRegistryKey RegisterExternalNode(Metasound::FCreateMetasoundNodeFunction&& InCreateNode, Metasound::FCreateMetasoundFrontendClassFunction&& InCreateDescription) = 0;
-	virtual bool UnregisterExternalNode(const FNodeRegistryKey& InKey) = 0;
-	virtual bool RegisterConversionNode(const FConverterNodeRegistryKey& InNodeKey, const FConverterNodeInfo& InNodeInfo) = 0;
-	virtual bool IsNodeRegistered(const FNodeRegistryKey& InKey) const = 0;
-
-	/** Return all node classes registered. */
-	virtual TArray<Metasound::Frontend::FNodeClassInfo> GetAllAvailableNodeClasses(Metasound::Frontend::FRegistryTransactionID* OutCurrentRegistryTransactionID) const = 0;
-
 	/** Return registry transactions since a given transaction id. */
 	virtual TArray<const Metasound::Frontend::IRegistryTransaction*> GetRegistryTransactionsSince(Metasound::Frontend::FRegistryTransactionID InSince, Metasound::Frontend::FRegistryTransactionID* OutCurrentRegistryTransactionID) const = 0;
 
+	/** Register a data type
+	 * @param InName - Name of data type.
+	 * @param InEntry - TUniquePtr to data type registry entry.
+	 *
+	 * @return True on success, false on failure.
+	 */
+	virtual bool RegisterDataType(const FName& InName, TUniquePtr<Metasound::Frontend::IDataTypeRegistryEntry>&& InEntry) = 0;
+
 	// Return any data types that can be used as a metasound input type or output type.
 	virtual TArray<FName> GetAllValidDataTypes() = 0;
 
-	// Iterates class types in registry.  If InClassType is set to a valid class type (optional), only iterates classes of the given type
-	virtual void IterateRegistry(Metasound::FIterateMetasoundFrontendClassFunction InIterFunc, EMetasoundFrontendClassType InClassType = EMetasoundFrontendClassType::Invalid) const = 0;
-
-	// Query for MetaSound Frontend document objects.
-	virtual bool FindFrontendClassFromRegistered(const Metasound::Frontend::FNodeClassInfo& InClassInfo, FMetasoundFrontendClass& OutClass) = 0;
-	virtual bool FindFrontendClassFromRegistered(const FMetasoundFrontendClassMetadata& InMetadata, FMetasoundFrontendClass& OutClass) = 0;
-	virtual bool FindInputNodeClassMetadataForDataType(const FName& InDataTypeName, FMetasoundFrontendClassMetadata& OutMetadata) = 0;
-	virtual bool FindVariableNodeClassMetadataForDataType(const FName& InDataTypeName, FMetasoundFrontendClassMetadata& OutMetadata) = 0;
-	virtual bool FindOutputNodeClassMetadataForDataType(const FName& InDataTypeName, FMetasoundFrontendClassMetadata& OutMetadata) = 0;
-
-	// Create a new instance of a C++ implemented node from the registry.
-	virtual TUniquePtr<Metasound::INode> ConstructInputNode(const FName& InInputType, Metasound::FInputNodeConstructorParams&& InParams) = 0;
-	virtual TUniquePtr<Metasound::INode> ConstructVariableNode(const FName& InVariableType, Metasound::FVariableNodeConstructorParams&& InParams) = 0;
-	virtual TUniquePtr<Metasound::INode> ConstructOutputNode(const FName& InOutputType, Metasound::FOutputNodeConstructorParams&& InParams) = 0;
-	virtual TUniquePtr<Metasound::INode> ConstructExternalNode(const Metasound::Frontend::FNodeRegistryKey& InRegistryKey, const Metasound::FNodeInitData& InInitData) = 0;
-
-	// Returns a list of possible nodes to use to convert from FromDataType to ToDataType.
-	// Returns an empty array if none are available.
-	virtual TArray<FConverterNodeInfo> GetPossibleConverterNodes(const FName& FromDataType, const FName& ToDataType) = 0;
-
-	virtual bool RegisterDataType(const FDataTypeRegistryInfo& InDataInfo, const FDataTypeConstructorCallbacks& InCallbacks) = 0;
-
-	virtual bool RegisterEnumDataInterface(FName InDataType, TSharedPtr<IEnumDataTypeInterface>&& InInterface) = 0;
+	// Get info about a specific data type (what kind of literals we can use, etc.)
+	// @returns false if InDataType wasn't found in the registry. 
+	virtual bool GetInfoForDataType(FName InDataType, FDataTypeRegistryInfo& OutInfo) = 0;
 
 	// Get the desired kind of literal for a given data type. Returns EConstructorArgType::Invalid if the data type couldn't be found.
 	virtual Metasound::ELiteralType GetDesiredLiteralTypeForDataType(FName InDataType) const = 0;
 
+	// Get whether we can build a literal of this specific type for InDataType.
+	virtual bool DoesDataTypeSupportLiteralType(FName InDataType, Metasound::ELiteralType InLiteralType) const = 0;
+
 	template<typename ArgType>
 	bool DoesDataTypeSupportLiteralType(FName InDataType) const
 	{
@@ -382,20 +408,64 @@ public:
 	}
 
 	virtual UClass* GetLiteralUClassForDataType(FName InDataType) const = 0;
-	virtual Metasound::FLiteral GenerateLiteralForUObject(const FName& InDataType, UObject* InObject) = 0;
-	virtual Metasound::FLiteral GenerateLiteralForUObjectArray(const FName& InDataType, TArray<UObject*> InObjectArray) = 0;
+	virtual Metasound::FLiteral CreateLiteralFromUObject(const FName& InDataType, UObject* InObject) = 0;
+	virtual Metasound::FLiteral CreateLiteralFromUObjectArray(const FName& InDataType, TArray<UObject*> InObjectArray) = 0;
 
-	// Get whether we can build a literal of this specific type for InDataType.
-	virtual bool DoesDataTypeSupportLiteralType(FName InDataType, Metasound::ELiteralType InLiteralType) const = 0;
+	// Return the enum interface for a data type. If the data type does not have 
+	// an enum interface, returns a nullptr.
+	virtual TSharedPtr<const Metasound::Frontend::IEnumDataTypeInterface> GetEnumInterfaceForDataType(FName InDataType) const = 0;
 
+	// Create a data channel for a data type. If the data type cannot be used in
+	// a data channel, return a nullptr.
+	virtual TSharedPtr<Metasound::IDataChannel, ESPMode::ThreadSafe> CreateDataChannelForDataType(const FName& InDataType, const Metasound::FOperatorSettings& InOperatorSettings) const = 0;
 
-	// Get info about a specific data type (what kind of literals we can use, etc.)
-	// @returns false if InDataType wasn't found in the registry. 
-	virtual bool GetInfoForDataType(FName InDataType, FDataTypeRegistryInfo& OutInfo) = 0;
 
-	virtual TSharedPtr<const Metasound::Frontend::IEnumDataTypeInterface> GetEnumInterfaceForDataType(FName InDataType) const = 0;
+	/** Register an external node with the frontend.
+	 *
+	 * @param InCreateNode - Function for creating node from FNodeInitData.
+	 * @param InCreateDescription - Function for creating a FMetasoundFrontendClass.
+	 *
+	 * @return A node registration key. If the registration failed, then the registry 
+	 *         key will be invalid.
+	 */
+	virtual FNodeRegistryKey RegisterNode(TUniquePtr<Metasound::Frontend::INodeRegistryEntry>&&) = 0;
 
-	virtual TSharedPtr<Metasound::IDataChannel, ESPMode::ThreadSafe> CreateDataChannelForDataType(const FName& InDataType, const Metasound::FOperatorSettings& InOperatorSettings) const = 0;
+	/** Unregister an external node from the frontend.
+	 *
+	 * @param InKey - The registration key for the node.
+	 *
+	 * @return True on success, false on failer.
+	 */
+	virtual bool UnregisterNode(const FNodeRegistryKey& InKey) = 0;
+
+	/** Returns true if the provided registry key corresponds to a valid registered node. */
+	virtual bool IsNodeRegistered(const FNodeRegistryKey& InKey) const = 0;
+
+	// Iterates class types in registry.  If InClassType is set to a valid class type (optional), only iterates classes of the given type
+	virtual void IterateRegistry(Metasound::FIterateMetasoundFrontendClassFunction InIterFunc, EMetasoundFrontendClassType InClassType = EMetasoundFrontendClassType::Invalid) const = 0;
+
+	// Query for MetaSound Frontend document objects.
+	virtual bool FindFrontendClassFromRegistered(const Metasound::Frontend::FNodeRegistryKey& InKey, FMetasoundFrontendClass& OutClass) = 0;
+	virtual bool FindNodeClassInfoFromRegistered(const Metasound::Frontend::FNodeRegistryKey& InKey, FNodeClassInfo& OutInfo) = 0;
+
+	virtual bool FindInputNodeRegistryKeyForDataType(const FName& InDataTypeName, FNodeRegistryKey& OutKey) = 0;
+	virtual bool FindVariableNodeRegistryKeyForDataType(const FName& InDataTypeName, FNodeRegistryKey& OutKey) = 0;
+	virtual bool FindOutputNodeRegistryKeyForDataType(const FName& InDataTypeName, FNodeRegistryKey& OutKey) = 0;
+
+	// Create a new instance of a C++ implemented node from the registry.
+	virtual TUniquePtr<Metasound::INode> CreateInputNode(const FName& InInputType, Metasound::FInputNodeConstructorParams&& InParams) = 0;
+	virtual TUniquePtr<Metasound::INode> CreateVariableNode(const FName& InVariableType, Metasound::FVariableNodeConstructorParams&& InParams) = 0;
+	virtual TUniquePtr<Metasound::INode> CreateOutputNode(const FName& InOutputType, Metasound::FOutputNodeConstructorParams&& InParams) = 0;
+
+	virtual TUniquePtr<Metasound::INode> CreateNode(const FNodeRegistryKey& InKey, Metasound::FDefaultNodeConstructorParams&&) const = 0;
+	virtual TUniquePtr<Metasound::INode> CreateNode(const FNodeRegistryKey& InKey, Metasound::FDefaultLiteralNodeConstructorParams&&) const = 0;
+	virtual TUniquePtr<Metasound::INode> CreateNode(const FNodeRegistryKey& InKey, const Metasound::FNodeInitData&) const = 0;
+
+	virtual bool RegisterConversionNode(const FConverterNodeRegistryKey& InNodeKey, const FConverterNodeInfo& InNodeInfo) = 0;
+
+	// Returns a list of possible nodes to use to convert from FromDataType to ToDataType.
+	// Returns an empty array if none are available.
+	virtual TArray<FConverterNodeInfo> GetPossibleConverterNodes(const FName& FromDataType, const FName& ToDataType) = 0;
 
 private:
 	static FMetasoundFrontendRegistryContainer* LazySingleton;
diff --git a/Engine/Plugins/Runtime/Metasound/Source/MetasoundFrontend/Public/MetasoundNodeRegistrationMacro.h b/Engine/Plugins/Runtime/Metasound/Source/MetasoundFrontend/Public/MetasoundNodeRegistrationMacro.h
--- a/Engine/Plugins/Runtime/Metasound/Source/MetasoundFrontend/Public/MetasoundNodeRegistrationMacro.h
+++ b/Engine/Plugins/Runtime/Metasound/Source/MetasoundFrontend/Public/MetasoundNodeRegistrationMacro.h
@@ -40,19 +40,59 @@ namespace Metasound
 
 		bAlreadyRegisteredThisDataType = true;
 
-		Frontend::FNodeRegistryKey Key = FMetasoundFrontendRegistryContainer::Get()->RegisterExternalNode(
-			[](const Metasound::FNodeInitData& InInitData) -> TUniquePtr<Metasound::INode>
+		// Node registry entry specialized to FNodeType.
+		class FNodeRegistryEntry : public Frontend::INodeRegistryEntry
+		{
+		public:
+			FNodeRegistryEntry(const Metasound::FNodeClassMetadata& InMetadata)
+			: ClassInfo(InMetadata)
+			, FrontendClass(Metasound::Frontend::GenerateClassDescription(InMetadata)) 
+			{
+			}
+
+			virtual ~FNodeRegistryEntry() = default;
+
+			virtual const Frontend::FNodeClassInfo& GetClassInfo() const override
+			{
+				return ClassInfo;
+			}
+
+			virtual TUniquePtr<INode> CreateNode(FDefaultNodeConstructorParams&& InParams) const override
+			{
+				return nullptr;
+			}
+		
+			virtual TUniquePtr<INode> CreateNode(FDefaultLiteralNodeConstructorParams&& InParams) const override
 			{
-				return TUniquePtr<Metasound::INode>(new FNodeType(InInitData));
-			},
-			[=]() -> FMetasoundFrontendClass
+				return nullptr;
+			}
+
+			virtual TUniquePtr<INode> CreateNode(const FNodeInitData& InParams) const override
 			{
-				return Metasound::Frontend::GenerateClassDescription(InMetadata);
+				return MakeUnique<FNodeType>(InParams);
 			}
-		);
-			
+
+			virtual const FMetasoundFrontendClass& GetFrontendClass() const override
+			{
+				return FrontendClass;
+			}
+
+			virtual TUniquePtr<INodeRegistryEntry> Clone() const 
+			{
+				return MakeUnique<FNodeRegistryEntry>(*this);
+			}
+
+		private:
+
+			Frontend::FNodeClassInfo ClassInfo;
+			FMetasoundFrontendClass FrontendClass;
+		};
+
+
+		Frontend::FNodeRegistryKey Key = FMetasoundFrontendRegistryContainer::Get()->RegisterNode(MakeUnique<FNodeRegistryEntry>(InMetadata));
 		const bool bSuccessfullyRegisteredNode = Frontend::IsValidNodeRegistryKey(Key);
 		ensureAlwaysMsgf(bSuccessfullyRegisteredNode, TEXT("Registering node class failed. Please check the logs."));
+
 		return bSuccessfullyRegisteredNode;
 	}
 
