UX Scaffolding for MetaSound Composition/Presets Part 2
- Leverage AssetRegistry/Manager to only load data on MetaSounds required to register (disabled until frontend registry supports this)
- Drag-And-Drop assets on MetaSound Graph
- Move ConvertToPreset to toolbar (WIP, still need to properly hide when clicked instead of closing editor blindly, has issue with refresh)

diff --git a/Engine/Plugins/Runtime/Metasound/Source/MetasoundEditor/Private/MetasoundAssetTypeActions.cpp b/Engine/Plugins/Runtime/Metasound/Source/MetasoundEditor/Private/MetasoundAssetTypeActions.cpp
--- a/Engine/Plugins/Runtime/Metasound/Source/MetasoundEditor/Private/MetasoundAssetTypeActions.cpp
+++ b/Engine/Plugins/Runtime/Metasound/Source/MetasoundEditor/Private/MetasoundAssetTypeActions.cpp
@@ -153,6 +153,7 @@ namespace Metasound
 		void FAssetTypeActions_MetaSoundSource::RegisterMenuActions()
 		{
 			using namespace AssetTypeActionsPrivate;
+			// TODO: Enable Composition
 			// RegisterPresetAction<UMetaSoundSource, UMetaSoundSourceFactory>(SourcePresetLabel, SourcePresetToolTip);
 			// RegisterPresetAction<UMetaSound, UMetaSoundFactory, UMetaSoundSource>(PresetLabel, PresetToolTip);
 		}
diff --git a/Engine/Plugins/Runtime/Metasound/Source/MetasoundEditor/Private/MetasoundEditor.cpp b/Engine/Plugins/Runtime/Metasound/Source/MetasoundEditor/Private/MetasoundEditor.cpp
--- a/Engine/Plugins/Runtime/Metasound/Source/MetasoundEditor/Private/MetasoundEditor.cpp
+++ b/Engine/Plugins/Runtime/Metasound/Source/MetasoundEditor/Private/MetasoundEditor.cpp
@@ -351,6 +351,7 @@ namespace Metasound
 			FString AssetPath = Metasound->GetPathName();
 			if (Frontend::FVersionDocument(AssetName, AssetPath).Transform(Document))
 			{
+				MetasoundAsset->RegisterGraphWithFrontend();
 				Metasound->MarkPackageDirty();
 			}
 			FGraphBuilder::SynchronizeGraph(*Metasound);
@@ -623,8 +624,8 @@ namespace Metasound
 				FToolBarExtensionDelegate::CreateLambda([this](FToolBarBuilder& ToolbarBuilder)
 				{
 				// TODO: Add OS SVD and clean this up post UE5.0 - Early Access
-// 					ToolbarBuilder.BeginSection("Utilities");
-// 					{
+ 					ToolbarBuilder.BeginSection("Utilities");
+ 					{
 // 						ToolbarBuilder.AddToolBarButton
 // 						(
 // 							FEditorCommands::Get().Import,
@@ -644,8 +645,21 @@ namespace Metasound
 // 							TAttribute<FSlateIcon>::Create([this]() { return GetExportStatusImage(); }),
 // 							"ExportMetasound"
 // 						);
-// 					}
-// 					ToolbarBuilder.EndSection();
+
+						if (!IsGraphEditable())
+						{
+							ToolbarBuilder.AddToolBarButton
+ 							(
+ 								FEditorCommands::Get().ConvertFromPreset,
+ 								NAME_None,
+ 								TAttribute<FText>(),
+ 								TAttribute<FText>(),
+ 								TAttribute<FSlateIcon>::Create([this]() { return GetExportStatusImage(); }),
+ 								"ConvertFromPreset"
+							);
+						}
+ 					}
+ 					ToolbarBuilder.EndSection();
 
 					if (Metasound->IsA<USoundBase>())
 					{
@@ -656,16 +670,20 @@ namespace Metasound
 						}
 						ToolbarBuilder.EndSection();
 					}
+
 					ToolbarBuilder.BeginSection("Utilities");
 					{
-						ToolbarBuilder.AddToolBarButton(
-							FEditorCommands::Get().EditGeneralSettings,
-							NAME_None,
-							TAttribute<FText>(),
-							TAttribute<FText>(),
-							TAttribute<FSlateIcon>::Create([this]() { return GetSettingsImage(); }),
-							"EditGeneralSettings"
-						);
+						if (Metasound->IsA<USoundBase>())
+						{
+							ToolbarBuilder.AddToolBarButton(
+								FEditorCommands::Get().EditSourceSettings,
+								NAME_None,
+								TAttribute<FText>(),
+								TAttribute<FText>(),
+								TAttribute<FSlateIcon>::Create([this]() { return GetSettingsImage(); }),
+								"EditSourceSettings"
+							);
+						}
 
 						ToolbarBuilder.AddToolBarButton(
 							FEditorCommands::Get().EditMetasoundSettings,
@@ -743,8 +761,12 @@ namespace Metasound
 				FExecuteAction::CreateSP(this, &FEditor::EditMetasoundSettings));
 
 			ToolkitCommands->MapAction(
-				Commands.EditGeneralSettings,
-				FExecuteAction::CreateSP(this, &FEditor::EditGeneralSettings));
+				Commands.EditSourceSettings,
+				FExecuteAction::CreateSP(this, &FEditor::EditSourceSettings));
+
+			ToolkitCommands->MapAction(
+				Commands.ConvertFromPreset,
+				FExecuteAction::CreateSP(this, &FEditor::ConvertFromPreset));
 
 			ToolkitCommands->MapAction(
 				FGenericCommands::Get().Delete,
@@ -923,24 +945,47 @@ namespace Metasound
 			}
 		}
 
-		void FEditor::ExecuteNode(UEdGraphNode* Node)
+		void FEditor::ExecuteNode(UEdGraphNode* InNode)
 		{
-			if (!IsPlaying())
+			using namespace Metasound;
+			using namespace Metasound::Frontend;
+
+			if (!GEditor)
 			{
 				return;
 			}
 
-			if (UMetasoundEditorGraphInputNode* InputNode = Cast<UMetasoundEditorGraphInputNode>(Node))
+			if (UMetasoundEditorGraphInputNode* InputNode = Cast<UMetasoundEditorGraphInputNode>(InNode))
+			{
+				if (IsPlaying())
+				{
+					if (UAudioComponent* PreviewComponent = GEditor->GetPreviewAudioComponent())
+					{
+						// TODO: fix how identifying the parameter to update is determined. It should not be done
+						// with a "DisplayName" but rather the vertex Guid.
+						if (TScriptInterface<IAudioParameterInterface> ParamInterface = PreviewComponent->GetParameterInterface())
+						{
+							FConstNodeHandle NodeHandle = InputNode->GetConstNodeHandle();
+							FVertexKey VertexKey = Metasound::FVertexKey(NodeHandle->GetDisplayName().ToString());
+							ParamInterface->Trigger(*VertexKey);
+						}
+					}
+				}
+			}
+			else if (UMetasoundEditorGraphExternalNode* ExternalNode = Cast<UMetasoundEditorGraphExternalNode>(InNode))
 			{
-				if (UAudioComponent* PreviewComponent = GEditor->GetPreviewAudioComponent())
+				FConstNodeHandle NodeHandle = ExternalNode->GetConstNodeHandle();
+				FNodeRegistryKey Key = FMetasoundFrontendRegistryContainer::Get()->GetRegistryKey(NodeHandle->GetClassMetadata());
+
+				FNodeClassInfo ClassInfo;
+				if (FMetasoundFrontendRegistryContainer::Get()->FindNodeClassInfoFromRegistered(Key, ClassInfo))
 				{
-					// TODO: fix how identifying the parameter to update is determined. It should not be done
-					// with a "DisplayName" but rather the vertex Guid.
-					if (TScriptInterface<IAudioParameterInterface> ParamInterface = PreviewComponent->GetParameterInterface())
+					if (ClassInfo.AssetClassID.IsValid())
 					{
-						Metasound::Frontend::FConstNodeHandle NodeHandle = InputNode->GetConstNodeHandle();
-						Metasound::FVertexKey VertexKey = Metasound::FVertexKey(NodeHandle->GetDisplayName().ToString());
-						ParamInterface->Trigger(*VertexKey);
+						if (UObject* AssetObject = ClassInfo.LoadAsset())
+						{
+							GEditor->GetEditorSubsystem<UAssetEditorSubsystem>()->OpenEditorForAsset(AssetObject);
+						}
 					}
 				}
 			}
@@ -967,7 +1012,25 @@ namespace Metasound
 			SetSelection({ Metasound });
 		}
 
-		void FEditor::EditGeneralSettings()
+		void FEditor::ConvertFromPreset()
+		{
+			check(GEditor);
+
+			if (Metasound)
+			{
+				FMetasoundAssetBase* MetasoundAsset = IMetasoundUObjectRegistry::Get().GetObjectAsAssetBase(Metasound);
+				check(MetasoundAsset);
+				MetasoundAsset->ConvertFromPreset();
+
+				// Hack until toolbar is polished up & corner text properly dynamically updates
+				if (UAssetEditorSubsystem* AssetEditorSubsystem = GEditor->GetEditorSubsystem<UAssetEditorSubsystem>())
+				{
+					AssetEditorSubsystem->CloseAllEditorsForAsset(Metasound);
+				}
+			}
+		}
+
+		void FEditor::EditSourceSettings()
 		{
 			if (UMetasoundEditorSettings* EditorSettings = GetMutableDefault<UMetasoundEditorSettings>())
 			{
@@ -1048,8 +1111,11 @@ namespace Metasound
 				GraphEditorCommands->MapAction(FEditorCommands::Get().EditMetasoundSettings,
 					FExecuteAction::CreateSP(this, &FEditor::EditMetasoundSettings));
 
-				GraphEditorCommands->MapAction(FEditorCommands::Get().EditGeneralSettings,
-					FExecuteAction::CreateSP(this, &FEditor::EditGeneralSettings));
+				if (Metasound->IsA<UMetaSoundSource>())
+				{
+					GraphEditorCommands->MapAction(FEditorCommands::Get().EditSourceSettings,
+						FExecuteAction::CreateSP(this, &FEditor::EditSourceSettings));
+				}
 
 				GraphEditorCommands->MapAction(FEditorCommands::Get().AddInput,
 					FExecuteAction::CreateSP(this, &FEditor::AddInput),
diff --git a/Engine/Plugins/Runtime/Metasound/Source/MetasoundEditor/Private/MetasoundEditor.h b/Engine/Plugins/Runtime/Metasound/Source/MetasoundEditor/Private/MetasoundEditor.h
--- a/Engine/Plugins/Runtime/Metasound/Source/MetasoundEditor/Private/MetasoundEditor.h
+++ b/Engine/Plugins/Runtime/Metasound/Source/MetasoundEditor/Private/MetasoundEditor.h
@@ -232,10 +232,13 @@ namespace Metasound
 			/** Sync the content browser to the current selection of nodes */
 			void SyncInBrowser();
 
-			/** Show the Metasound object's General settings in the Inspector */
-			void EditGeneralSettings();
+			/** Converts the MetaSound from a preset to a fully modifiable MetaSound. */
+			void ConvertFromPreset();
 
-			/** Show the Metasound object's Metasound-specific settings in the Inspector */
+			/** Show the Metasound object's Source settings in the Inspector */
+			void EditSourceSettings();
+
+			/** Show the Metasound object's settings in the Inspector */
 			void EditMetasoundSettings();
 
 			/** Add an input to the currently selected node */
diff --git a/Engine/Plugins/Runtime/Metasound/Source/MetasoundEditor/Private/MetasoundEditorCommands.cpp b/Engine/Plugins/Runtime/Metasound/Source/MetasoundEditor/Private/MetasoundEditorCommands.cpp
--- a/Engine/Plugins/Runtime/Metasound/Source/MetasoundEditor/Private/MetasoundEditorCommands.cpp
+++ b/Engine/Plugins/Runtime/Metasound/Source/MetasoundEditor/Private/MetasoundEditorCommands.cpp
@@ -21,9 +21,11 @@ namespace Metasound
 			UI_COMMAND(DeleteInput, "Delete Input", "Removes an input from the node", EUserInterfaceActionType::Button, FInputChord());
 
 			UI_COMMAND(EditMetasoundSettings, "MetaSound", "Edit MetaSound settings", EUserInterfaceActionType::Button, FInputChord());
-			UI_COMMAND(EditGeneralSettings, "General", "Edit general sound settings", EUserInterfaceActionType::Button, FInputChord());
+			UI_COMMAND(EditSourceSettings, "Source", "Edit source specific settings", EUserInterfaceActionType::Button, FInputChord());
 
 			UI_COMMAND(UpdateNodes, "Update Nodes", "Update selected node(s) that have an available update.", EUserInterfaceActionType::Button, FInputChord());
+
+			UI_COMMAND(ConvertFromPreset, "Convert From Preset", "Convert's this preset to a fully accessible MetaSound.", EUserInterfaceActionType::Button, FInputChord());
 		}
 	} // namespace Editor
 } // namespace Metasound
diff --git a/Engine/Plugins/Runtime/Metasound/Source/MetasoundEditor/Private/MetasoundEditorGraphBuilder.cpp b/Engine/Plugins/Runtime/Metasound/Source/MetasoundEditor/Private/MetasoundEditorGraphBuilder.cpp
--- a/Engine/Plugins/Runtime/Metasound/Source/MetasoundEditor/Private/MetasoundEditorGraphBuilder.cpp
+++ b/Engine/Plugins/Runtime/Metasound/Source/MetasoundEditor/Private/MetasoundEditorGraphBuilder.cpp
@@ -773,7 +773,7 @@ namespace Metasound
 			FMetasoundFrontendClassMetadata Metadata;
 
 			// 1. Set default class Metadata
-			Metadata.ClassName = FMetasoundFrontendClassName("GraphAsset", InMetaSound.GetFName(), *InMetaSound.GetPathName());
+			Metadata.ClassName = FMetasoundFrontendClassName(FName(), *FGuid::NewGuid().ToString(), FName());
 			Metadata.Version.Major = 1;
 			Metadata.Version.Minor = 0;
 			Metadata.DisplayName = FText::FromString(InMetaSound.GetName());
@@ -784,7 +784,7 @@ namespace Metasound
 			check(MetaSoundAsset);
 			MetaSoundAsset->SetMetadata(Metadata);
 
-			// 2. Set default version Metadata
+			// 2. Set default doc version Metadata
 			FDocumentHandle DocumentHandle = MetaSoundAsset->GetDocumentHandle();
 			FMetasoundFrontendDocumentMetadata DocMetadata = DocumentHandle->GetMetadata();
 			DocMetadata.Version.Number = FVersionDocument::GetMaxVersion();
@@ -824,14 +824,16 @@ namespace Metasound
 				Style.Display.Locations.Add(FGuid::NewGuid(), NewLocation);
 				NodeHandle->SetNodeStyle(Style);
 			}
+
+			MetaSoundAsset->RegisterGraphWithFrontend();
 		}
 
-		void FGraphBuilder::InitMetaSoundPreset(const UObject& InMetaSoundReferenced, UObject& InMetaSoundPreset)
+		void FGraphBuilder::InitMetaSoundPreset(UObject& InMetaSoundReferenced, UObject& InMetaSoundPreset)
 		{
 			using namespace Frontend;
 			using namespace GraphBuilderPrivate;
 
-			const FMetasoundAssetBase* MetaSoundReferencedAsset = IMetasoundUObjectRegistry::Get().GetObjectAsAssetBase(&InMetaSoundReferenced);
+			FMetasoundAssetBase* MetaSoundReferencedAsset = IMetasoundUObjectRegistry::Get().GetObjectAsAssetBase(&InMetaSoundReferenced);
 			check(MetaSoundReferencedAsset);
 
 			FConstGraphHandle ReferencedGraphHandle = MetaSoundReferencedAsset->GetRootGraphHandle();
@@ -848,6 +850,20 @@ namespace Metasound
 				return;
 			}
 
+			// Ensure referenced MetaSound is up-to-date
+			check(MetaSoundReferencedAsset);
+			FName Name = InMetaSoundReferenced.GetFName();
+			FString Path = InMetaSoundReferenced.GetPathName();
+			if (Frontend::FVersionDocument(Name, Path).Transform(MetaSoundReferencedAsset->GetDocumentHandle()))
+			{
+				InMetaSoundReferenced.MarkPackageDirty();
+			}
+
+			// Ensure asset is registered with frontend so the class
+			// can be queried when reference is added. This may modify
+			// registry properties, so must be non-const.
+			MetaSoundReferencedAsset->RegisterGraphWithFrontend();
+
 			// Initialize EdGraph, but don't conform to archetype as this would provide the
 			// default interface (inputs/outputs), which will be added below when duplicating
 			// node's interface to the preset.
@@ -872,71 +888,82 @@ namespace Metasound
 			FVector2D ExternalNodeLocation = InputNodeLocation + NodeLayout::BufferX;
 			FVector2D OutputNodeLocation = ExternalNodeLocation + NodeLayout::BufferX;
 
-			ReferencedGraphHandle->IterateConstNodes([&](FConstNodeHandle InputNode)
+			TMap<FString, FString> PresetInputToReferenceInputMap;
+			TMap<FString, FString> PresetOutputToReferenceOutputMap;
+
+			ReferencedGraphHandle->IterateConstNodes([&](FConstNodeHandle RefGraphInputNode)
 			{
 				// Required inputs are added via the prior step (archetype conformation).
 				// For these, find the corrisponding input & create an editor node only.
-				const bool bIsRequired = InputNode->IsRequired(Archetype);
+				const bool bIsRequired = RefGraphInputNode->IsRequired(Archetype);
 
-				const FText& DisplayName = InputNode->GetDisplayName();
-				const FText& Description = InputNode->GetDescription();
+				const FText& DisplayName = RefGraphInputNode->GetDisplayName();
+				const FText& Description = RefGraphInputNode->GetDescription();
 
 				// Should only ever be one
-				ensure(InputNode->GetNumInputs() == 1);
-				InputNode->IterateConstInputs([&](FConstInputHandle InputHandle)
+				ensure(RefGraphInputNode->GetNumInputs() == 1);
+				RefGraphInputNode->IterateConstInputs([&](FConstInputHandle RefGraphInputHandle)
 				{
 					FNodeHandle NodeHandle = INodeController::GetInvalidHandle();
 					if (bIsRequired)
 					{
-						NodeHandle = PresetGraphHandle->GetInputNodeWithName(InputHandle->GetName());
+						NodeHandle = PresetGraphHandle->GetInputNodeWithName(RefGraphInputHandle->GetName());
 					}
 					else
 					{
-						FGuid InputID = ReferencedGraphHandle->GetVertexIDForInputVertex(InputHandle->GetName());
+						FGuid InputID = ReferencedGraphHandle->GetVertexIDForInputVertex(RefGraphInputHandle->GetName());
 						const FMetasoundFrontendLiteral DefaultLiteral = ReferencedGraphHandle->GetDefaultInput(InputID);
 
-						const FName DataTypeName = InputHandle->GetDataType();
+						const FName DataTypeName = RefGraphInputHandle->GetDataType();
 						NodeHandle = AddInputNodeHandle(InMetaSoundPreset, DataTypeName, Description, &DefaultLiteral);
 						NodeHandle->SetDisplayName(DisplayName);
 					}
 
 					if (ensure(NodeHandle->IsValid()))
 					{
-						AddInputNode(InMetaSoundPreset, NodeHandle, InputNodeLocation, false /* bSelectNewNode */);
+						UMetasoundEditorGraphInputNode* InputNode = AddInputNode(InMetaSoundPreset, NodeHandle, InputNodeLocation, false /* bSelectNewNode */);
+						if (ensure(InputNode))
+						{
+							PresetInputToReferenceInputMap.Add(NodeHandle->GetNodeName(), RefGraphInputHandle->GetName());
+						}
 						InputNodeLocation += NodeLayout::BufferY;
 					}
 				});
 			},
 			EMetasoundFrontendClassType::Input);
 
-			ReferencedGraphHandle->IterateConstNodes([&](FConstNodeHandle OutputNode)
+			ReferencedGraphHandle->IterateConstNodes([&](FConstNodeHandle RefGraphOutputNode)
 			{
 				// Required inputs are added via the prior step (archetype conformation).
 				// For these, find the corrisponding input & create an editor node only.
-				const bool bIsRequired = OutputNode->IsRequired(Archetype);
+				const bool bIsRequired = RefGraphOutputNode->IsRequired(Archetype);
 
-				const FText& DisplayName = OutputNode->GetDisplayName();
-				const FText& Description = OutputNode->GetDescription();
+				const FText& DisplayName = RefGraphOutputNode->GetDisplayName();
+				const FText& Description = RefGraphOutputNode->GetDescription();
 
 				// Should only ever be one
-				ensure(OutputNode->GetNumOutputs() == 1);
-				OutputNode->IterateConstOutputs([&](FConstOutputHandle OutputHandle)
+				ensure(RefGraphOutputNode->GetNumOutputs() == 1);
+				RefGraphOutputNode->IterateConstOutputs([&](FConstOutputHandle RefGraphOutputHandle)
 				{
 					Frontend::FNodeHandle NodeHandle = Frontend::INodeController::GetInvalidHandle();
 					if (bIsRequired)
 					{
-						NodeHandle = PresetGraphHandle->GetOutputNodeWithName(OutputHandle->GetName());
+						NodeHandle = PresetGraphHandle->GetOutputNodeWithName(RefGraphOutputHandle->GetName());
 					}
 					else
 					{
-						const FName DataTypeName = OutputHandle->GetDataType();
+						const FName DataTypeName = RefGraphOutputHandle->GetDataType();
 						NodeHandle = AddOutputNodeHandle(InMetaSoundPreset, DataTypeName, Description);
 						NodeHandle->SetDisplayName(DisplayName);
 					}
 
 					if (ensure(NodeHandle->IsValid()))
 					{
-						AddOutputNode(InMetaSoundPreset, NodeHandle, OutputNodeLocation, false /* bSelectNewNode */);
+						UMetasoundEditorGraphOutputNode* OutputNode = AddOutputNode(InMetaSoundPreset, NodeHandle, OutputNodeLocation, false /* bSelectNewNode */);
+						if (ensure(OutputNode))
+						{
+							PresetOutputToReferenceOutputMap.Add(OutputNode->GetNodeHandle()->GetNodeName(), RefGraphOutputHandle->GetName());
+						}
 						OutputNodeLocation += NodeLayout::BufferY;
 					}
 				});
@@ -944,7 +971,10 @@ namespace Metasound
 			EMetasoundFrontendClassType::Output);
 
 			// 3. Generate a referencing node to the given referenced MetaSound.
-			const FMetasoundFrontendClassMetadata ReferencedClassMetadata = ReferencedGraphHandle->GetGraphMetadata();
+			FMetasoundFrontendClassMetadata ReferencedClassMetadata = ReferencedGraphHandle->GetGraphMetadata();
+
+			// Swap type on look-up as it will be referenced as an externally defined class relative to the new Preset asset
+			ReferencedClassMetadata.Type = EMetasoundFrontendClassType::External;
 			UMetasoundEditorGraphExternalNode* ReferencedNode = AddExternalNode(InMetaSoundPreset, ReferencedClassMetadata, ExternalNodeLocation, false /* bSelectNewNode */);
 			FNodeHandle ReferencedNodeHandle = ReferencedNode->GetNodeHandle();
 
@@ -957,7 +987,8 @@ namespace Metasound
 				ensure(InputNode->GetNumOutputs() == 1);
 				InputNode->IterateOutputs([&](FOutputHandle OutputHandle)
 				{
-					const TArray<FInputHandle> ReferenceInputs = ReferencedNodeHandle->GetInputsWithVertexName(Name);
+					const FString& InputName = PresetInputToReferenceInputMap.FindChecked(Name);
+					const TArray<FInputHandle> ReferenceInputs = ReferencedNodeHandle->GetInputsWithVertexName(InputName);
 					if (ensure(ReferenceInputs.Num() == 1))
 					{
 						ensure(ReferenceInputs[0]->Connect(*OutputHandle));
@@ -974,7 +1005,8 @@ namespace Metasound
 				ensure(OutputNode->GetNumInputs() == 1);
 				OutputNode->IterateInputs([&](FInputHandle InputHandle)
 				{
-					const TArray<FOutputHandle> ReferenceInputs = ReferencedNodeHandle->GetOutputsWithVertexName(Name);
+					const FString& OutputName = PresetOutputToReferenceOutputMap.FindChecked(Name);
+					const TArray<FOutputHandle> ReferenceInputs = ReferencedNodeHandle->GetOutputsWithVertexName(OutputName);
 					if (ensure(ReferenceInputs.Num() == 1))
 					{
 						ensure(InputHandle->Connect(*ReferenceInputs[0]));
diff --git a/Engine/Plugins/Runtime/Metasound/Source/MetasoundEditor/Private/MetasoundEditorGraphBuilder.h b/Engine/Plugins/Runtime/Metasound/Source/MetasoundEditor/Private/MetasoundEditorGraphBuilder.h
--- a/Engine/Plugins/Runtime/Metasound/Source/MetasoundEditor/Private/MetasoundEditorGraphBuilder.h
+++ b/Engine/Plugins/Runtime/Metasound/Source/MetasoundEditor/Private/MetasoundEditorGraphBuilder.h
@@ -136,7 +136,7 @@ namespace Metasound
 
 			// Initializes a MetaSound Preset using the provided ReferencedMetaSound asset's
 			// root graph as the sole, encapsulated topology.
-			static void InitMetaSoundPreset(const UObject& InMetaSoundReferenced, UObject& InMetaSoundPreset);
+			static void InitMetaSoundPreset(UObject& InMetaSoundReferenced, UObject& InMetaSoundPreset);
 
 			// Rebuilds all editor node pins based on the provided node handle's class definition.
 			static void RebuildNodePins(UMetasoundEditorGraphNode& InGraphNode);
diff --git a/Engine/Plugins/Runtime/Metasound/Source/MetasoundEditor/Private/MetasoundEditorGraphSchema.cpp b/Engine/Plugins/Runtime/Metasound/Source/MetasoundEditor/Private/MetasoundEditorGraphSchema.cpp
--- a/Engine/Plugins/Runtime/Metasound/Source/MetasoundEditor/Private/MetasoundEditorGraphSchema.cpp
+++ b/Engine/Plugins/Runtime/Metasound/Source/MetasoundEditor/Private/MetasoundEditorGraphSchema.cpp
@@ -24,6 +24,7 @@
 #include "MetasoundEditorModule.h"
 #include "MetasoundEditorSettings.h"
 #include "MetasoundFrontend.h"
+#include "MetasoundFrontendRegistries.h"
 #include "MetasoundFrontendSearchEngine.h"
 #include "ScopedTransaction.h"
 #include "Toolkits/ToolkitManager.h"
@@ -742,11 +743,30 @@ FText UMetasoundEditorGraphSchema::GetPinDisplayName(const UEdGraphPin* Pin) con
 			}
 		}
 
+		case EMetasoundFrontendClassType::External:
+		{
+			if (Pin->Direction == EGPD_Input)
+			{
+				TArray<FInputHandle> InputHandles = NodeHandle->GetInputsWithVertexName(Pin->GetName());
+				if (ensure(InputHandles.Num() > 0))
+				{
+					return InputHandles[0]->GetDisplayName();
+				}
+			}
+			else
+			{
+				TArray<FOutputHandle> OutputHandles = NodeHandle->GetOutputsWithVertexName(Pin->GetName());
+				if (ensure(OutputHandles.Num() > 0))
+				{
+					return OutputHandles[0]->GetDisplayName();
+				}
+			}
+		}
+
 		case EMetasoundFrontendClassType::Variable:
 		case EMetasoundFrontendClassType::Graph:
 			// TODO: Implement above
 
-		case EMetasoundFrontendClassType::External:
 		case EMetasoundFrontendClassType::Invalid:
 		default:
 		{
@@ -799,21 +819,77 @@ void UMetasoundEditorGraphSchema::BreakPinLinks(UEdGraphPin& TargetPin, bool bSe
 	Super::BreakPinLinks(TargetPin, bSendsNodeNotifcation);
 }
 
-void UMetasoundEditorGraphSchema::GetAssetsGraphHoverMessage(const TArray<FAssetData>& Assets, const UEdGraph* HoverGraph, FString& OutTooltipText, bool& bOutOkIcon) const
+void UMetasoundEditorGraphSchema::GetAssetsGraphHoverMessage(const TArray<FAssetData>& Assets, const UEdGraph* HoverGraph, FString& OutTooltipText, bool& OutOkIcon) const
 {
-	bOutOkIcon = false;
-
-	// TODO: Implement
+	// Enable Composition
+	//using namespace Metasound::Editor;
+
+	//OutOkIcon = true;
+
+	//if (!HoverGraph)
+	//{
+	//	OutOkIcon = false;
+	//	return;
+	//}
+
+	//OutTooltipText = TEXT("Add MetaSound reference to Graph.");
+
+	//IMetasoundEditorModule& EditorModule = FModuleManager::GetModuleChecked<IMetasoundEditorModule>("MetaSoundEditor");
+	//for (const FAssetData& Data : Assets)
+	//{
+	//	if (!EditorModule.IsMetaSoundAssetClass(Data.GetClass()->GetFName()))
+	//	{
+	//		OutOkIcon = false;
+	//		OutTooltipText = TEXT("Asset(s) must all be MetaSounds.");
+	//		break;
+	//	}
+	//}
 }
 
 void UMetasoundEditorGraphSchema::DroppedAssetsOnGraph(const TArray<FAssetData>& Assets, const FVector2D& GraphPosition, UEdGraph* Graph) const
 {
-	// TODO: Implement
+	// Enable Composition
+	//using namespace Metasound;
+	//using namespace Metasound::Frontend;
+
+	//const FScopedTransaction Transaction(LOCTEXT("DropMetaSoundOnGraph", "Drop MetaSound On Graph"));
+
+	//UMetasoundEditorGraph* MetaSoundGraph = CastChecked<UMetasoundEditorGraph>(Graph);
+	//UObject& MetaSound = MetaSoundGraph->GetMetasoundChecked();
+	//MetaSound.Modify();
+	//Graph->Modify();
+
+	//for (const FAssetData& DroppedAsset : Assets)
+	//{
+	//	if (UObject* DroppedMetaSound = DroppedAsset.GetAsset())
+	//	{
+	//		FMetasoundAssetBase* DroppedMetaSoundAsset = IMetasoundUObjectRegistry::Get().GetObjectAsAssetBase(&DroppedMetaSound);
+	//		check(DroppedMetaSoundAsset);
+
+	//		const FNodeClassInfo ClassInfo = DroppedMetaSoundAsset->GetAssetClassInfo();
+	//		const FNodeRegistryKey RegistryKey = FMetasoundFrontendRegistryContainer::Get()->GetRegistryKey(ClassInfo);
+	//		if (ensure(IsValidNodeRegistryKey(RegistryKey)))
+	//		{
+	//			if (!FMetasoundFrontendRegistryContainer::Get()->IsNodeRegistered(RegistryKey))
+	//			{
+	//				DroppedMetaSoundAsset->RegisterGraphWithFrontend();
+	//			}
+	//		}
+
+	//		const FMetasoundFrontendClassName ClassName = MetaSoundAsset->GetAssetClassInfo().ClassName;
+
+	//		FMetasoundFrontendClass Class;
+	//		if (ensure(ISearchEngine::Get().FindClassWithHighestVersion(ClassName.ToNodeClassName(), Class)))
+	//		{
+	//			Metasound::Editor::FGraphBuilder::AddExternalNode(MetaSound, Class.Metadata, GraphPosition);
+	//		}
+	//	}
+	//}
 }
 
 void UMetasoundEditorGraphSchema::DroppedAssetsOnNode(const TArray<FAssetData>& Assets, const FVector2D& GraphPosition, UEdGraphNode* Node) const
 {
-	// TODO: Implement for metasound references
+	// Still needed?
 }
 
 void UMetasoundEditorGraphSchema::GetConversionActions(FGraphActionMenuBuilder& ActionMenuBuilder, Metasound::Editor::FActionClassFilters InFilters, bool bShowSelectedActions) const
@@ -835,7 +911,7 @@ void UMetasoundEditorGraphSchema::GetConversionActions(FGraphActionMenuBuilder&
 			continue;
 		}
 
-		const FMetasoundFrontendClassMetadata& Metadata = FrontendClass.Metadata;
+		FMetasoundFrontendClassMetadata Metadata = FrontendClass.Metadata;
 		const FText Tooltip = Metadata.Author.IsEmpty()
 			? Metadata.Description
 			: FText::Format(LOCTEXT("MetasoundTooltipAuthorFormat", "{0}\nAuthor: {1}"), Metadata.Description, Metadata.Author);
@@ -850,6 +926,7 @@ void UMetasoundEditorGraphSchema::GetConversionActions(FGraphActionMenuBuilder&
 				0
 			);
 
+			Metadata.Type = EMetasoundFrontendClassType::External;
 			NewNodeAction->ClassMetadata = Metadata;
 			ActionMenuBuilder.AddAction(NewNodeAction);
 
@@ -934,7 +1011,7 @@ void UMetasoundEditorGraphSchema::GetFunctionActions(FGraphActionMenuBuilder& Ac
 			continue;
 		}
 
-		const FMetasoundFrontendClassMetadata& Metadata = FrontendClass.Metadata;
+		FMetasoundFrontendClassMetadata Metadata = FrontendClass.Metadata;
 		const FText Tooltip = Metadata.Author.IsEmpty()
 			? Metadata.Description
 			: FText::Format(LOCTEXT("MetasoundTooltipAuthorFormat", "{0}\nAuthor: {1}"), Metadata.Description, Metadata.Author);
@@ -950,6 +1027,7 @@ void UMetasoundEditorGraphSchema::GetFunctionActions(FGraphActionMenuBuilder& Ac
 				0
 			);
 
+			Metadata.Type = EMetasoundFrontendClassType::External;
 			NewNodeAction->ClassMetadata = Metadata;
 			ActionMenuBuilder.AddAction(NewNodeAction);
 		}
diff --git a/Engine/Plugins/Runtime/Metasound/Source/MetasoundEditor/Private/MetasoundEditorModule.cpp b/Engine/Plugins/Runtime/Metasound/Source/MetasoundEditor/Private/MetasoundEditorModule.cpp
--- a/Engine/Plugins/Runtime/Metasound/Source/MetasoundEditor/Private/MetasoundEditorModule.cpp
+++ b/Engine/Plugins/Runtime/Metasound/Source/MetasoundEditor/Private/MetasoundEditorModule.cpp
@@ -1,6 +1,9 @@
 // Copyright Epic Games, Inc. All Rights Reserved.
 #include "MetasoundEditorModule.h"
 
+
+#include "AssetRegistry/AssetData.h"
+#include "AssetRegistry/AssetRegistryModule.h"
 #include "AssetRegistryModule.h"
 #include "AssetTypeActions_Base.h"
 #include "Brushes/SlateImageBrush.h"
@@ -12,8 +15,8 @@
 #include "IDetailCustomization.h"
 #include "ISettingsModule.h"
 #include "Metasound.h"
-#include "MetasoundAudioBuffer.h"
 #include "MetasoundAssetTypeActions.h"
+#include "MetasoundAudioBuffer.h"
 #include "MetasoundDetailCustomization.h"
 #include "MetasoundEditorGraph.h"
 #include "MetasoundEditorGraphBuilder.h"
@@ -21,11 +24,13 @@
 #include "MetasoundEditorGraphNodeFactory.h"
 #include "MetasoundEditorGraphSchema.h"
 #include "MetasoundEditorSettings.h"
+#include "MetasoundFrontendDocument.h"
 #include "MetasoundFrontendRegistries.h"
 #include "MetasoundNodeDetailCustomization.h"
 #include "MetasoundSource.h"
 #include "MetasoundTime.h"
 #include "MetasoundTrigger.h"
+#include "MetasoundUObjectRegistry.h"
 #include "Modules/ModuleInterface.h"
 #include "Modules/ModuleManager.h"
 #include "PropertyEditorDelegates.h"
@@ -45,7 +50,7 @@ namespace Metasound
 {
 	namespace Editor
 	{
-		static const FName AssetToolName = TEXT("AssetTools");
+		static const FName AssetToolName { "AssetTools" };
 
 		template <typename T>
 		void AddAssetAction(IAssetTools& AssetTools, TArray<TSharedPtr<FAssetTypeActions_Base>>& AssetArray)
@@ -134,6 +139,52 @@ namespace Metasound
 
 		class FModule : public IMetasoundEditorModule
 		{
+			void AddOrUpdateClassRegistryAsset(const FAssetData& InAssetData)
+			{
+				using namespace Metasound;
+				using namespace Metasound::Frontend;
+
+				if (IsMetaSoundAssetClass(InAssetData.AssetClass))
+				{
+					UMetaSoundAssetSubsystem::Get().AddOrUpdateAsset(InAssetData);
+				}
+			}
+
+			void InitializeAssetClassRegistry()
+			{
+				FARFilter Filter;
+				Filter.ClassNames = MetaSoundClassNames;
+
+				FAssetRegistryModule& AssetRegistryModule = FModuleManager::LoadModuleChecked<FAssetRegistryModule>("AssetRegistry");
+				AssetRegistryModule.Get().EnumerateAssets(Filter, [this](const FAssetData& AssetData) { AddOrUpdateClassRegistryAsset(AssetData); return true; });
+				AssetRegistryModule.Get().OnAssetAdded().AddRaw(this, &FModule::AddOrUpdateClassRegistryAsset);
+				AssetRegistryModule.Get().OnAssetUpdated().AddRaw(this, &FModule::AddOrUpdateClassRegistryAsset);
+				AssetRegistryModule.Get().OnAssetRemoved().AddRaw(this, &FModule::RemoveAssetFromClassRegistry);
+				AssetRegistryModule.Get().OnAssetRenamed().AddRaw(this, &FModule::RenameAssetInClassRegistry);
+			}
+
+			void RemoveAssetFromClassRegistry(const FAssetData& InAssetData)
+			{
+				using namespace Metasound;
+				using namespace Metasound::Frontend;
+
+				if (IsMetaSoundAssetClass(InAssetData.AssetClass))
+				{
+					UMetaSoundAssetSubsystem::Get().RemoveAsset(InAssetData);
+				}
+			}
+
+			void RenameAssetInClassRegistry(const FAssetData& InAssetData, const FString& InOldObjectPath)
+			{
+				using namespace Metasound;
+				using namespace Metasound::Frontend;
+
+				if (IsMetaSoundAssetClass(InAssetData.AssetClass))
+				{
+					UMetaSoundAssetSubsystem::Get().AddOrUpdateAsset(InAssetData);
+				}
+			}
+
 			void RegisterNodeInputClasses()
 			{
 				TSubclassOf<UMetasoundEditorGraphInputLiteral> NodeClass;
@@ -162,49 +213,6 @@ namespace Metasound
 				}
 			}
 
-			const TSubclassOf<UMetasoundEditorGraphInputLiteral> FindInputLiteralClass(EMetasoundFrontendLiteralType InLiteralType) const override
-			{
-				return NodeInputClassRegistry.FindRef(InLiteralType);
-			}
-
-			virtual const FEditorDataType& FindDataType(FName InDataTypeName) const override
-			{
-				return DataTypeInfo.FindChecked(InDataTypeName);
-			}
-
-			virtual bool IsRegisteredDataType(FName InDataTypeName) const override
-			{
-				return DataTypeInfo.Contains(InDataTypeName);
-			}
-
-			virtual void IterateDataTypes(TUniqueFunction<void(const FEditorDataType&)> InDataTypeFunction) const override
-			{
-				for (const TPair<FName, FEditorDataType>& Pair : DataTypeInfo)
-				{
-					InDataTypeFunction(Pair.Value);
-				}
-			}
-
-			virtual void RegisterDataType(FName InPinCategoryName, FName InPinSubCategoryName, const FDataTypeRegistryInfo& InRegistryInfo) override
-			{
-				const bool bIsArray = InRegistryInfo.IsArrayType();
-				const EPinContainerType ContainerType = bIsArray ? EPinContainerType::Array : EPinContainerType::None;
-				FEdGraphPinType PinType(InPinCategoryName, InPinSubCategoryName, nullptr, ContainerType, false, FEdGraphTerminalType());
-				DataTypeInfo.Emplace(InRegistryInfo.DataTypeName, Editor::FEditorDataType(MoveTemp(PinType), InRegistryInfo));
-			}
-
-			static bool GetClassAssetData(const FAssetData& InAssetData, FMetasoundFrontendClassAssetTags& OutTags)
-			{
-				bool bSuccess = false;
-				bSuccess |= InAssetData.GetTagValue(GET_MEMBER_NAME_CHECKED(FMetasoundFrontendClassAssetTags, ID), OutTags.ID);
-				bSuccess |= InAssetData.GetTagValue(GET_MEMBER_NAME_CHECKED(FMetasoundFrontendClassAssetTags, Namespace), OutTags.Namespace);
-				bSuccess |= InAssetData.GetTagValue(GET_MEMBER_NAME_CHECKED(FMetasoundFrontendClassAssetTags, Name), OutTags.Name);
-				bSuccess |= InAssetData.GetTagValue(GET_MEMBER_NAME_CHECKED(FMetasoundFrontendClassAssetTags, MajorVersion), OutTags.MajorVersion);
-				bSuccess |= InAssetData.GetTagValue(GET_MEMBER_NAME_CHECKED(FMetasoundFrontendClassAssetTags, MinorVersion), OutTags.MinorVersion);
-
-				return bSuccess;
-			}
-
 			void RegisterCoreDataTypes()
 			{
 				TArray<FName> DataTypeNames = Frontend::GetAllAvailableDataTypes();
@@ -309,54 +317,59 @@ namespace Metasound
 				}
 			}
 
-			//void AddOrUpdateAssetTags(const FAssetData& InAssetData)
-			//{
-			//	FMetasoundFrontendClassAssetTags TagData;
-			//	if (GetClassAssetData(InAssetData, TagData))
-			//	{
-			//		AssetTagData.FindOrAdd(TagData.ID) = MoveTemp(TagData);
-			//	}
-			//}
-
-			//void RemoveAssetTags(const FAssetData& InAssetData)
-			//{
-			//	FMetasoundFrontendClassAssetTags TagData;
-			//	if (!GetClassAssetData(InAssetData, TagData))
-			//	{
-			//		AssetTagData.Remove(TagData.ID);
-			//	}
-			//}
-
-			void RegisterAssetData()
+			void ShutdownAssetClassRegistry()
+			{
+				if (FAssetRegistryModule* AssetRegistryModule = static_cast<FAssetRegistryModule*>(FModuleManager::Get().GetModule("AssetRegistry")))
+				{
+					AssetRegistryModule->Get().OnAssetAdded().RemoveAll(this);
+					AssetRegistryModule->Get().OnAssetUpdated().RemoveAll(this);
+					AssetRegistryModule->Get().OnAssetRemoved().RemoveAll(this);
+					AssetRegistryModule->Get().OnAssetRenamed().RemoveAll(this);
+				}
+			}
+
+			virtual const TSubclassOf<UMetasoundEditorGraphInputLiteral> FindInputLiteralClass(EMetasoundFrontendLiteralType InLiteralType) const override
+			{
+				return NodeInputClassRegistry.FindRef(InLiteralType);
+			}
+
+			virtual const FEditorDataType& FindDataType(FName InDataTypeName) const override
+			{
+				return DataTypeInfo.FindChecked(InDataTypeName);
+			}
+
+			virtual bool IsRegisteredDataType(FName InDataTypeName) const override
 			{
-				//FAssetRegistryModule& AssetRegistryModule = FModuleManager::LoadModuleChecked<FAssetRegistryModule>("AssetRegistry");
+				return DataTypeInfo.Contains(InDataTypeName);
+			}
 
-				//FARFilter Filter;
-				//Filter.ClassNames.Add(UMetaSound::StaticClass()->GetFName());
-				//Filter.ClassNames.Add(UMetaSoundSource::StaticClass()->GetFName());
+			virtual void IterateDataTypes(TUniqueFunction<void(const FEditorDataType&)> InDataTypeFunction) const override
+			{
+				for (const TPair<FName, FEditorDataType>& Pair : DataTypeInfo)
+				{
+					InDataTypeFunction(Pair.Value);
+				}
+			}
 
-				//AssetRegistryModule.Get().EnumerateAssets(Filter, [this](const FAssetData& AssetData) { AddOrUpdateAssetTags(AssetData); return true; });
-				//AssetRegistryModule.Get().OnAssetAdded().AddRaw(this, &FModule::AddOrUpdateAssetTags);
-				//AssetRegistryModule.Get().OnAssetUpdated().AddRaw(this, &FModule::AddOrUpdateAssetTags);
-				//AssetRegistryModule.Get().OnAssetRemoved().AddRaw(this, &FModule::RemoveAssetTags);
+			virtual void RegisterDataType(FName InPinCategoryName, FName InPinSubCategoryName, const FDataTypeRegistryInfo& InRegistryInfo) override
+			{
+				const bool bIsArray = InRegistryInfo.IsArrayType();
+				const EPinContainerType ContainerType = bIsArray ? EPinContainerType::Array : EPinContainerType::None;
+				FEdGraphPinType PinType(InPinCategoryName, InPinSubCategoryName, nullptr, ContainerType, false, FEdGraphTerminalType());
+				DataTypeInfo.Emplace(InRegistryInfo.DataTypeName, Editor::FEditorDataType(MoveTemp(PinType), InRegistryInfo));
 			}
 
-			void UnregisterAssetData()
+			virtual bool IsMetaSoundAssetClass(const FName InClassName) const override
 			{
-				//AssetTagData.Reset();
-
-				//if (FAssetRegistryModule* AssetRegistryModule = static_cast<FAssetRegistryModule*>(FModuleManager::Get().GetModule("AssetRegistry")))
-				//{
-				//	AssetRegistryModule->Get().OnAssetAdded().RemoveAll(this);
-				//	AssetRegistryModule->Get().OnAssetUpdated().RemoveAll(this);
-				//	AssetRegistryModule->Get().OnAssetRemoved().RemoveAll(this);
-				//}
+				return MetaSoundClassNames.Contains(InClassName);
 			}
 
 			virtual void StartupModule() override
 			{
 				// Register Metasound asset type actions
 				IAssetTools& AssetTools = FModuleManager::LoadModuleChecked<FAssetToolsModule>(AssetToolName).Get();
+
+				// TODO: Enable Composition
 				//AddAssetAction<FAssetTypeActions_MetaSound>(AssetTools, AssetActions);
 				AddAssetAction<FAssetTypeActions_MetaSoundSource>(AssetTools, AssetActions);
 
@@ -392,7 +405,6 @@ namespace Metasound
 
 				StyleSet = MakeShared<FSlateStyle>();
 
-				RegisterAssetData();
 				RegisterCoreDataTypes();
 				RegisterNodeInputClasses();
 
@@ -413,8 +425,13 @@ namespace Metasound
 					GetMutableDefault<UMetasoundEditorSettings>()
 				);
 
+				MetaSoundClassNames.Add(UMetaSound::StaticClass()->GetFName());
+				MetaSoundClassNames.Add(UMetaSoundSource::StaticClass()->GetFName());
+
+				//  TODO: Enable Composition
 				//FAssetTypeActions_MetaSound::RegisterMenuActions();
 				//FAssetTypeActions_MetaSoundSource::RegisterMenuActions();
+				// InitializeAssetClassRegistry();
 			}
 
 			virtual void ShutdownModule() override
@@ -450,13 +467,15 @@ namespace Metasound
 					GraphPanelPinFactory.Reset();
 				}
 
-				UnregisterAssetData();
+				//  TODO: Enable Composition
+				// ShutdownAssetClassRegistry();
 
 				AssetActions.Reset();
 				DataTypeInfo.Reset();
+				MetaSoundClassNames.Reset();
 			}
 
-			//TMap<FGuid, FMetasoundFrontendClassAssetTags> AssetTagData;
+			TArray<FName> MetaSoundClassNames;
 
 			TArray<TSharedPtr<FAssetTypeActions_Base>> AssetActions;
 			TMap<FName, FEditorDataType> DataTypeInfo;
diff --git a/Engine/Plugins/Runtime/Metasound/Source/MetasoundEditor/Private/MetasoundEditorModule.h b/Engine/Plugins/Runtime/Metasound/Source/MetasoundEditor/Private/MetasoundEditorModule.h
--- a/Engine/Plugins/Runtime/Metasound/Source/MetasoundEditor/Private/MetasoundEditorModule.h
+++ b/Engine/Plugins/Runtime/Metasound/Source/MetasoundEditor/Private/MetasoundEditorModule.h
@@ -39,6 +39,7 @@ namespace Metasound
 		public:
 			virtual void RegisterDataType(FName InPinCategoryName, FName InPinSubCategoryName, const FDataTypeRegistryInfo& InRegistryInfo) = 0;
 			virtual const FEditorDataType& FindDataType(FName InDataTypeName) const = 0;
+			virtual bool IsMetaSoundAssetClass(const FName InClassName) const = 0;
 			virtual bool IsRegisteredDataType(FName InDataTypeName) const = 0;
 			virtual void IterateDataTypes(TUniqueFunction<void(const FEditorDataType&)> InDataTypeFunction) const = 0;
 
diff --git a/Engine/Plugins/Runtime/Metasound/Source/MetasoundEditor/Private/MetasoundFactory.cpp b/Engine/Plugins/Runtime/Metasound/Source/MetasoundEditor/Private/MetasoundFactory.cpp
--- a/Engine/Plugins/Runtime/Metasound/Source/MetasoundEditor/Private/MetasoundFactory.cpp
+++ b/Engine/Plugins/Runtime/Metasound/Source/MetasoundEditor/Private/MetasoundFactory.cpp
@@ -22,16 +22,6 @@ namespace Metasound
 
 			if (InReferencedMetaSoundObject)
 			{
-				// Ensure referenced MetaSound is up-to-date
-				FMetasoundAssetBase* RefMetasoundAsset = IMetasoundUObjectRegistry::Get().GetObjectAsAssetBase(InReferencedMetaSoundObject);
-				check(RefMetasoundAsset);
-				FName Name = InReferencedMetaSoundObject->GetFName();
-				FString Path = InReferencedMetaSoundObject->GetPathName();
-				if (Frontend::FVersionDocument(Name, Path).Transform(RefMetasoundAsset->GetDocumentHandle()))
-				{
-					InReferencedMetaSoundObject->MarkPackageDirty();
-				}
-
 				FGraphBuilder::InitMetaSound(*MetaSoundObject, UKismetSystemLibrary::GetPlatformUserName(), false /* bConformToArchetype */);
 				FGraphBuilder::InitMetaSoundPreset(*InReferencedMetaSoundObject, *MetaSoundObject);
 			}
diff --git a/Engine/Plugins/Runtime/Metasound/Source/MetasoundEditor/Public/MetasoundEditorCommands.h b/Engine/Plugins/Runtime/Metasound/Source/MetasoundEditor/Public/MetasoundEditorCommands.h
--- a/Engine/Plugins/Runtime/Metasound/Source/MetasoundEditor/Public/MetasoundEditorCommands.h
+++ b/Engine/Plugins/Runtime/Metasound/Source/MetasoundEditor/Public/MetasoundEditorCommands.h
@@ -46,8 +46,8 @@ namespace Metasound
 			/** Removes an input from the node */
 			TSharedPtr<FUICommandInfo> DeleteInput;
 
-			/** Shows the Metasound's top-level Metasound-specific settings in the Inspector */
-			TSharedPtr<FUICommandInfo> EditGeneralSettings;
+			/** Shows the MetasoundSource's specific settings in the Inspector (if of respective type) */
+			TSharedPtr<FUICommandInfo> EditSourceSettings;
 
 			/** Shows the Metasound's top-level general object settings in the Inspector */
 			TSharedPtr<FUICommandInfo> EditMetasoundSettings;
@@ -55,6 +55,9 @@ namespace Metasound
 			// Run updates on selected nodes if required by any subset of selected nodes
 			TSharedPtr<FUICommandInfo> UpdateNodes;
 
+			// Converts a MetaSound from a restricted, preset edit state to a fully accessible MetaSound.
+			TSharedPtr<FUICommandInfo> ConvertFromPreset;
+
 			/** Initialize commands */
 			virtual void RegisterCommands() override;
 		};
diff --git a/Engine/Plugins/Runtime/Metasound/Source/MetasoundEngine/Private/Metasound.cpp b/Engine/Plugins/Runtime/Metasound/Source/MetasoundEngine/Private/Metasound.cpp
--- a/Engine/Plugins/Runtime/Metasound/Source/MetasoundEngine/Private/Metasound.cpp
+++ b/Engine/Plugins/Runtime/Metasound/Source/MetasoundEngine/Private/Metasound.cpp
@@ -26,6 +26,8 @@
 #endif // WITH_EDITORONLY_DATA
 
 #define LOCTEXT_NAMESPACE "MetaSound"
+
+
 UMetaSound::UMetaSound(const FObjectInitializer& ObjectInitializer)
 	: Super(ObjectInitializer)
 	, FMetasoundAssetBase()
@@ -36,8 +38,9 @@ UMetaSound::UMetaSound(const FObjectInitializer& ObjectInitializer)
 void UMetaSound::PreSave(FObjectPreSaveContext SaveContext)
 {
 	Super::PreSave(SaveContext);
-	UpdateAssetTags(AssetTags);
-	RegisterGraphWithFrontend();
+
+	// TODO: Enable Composition
+	// RegisterGraphWithFrontend();
 }
 
 void UMetaSound::PostEditUndo()
@@ -53,6 +56,7 @@ void UMetaSound::PostEditUndo()
 void UMetaSound::PostEditChangeProperty(FPropertyChangedEvent& InEvent)
 {
 	Super::PostEditChangeProperty(InEvent);
+	Metasound::PostEditChangeProperty(*this, InEvent);
 }
 
 #endif // WITHEDITOR
@@ -85,15 +89,16 @@ FText UMetaSound::GetDisplayName() const
 	FString TypeName = UMetaSound::StaticClass()->GetName();
 	return FMetasoundAssetBase::GetDisplayName(MoveTemp(TypeName));
 }
+
+void UMetaSound::SetRegistryAssetClassInfo(const Metasound::Frontend::FNodeClassInfo& InNodeInfo)
+{
+	Metasound::SetMetaSoundRegistryAssetClassInfo(*this, InNodeInfo);
+}
 #endif // WITH_EDITORONLY_DATA
 
-void UMetaSound::ConvertFromPreset()
+Metasound::Frontend::FNodeClassInfo UMetaSound::GetAssetClassInfo() const
 {
-	using namespace Metasound::Frontend;
-	FGraphHandle GraphHandle = GetRootGraphHandle();
-	FMetasoundFrontendGraphStyle Style = GraphHandle->GetGraphStyle();
-	Style.bIsGraphEditable = true;
-	GraphHandle->SetGraphStyle(Style);
+	return { GetDocumentChecked().RootGraph, *GetPathName() };
 }
 
 const FMetasoundFrontendArchetype& UMetaSound::GetArchetype() const
@@ -140,7 +145,7 @@ const FMetasoundFrontendArchetype& UMetaSound::GetBaseArchetype()
 	auto CreateBaseArchetype = []() -> FMetasoundFrontendArchetype
 	{
 		FMetasoundFrontendArchetype Archetype;
-		
+
 		FMetasoundFrontendEnvironmentVariable AudioDeviceHandle;
 		AudioDeviceHandle.Name = GetAudioDeviceHandleVariableName();
 		AudioDeviceHandle.Metadata.DisplayName = FText::FromString(AudioDeviceHandle.Name);
diff --git a/Engine/Plugins/Runtime/Metasound/Source/MetasoundEngine/Private/MetasoundSource.cpp b/Engine/Plugins/Runtime/Metasound/Source/MetasoundEngine/Private/MetasoundSource.cpp
--- a/Engine/Plugins/Runtime/Metasound/Source/MetasoundEngine/Private/MetasoundSource.cpp
+++ b/Engine/Plugins/Runtime/Metasound/Source/MetasoundEngine/Private/MetasoundSource.cpp
@@ -37,8 +37,6 @@ FAutoConsoleVariableRef CVarMetaSoundBlockRate(
 	TEXT("Default: 100.0f"),
 	ECVF_Default);
 
-static const FName MetasoundSourceArchetypeName = "Metasound Source";
-
 UMetaSoundSource::UMetaSoundSource(const FObjectInitializer& ObjectInitializer)
 	: Super(ObjectInitializer)
 	, FMetasoundAssetBase()
@@ -60,13 +58,9 @@ void UMetaSoundSource::PreSave(FObjectPreSaveContext SaveContext)
 
 	// TODO: Move these to be run anytime the interface changes or
 	// the node is versioned
-	UpdateAssetTags(AssetTags);
+
+	//  TODO: Enable Composition
 	// RegisterGraphWithFrontend();
-	// TODO: Post re-register, determine if interface changed and
-	// test/figure out when we should run updates on referencing
-	// assets (presets only and on breaking changes? pop-up dialog?
-	// Only done through an action user can run and just let other
-	// be broken? Fail on cook if so?)
 }
 
 void UMetaSoundSource::PostEditUndo()
@@ -85,6 +79,7 @@ void UMetaSoundSource::PostEditChangeProperty(FPropertyChangedEvent& InEvent)
 	using namespace Metasound::Frontend;
 
 	Super::PostEditChangeProperty(InEvent);
+	Metasound::PostEditChangeProperty(*this, InEvent);
 
 	if (InEvent.GetPropertyName() == GET_MEMBER_NAME_CHECKED(UMetaSoundSource, OutputFormat))
 	{
@@ -147,15 +142,16 @@ FText UMetaSoundSource::GetDisplayName() const
 	FString TypeName = UMetaSoundSource::StaticClass()->GetName();
 	return FMetasoundAssetBase::GetDisplayName(MoveTemp(TypeName));
 }
+
+void UMetaSoundSource::SetRegistryAssetClassInfo(const Metasound::Frontend::FNodeClassInfo& InNodeInfo)
+{
+	Metasound::SetMetaSoundRegistryAssetClassInfo(*this, InNodeInfo);
+}
 #endif // WITH_EDITORONLY_DATA
 
-void UMetaSoundSource::ConvertFromPreset()
+Metasound::Frontend::FNodeClassInfo UMetaSoundSource::GetAssetClassInfo() const
 {
-	using namespace Metasound::Frontend;
-	FGraphHandle GraphHandle = GetRootGraphHandle();
-	FMetasoundFrontendGraphStyle Style = GraphHandle->GetGraphStyle();
-	Style.bIsGraphEditable = true;
-	GraphHandle->SetGraphStyle(Style);
+	return { GetDocumentChecked().RootGraph, *GetPathName() };
 }
 
 bool UMetaSoundSource::IsPlayable() const
diff --git a/Engine/Plugins/Runtime/Metasound/Source/MetasoundEngine/Private/MetasoundUObjectRegistry.cpp b/Engine/Plugins/Runtime/Metasound/Source/MetasoundEngine/Private/MetasoundUObjectRegistry.cpp
--- a/Engine/Plugins/Runtime/Metasound/Source/MetasoundEngine/Private/MetasoundUObjectRegistry.cpp
+++ b/Engine/Plugins/Runtime/Metasound/Source/MetasoundEngine/Private/MetasoundUObjectRegistry.cpp
@@ -3,14 +3,218 @@
 #include "MetasoundUObjectRegistry.h"
 
 #include "Algo/Copy.h"
+#include "AssetData.h"
 #include "AssetRegistry/AssetRegistryModule.h"
 #include "Async/Async.h"
 #include "Async/TaskGraphInterfaces.h"
 #include "CoreMinimal.h"
+#include "Engine/AssetManager.h"
+#include "Metasound.h"
+#include "MetasoundAssetBase.h"
+#include "MetasoundFrontendRegistries.h"
+#include "MetasoundFrontendTransform.h"
+#include "MetasoundSource.h"
 #include "UObject/Object.h"
 
+
+namespace Metasound
+{
+	namespace AssetSubsystemPrivate
+	{
+		static bool GetAssetClassInfo(const FAssetData& InAssetData, Frontend::FNodeClassInfo& OutInfo)
+		{
+			using namespace Metasound;
+			using namespace Metasound::Frontend;
+
+			bool bSuccess = true;
+
+			OutInfo.Type = EMetasoundFrontendClassType::External;
+			OutInfo.AssetPath = InAssetData.ObjectPath;
+
+			FString AssetClassID;
+			bSuccess &= InAssetData.GetTagValue(AssetTags::AssetClassID, AssetClassID);
+			OutInfo.AssetClassID = FGuid(AssetClassID);
+			OutInfo.ClassName = FMetasoundFrontendClassName(FName(), *AssetClassID, FName());
+
+			int32 RegistryVersionMajor = 0;
+			bSuccess &= InAssetData.GetTagValue(AssetTags::RegistryVersionMajor, RegistryVersionMajor);
+			OutInfo.Version.Major = RegistryVersionMajor;
+
+			int32 RegistryVersionMinor = 0;
+			bSuccess &= InAssetData.GetTagValue(AssetTags::RegistryVersionMinor, RegistryVersionMinor);
+			OutInfo.Version.Minor = RegistryVersionMinor;
+
+#if WITH_EDITORONLY_DATA
+			auto ParseTypesString = [&](const FName AssetTag, TArray<FName>& OutTypes)
+			{
+				FString TypesString;
+				if (InAssetData.GetTagValue(AssetTag, TypesString))
+				{
+					TArray<FString> DataTypeStrings;
+					TypesString.ParseIntoArray(DataTypeStrings, *AssetTags::ArrayDelim);
+					Algo::Transform(DataTypeStrings, OutTypes, [](const FString& DataType) { return *DataType; });
+					return true;
+				}
+
+				return false;
+			};
+
+			OutInfo.InputTypes.Reset();
+			bSuccess &= ParseTypesString(AssetTags::RegistryInputTypes, OutInfo.InputTypes);
+
+			OutInfo.OutputTypes.Reset();
+			bSuccess &= ParseTypesString(AssetTags::RegistryOutputTypes, OutInfo.OutputTypes);
+#endif // WITH_EDITORONLY_DATA
+
+			return bSuccess;
+		}
+	}
+}
+
+void UMetaSoundAssetSubsystem::Initialize(FSubsystemCollectionBase& InCollection)
+{
+//  TODO: Enable Composition
+//	FCoreDelegates::OnPostEngineInit.AddUObject(this, &UMetaSoundAssetSubsystem::PostEngineInit);
+}
+
+void UMetaSoundAssetSubsystem::PostEngineInit()
+{
+	if (UAssetManager* AssetManager = UAssetManager::GetIfValid())
+	{
+		AssetManager->CallOrRegister_OnCompletedInitialScan(FSimpleMulticastDelegate::FDelegate::CreateUObject(this, &UMetaSoundAssetSubsystem::PostInitAssetScan));
+	}
+	else
+	{
+		UE_LOG(LogMetaSound, Error, TEXT("Cannot initialize MetaSoundAssetSubsystem: Enable AssetManager or disable MetaSound plugin"));
+	}
+}
+
+void UMetaSoundAssetSubsystem::PostInitAssetScan()
+{
+	UAssetManager& AssetManager = UAssetManager::Get();
+
+	FAssetManagerSearchRules Rules;
+	Rules.AssetScanPaths.Add(TEXT("/Game"));
+
+	Rules.AssetBaseClass = UMetaSound::StaticClass();
+	TArray<FAssetData> MetaSoundAssets;
+	AssetManager.SearchAssetRegistryPaths(MetaSoundAssets, Rules);
+	for (const FAssetData& AssetData : MetaSoundAssets)
+	{
+		AddOrUpdateAsset(AssetData);
+	}
+
+	Rules.AssetBaseClass = UMetaSoundSource::StaticClass();
+	TArray<FAssetData> MetaSoundSourceAssets;
+	AssetManager.SearchAssetRegistryPaths(MetaSoundSourceAssets, Rules);
+	for (const FAssetData& AssetData : MetaSoundSourceAssets)
+	{
+		AddOrUpdateAsset(AssetData);
+	}
+}
+
+void UMetaSoundAssetSubsystem::Deinitialize()
+{
+}
+
+void UMetaSoundAssetSubsystem::AddOrUpdateAsset(const FAssetData& InAssetData)
+{
+	using namespace Metasound;
+	using namespace Metasound::AssetSubsystemPrivate;
+	using namespace Metasound::Frontend;
+
+	// TODO: Set to false for editor data only builds once registering
+	// without loading asset is supported.
+	static const bool bLoadRequiredToRegisterAssetClasses = true;
+
+	FNodeClassInfo ClassInfo;
+	bool bClassInfoFound = GetAssetClassInfo(InAssetData, ClassInfo);
+	if (!bClassInfoFound || bLoadRequiredToRegisterAssetClasses)
+	{
+		UObject* Object = nullptr;
+
+		FSoftObjectPath Path(InAssetData.ObjectPath);
+		if (InAssetData.IsAssetLoaded())
+		{
+			Object = Path.ResolveObject();
+		}
+		else
+		{
+			if (!bLoadRequiredToRegisterAssetClasses)
+			{
+				UE_LOG(LogMetaSound, Warning,
+					TEXT("Failed to find serialized MetaSound asset registry data for asset '%s'. "
+						"Forcing synchronous load which increases load times. Re-save asset to avoid this."),
+					*InAssetData.ObjectPath.ToString());
+			}
+			Object = Path.TryLoad();
+		}
+
+		if (ensure(Object))
+		{
+			FMetasoundAssetBase* MetaSoundAsset = Metasound::IMetasoundUObjectRegistry::Get().GetObjectAsAssetBase(Object);
+			check(MetaSoundAsset);
+
+			FDocumentHandle Document = MetaSoundAsset->GetDocumentHandle();
+
+			// Must version to ensure registration uses the correct key,
+			// which must be based off of most up-to-date document model
+			// for safety.
+			const FName AssetName = Object->GetFName();
+			const FString AssetPath = Object->GetPathName();
+			FVersionDocument(AssetName, AssetPath).Transform(Document);
+
+			MetaSoundAsset->RegisterGraphWithFrontend();
+		}
+	}
+}
+
+void UMetaSoundAssetSubsystem::RemoveAsset(const FAssetData& InAssetData)
+{
+	using namespace Metasound;
+	using namespace Metasound::AssetSubsystemPrivate;
+	using namespace Metasound::Frontend;
+
+	FNodeClassInfo ClassInfo;
+	if (!GetAssetClassInfo(InAssetData, ClassInfo))
+	{
+		UObject* Object = nullptr;
+
+		FSoftObjectPath Path(InAssetData.ObjectPath);
+		if (InAssetData.IsAssetLoaded())
+		{
+			Object = Path.ResolveObject();
+		}
+		else
+		{
+			Object = Path.TryLoad();
+		}
+
+		if (ensure(Object))
+		{
+			FMetasoundAssetBase* MetaSoundAsset = Metasound::IMetasoundUObjectRegistry::Get().GetObjectAsAssetBase(Object);
+			check(MetaSoundAsset);
+
+			FDocumentHandle Document = MetaSoundAsset->GetDocumentHandle();
+
+			// Must version to ensure registration uses the correct key,
+			// which must be based off of most up-to-date document model
+			// for safety.
+			const FName AssetName = Object->GetFName();
+			const FString AssetPath = Object->GetPathName();
+			FVersionDocument(AssetName, AssetPath).Transform(Document);
+
+			ClassInfo = MetaSoundAsset->GetAssetClassInfo();
+		}
+	}
+
+	const FNodeRegistryKey RegistryKey = FMetasoundFrontendRegistryContainer::Get()->GetRegistryKey(ClassInfo);
+	FMetasoundFrontendRegistryContainer::Get()->UnregisterNode(RegistryKey);
+}
+
 namespace Metasound
 {
+
 	class FMetasoundUObjectRegistry : public IMetasoundUObjectRegistry
 	{
 		public:
diff --git a/Engine/Plugins/Runtime/Metasound/Source/MetasoundEngine/Public/Metasound.h b/Engine/Plugins/Runtime/Metasound/Source/MetasoundEngine/Public/Metasound.h
--- a/Engine/Plugins/Runtime/Metasound/Source/MetasoundEngine/Public/Metasound.h
+++ b/Engine/Plugins/Runtime/Metasound/Source/MetasoundEngine/Public/Metasound.h
@@ -2,17 +2,85 @@
 #pragma once
 
 #include "CoreMinimal.h"
+
 #include "EdGraph/EdGraph.h"
 #include "MetasoundAssetBase.h"
 #include "MetasoundFrontend.h"
 #include "MetasoundFrontendDocument.h"
 #include "MetasoundOperatorSettings.h"
 #include "MetasoundRouter.h"
-#include "MetasoundSource.h"
 
 #include "Metasound.generated.h"
 
 
+namespace Metasound
+{
+	namespace ConsoleVariables
+	{
+		static float BlockRate = 100.f;
+	} // namespace ConsoleVariables
+
+	template <typename TMetaSoundObject>
+	void PostEditChangeProperty(TMetaSoundObject& InMetaSound, FPropertyChangedEvent& InEvent)
+	{
+		// TODO: Update registry info here if interface has changed,
+		// potentially enabling removal of PreSave registration
+	}
+
+#if WITH_EDITORONLY_DATA
+	template <typename TMetaSoundObject>
+	void SetMetaSoundRegistryAssetClassInfo(TMetaSoundObject& InMetaSound, const Metasound::Frontend::FNodeClassInfo& InClassInfo)
+	{
+		using namespace Metasound;
+
+		check(AssetTags::AssetClassID == GET_MEMBER_NAME_CHECKED(TMetaSoundObject, AssetClassID));
+		check(AssetTags::RegistryInputTypes == GET_MEMBER_NAME_CHECKED(TMetaSoundObject, RegistryInputTypes));
+		check(AssetTags::RegistryOutputTypes == GET_MEMBER_NAME_CHECKED(TMetaSoundObject, RegistryOutputTypes));
+		check(AssetTags::RegistryVersionMajor == GET_MEMBER_NAME_CHECKED(TMetaSoundObject, RegistryVersionMajor));
+		check(AssetTags::RegistryVersionMinor == GET_MEMBER_NAME_CHECKED(TMetaSoundObject, RegistryVersionMinor));
+
+		bool bMarkDirty = InMetaSound.AssetClassID != InClassInfo.AssetClassID;
+		bMarkDirty |= InMetaSound.RegistryVersionMajor != InClassInfo.Version.Major;
+		bMarkDirty |= InMetaSound.RegistryVersionMinor != InClassInfo.Version.Minor;
+
+		InMetaSound.AssetClassID = InClassInfo.AssetClassID;
+		InMetaSound.RegistryVersionMajor = InClassInfo.Version.Major;
+		InMetaSound.RegistryVersionMinor = InClassInfo.Version.Minor;
+
+		{
+			TArray<FString> InputTypes;
+			Algo::Transform(InClassInfo.InputTypes, InputTypes, [](const FName& Name) { return Name.ToString(); });
+
+			const FString TypeString = FString::Join(InputTypes, *AssetTags::ArrayDelim);
+			bMarkDirty |= InMetaSound.RegistryInputTypes != TypeString;
+			InMetaSound.RegistryInputTypes = TypeString;
+		}
+
+		{
+			TArray<FString> OutputTypes;
+			Algo::Transform(InClassInfo.OutputTypes, OutputTypes, [](const FName& Name) { return Name.ToString(); });
+
+			const FString TypeString = FString::Join(OutputTypes, *AssetTags::ArrayDelim);
+			bMarkDirty |= InMetaSound.RegistryOutputTypes != TypeString;
+			InMetaSound.RegistryOutputTypes = TypeString;
+		}
+	}
+#endif // WITH_EDITORONLY_DATA
+} // namespace Metasound
+
+UCLASS()
+class METASOUNDENGINE_API UMetasoundEditorGraphBase : public UEdGraph
+{
+	GENERATED_BODY()
+
+public:
+	virtual bool IsEditorOnly() const override { return true; }
+	virtual bool NeedsLoadForEditorGame() const override { return false; }
+
+	virtual void Synchronize() { }
+};
+
+
 /**
  * This asset type is used for Metasound assets that can only be used as nodes in other Metasound graphs.
  * Because of this, they contain no required inputs or outputs.
@@ -27,9 +95,6 @@ protected:
 	FMetasoundFrontendDocument RootMetaSoundDocument;
 
 #if WITH_EDITORONLY_DATA
-	UPROPERTY(AssetRegistrySearchable)
-	FMetasoundFrontendClassAssetTags AssetTags;
-
 	UPROPERTY()
 	UMetasoundEditorGraphBase* Graph;
 #endif // WITH_EDITORONLY_DATA
@@ -37,7 +102,25 @@ protected:
 public:
 	UMetaSound(const FObjectInitializer& ObjectInitializer);
 
+	UPROPERTY(AssetRegistrySearchable)
+	FGuid AssetClassID;
+
 #if WITH_EDITORONLY_DATA
+	UPROPERTY(AssetRegistrySearchable)
+	FString RegistryInputTypes;
+
+	UPROPERTY(AssetRegistrySearchable)
+	FString RegistryOutputTypes;
+
+	UPROPERTY(AssetRegistrySearchable)
+	int32 RegistryVersionMajor = 0;
+
+	UPROPERTY(AssetRegistrySearchable)
+	int32 RegistryVersionMinor = 0;
+
+	// Sets Asset Registry Metadata associated with this MetaSound
+	virtual void SetRegistryAssetClassInfo(const Metasound::Frontend::FNodeClassInfo& InClassInfo) override;
+
 	// Returns document name (for editor purposes, and avoids making document public for edit
 	// while allowing editor to reference directly)
 	static FName GetDocumentPropertyName()
@@ -71,13 +154,10 @@ public:
 	virtual void PostEditChangeProperty(FPropertyChangedEvent& InEvent) override;
 #endif // WITH_EDITOR
 
-	virtual const FMetasoundFrontendArchetype& GetArchetype() const override;
+	// Returns Asset Metadata associated with this MetaSound
+	virtual Metasound::Frontend::FNodeClassInfo GetAssetClassInfo() const override;
 
-	// If set to be a preset, converts to a full-access MetaSound,
-	// removing edit restrictions and excluding it from automatic
-	// interface versioning.
-	UFUNCTION(Category = Metasound, meta = (CallInEditor = "true"))
-	void ConvertFromPreset();
+	virtual const FMetasoundFrontendArchetype& GetArchetype() const override;
 
 	UObject* GetOwningAsset() override
 	{
diff --git a/Engine/Plugins/Runtime/Metasound/Source/MetasoundEngine/Public/MetasoundSource.h b/Engine/Plugins/Runtime/Metasound/Source/MetasoundEngine/Public/MetasoundSource.h
--- a/Engine/Plugins/Runtime/Metasound/Source/MetasoundEngine/Public/MetasoundSource.h
+++ b/Engine/Plugins/Runtime/Metasound/Source/MetasoundEngine/Public/MetasoundSource.h
@@ -3,6 +3,7 @@
 
 #include "CoreMinimal.h"
 #include "EdGraph/EdGraph.h"
+#include "Metasound.h"
 #include "MetasoundAssetBase.h"
 #include "MetasoundFrontend.h"
 #include "MetasoundFrontendDocument.h"
@@ -10,18 +11,11 @@
 #include "MetasoundOperatorSettings.h"
 #include "MetasoundRouter.h"
 #include "Sound/SoundWaveProcedural.h"
+#include "UObject/MetaData.h"
 
 #include "MetasoundSource.generated.h"
 
 
-namespace Metasound
-{
-	namespace ConsoleVariables
-	{
-		static float BlockRate = 100.f;
-	}
-} // namespace Metasound
-
 /** Declares the output audio format of the UMetaSoundSource */
 UENUM()
 enum class EMetasoundSourceAudioFormat : uint8
@@ -35,17 +29,6 @@ enum class EMetasoundSourceAudioFormat : uint8
 	COUNT UMETA(Hidden)
 };
 
-UCLASS()
-class METASOUNDENGINE_API UMetasoundEditorGraphBase : public UEdGraph
-{
-	GENERATED_BODY()
-
-public:
-	virtual bool IsEditorOnly() const override { return true; }
-	virtual bool NeedsLoadForEditorGame() const override { return false; }
-
-	virtual void Synchronize() { }
-};
 
 /**
  * This Metasound type can be played as an audio source.
@@ -60,9 +43,6 @@ protected:
 	FMetasoundFrontendDocument RootMetasoundDocument;
 
 #if WITH_EDITORONLY_DATA
-	UPROPERTY(AssetRegistrySearchable)
-	FMetasoundFrontendClassAssetTags AssetTags;
-
 	UPROPERTY()
 	UMetasoundEditorGraphBase* Graph;
 #endif // WITH_EDITORONLY_DATA
@@ -74,7 +54,25 @@ public:
 	UPROPERTY(EditAnywhere, BlueprintReadOnly, Category = Metasound)
 	EMetasoundSourceAudioFormat OutputFormat;
 
+	UPROPERTY(AssetRegistrySearchable)
+	FGuid AssetClassID;
+
 #if WITH_EDITORONLY_DATA
+	UPROPERTY(AssetRegistrySearchable)
+	FString RegistryInputTypes;
+
+	UPROPERTY(AssetRegistrySearchable)
+	FString RegistryOutputTypes;
+
+	UPROPERTY(AssetRegistrySearchable)
+	int32 RegistryVersionMajor = 0;
+
+	UPROPERTY(AssetRegistrySearchable)
+	int32 RegistryVersionMinor = 0;
+
+	// Sets Asset Registry Metadata associated with this MetaSoundSource
+	virtual void SetRegistryAssetClassInfo(const Metasound::Frontend::FNodeClassInfo& InNodeInfo) override;
+
 	// Returns document name (for editor purposes, and avoids making document public for edit
 	// while allowing editor to reference directly)
 	static FName GetDocumentPropertyName()
@@ -123,13 +121,10 @@ public:
 	virtual void PostEditChangeProperty(FPropertyChangedEvent& InEvent) override;
 #endif // WITH_EDITOR
 
-	virtual const FMetasoundFrontendArchetype& GetArchetype() const override;
+	// Returns Asset Metadata associated with this MetaSoundSource
+	virtual Metasound::Frontend::FNodeClassInfo GetAssetClassInfo() const override;
 
-	// If set to be a preset, converts to a full-access MetaSound,
-	// removing edit restrictions and excluding it from automatic
-	// interface versioning.
-	UFUNCTION(Category = Metasound, meta = (CallInEditor = "true"))
-	void ConvertFromPreset();
+	virtual const FMetasoundFrontendArchetype& GetArchetype() const override;
 
 	UObject* GetOwningAsset() override
 	{
diff --git a/Engine/Plugins/Runtime/Metasound/Source/MetasoundEngine/Public/MetasoundUObjectRegistry.h b/Engine/Plugins/Runtime/Metasound/Source/MetasoundEngine/Public/MetasoundUObjectRegistry.h
--- a/Engine/Plugins/Runtime/Metasound/Source/MetasoundEngine/Public/MetasoundUObjectRegistry.h
+++ b/Engine/Plugins/Runtime/Metasound/Source/MetasoundEngine/Public/MetasoundUObjectRegistry.h
@@ -1,10 +1,35 @@
 // Copyright Epic Games, Inc. All Rights Reserved.
-
 #pragma once
 
 #include "CoreMinimal.h"
 #include "MetasoundAssetBase.h"
 #include "UObject/Object.h"
+#include "MetasoundUObjectRegistry.generated.h"
+
+
+/** The subsystem in charge of the MetaSound asset registry */
+UCLASS()
+class METASOUNDENGINE_API UMetaSoundAssetSubsystem : public UEngineSubsystem
+{
+	GENERATED_BODY()
+
+public:
+	virtual void Initialize(FSubsystemCollectionBase& InCollection) override;
+	virtual void Deinitialize() override;
+
+	static UMetaSoundAssetSubsystem& Get()
+	{
+		check(GEngine);
+		return *GEngine->GetEngineSubsystem<UMetaSoundAssetSubsystem>();
+	}
+
+	void AddOrUpdateAsset(const FAssetData& InAssetData);
+	void RemoveAsset(const FAssetData& InAssetData);
+
+protected:
+	void PostEngineInit();
+	void PostInitAssetScan();
+};
 
 namespace Metasound
 {
@@ -195,4 +220,4 @@ namespace Metasound
 			 */
 			virtual const FMetasoundAssetBase* GetObjectAsAssetBase(const UObject* InObject) const = 0;
 	};
-}
+} // namespace Metasound
diff --git a/Engine/Plugins/Runtime/Metasound/Source/MetasoundFrontend/Private/MetasoundAssetBase.cpp b/Engine/Plugins/Runtime/Metasound/Source/MetasoundFrontend/Private/MetasoundAssetBase.cpp
--- a/Engine/Plugins/Runtime/Metasound/Source/MetasoundFrontend/Private/MetasoundAssetBase.cpp
+++ b/Engine/Plugins/Runtime/Metasound/Source/MetasoundFrontend/Private/MetasoundAssetBase.cpp
@@ -3,6 +3,7 @@
 #include "MetasoundAssetBase.h"
 
 #include "Algo/AnyOf.h"
+#include "Algo/Transform.h"
 #include "HAL/FileManager.h"
 #include "Internationalization/Text.h"
 #include "IStructSerializerBackend.h"
@@ -15,9 +16,27 @@
 #include "MetasoundLog.h"
 #include "MetasoundReceiveNode.h"
 #include "StructSerializer.h"
+#include "UObject/MetaData.h"
 
 #define LOCTEXT_NAMESPACE "MetaSound"
 
+namespace Metasound
+{
+	namespace AssetTags
+	{
+		const FString ArrayDelim = TEXT(",");
+
+		const FName AssetClassID = "AssetClassID";
+		const FName RegistryVersionMajor = "RegistryVersionMajor";
+		const FName RegistryVersionMinor = "RegistryVersionMinor";
+
+#if WITH_EDITORONLY_DATA
+		const FName RegistryInputTypes = "RegistryInputTypes";
+		const FName RegistryOutputTypes = "RegistryOutputTypes";
+#endif // WITH_EDITORONLY_DATA
+	} // namespace AssetTags
+} // namespace Metasound
+
 const FString FMetasoundAssetBase::FileExtension(TEXT(".metasound"));
 
 void FMetasoundAssetBase::RegisterGraphWithFrontend()
@@ -27,23 +46,25 @@ void FMetasoundAssetBase::RegisterGraphWithFrontend()
 
 	FConstDocumentAccessPtr DocumentPtr = GetDocument();
 	FString AssetName;
+	FString AssetPath;
 	if (const UObject* OwningAsset = GetOwningAsset())
 	{
 		AssetName = OwningAsset->GetName();
+		AssetPath = OwningAsset->GetPathName();
 	}
 
 	// Registers node by copying document. Updates to document require re-registration.
 	class FNodeRegistryEntry : public INodeRegistryEntry
 	{
 	public:
-		FNodeRegistryEntry(const FString& InName, const FMetasoundFrontendDocument& InDocument)
+		FNodeRegistryEntry(const FString& InName, const FMetasoundFrontendDocument& InDocument, FName InAssetPath)
 		: Name(InName)
 		, Document(InDocument)
 		{
 			// Copy frontend class to preserve original document.
 			FrontendClass = Document.RootGraph;
 			FrontendClass.Metadata.Type = EMetasoundFrontendClassType::External;
-			ClassInfo = FNodeClassInfo{FrontendClass};
+			ClassInfo = FNodeClassInfo(Document.RootGraph, InAssetPath);
 		}
 
 		virtual ~FNodeRegistryEntry() = default;
@@ -75,7 +96,7 @@ void FMetasoundAssetBase::RegisterGraphWithFrontend()
 
 		virtual TUniquePtr<INodeRegistryEntry> Clone() const override
 		{
-			return MakeUnique<FNodeRegistryEntry>(Name, Document);
+			return MakeUnique<FNodeRegistryEntry>(Name, Document, ClassInfo.AssetPath);
 		}
 
 	private:
@@ -86,19 +107,48 @@ void FMetasoundAssetBase::RegisterGraphWithFrontend()
 		FNodeClassInfo ClassInfo;
 	};
 
+	FNodeClassInfo ClassInfo = GetAssetClassInfo();
+	FNodeRegistryKey RegistryKey = FMetasoundFrontendRegistryContainer::Get()->GetRegistryKey(ClassInfo);
 
 	if (IsValidNodeRegistryKey(RegistryKey))
 	{
 		// Unregister prior version if it exists.
-		ensure(FMetasoundFrontendRegistryContainer::Get()->UnregisterNode(RegistryKey));
+		if (bHasRegistered)
+		{
+			ensure(FMetasoundFrontendRegistryContainer::Get()->UnregisterNode(RegistryKey));
+		}
 	}
-	
+
 	if (const FMetasoundFrontendDocument* Doc = DocumentPtr.Get())
 	{
-		RegistryKey = FMetasoundFrontendRegistryContainer::Get()->RegisterNode(MakeUnique<FNodeRegistryEntry>(AssetName, *Doc));
+		RegistryKey = FMetasoundFrontendRegistryContainer::Get()->RegisterNode(MakeUnique<FNodeRegistryEntry>(AssetName, *Doc, ClassInfo.AssetPath));
 	}
 
-	if (!IsValidNodeRegistryKey(RegistryKey))
+	if (IsValidNodeRegistryKey(RegistryKey))
+	{
+		bHasRegistered = true;
+
+#if WITH_EDITORONLY_DATA
+		// Refresh Asset Registry Info if successfully registered with Frontend
+		FConstGraphHandle GraphHandle = GetDocumentHandle()->GetRootGraph();
+		const FMetasoundFrontendClassMetadata& ClassFrontendMetadata = GraphHandle->GetGraphMetadata();
+		ClassInfo.AssetClassID = FGuid(ClassFrontendMetadata.ClassName.Name.ToString());
+		FNodeClassName ClassName = ClassFrontendMetadata.ClassName.ToNodeClassName();
+		FMetasoundFrontendClass GraphClass;
+		ensure(ISearchEngine::Get().FindClassWithMajorVersion(ClassName, ClassFrontendMetadata.Version.Major, GraphClass));
+
+		ClassInfo.Version = ClassFrontendMetadata.Version;
+
+		ClassInfo.InputTypes.Reset();
+		Algo::Transform(GraphClass.Interface.Inputs, ClassInfo.InputTypes, [] (const FMetasoundFrontendClassInput& Input) { return Input.TypeName; });
+
+		ClassInfo.OutputTypes.Reset();
+		Algo::Transform(GraphClass.Interface.Outputs, ClassInfo.OutputTypes, [](const FMetasoundFrontendClassOutput& Output) { return Output.TypeName; });
+
+		SetRegistryAssetClassInfo(MoveTemp(ClassInfo));
+#endif // WITH_EDITORONLY_DATA
+	}
+	else
 	{
 		UE_LOG(LogMetaSound, Error, TEXT("Failed to register node for MetaSoundSource [Name:%s]"), *AssetName);
 	}
@@ -259,6 +309,15 @@ Metasound::FSendAddress FMetasoundAssetBase::CreateSendAddress(uint64 InInstance
 	return Address;
 }
 
+void FMetasoundAssetBase::ConvertFromPreset()
+{
+	using namespace Metasound::Frontend;
+	FGraphHandle GraphHandle = GetRootGraphHandle();
+	FMetasoundFrontendGraphStyle Style = GraphHandle->GetGraphStyle();
+	Style.bIsGraphEditable = true;
+	GraphHandle->SetGraphStyle(Style);
+}
+
 TArray<FMetasoundAssetBase::FSendInfoAndVertexName> FMetasoundAssetBase::GetSendInfos(uint64 InInstanceID) const
 {
 	using FSendInfo = Metasound::FMetasoundInstanceTransmitter::FSendInfo;
@@ -388,24 +447,6 @@ bool FMetasoundAssetBase::GetReceiveNodeMetadataForDataType(const FName& InTypeN
 	return true;
 }
 
-void FMetasoundAssetBase::UpdateAssetTags(FMetasoundFrontendClassAssetTags& OutTags)
-{
-	using namespace Metasound::Frontend;
-
-	FGraphHandle GraphHandle = GetRootGraphHandle();
-
-	OutTags.ID = GraphHandle->GetClassID();
-
-	OutTags.Name = GraphHandle->GetGraphMetadata().ClassName.Name;
-	OutTags.Namespace = GraphHandle->GetGraphMetadata().ClassName.Namespace;
-	OutTags.Variant = GraphHandle->GetGraphMetadata().ClassName.Variant;
-
-	OutTags.MajorVersion = GraphHandle->GetGraphMetadata().Version.Major;
-	OutTags.MinorVersion = GraphHandle->GetGraphMetadata().Version.Minor;
-
-
-}
-
 bool FMetasoundAssetBase::MarkMetasoundDocumentDirty() const
 {
 	if (const UObject* OwningAsset = GetOwningAsset())
@@ -415,11 +456,6 @@ bool FMetasoundAssetBase::MarkMetasoundDocumentDirty() const
 	return false;
 }
 
-FMetasoundFrontendClassMetadata FMetasoundAssetBase::GetMetadata()
-{
-	return GetDocumentChecked().RootGraph.Metadata;
-}
-
 Metasound::Frontend::FDocumentHandle FMetasoundAssetBase::GetDocumentHandle()
 {
 	return Metasound::Frontend::IDocumentController::CreateDocumentHandle(GetDocument());
diff --git a/Engine/Plugins/Runtime/Metasound/Source/MetasoundFrontend/Private/MetasoundFrontendRegistries.cpp b/Engine/Plugins/Runtime/Metasound/Source/MetasoundFrontend/Private/MetasoundFrontendRegistries.cpp
--- a/Engine/Plugins/Runtime/Metasound/Source/MetasoundFrontend/Private/MetasoundFrontendRegistries.cpp
+++ b/Engine/Plugins/Runtime/Metasound/Source/MetasoundFrontend/Private/MetasoundFrontendRegistries.cpp
@@ -170,7 +170,7 @@ namespace Metasound
 					if (DataTypeEntry.IsValid())
 					{
 						FrontendClass = DataTypeEntry->GetFrontendInputClass();
-						ClassInfo = FNodeClassInfo(FrontendClass);
+						ClassInfo = FNodeClassInfo(FrontendClass.Metadata);
 					}
 				}
 
@@ -233,7 +233,7 @@ namespace Metasound
 					if (DataTypeEntry.IsValid())
 					{
 						FrontendClass = DataTypeEntry->GetFrontendOutputClass();
-						ClassInfo = FNodeClassInfo(FrontendClass);
+						ClassInfo = FNodeClassInfo(FrontendClass.Metadata);
 					}
 				}
 
@@ -296,7 +296,7 @@ namespace Metasound
 					if (DataTypeEntry.IsValid())
 					{
 						FrontendClass = DataTypeEntry->GetFrontendVariableClass();
-						ClassInfo = FNodeClassInfo(FrontendClass);
+						ClassInfo = FNodeClassInfo(FrontendClass.Metadata);
 					}
 				}
 
@@ -985,20 +985,30 @@ namespace Metasound
 		}
 
 		FNodeClassInfo::FNodeClassInfo(const FMetasoundFrontendClassMetadata& InMetadata)
-		: ClassName(InMetadata.ClassName)
-		, Type(InMetadata.Type)
-		, Version(InMetadata.Version)
+			: ClassName(InMetadata.ClassName)
+			, Type(InMetadata.Type)
+			, Version(InMetadata.Version)
 		{
 		}
 
-		FNodeClassInfo::FNodeClassInfo(const FMetasoundFrontendClass& InClass)
-		: FNodeClassInfo(InClass.Metadata)
+		FNodeClassInfo::FNodeClassInfo(const FMetasoundFrontendGraphClass& InClass, FName InAssetPath)
+			: ClassName(InClass.Metadata.ClassName)
+			, Type(EMetasoundFrontendClassType::External) // Overridden as it is considered the same as an external class in the registry
+			, AssetClassID(FGuid(ClassName.Name.ToString()))
+			, AssetPath(InAssetPath)
+			, Version(InClass.Metadata.Version)
 		{
-		}
+			ensure(!AssetPath.IsNone());
 
-		FNodeClassInfo::FNodeClassInfo(const FNodeClassMetadata& InMetadata)
-		: FNodeClassInfo(FMetasoundFrontendClassMetadata(InMetadata))
-		{
+			for (const FMetasoundFrontendClassInput& Input : InClass.Interface.Inputs)
+			{
+				InputTypes.Add(Input.TypeName);
+			}
+
+			for (const FMetasoundFrontendClassOutput& Output : InClass.Interface.Outputs)
+			{
+				OutputTypes.Add(Output.TypeName);
+			}
 		}
 	} // namespace Frontend
 } // namespace Metasound
diff --git a/Engine/Plugins/Runtime/Metasound/Source/MetasoundFrontend/Private/MetasoundFrontendTransform.cpp b/Engine/Plugins/Runtime/Metasound/Source/MetasoundFrontend/Private/MetasoundFrontendTransform.cpp
--- a/Engine/Plugins/Runtime/Metasound/Source/MetasoundFrontend/Private/MetasoundFrontendTransform.cpp
+++ b/Engine/Plugins/Runtime/Metasound/Source/MetasoundFrontend/Private/MetasoundFrontendTransform.cpp
@@ -1,8 +1,11 @@
 // Copyright Epic Games, Inc. All Rights Reserved.
 #include "MetasoundFrontendTransform.h"
 
+#include "CoreMinimal.h"
 #include "MetasoundFrontendDocument.h"
 #include "MetasoundLog.h"
+#include "Misc/App.h"
+
 
 namespace Metasound
 {
@@ -182,7 +185,7 @@ namespace Metasound
 
 			FMetasoundFrontendVersionNumber GetTargetVersion() const override
 			{
-				return FVersionDocument::GetMaxVersion();
+				return { 1, 2 };
 			}
 
 			void TransformInternal(FDocumentHandle InDocument) const override
@@ -192,7 +195,29 @@ namespace Metasound
 
 				Metadata.ClassName = FMetasoundFrontendClassName { "GraphAsset", Name, *Path };
 				Metadata.DisplayName = FText::FromString(Name.ToString());
+				GraphHandle->SetGraphMetadata(Metadata);
+			}
+		};
 
+		/** Versions document from 1.1 to 1.3. */
+		class FVersionDocument_1_3 : public FVersionDocumentTransform
+		{
+		public:
+			FVersionDocument_1_3()
+			{
+			}
+
+			FMetasoundFrontendVersionNumber GetTargetVersion() const override
+			{
+				return FVersionDocument::GetMaxVersion();
+			}
+
+			void TransformInternal(FDocumentHandle InDocument) const override
+			{
+				FGraphHandle GraphHandle = InDocument->GetRootGraph();
+				FMetasoundFrontendClassMetadata Metadata = GraphHandle->GetGraphMetadata();
+
+				Metadata.ClassName = FMetasoundFrontendClassName { FName(), *FGuid::NewGuid().ToString(), FName() };
 				GraphHandle->SetGraphMetadata(Metadata);
 			}
 		};
@@ -215,6 +240,8 @@ namespace Metasound
 			// Add additional transforms here after defining them above, example below.
 			bWasUpdated |= FVersionDocument_1_1().Transform(InDocument);
 			bWasUpdated |= FVersionDocument_1_2(Name, Path).Transform(InDocument);
+			bWasUpdated |= FVersionDocument_1_3().Transform(InDocument);
+
 			return bWasUpdated;
 		}
 	} // namespace Frontend
diff --git a/Engine/Plugins/Runtime/Metasound/Source/MetasoundFrontend/Public/MetasoundAssetBase.h b/Engine/Plugins/Runtime/Metasound/Source/MetasoundFrontend/Public/MetasoundAssetBase.h
--- a/Engine/Plugins/Runtime/Metasound/Source/MetasoundFrontend/Public/MetasoundAssetBase.h
+++ b/Engine/Plugins/Runtime/Metasound/Source/MetasoundFrontend/Public/MetasoundAssetBase.h
@@ -2,6 +2,7 @@
 #pragma once
 
 #include "CoreMinimal.h"
+
 #include "MetasoundAccessPtr.h"
 #include "MetasoundFrontend.h"
 #include "MetasoundFrontendController.h"
@@ -12,9 +13,25 @@
 #include "UObject/WeakObjectPtrTemplates.h"
 
 
-
 class UEdGraph;
 
+namespace Metasound
+{
+	namespace AssetTags
+	{
+		extern const FString METASOUNDFRONTEND_API ArrayDelim;
+
+		extern const FName METASOUNDFRONTEND_API AssetClassID;
+		extern const FName METASOUNDFRONTEND_API RegistryVersionMajor;
+		extern const FName METASOUNDFRONTEND_API RegistryVersionMinor;
+
+#if WITH_EDITORONLY_DATA
+		extern const FName METASOUNDFRONTEND_API RegistryInputTypes;
+		extern const FName METASOUNDFRONTEND_API RegistryOutputTypes;
+#endif // WITH_EDITORONLY_DATA
+	} // namespace AssetTags
+} // namespace Metasound
+
 /** FMetasoundAssetBase is intended to be a mix-in subclass for UObjects which utilize
  * Metasound assets.  It provides consistent access to FMetasoundFrontendDocuments, control
  * over the FMetasoundFrontendArchetype of the FMetasoundFrontendDocument.  It also enables the UObject
@@ -44,6 +61,9 @@ public:
 	// @param Editor graph associated with this metasound object.
 	virtual void SetGraph(UEdGraph* InGraph) = 0;
 
+	// Only required for editor builds. Adds metadata to properties available when the object is
+	// not loaded for use by the Asset Registry.
+	virtual void SetRegistryAssetClassInfo(const Metasound::Frontend::FNodeClassInfo& InClassInfo) = 0;
 #endif // WITH_EDITORONLY_DATA
 
 	// Registers the root graph of the given asset with the MetaSound Frontend.
@@ -66,13 +86,15 @@ public:
 	// Returns the preferred archetype for the given document.
 	virtual const FMetasoundFrontendArchetype& GetPreferredArchetypes(const FMetasoundFrontendDocument& InDocument, const FMetasoundFrontendArchetype& InDefaultArchetype) const;
 
-	// Returns the root class metadata
-	FMetasoundFrontendClassMetadata GetMetadata();
-
 	// TODO:
 	//virtual void OnMetaSoundDependencyAdded(const FMetasoundFrontendClassMetadata& InMetadata) = 0;
 	//virtual void OnMetaSoundDependencyRemoved(const FMetasoundFrontendClassMetadata& InMetadata) = 0;
 
+	// Gets the asset class info.
+	virtual Metasound::Frontend::FNodeClassInfo GetAssetClassInfo() const = 0;
+
+	void ConvertFromPreset();
+
 	// Imports data from a JSON string directly
 	bool ImportFromJSON(const FString& InJSON);
 
@@ -87,7 +109,6 @@ public:
 	Metasound::Frontend::FGraphHandle GetRootGraphHandle();
 	Metasound::Frontend::FConstGraphHandle GetRootGraphHandle() const;
 
-
 	// Overwrites the existing document. If the document's archetype is not supported,
 	// the FMetasoundAssetBase be while queried for a new one using `GetPreferredArchetype`. If `bForceUpdateArchetype`
 	// is true, `GetPreferredArchetype` will be used whether or not the provided document's archetype
@@ -116,7 +137,6 @@ protected:
 	FText GetDisplayName(FString&& InTypeName) const;
 #endif // WITH_EDITORONLY_DATA
 
-	void UpdateAssetTags(FMetasoundFrontendClassAssetTags& OutTags);
 
 	// Returns an access pointer to the document.
 	virtual Metasound::Frontend::FDocumentAccessPtr GetDocument() = 0;
@@ -132,10 +152,10 @@ protected:
 
 
 private:
+	bool bHasRegistered = false;
+
 	bool GetReceiveNodeMetadataForDataType(const FName& InTypeName, FMetasoundFrontendClassMetadata& OutMetadata) const;
 	TArray<FString> GetTransmittableInputVertexNames() const;
 	Metasound::FSendAddress CreateSendAddress(uint64 InInstanceID, const FString& InVertexName, const FName& InDataTypeName) const;
 	Metasound::Frontend::FNodeHandle AddInputPinForSendAddress(const Metasound::FMetasoundInstanceTransmitter::FSendInfo& InSendInfo, Metasound::Frontend::FGraphHandle InGraph) const;
-
-	Metasound::Frontend::FNodeRegistryKey RegistryKey;
 };
diff --git a/Engine/Plugins/Runtime/Metasound/Source/MetasoundFrontend/Public/MetasoundFrontendDocument.h b/Engine/Plugins/Runtime/Metasound/Source/MetasoundFrontend/Public/MetasoundFrontendDocument.h
--- a/Engine/Plugins/Runtime/Metasound/Source/MetasoundFrontend/Public/MetasoundFrontendDocument.h
+++ b/Engine/Plugins/Runtime/Metasound/Source/MetasoundFrontend/Public/MetasoundFrontendDocument.h
@@ -565,24 +565,28 @@ struct FMetasoundFrontendInterfaceStyle
 	{
 		TArray<HandleType> SortedHandles = InHandles;
 
-		if (ensure(SortedHandles.Num() == DefaultSortOrder.Num()))
+		// TODO: Hack for assets which aren't getting sort order set for inputs/outputs. Fix this & remove size check.
+		if (DefaultSortOrder.Num() > 0)
 		{
-			TMap<FGuid, int32> HandleIDToSortIndex;
-			for (int32 i = 0; i < DefaultSortOrder.Num(); ++i)
+			if (ensure(SortedHandles.Num() == DefaultSortOrder.Num()))
 			{
-				if (InHandles.IsValidIndex(i))
+				TMap<FGuid, int32> HandleIDToSortIndex;
+				for (int32 i = 0; i < DefaultSortOrder.Num(); ++i)
 				{
-					const int32 SortIndex = DefaultSortOrder[i];
-					HandleIDToSortIndex.Add(InHandles[i]->GetID(), SortIndex);
+					if (InHandles.IsValidIndex(i))
+					{
+						const int32 SortIndex = DefaultSortOrder[i];
+						HandleIDToSortIndex.Add(InHandles[i]->GetID(), SortIndex);
+					}
 				}
-			}
 
-			SortedHandles.Sort([&](const HandleType& HandleA, const HandleType& HandleB)
-			{
-				const FGuid HandleAID = HandleA->GetID();
-				const FGuid HandleBID = HandleB->GetID();
-				return HandleIDToSortIndex[HandleAID] < HandleIDToSortIndex[HandleBID];
-			});
+				SortedHandles.Sort([&](const HandleType& HandleA, const HandleType& HandleB)
+					{
+						const FGuid HandleAID = HandleA->GetID();
+						const FGuid HandleBID = HandleB->GetID();
+						return HandleIDToSortIndex[HandleAID] < HandleIDToSortIndex[HandleBID];
+					});
+			}
 		}
 
 		return SortedHandles;
@@ -661,7 +665,7 @@ struct METASOUNDFRONTEND_API FMetasoundFrontendClassName
 
 
 USTRUCT()
-struct FMetasoundFrontendClassMetadata
+struct METASOUNDFRONTEND_API FMetasoundFrontendClassMetadata
 {
 	GENERATED_BODY()
 
@@ -718,30 +722,6 @@ struct FMetasoundFrontendEditorData
 	TArray<uint8> Data;
 };
 
-USTRUCT()
-struct METASOUNDFRONTEND_API FMetasoundFrontendClassAssetTags
-{
-	GENERATED_BODY()
-
-	UPROPERTY()
-	FGuid ID = Metasound::FrontendInvalidID;
-
-	UPROPERTY()
-	FName Namespace;
-
-	UPROPERTY()
-	FName Name;
-
-	UPROPERTY()
-	FName Variant;
-
-	UPROPERTY()
-	int32 MajorVersion = -1;
-
-	UPROPERTY()
-	int32 MinorVersion = -1;
-};
-
 USTRUCT()
 struct METASOUNDFRONTEND_API FMetasoundFrontendClass
 {
diff --git a/Engine/Plugins/Runtime/Metasound/Source/MetasoundFrontend/Public/MetasoundFrontendRegistries.h b/Engine/Plugins/Runtime/Metasound/Source/MetasoundFrontend/Public/MetasoundFrontendRegistries.h
--- a/Engine/Plugins/Runtime/Metasound/Source/MetasoundFrontend/Public/MetasoundFrontendRegistries.h
+++ b/Engine/Plugins/Runtime/Metasound/Source/MetasoundFrontend/Public/MetasoundFrontendRegistries.h
@@ -1,5 +1,4 @@
 // Copyright Epic Games, Inc. All Rights Reserved.
-
 #pragma once
 
 #include "CoreMinimal.h"
@@ -12,6 +11,7 @@
 #include "MetasoundOperatorInterface.h"
 #include "MetasoundRouter.h"
 
+
 namespace Metasound
 {
 	// Base implementation for NodeConstructorCallbacks
@@ -64,14 +64,46 @@ namespace Metasound
 		 */
 		struct METASOUNDFRONTEND_API FNodeClassInfo
 		{
+			// ClassName of the given class
 			FMetasoundFrontendClassName ClassName;
+
+			// The type of this node class
 			EMetasoundFrontendClassType Type = EMetasoundFrontendClassType::Invalid;
+
+			// The ID used for the Asset Classes. If zero, class is natively defined.
+			FGuid AssetClassID;
+
+			// Path to asset containing graph if external type and references asset class.
+			FName AssetPath;
+
+			// Version of the registered class
 			FMetasoundFrontendVersionNumber Version;
 
+			// Types of class inputs
+			TArray<FName> InputTypes;
+
+			// Types of class outputs
+			TArray<FName> OutputTypes;
+
 			FNodeClassInfo() = default;
+
+			// Constructor used to generate NodeClassInfo from a native class' Metadata.
 			FNodeClassInfo(const FMetasoundFrontendClassMetadata& InMetadata);
-			FNodeClassInfo(const FMetasoundFrontendClass& InClass);
-			FNodeClassInfo(const FNodeClassMetadata& InMetadata);
+
+			// Constructor used to generate NodeClassInfo from an asset
+			FNodeClassInfo(const FMetasoundFrontendGraphClass& InClass, FName InAssetPath);
+
+			// Loads the asset from the provided path, ensuring that the class is of type graph.
+			UObject* LoadAsset() const
+			{
+				if (ensure(Type == EMetasoundFrontendClassType::External))
+				{
+					FSoftObjectPath SoftObjectPath(AssetPath);
+					return SoftObjectPath.TryLoad();
+				}
+
+				return nullptr;
+			}
 		};
 
 		/** INodeRegistryEntry declares the interface for a node registry entry.
@@ -434,7 +466,7 @@ public:
 	 *
 	 * @param InKey - The registration key for the node.
 	 *
-	 * @return True on success, false on failer.
+	 * @return True on success, false on failure.
 	 */
 	virtual bool UnregisterNode(const FNodeRegistryKey& InKey) = 0;
 
diff --git a/Engine/Plugins/Runtime/Metasound/Source/MetasoundFrontend/Public/MetasoundFrontendTransform.h b/Engine/Plugins/Runtime/Metasound/Source/MetasoundFrontend/Public/MetasoundFrontendTransform.h
--- a/Engine/Plugins/Runtime/Metasound/Source/MetasoundFrontend/Public/MetasoundFrontendTransform.h
+++ b/Engine/Plugins/Runtime/Metasound/Source/MetasoundFrontend/Public/MetasoundFrontendTransform.h
@@ -41,7 +41,7 @@ namespace Metasound
 			public:
 				static FMetasoundFrontendVersionNumber GetMaxVersion()
 				{
-					return FMetasoundFrontendVersionNumber { 1, 2 };
+					return FMetasoundFrontendVersionNumber { 1, 3 };
 				}
 
 				FVersionDocument(FName InName, const FString& InPath);
diff --git a/Engine/Plugins/Runtime/Metasound/Source/MetasoundGraphCore/Public/MetasoundNodeInterface.h b/Engine/Plugins/Runtime/Metasound/Source/MetasoundGraphCore/Public/MetasoundNodeInterface.h
--- a/Engine/Plugins/Runtime/Metasound/Source/MetasoundGraphCore/Public/MetasoundNodeInterface.h
+++ b/Engine/Plugins/Runtime/Metasound/Source/MetasoundGraphCore/Public/MetasoundNodeInterface.h
@@ -38,7 +38,7 @@ namespace Metasound
 	 *
 	 * FNodeClassName is used for lookup and declaring interoperability.
 	 *
-	 * Namespaces are provided as a convenience to simply name collisions.
+	 * Namespaces are provided as a convenience to avoid name collisions.
 	 *
 	 * Nodes with equal Namespace and Name, but different Variants are considered
 	 * to be interoperable. They can be used to define nodes that perform the same
