Minor comment fix-up

diff --git a/Engine/Plugins/Runtime/Metasound/Source/MetasoundEngine/Private/MetasoundUObjectRegistry.cpp b/Engine/Plugins/Runtime/Metasound/Source/MetasoundEngine/Private/MetasoundUObjectRegistry.cpp
--- a/Engine/Plugins/Runtime/Metasound/Source/MetasoundEngine/Private/MetasoundUObjectRegistry.cpp
+++ b/Engine/Plugins/Runtime/Metasound/Source/MetasoundEngine/Private/MetasoundUObjectRegistry.cpp
@@ -123,8 +123,7 @@ void UMetaSoundAssetSubsystem::AddOrUpdateAsset(const FAssetData& InAssetData)
 	using namespace Metasound::AssetSubsystemPrivate;
 	using namespace Metasound::Frontend;
 
-	// TODO: Set to false for editor data only builds once registering
-	// without loading asset is supported.
+	// TODO: Set to false for builds once registering without loading asset is supported.
 	static const bool bLoadRequiredToRegisterAssetClasses = true;
 
 	FNodeClassInfo ClassInfo;
