IWYU fixes
#rnx

diff --git a/Engine/Plugins/Runtime/Metasound/Source/MetasoundEngine/Public/Metasound.h b/Engine/Plugins/Runtime/Metasound/Source/MetasoundEngine/Public/Metasound.h
--- a/Engine/Plugins/Runtime/Metasound/Source/MetasoundEngine/Public/Metasound.h
+++ b/Engine/Plugins/Runtime/Metasound/Source/MetasoundEngine/Public/Metasound.h
@@ -10,6 +10,10 @@
 #include "MetasoundOperatorSettings.h"
 #include "MetasoundRouter.h"
 
+#if WITH_EDITORONLY_DATA
+#include "Algo/Transform.h"
+#endif // WITH_EDITORONLY_DATA
+
 #include "Metasound.generated.h"
 
 
diff --git a/Engine/Plugins/Runtime/Metasound/Source/MetasoundEngine/Public/MetasoundUObjectRegistry.h b/Engine/Plugins/Runtime/Metasound/Source/MetasoundEngine/Public/MetasoundUObjectRegistry.h
--- a/Engine/Plugins/Runtime/Metasound/Source/MetasoundEngine/Public/MetasoundUObjectRegistry.h
+++ b/Engine/Plugins/Runtime/Metasound/Source/MetasoundEngine/Public/MetasoundUObjectRegistry.h
@@ -3,7 +3,9 @@
 
 #include "CoreMinimal.h"
 #include "MetasoundAssetBase.h"
+#include "Subsystems/EngineSubsystem.h"
 #include "UObject/Object.h"
+
 #include "MetasoundUObjectRegistry.generated.h"
 
 
