Add FrequencyToMIDI node to Metasounds

diff --git a/Engine/Plugins/Runtime/Metasound/Source/MetasoundStandardNodes/Private/MetasoundFreqToMIDINode.cpp b/Engine/Plugins/Runtime/Metasound/Source/MetasoundStandardNodes/Private/MetasoundFreqToMIDINode.cpp
new file mode 100644
--- /dev/null
+++ b/Engine/Plugins/Runtime/Metasound/Source/MetasoundStandardNodes/Private/MetasoundFreqToMIDINode.cpp
@@ -0,0 +1,160 @@
+// Copyright Epic Games, Inc. All Rights Reserved.
+
+#include "Internationalization/Text.h"
+#include "MetasoundFacade.h"
+#include "MetasoundExecutableOperator.h"
+#include "MetasoundNodeRegistrationMacro.h"
+#include "MetasoundPrimitives.h"
+#include "MetasoundStandardNodesNames.h"
+#include "MetasoundStandardNodesCategories.h"
+#include "DSP/Dsp.h"
+#include "MetasoundParamHelper.h"
+
+#define LOCTEXT_NAMESPACE "MetasoundStandardNodes_FreqToMidiNode"
+
+namespace Metasound
+{
+	
+		
+	namespace FrequencyToMidiVertexNames
+	{
+		METASOUND_PARAM(InputFreq, "Frequency In", "Input frequency value in Hz.");
+		METASOUND_PARAM(OutputMidi, "Out MIDI", "Output MIDI note value that corresponds to the input frequency value.");
+	}
+
+
+	class FFreqToMidiOperator : public TExecutableOperator<FFreqToMidiOperator>
+	{
+	public:
+
+		static const FNodeClassMetadata& GetNodeInfo();
+		static const FVertexInterface& GetVertexInterface();
+		static TUniquePtr<IOperator> CreateOperator(const FCreateOperatorParams& InParams, FBuildErrorArray& OutErrors);
+
+		FFreqToMidiOperator(const FOperatorSettings& InSettings, const FFloatReadRef& InMidiNote);
+
+		virtual FDataReferenceCollection GetInputs() const override;
+		virtual FDataReferenceCollection GetOutputs() const override;
+		void Execute();
+
+	private:
+		// The input frequency value
+		FFloatReadRef FreqInput;
+
+		// The output MIDI note value
+		FFloatWriteRef MidiOutput;
+
+		// Cached freq note value. Used to catch if the value changes to recompute Midi output.
+		float PrevFreq;
+	};
+
+	FFreqToMidiOperator::FFreqToMidiOperator(const FOperatorSettings& InSettings, const FFloatReadRef& InFreq)
+		: FreqInput(InFreq)
+		, MidiOutput(FFloatWriteRef::CreateNew(Audio::GetMidiFromFrequency(*InFreq)))
+		, PrevFreq(*InFreq)
+	{
+	}
+
+	
+	FDataReferenceCollection FFreqToMidiOperator::GetInputs() const
+	{
+		using namespace FrequencyToMidiVertexNames;
+
+		FDataReferenceCollection InputDataReferences;
+		InputDataReferences.AddDataReadReference(METASOUND_GET_PARAM_NAME(InputFreq), FreqInput);
+
+		return InputDataReferences;
+	}
+
+	FDataReferenceCollection FFreqToMidiOperator::GetOutputs() const
+	{
+		using namespace FrequencyToMidiVertexNames;
+
+		FDataReferenceCollection OutputDataReferences;
+		OutputDataReferences.AddDataReadReference(METASOUND_GET_PARAM_NAME(OutputMidi), MidiOutput);
+		return OutputDataReferences;
+	}
+
+	void FFreqToMidiOperator::Execute()
+	{
+		// Only do anything if the Midi note changes
+		if (!FMath::IsNearlyEqual(*FreqInput, PrevFreq))
+		{
+			PrevFreq = *FreqInput;
+
+			*MidiOutput = Audio::GetMidiFromFrequency(PrevFreq);
+		}
+	}
+
+	const FVertexInterface& FFreqToMidiOperator::GetVertexInterface()
+	{
+		using namespace FrequencyToMidiVertexNames;
+
+		static const FVertexInterface Interface(
+			FInputVertexInterface(
+				TInputDataVertexModel<float>(METASOUND_GET_PARAM_NAME_AND_TT(InputFreq), 440.0f)
+			),
+			FOutputVertexInterface(
+				TOutputDataVertexModel<float>(METASOUND_GET_PARAM_NAME_AND_TT(OutputMidi))
+			)
+		);
+
+		return Interface;
+	}
+
+	const FNodeClassMetadata& FFreqToMidiOperator::GetNodeInfo()
+	{
+		auto InitNodeInfo = []() -> FNodeClassMetadata
+		{
+			const FName DataTypeName = GetMetasoundDataTypeName<float>();
+			const FName OperatorName = TEXT("Frequency to MIDI");
+			const FText NodeDisplayName = LOCTEXT("Metasound_FreqToMidiNodeName", "Frequency To MIDI");
+			const FText NodeDescription = LOCTEXT("Metasound_FreqToMidiNodeDescription", "Converts a frequency (Hz) value to a Midi Note value.");
+
+			FNodeClassMetadata Info;
+			Info.ClassName = { Metasound::StandardNodes::Namespace, OperatorName, DataTypeName };
+			Info.MajorVersion = 1;
+			Info.MinorVersion = 0;
+			Info.DisplayName = NodeDisplayName;
+			Info.Description = NodeDescription;
+			Info.Author = PluginAuthor;
+			Info.PromptIfMissing = PluginNodeMissingPrompt;
+			Info.DefaultInterface = GetVertexInterface();
+			Info.CategoryHierarchy.Emplace(StandardNodes::Music);
+
+			return Info;
+		};
+
+		static const FNodeClassMetadata Info = InitNodeInfo();
+
+		return Info;
+	}
+
+	TUniquePtr<IOperator> FFreqToMidiOperator::CreateOperator(const FCreateOperatorParams& InParams, FBuildErrorArray& OutErrors)
+	{
+		using namespace FrequencyToMidiVertexNames;
+
+		const FDataReferenceCollection& InputCollection = InParams.InputDataReferences;
+		const FInputVertexInterface& InputInterface = GetVertexInterface().GetInputInterface();
+
+		FFloatReadRef InFreq = InputCollection.GetDataReadReferenceOrConstructWithVertexDefault<float>(InputInterface, METASOUND_GET_PARAM_NAME(InputFreq), InParams.OperatorSettings);
+
+		return MakeUnique<FFreqToMidiOperator>(InParams.OperatorSettings, InFreq);
+	}
+
+	class METASOUNDSTANDARDNODES_API FFreqToMidiNode : public FNodeFacade
+	{
+	public:
+		/**
+		 * Constructor used by the Metasound Frontend.
+		 */
+		FFreqToMidiNode(const FNodeInitData& InitData)
+			: FNodeFacade(InitData.InstanceName, InitData.InstanceID, TFacadeOperatorClass <FFreqToMidiOperator>())
+		{
+		}
+	};
+
+	METASOUND_REGISTER_NODE(FFreqToMidiNode)
+}
+
+#undef LOCTEXT_NAMESPACE
