Add LinearToLogFrequency node to Metasounds

diff --git a/Engine/Plugins/Runtime/Metasound/Source/MetasoundStandardNodes/Private/MetasoundLinearToLogFrequencyNode.cpp b/Engine/Plugins/Runtime/Metasound/Source/MetasoundStandardNodes/Private/MetasoundLinearToLogFrequencyNode.cpp
new file mode 100644
--- /dev/null
+++ b/Engine/Plugins/Runtime/Metasound/Source/MetasoundStandardNodes/Private/MetasoundLinearToLogFrequencyNode.cpp
@@ -0,0 +1,182 @@
+// Copyright Epic Games, Inc. All Rights Reserved.
+
+#include "Internationalization/Text.h"
+#include "MetasoundFacade.h"
+#include "MetasoundExecutableOperator.h"
+#include "MetasoundNodeRegistrationMacro.h"
+#include "MetasoundPrimitives.h"
+#include "MetasoundStandardNodesNames.h"
+#include "MetasoundStandardNodesCategories.h"
+#include "DSP/Dsp.h"
+#include "MetasoundParamHelper.h"
+
+#define LOCTEXT_NAMESPACE "MetasoundStandardNodes_LinearToLogFreqNode"
+
+namespace Metasound
+{
+	namespace LinearToLogFrequencyVertexNames
+	{
+		METASOUND_PARAM(InputValue, "Value", "Linear input value to map to log frequency output. Input and output are clamped to specified domain and range.");
+		METASOUND_PARAM(InputDomainMin, "Min Domain", "Min domain for the input value.");
+		METASOUND_PARAM(InputDomainMax, "Max Domain", "Max domain for the input value.");
+		METASOUND_PARAM(InputRangeMin, "Min Range", "Min positive range for the output frequency (hz) value.");
+		METASOUND_PARAM(InputRangeMax, "Max Range", "Max poasitive range for the output frequency (hz) value.");
+		METASOUND_PARAM(OutputFreq, "Frequency", "Output frequency value in hertz that is the log frequency of the input value.");
+	}
+
+	class FLinearToLogFrequencyOperator : public TExecutableOperator<FLinearToLogFrequencyOperator>
+	{
+	public:
+
+		static const FNodeClassMetadata& GetNodeInfo()
+		{
+			auto InitNodeInfo = []() -> FNodeClassMetadata
+			{
+				FName OperatorName = TEXT("Linear To Log Frequency");
+				FText NodeDisplayName = LOCTEXT("LinearToLogFrequencyName", "Linear To Log Frequency");
+				FText NodeDescription = LOCTEXT("LinearToLogFrequencyDescription", "Converts a linear space input value to log-frequency space.");
+
+				FNodeClassMetadata Info;
+				Info.ClassName = { Metasound::StandardNodes::Namespace, OperatorName, TEXT("") };
+				Info.MajorVersion = 1;
+				Info.MinorVersion = 0;
+				Info.DisplayName = NodeDisplayName;
+				Info.Description = NodeDescription;
+				Info.Author = PluginAuthor;
+				Info.PromptIfMissing = PluginNodeMissingPrompt;
+				Info.DefaultInterface = GetVertexInterface();
+				Info.CategoryHierarchy.Emplace(StandardNodes::Music);
+
+				return Info;
+			};
+
+			static const FNodeClassMetadata Info = InitNodeInfo();
+
+			return Info;
+		}
+
+		static const FVertexInterface& GetVertexInterface()
+		{
+			using namespace LinearToLogFrequencyVertexNames;
+
+			static const FVertexInterface Interface(
+				FInputVertexInterface(
+					TInputDataVertexModel<float>(METASOUND_GET_PARAM_NAME_AND_TT(InputValue),     0.0f),
+					TInputDataVertexModel<float>(METASOUND_GET_PARAM_NAME_AND_TT(InputDomainMin), 0.0f),
+					TInputDataVertexModel<float>(METASOUND_GET_PARAM_NAME_AND_TT(InputDomainMax), 1.0f),
+					TInputDataVertexModel<float>(METASOUND_GET_PARAM_NAME_AND_TT(InputRangeMin),  20.0f),
+					TInputDataVertexModel<float>(METASOUND_GET_PARAM_NAME_AND_TT(InputRangeMax),  20000.0f)
+					),
+				FOutputVertexInterface(
+					TOutputDataVertexModel<float>(METASOUND_GET_PARAM_NAME_AND_TT(OutputFreq))
+				)
+			);
+
+			return Interface;
+		}
+
+		static TUniquePtr<IOperator> CreateOperator(const FCreateOperatorParams& InParams, FBuildErrorArray& OutErrors)
+		{
+			using namespace LinearToLogFrequencyVertexNames;
+
+			const FDataReferenceCollection& InputCollection = InParams.InputDataReferences;
+			const FInputVertexInterface& InputInterface = GetVertexInterface().GetInputInterface();
+
+			TDataReadReference<float> InValue = InputCollection.GetDataReadReferenceOrConstructWithVertexDefault<float>(InputInterface, METASOUND_GET_PARAM_NAME(InputValue), InParams.OperatorSettings);
+			TDataReadReference<float> InDomainMin = InputCollection.GetDataReadReferenceOrConstructWithVertexDefault<float>(InputInterface, METASOUND_GET_PARAM_NAME(InputDomainMin), InParams.OperatorSettings);
+			TDataReadReference<float> InDomainMax = InputCollection.GetDataReadReferenceOrConstructWithVertexDefault<float>(InputInterface, METASOUND_GET_PARAM_NAME(InputDomainMax), InParams.OperatorSettings);
+			TDataReadReference<float> InRangeMin = InputCollection.GetDataReadReferenceOrConstructWithVertexDefault<float>(InputInterface, METASOUND_GET_PARAM_NAME(InputRangeMin), InParams.OperatorSettings);
+			TDataReadReference<float> InRangeMax = InputCollection.GetDataReadReferenceOrConstructWithVertexDefault<float>(InputInterface, METASOUND_GET_PARAM_NAME(InputRangeMax), InParams.OperatorSettings);
+
+			return MakeUnique<FLinearToLogFrequencyOperator>(InValue, InDomainMin, InDomainMax, InRangeMin, InRangeMax);
+		}
+
+		FLinearToLogFrequencyOperator(
+			const TDataReadReference<float>& InValue, 
+			const TDataReadReference<float>& InDomainMin,
+			const TDataReadReference<float>& InDomainMax,
+			const TDataReadReference<float>& InRangeMin,
+			const TDataReadReference<float>& InRangeMax)
+			: Value(InValue)
+			, DomainMin(InDomainMin)
+			, DomainMax(InDomainMax)
+			, RangeMin(InRangeMin)
+			, RangeMax(InRangeMax)
+			, FreqOutput(FFloatWriteRef::CreateNew(GetOutputValue()))
+			, PrevInputValue(*Value)
+		{
+		}
+
+		virtual FDataReferenceCollection GetInputs() const override
+		{
+			using namespace LinearToLogFrequencyVertexNames;
+
+			FDataReferenceCollection InputDataReferences;
+			InputDataReferences.AddDataReadReference(METASOUND_GET_PARAM_NAME(InputValue), Value);
+			InputDataReferences.AddDataReadReference(METASOUND_GET_PARAM_NAME(InputDomainMin), DomainMin);
+			InputDataReferences.AddDataReadReference(METASOUND_GET_PARAM_NAME(InputDomainMax), DomainMax);
+			InputDataReferences.AddDataReadReference(METASOUND_GET_PARAM_NAME(InputRangeMin), RangeMin);
+			InputDataReferences.AddDataReadReference(METASOUND_GET_PARAM_NAME(InputRangeMax), RangeMax);
+
+			return InputDataReferences;
+
+		}
+
+		virtual FDataReferenceCollection GetOutputs() const override
+		{
+			using namespace LinearToLogFrequencyVertexNames;
+
+			FDataReferenceCollection OutputDataReferences;
+			OutputDataReferences.AddDataReadReference(METASOUND_GET_PARAM_NAME(OutputFreq), FreqOutput);
+			return OutputDataReferences;
+		}
+
+		void Execute()
+		{
+			using namespace LinearToLogFrequencyVertexNames;
+
+			if (!FMath::IsNearlyEqual(*Value, PrevInputValue))
+			{
+				PrevInputValue = *Value;
+				*FreqOutput = GetOutputValue();
+			}
+		}
+
+	private:
+		float GetOutputValue() const
+		{
+			float RangeMinClamped = FMath::Max(*RangeMin, SMALL_NUMBER);
+			float RangeMaxClamped = FMath::Max(*RangeMax, SMALL_NUMBER);
+			float Result = Audio::GetLogFrequencyClamped(*Value, { *DomainMin, *DomainMax }, { RangeMinClamped, RangeMaxClamped });
+			return Result;
+		}
+
+		TDataReadReference<float> Value;
+		TDataReadReference<float> DomainMin;
+		TDataReadReference<float> DomainMax;
+		TDataReadReference<float> RangeMin;
+		TDataReadReference<float> RangeMax;
+
+		FFloatWriteRef FreqOutput;
+
+		float PrevInputValue;
+	};
+
+
+	class METASOUNDSTANDARDNODES_API FLinearToLogFrequencyNode : public FNodeFacade
+	{
+	public:
+		/**
+		 * Constructor used by the Metasound Frontend.
+		 */
+		FLinearToLogFrequencyNode(const FNodeInitData& InitData)
+			: FNodeFacade(InitData.InstanceName, InitData.InstanceID, TFacadeOperatorClass<FLinearToLogFrequencyOperator>())
+		{
+
+		}
+	};
+
+	METASOUND_REGISTER_NODE(FLinearToLogFrequencyNode)
+}
+
+#undef LOCTEXT_NAMESPACE
