Add 4 math/conversion nodes to Metasounds
LinearGainToDecibels, DecibelsToLinearGain, PitchScaleToSemitones, SemitoneToPitchScale

diff --git a/Engine/Plugins/Runtime/Metasound/Source/MetasoundStandardNodes/Private/MetasoundDecibelsToLinearGain.cpp b/Engine/Plugins/Runtime/Metasound/Source/MetasoundStandardNodes/Private/MetasoundDecibelsToLinearGain.cpp
new file mode 100644
--- /dev/null
+++ b/Engine/Plugins/Runtime/Metasound/Source/MetasoundStandardNodes/Private/MetasoundDecibelsToLinearGain.cpp
@@ -0,0 +1,150 @@
+// Copyright Epic Games, Inc. All Rights Reserved.
+
+#include "Internationalization/Text.h"
+#include "MetasoundFacade.h"
+#include "MetasoundExecutableOperator.h"
+#include "MetasoundNodeRegistrationMacro.h"
+#include "MetasoundPrimitives.h"
+#include "MetasoundStandardNodesNames.h"
+#include "MetasoundStandardNodesCategories.h"
+#include "DSP/Dsp.h"
+#include "MetasoundParamHelper.h"
+
+#define LOCTEXT_NAMESPACE "MetasoundStandardNodes_DecibelsToLinearGain"
+
+namespace Metasound
+{
+
+	namespace DecibelsToLinearGainVertexNames
+	{
+		METASOUND_PARAM(InputDecibelGain, "Decibels", "Input logarithmic (dB) gain.");
+
+		METASOUND_PARAM(OutputLinearGain, "Linear Gain", "Output corresponding linear gain.");
+	}
+
+
+	class FDecibelsToLinearGainOperator : public TExecutableOperator<FDecibelsToLinearGainOperator>
+	{
+	public:
+
+		static const FNodeClassMetadata& GetNodeInfo();
+		static const FVertexInterface& GetVertexInterface();
+		static TUniquePtr<IOperator> CreateOperator(const FCreateOperatorParams& InParams, FBuildErrorArray& OutErrors);
+
+		FDecibelsToLinearGainOperator(const FOperatorSettings& InSettings, const FFloatReadRef& InDecibelGain);
+
+		virtual FDataReferenceCollection GetInputs() const override;
+		virtual FDataReferenceCollection GetOutputs() const override;
+		void Execute();
+
+	private:
+		// The input dB value
+		FFloatReadRef DecibelGainInput;
+
+		// The output linear gain
+		FFloatWriteRef LinearGainOutput;
+	};
+
+	FDecibelsToLinearGainOperator::FDecibelsToLinearGainOperator(const FOperatorSettings& InSettings, const FFloatReadRef& InDecibelGain)
+		: DecibelGainInput(InDecibelGain)
+		, LinearGainOutput(FFloatWriteRef::CreateNew(Audio::ConvertToLinear(*InDecibelGain)))
+	{
+	}
+
+	
+	FDataReferenceCollection FDecibelsToLinearGainOperator::GetInputs() const
+	{
+		using namespace DecibelsToLinearGainVertexNames;
+
+		FDataReferenceCollection InputDataReferences;
+		InputDataReferences.AddDataReadReference(METASOUND_GET_PARAM_NAME(InputDecibelGain), DecibelGainInput);
+
+		return InputDataReferences;
+	}
+
+	FDataReferenceCollection FDecibelsToLinearGainOperator::GetOutputs() const
+	{
+		using namespace DecibelsToLinearGainVertexNames;
+
+		FDataReferenceCollection OutputDataReferences;
+		OutputDataReferences.AddDataReadReference(METASOUND_GET_PARAM_NAME(OutputLinearGain), LinearGainOutput);
+		return OutputDataReferences;
+	}
+
+	void FDecibelsToLinearGainOperator::Execute()
+	{
+		*LinearGainOutput = Audio::ConvertToLinear(*DecibelGainInput);
+	}
+
+	const FVertexInterface& FDecibelsToLinearGainOperator::GetVertexInterface()
+	{
+		using namespace DecibelsToLinearGainVertexNames;
+
+		static const FVertexInterface Interface(
+			FInputVertexInterface(
+				TInputDataVertexModel<float>(METASOUND_GET_PARAM_NAME_AND_TT(InputDecibelGain), 0.0f)
+			),
+			FOutputVertexInterface(
+				TOutputDataVertexModel<float>(METASOUND_GET_PARAM_NAME_AND_TT(OutputLinearGain))
+			)
+		);
+
+		return Interface;
+	}
+
+	const FNodeClassMetadata& FDecibelsToLinearGainOperator::GetNodeInfo()
+	{
+		auto InitNodeInfo = []() -> FNodeClassMetadata
+		{
+			const FName DataTypeName = GetMetasoundDataTypeName<float>();
+			const FName OperatorName = TEXT("Decibels to Linear Gain");
+			const FText NodeDisplayName = LOCTEXT("Metasound_DecibelsToLinearGainName", "Decibels to Linear Gain");
+			const FText NodeDescription = LOCTEXT("Metasound_DecibelsToLinearGainDescription", "Converts a logarithmic (dB) gain value to a linear gain value.");
+
+			FNodeClassMetadata Info;
+			Info.ClassName = { Metasound::StandardNodes::Namespace, OperatorName, DataTypeName };
+			Info.MajorVersion = 1;
+			Info.MinorVersion = 0;
+			Info.DisplayName = NodeDisplayName;
+			Info.Description = NodeDescription;
+			Info.Author = PluginAuthor;
+			Info.PromptIfMissing = PluginNodeMissingPrompt;
+			Info.DefaultInterface = GetVertexInterface();
+			Info.CategoryHierarchy.Emplace(StandardNodes::Dynamics);
+
+			return Info;
+		};
+
+		static const FNodeClassMetadata Info = InitNodeInfo();
+
+		return Info;
+	}
+
+	TUniquePtr<IOperator> FDecibelsToLinearGainOperator::CreateOperator(const FCreateOperatorParams& InParams, FBuildErrorArray& OutErrors)
+	{
+		using namespace DecibelsToLinearGainVertexNames;
+
+		const FDataReferenceCollection& InputCollection = InParams.InputDataReferences;
+		const FInputVertexInterface& InputInterface = GetVertexInterface().GetInputInterface();
+
+		FFloatReadRef InDecibelGain = InputCollection.GetDataReadReferenceOrConstructWithVertexDefault<float>(InputInterface, METASOUND_GET_PARAM_NAME(InputDecibelGain), InParams.OperatorSettings);
+
+		return MakeUnique<FDecibelsToLinearGainOperator>(InParams.OperatorSettings, InDecibelGain);
+	}
+
+	class METASOUNDSTANDARDNODES_API FDecibelsToLinearGainNode : public FNodeFacade
+	{
+	public:
+		/**
+		 * Constructor used by the Metasound Frontend.
+		 */
+		FDecibelsToLinearGainNode(const FNodeInitData& InitData)
+			: FNodeFacade(InitData.InstanceName, InitData.InstanceID, TFacadeOperatorClass <FDecibelsToLinearGainOperator>())
+		{
+		}
+	};
+
+	METASOUND_REGISTER_NODE(FDecibelsToLinearGainNode)
+}
+
+#undef LOCTEXT_NAMESPACE
diff --git a/Engine/Plugins/Runtime/Metasound/Source/MetasoundStandardNodes/Private/MetasoundFreqMultiplierToSemitoneNode.cpp b/Engine/Plugins/Runtime/Metasound/Source/MetasoundStandardNodes/Private/MetasoundFreqMultiplierToSemitoneNode.cpp
new file mode 100644
--- /dev/null
+++ b/Engine/Plugins/Runtime/Metasound/Source/MetasoundStandardNodes/Private/MetasoundFreqMultiplierToSemitoneNode.cpp
@@ -0,0 +1,159 @@
+// Copyright Epic Games, Inc. All Rights Reserved.
+
+#include "Internationalization/Text.h"
+#include "MetasoundFacade.h"
+#include "MetasoundExecutableOperator.h"
+#include "MetasoundNodeRegistrationMacro.h"
+#include "MetasoundPrimitives.h"
+#include "MetasoundStandardNodesNames.h"
+#include "MetasoundStandardNodesCategories.h"
+#include "DSP/Dsp.h"
+#include "MetasoundParamHelper.h"
+
+#define LOCTEXT_NAMESPACE "MetasoundStandardNodes_FreqMultiplierToSemitone"
+
+namespace Metasound
+{
+
+	namespace FrequencyMultiplierToSemitoneVertexNames
+	{
+		METASOUND_PARAM(InputFrequencyMultiplier, "Frequency Multiplier", "Input frequency multiplier.");
+
+		METASOUND_PARAM(OutputSemitone, "Semitones", "Output corresponding difference in semitones.");
+	}
+
+
+	class FFrequencyMultiplierToSemitoneOperator : public TExecutableOperator<FFrequencyMultiplierToSemitoneOperator>
+	{
+	public:
+
+		static const FNodeClassMetadata& GetNodeInfo();
+		static const FVertexInterface& GetVertexInterface();
+		static TUniquePtr<IOperator> CreateOperator(const FCreateOperatorParams& InParams, FBuildErrorArray& OutErrors);
+
+		FFrequencyMultiplierToSemitoneOperator(const FOperatorSettings& InSettings, const FFloatReadRef& InFrequencyMultiplier);
+
+		virtual FDataReferenceCollection GetInputs() const override;
+		virtual FDataReferenceCollection GetOutputs() const override;
+		void Execute();
+
+	private:
+		// The input frequency multiplier
+		FFloatReadRef FrequencyMultiplierInput;
+
+		// The output difference in semitones
+		FFloatWriteRef SemitoneOutput;
+
+		// Cached Frequency Multiplier value. Used to catch if the value changes to recompute semitone output.
+		float PrevFrequencyMultiplier;
+	};
+
+	FFrequencyMultiplierToSemitoneOperator::FFrequencyMultiplierToSemitoneOperator(const FOperatorSettings& InSettings, const FFloatReadRef& InFrequencyMultiplier)
+		: FrequencyMultiplierInput(InFrequencyMultiplier)
+		, SemitoneOutput(FFloatWriteRef::CreateNew(Audio::GetSemitones(*InFrequencyMultiplier)))
+		, PrevFrequencyMultiplier(*InFrequencyMultiplier)
+	{
+	}
+
+	
+	FDataReferenceCollection FFrequencyMultiplierToSemitoneOperator::GetInputs() const
+	{
+		using namespace FrequencyMultiplierToSemitoneVertexNames;
+
+		FDataReferenceCollection InputDataReferences;
+		InputDataReferences.AddDataReadReference(METASOUND_GET_PARAM_NAME(InputFrequencyMultiplier), FrequencyMultiplierInput);
+
+		return InputDataReferences;
+	}
+
+	FDataReferenceCollection FFrequencyMultiplierToSemitoneOperator::GetOutputs() const
+	{
+		using namespace FrequencyMultiplierToSemitoneVertexNames;
+
+		FDataReferenceCollection OutputDataReferences;
+		OutputDataReferences.AddDataReadReference(METASOUND_GET_PARAM_NAME(OutputSemitone), SemitoneOutput);
+		return OutputDataReferences;
+	}
+
+	void FFrequencyMultiplierToSemitoneOperator::Execute()
+	{
+		// Only do anything if the scale changes
+		if (!FMath::IsNearlyEqual(*FrequencyMultiplierInput, PrevFrequencyMultiplier))
+		{
+			PrevFrequencyMultiplier = *FrequencyMultiplierInput;
+			*SemitoneOutput = Audio::GetSemitones(*FrequencyMultiplierInput);
+		}
+	}
+
+	const FVertexInterface& FFrequencyMultiplierToSemitoneOperator::GetVertexInterface()
+	{
+		using namespace FrequencyMultiplierToSemitoneVertexNames;
+
+		static const FVertexInterface Interface(
+			FInputVertexInterface(
+				TInputDataVertexModel<float>(METASOUND_GET_PARAM_NAME_AND_TT(InputFrequencyMultiplier), 1.0f)
+			),
+			FOutputVertexInterface(
+				TOutputDataVertexModel<float>(METASOUND_GET_PARAM_NAME_AND_TT(OutputSemitone))
+			)
+		);
+
+		return Interface;
+	}
+
+	const FNodeClassMetadata& FFrequencyMultiplierToSemitoneOperator::GetNodeInfo()
+	{
+		auto InitNodeInfo = []() -> FNodeClassMetadata
+		{
+			const FName DataTypeName = GetMetasoundDataTypeName<float>();
+			const FName OperatorName = TEXT("Frequency Multiplier to Semitone");
+			const FText NodeDisplayName = LOCTEXT("Metasound_FrequencyMultiplierToSemitoneName", "Frequency Multiplier to Semitone");
+			const FText NodeDescription = LOCTEXT("Metasound_FrequencyMultiplierToSemitoneDescription", "Converts a frequency multiplier to the corresponding difference in semitones.");
+
+			FNodeClassMetadata Info;
+			Info.ClassName = { Metasound::StandardNodes::Namespace, OperatorName, DataTypeName };
+			Info.MajorVersion = 1;
+			Info.MinorVersion = 0;
+			Info.DisplayName = NodeDisplayName;
+			Info.Description = NodeDescription;
+			Info.Author = PluginAuthor;
+			Info.PromptIfMissing = PluginNodeMissingPrompt;
+			Info.DefaultInterface = GetVertexInterface();
+			Info.CategoryHierarchy.Emplace(StandardNodes::Conversions);
+
+			return Info;
+		};
+
+		static const FNodeClassMetadata Info = InitNodeInfo();
+
+		return Info;
+	}
+
+	TUniquePtr<IOperator> FFrequencyMultiplierToSemitoneOperator::CreateOperator(const FCreateOperatorParams& InParams, FBuildErrorArray& OutErrors)
+	{
+		using namespace FrequencyMultiplierToSemitoneVertexNames;
+
+		const FDataReferenceCollection& InputCollection = InParams.InputDataReferences;
+		const FInputVertexInterface& InputInterface = GetVertexInterface().GetInputInterface();
+
+		FFloatReadRef InFrequencyMultiplier = InputCollection.GetDataReadReferenceOrConstructWithVertexDefault<float>(InputInterface, METASOUND_GET_PARAM_NAME(InputFrequencyMultiplier), InParams.OperatorSettings);
+
+		return MakeUnique<FFrequencyMultiplierToSemitoneOperator>(InParams.OperatorSettings, InFrequencyMultiplier);
+	}
+
+	class METASOUNDSTANDARDNODES_API FFrequencyMultiplierToSemitoneNode : public FNodeFacade
+	{
+	public:
+		/**
+		 * Constructor used by the Metasound Frontend.
+		 */
+		FFrequencyMultiplierToSemitoneNode(const FNodeInitData& InitData)
+			: FNodeFacade(InitData.InstanceName, InitData.InstanceID, TFacadeOperatorClass <FFrequencyMultiplierToSemitoneOperator>())
+		{
+		}
+	};
+
+	METASOUND_REGISTER_NODE(FFrequencyMultiplierToSemitoneNode)
+}
+
+#undef LOCTEXT_NAMESPACE
diff --git a/Engine/Plugins/Runtime/Metasound/Source/MetasoundStandardNodes/Private/MetasoundLinearGainToDecibels.cpp b/Engine/Plugins/Runtime/Metasound/Source/MetasoundStandardNodes/Private/MetasoundLinearGainToDecibels.cpp
new file mode 100644
--- /dev/null
+++ b/Engine/Plugins/Runtime/Metasound/Source/MetasoundStandardNodes/Private/MetasoundLinearGainToDecibels.cpp
@@ -0,0 +1,155 @@
+// Copyright Epic Games, Inc. All Rights Reserved.
+
+#include "Internationalization/Text.h"
+#include "MetasoundFacade.h"
+#include "MetasoundExecutableOperator.h"
+#include "MetasoundNodeRegistrationMacro.h"
+#include "MetasoundPrimitives.h"
+#include "MetasoundStandardNodesNames.h"
+#include "MetasoundStandardNodesCategories.h"
+#include "DSP/Dsp.h"
+#include "MetasoundParamHelper.h"
+
+#define LOCTEXT_NAMESPACE "MetasoundStandardNodes_LinearGainToDecibels"
+
+namespace Metasound
+{
+
+	namespace LinearGainToDecibelsVertexNames
+	{
+		METASOUND_PARAM(InputLinearGain, "Linear Gain", "Input linear gain.");
+
+		METASOUND_PARAM(OutputDecibels, "Decibels", "Output corresponding logarithmic (dB) gain.");
+	}
+
+
+	class FLinearGainToDecibelsOperator : public TExecutableOperator<FLinearGainToDecibelsOperator>
+	{
+	public:
+
+		static const FNodeClassMetadata& GetNodeInfo();
+		static const FVertexInterface& GetVertexInterface();
+		static TUniquePtr<IOperator> CreateOperator(const FCreateOperatorParams& InParams, FBuildErrorArray& OutErrors);
+
+		FLinearGainToDecibelsOperator(const FOperatorSettings& InSettings, const FFloatReadRef& InLinearGain);
+
+		virtual FDataReferenceCollection GetInputs() const override;
+		virtual FDataReferenceCollection GetOutputs() const override;
+		void Execute();
+
+	private:
+		// The input linear gain value.
+		FFloatReadRef LinearGainInput;
+
+		// The output gain in dB.
+		FFloatWriteRef DecibelOutput;
+	};
+
+	FLinearGainToDecibelsOperator::FLinearGainToDecibelsOperator(const FOperatorSettings& InSettings, const FFloatReadRef& InLinearGain)
+		: LinearGainInput(InLinearGain)
+		, DecibelOutput(FFloatWriteRef::CreateNew(FMath::Max(Audio::ConvertToDecibels(*LinearGainInput, TNumericLimits<float>::Min()), -240.0f)))
+	{
+	}
+
+	
+	FDataReferenceCollection FLinearGainToDecibelsOperator::GetInputs() const
+	{
+		using namespace LinearGainToDecibelsVertexNames;
+
+		FDataReferenceCollection InputDataReferences;
+		InputDataReferences.AddDataReadReference(METASOUND_GET_PARAM_NAME(InputLinearGain), LinearGainInput);
+
+		return InputDataReferences;
+	}
+
+	FDataReferenceCollection FLinearGainToDecibelsOperator::GetOutputs() const
+	{
+		using namespace LinearGainToDecibelsVertexNames;
+
+		FDataReferenceCollection OutputDataReferences;
+		OutputDataReferences.AddDataReadReference(METASOUND_GET_PARAM_NAME(OutputDecibels), DecibelOutput);
+		return OutputDataReferences;
+	}
+
+	void FLinearGainToDecibelsOperator::Execute()
+	{
+		// Linear gain value must be above 0
+		const float ClampedInput = FMath::Max(*LinearGainInput, TNumericLimits<float>::Min());
+		// Convert input to decibel gain value
+		const float UnclampedDBOutput = Audio::ConvertToDecibels(ClampedInput, TNumericLimits<float>::Min());
+		// Clamp to extremely low decibel value
+		*DecibelOutput = FMath::Max(UnclampedDBOutput, -240.0f);
+	}
+
+	const FVertexInterface& FLinearGainToDecibelsOperator::GetVertexInterface()
+	{
+		using namespace LinearGainToDecibelsVertexNames;
+
+		static const FVertexInterface Interface(
+			FInputVertexInterface(
+				TInputDataVertexModel<float>(METASOUND_GET_PARAM_NAME_AND_TT(InputLinearGain), 1.0f)
+			),
+			FOutputVertexInterface(
+				TOutputDataVertexModel<float>(METASOUND_GET_PARAM_NAME_AND_TT(OutputDecibels))
+			)
+		);
+
+		return Interface;
+	}
+
+	const FNodeClassMetadata& FLinearGainToDecibelsOperator::GetNodeInfo()
+	{
+		auto InitNodeInfo = []() -> FNodeClassMetadata
+		{
+			const FName DataTypeName = GetMetasoundDataTypeName<float>();
+			const FName OperatorName = TEXT("Linear Gain to Decibels");
+			const FText NodeDisplayName = LOCTEXT("Metasound_LinearGainToDecibelsName", "Linear Gain to Decibels");
+			const FText NodeDescription = LOCTEXT("Metasound_LinearGainToDecibelsDescription", "Converts a linear gain value to a logarithmic (dB) gain value.");
+
+			FNodeClassMetadata Info;
+			Info.ClassName = { Metasound::StandardNodes::Namespace, OperatorName, DataTypeName };
+			Info.MajorVersion = 1;
+			Info.MinorVersion = 0;
+			Info.DisplayName = NodeDisplayName;
+			Info.Description = NodeDescription;
+			Info.Author = PluginAuthor;
+			Info.PromptIfMissing = PluginNodeMissingPrompt;
+			Info.DefaultInterface = GetVertexInterface();
+			Info.CategoryHierarchy.Emplace(StandardNodes::Dynamics);
+
+			return Info;
+		};
+
+		static const FNodeClassMetadata Info = InitNodeInfo();
+
+		return Info;
+	}
+
+	TUniquePtr<IOperator> FLinearGainToDecibelsOperator::CreateOperator(const FCreateOperatorParams& InParams, FBuildErrorArray& OutErrors)
+	{
+		using namespace LinearGainToDecibelsVertexNames;
+
+		const FDataReferenceCollection& InputCollection = InParams.InputDataReferences;
+		const FInputVertexInterface& InputInterface = GetVertexInterface().GetInputInterface();
+
+		FFloatReadRef InLinearGain = InputCollection.GetDataReadReferenceOrConstructWithVertexDefault<float>(InputInterface, METASOUND_GET_PARAM_NAME(InputLinearGain), InParams.OperatorSettings);
+		
+		return MakeUnique<FLinearGainToDecibelsOperator>(InParams.OperatorSettings, InLinearGain);
+	}
+
+	class METASOUNDSTANDARDNODES_API FLinearGainToDecibelsNode : public FNodeFacade
+	{
+	public:
+		/**
+		 * Constructor used by the Metasound Frontend.
+		 */
+		FLinearGainToDecibelsNode(const FNodeInitData& InitData)
+			: FNodeFacade(InitData.InstanceName, InitData.InstanceID, TFacadeOperatorClass <FLinearGainToDecibelsOperator>())
+		{
+		}
+	};
+
+	METASOUND_REGISTER_NODE(FLinearGainToDecibelsNode)
+}
+
+#undef LOCTEXT_NAMESPACE
diff --git a/Engine/Plugins/Runtime/Metasound/Source/MetasoundStandardNodes/Private/MetasoundSemitoneToFreqMultiplierNode.cpp b/Engine/Plugins/Runtime/Metasound/Source/MetasoundStandardNodes/Private/MetasoundSemitoneToFreqMultiplierNode.cpp
new file mode 100644
--- /dev/null
+++ b/Engine/Plugins/Runtime/Metasound/Source/MetasoundStandardNodes/Private/MetasoundSemitoneToFreqMultiplierNode.cpp
@@ -0,0 +1,161 @@
+// Copyright Epic Games, Inc. All Rights Reserved.
+
+#include "Internationalization/Text.h"
+#include "MetasoundFacade.h"
+#include "MetasoundExecutableOperator.h"
+#include "MetasoundNodeRegistrationMacro.h"
+#include "MetasoundPrimitives.h"
+#include "MetasoundStandardNodesNames.h"
+#include "MetasoundStandardNodesCategories.h"
+#include "DSP/Dsp.h"
+#include "MetasoundParamHelper.h"
+
+#define LOCTEXT_NAMESPACE "MetasoundStandardNodes_SemitoneToFrequencyMultiplier"
+
+namespace Metasound
+{
+
+	namespace SemitoneToFrequencyMultiplierVertexNames
+	{
+		METASOUND_PARAM(InputSemitone, "Semitones", "Input difference in semitones.");
+
+		METASOUND_PARAM(OutputFrequencyMultiplier, "Frequency Multiplier", "Output corresponding frequency multiplier.");
+	}
+
+
+	class FSemitoneToFrequencyMultiplierOperator : public TExecutableOperator<FSemitoneToFrequencyMultiplierOperator>
+	{
+	public:
+
+		static const FNodeClassMetadata& GetNodeInfo();
+		static const FVertexInterface& GetVertexInterface();
+		static TUniquePtr<IOperator> CreateOperator(const FCreateOperatorParams& InParams, FBuildErrorArray& OutErrors);
+
+		FSemitoneToFrequencyMultiplierOperator(const FOperatorSettings& InSettings, const FFloatReadRef& InSemitone);
+
+		virtual FDataReferenceCollection GetInputs() const override;
+		virtual FDataReferenceCollection GetOutputs() const override;
+		void Execute();
+
+	private:
+		// The input difference in semitones
+		FFloatReadRef SemitoneInput;
+
+		// The output frequency multiplier
+		FFloatWriteRef FrequencyMultiplierOutput;
+
+		// Cached semitone value. Used to catch if the value changes to recompute Frequency Multiplier output.
+		float PrevSemitone;
+	};
+
+	FSemitoneToFrequencyMultiplierOperator::FSemitoneToFrequencyMultiplierOperator(const FOperatorSettings& InSettings, const FFloatReadRef& InSemitone)
+		: SemitoneInput(InSemitone)
+		, FrequencyMultiplierOutput(FFloatWriteRef::CreateNew(Audio::GetFrequencyMultiplier(*InSemitone)))
+		, PrevSemitone(*InSemitone)
+	{
+	}
+
+	
+	FDataReferenceCollection FSemitoneToFrequencyMultiplierOperator::GetInputs() const
+	{
+		using namespace SemitoneToFrequencyMultiplierVertexNames;
+
+		FDataReferenceCollection InputDataReferences;
+		InputDataReferences.AddDataReadReference(METASOUND_GET_PARAM_NAME(InputSemitone), SemitoneInput);
+
+		return InputDataReferences;
+	}
+
+	FDataReferenceCollection FSemitoneToFrequencyMultiplierOperator::GetOutputs() const
+	{
+		using namespace SemitoneToFrequencyMultiplierVertexNames;
+
+		FDataReferenceCollection OutputDataReferences;
+		OutputDataReferences.AddDataReadReference(METASOUND_GET_PARAM_NAME(OutputFrequencyMultiplier), FrequencyMultiplierOutput);
+		return OutputDataReferences;
+	}
+
+	void FSemitoneToFrequencyMultiplierOperator::Execute()
+	{
+		const float CurrSemitone = *SemitoneInput;
+
+		// Only do anything if the input changes
+		if (!FMath::IsNearlyEqual(CurrSemitone, PrevSemitone))
+		{
+			PrevSemitone = CurrSemitone;
+			*FrequencyMultiplierOutput = Audio::GetFrequencyMultiplier(CurrSemitone);
+		}
+	}
+
+	const FVertexInterface& FSemitoneToFrequencyMultiplierOperator::GetVertexInterface()
+	{
+		using namespace SemitoneToFrequencyMultiplierVertexNames;
+
+		static const FVertexInterface Interface(
+			FInputVertexInterface(
+				TInputDataVertexModel<float>(METASOUND_GET_PARAM_NAME_AND_TT(InputSemitone), 0.0f)
+			),
+			FOutputVertexInterface(
+				TOutputDataVertexModel<float>(METASOUND_GET_PARAM_NAME_AND_TT(OutputFrequencyMultiplier))
+			)
+		);
+
+		return Interface;
+	}
+
+	const FNodeClassMetadata& FSemitoneToFrequencyMultiplierOperator::GetNodeInfo()
+	{
+		auto InitNodeInfo = []() -> FNodeClassMetadata
+		{
+			const FName DataTypeName = GetMetasoundDataTypeName<float>();
+			const FName OperatorName = TEXT("Semitone to Frequency Multiplier");
+			const FText NodeDisplayName = LOCTEXT("Metasound_SemitoneToFrequencyMultiplierName", "Semitone to Frequency Multiplier");
+			const FText NodeDescription = LOCTEXT("Metasound_SemitoneToFrequencyMultiplierDescription", "Converts a number of semitones to the corresponding frequency multiplier.");
+
+			FNodeClassMetadata Info;
+			Info.ClassName = { Metasound::StandardNodes::Namespace, OperatorName, DataTypeName };
+			Info.MajorVersion = 1;
+			Info.MinorVersion = 0;
+			Info.DisplayName = NodeDisplayName;
+			Info.Description = NodeDescription;
+			Info.Author = PluginAuthor;
+			Info.PromptIfMissing = PluginNodeMissingPrompt;
+			Info.DefaultInterface = GetVertexInterface();
+			Info.CategoryHierarchy.Emplace(StandardNodes::Conversions);
+
+			return Info;
+		};
+
+		static const FNodeClassMetadata Info = InitNodeInfo();
+
+		return Info;
+	}
+
+	TUniquePtr<IOperator> FSemitoneToFrequencyMultiplierOperator::CreateOperator(const FCreateOperatorParams& InParams, FBuildErrorArray& OutErrors)
+	{
+		using namespace SemitoneToFrequencyMultiplierVertexNames;
+
+		const FDataReferenceCollection& InputCollection = InParams.InputDataReferences;
+		const FInputVertexInterface& InputInterface = GetVertexInterface().GetInputInterface();
+
+		FFloatReadRef InSemitone = InputCollection.GetDataReadReferenceOrConstructWithVertexDefault<float>(InputInterface, METASOUND_GET_PARAM_NAME(InputSemitone), InParams.OperatorSettings);
+
+		return MakeUnique<FSemitoneToFrequencyMultiplierOperator>(InParams.OperatorSettings, InSemitone);
+	}
+
+	class METASOUNDSTANDARDNODES_API FSemitoneToFrequencyMultiplierNode : public FNodeFacade
+	{
+	public:
+		/**
+		 * Constructor used by the Metasound Frontend.
+		 */
+		FSemitoneToFrequencyMultiplierNode(const FNodeInitData& InitData)
+			: FNodeFacade(InitData.InstanceName, InitData.InstanceID, TFacadeOperatorClass <FSemitoneToFrequencyMultiplierOperator>())
+		{
+		}
+	};
+
+	METASOUND_REGISTER_NODE(FSemitoneToFrequencyMultiplierNode)
+}
+
+#undef LOCTEXT_NAMESPACE
diff --git a/Engine/Source/Runtime/SignalProcessing/Public/DSP/Dsp.h b/Engine/Source/Runtime/SignalProcessing/Public/DSP/Dsp.h
--- a/Engine/Source/Runtime/SignalProcessing/Public/DSP/Dsp.h
+++ b/Engine/Source/Runtime/SignalProcessing/Public/DSP/Dsp.h
@@ -362,6 +362,16 @@ namespace Audio
 		return FMath::Pow(2.0f, InPitchSemitones / 12.0f);
 	}
 
+	// Returns the number of semitones relative to a base frequency given the input frequency multiplier
+	static FORCEINLINE float GetSemitones(const float InMultiplier)
+	{
+		if (InMultiplier <= 0.0f)
+		{
+			return 12.0f * FMath::Log2(SMALL_NUMBER);
+		}
+		return 12.0f * FMath::Log2(InMultiplier);
+	}
+
 	// Calculates equal power stereo pan using sinusoidal-panning law and cheap approximation for sin
 	// InLinear pan is [-1.0, 1.0] so it can be modulated by a bipolar LFO
 	static FORCEINLINE void GetStereoPan(const float InLinearPan, float& OutLeft, float& OutRight)
