Add Bitcrusher node to metasounds

diff --git a/Engine/Plugins/Runtime/Metasound/Source/MetasoundStandardNodes/Private/MetasoundBitcrusherNode.cpp b/Engine/Plugins/Runtime/Metasound/Source/MetasoundStandardNodes/Private/MetasoundBitcrusherNode.cpp
new file mode 100644
--- /dev/null
+++ b/Engine/Plugins/Runtime/Metasound/Source/MetasoundStandardNodes/Private/MetasoundBitcrusherNode.cpp
@@ -0,0 +1,196 @@
+// Copyright Epic Games, Inc. All Rights Reserved.
+
+#include "Internationalization/Text.h"
+#include "MetasoundExecutableOperator.h"
+#include "MetasoundEnumRegistrationMacro.h"
+#include "MetasoundNodeRegistrationMacro.h"
+#include "MetasoundPrimitives.h"
+#include "MetasoundStandardNodesNames.h"
+#include "MetasoundAudioBuffer.h"
+#include "DSP/Dsp.h"
+#include "MetasoundStandardNodesCategories.h"
+#include "MetasoundFacade.h"
+#include "MetasoundParamHelper.h"
+#include "DSP/FloatArrayMath.h"
+#include "DSP/BitCrusher.h"
+
+#define LOCTEXT_NAMESPACE "MetasoundStandardNodes_BitcrusherNode"
+
+namespace Metasound
+{
+	/* Mid-Side Encoder */
+	namespace BitcrusherVertexNames
+	{
+		METASOUND_PARAM(InputAudio, "Audio", "Incoming audio signal to bitcrush.");
+		METASOUND_PARAM(InputSampleRate, "Sample Rate", "The sampling frequency to downsample the audio to.");
+		METASOUND_PARAM(InputBitDepth, "Bit Depth", "The bit resolution to reduce the audio to.");
+
+		METASOUND_PARAM(OutputAudio, "Audio", "The bitcrushed audio signal.");
+	}
+
+	// Operator Class
+	class FBitcrusherOperator : public TExecutableOperator<FBitcrusherOperator>
+	{
+	public:
+
+		FBitcrusherOperator(const FOperatorSettings& InSettings,
+			const FAudioBufferReadRef& InAudio,
+			const FFloatReadRef& InCrushedSampleRate,
+			const FFloatReadRef& InCrushedBitDepth)
+			: AudioInput(InAudio)
+			, AudioOutput(FAudioBufferWriteRef::CreateNew(InSettings))
+			, CrushedSampleRate(InCrushedSampleRate)
+			, CrushedBitDepth(InCrushedBitDepth)
+		{
+			// Passing in 1 for NumChannels because the node takes one audio input
+			Bitcrusher.Init(InSettings.GetSampleRate(), 1);
+			Bitcrusher.SetSampleRateCrush(*CrushedSampleRate);
+			Bitcrusher.SetBitDepthCrush(*CrushedBitDepth);
+		}
+
+		static const FNodeClassMetadata& GetNodeInfo()
+		{
+			auto CreateNodeClassMetadata = []() -> FNodeClassMetadata
+			{
+				FVertexInterface NodeInterface = DeclareVertexInterface();
+
+				FNodeClassMetadata Metadata
+				{
+					FNodeClassName{StandardNodes::Namespace, TEXT("Bitcrusher"), StandardNodes::AudioVariant},
+					1, // Major Version
+					0, // Minor Version
+					LOCTEXT("BitcrusherDisplayName", "Bitcrusher"),
+					LOCTEXT("BitcrusherDesc", "Downsamples and lowers the bit-depth of an incoming audio signal."),
+					PluginAuthor,
+					PluginNodeMissingPrompt,
+					NodeInterface,
+					{StandardNodes::Filters},
+					{TEXT("Bitcrusher")},
+					FNodeDisplayStyle{}
+				};
+
+				return Metadata;
+			};
+
+			static const FNodeClassMetadata Metadata = CreateNodeClassMetadata();
+			return Metadata;
+		}
+
+		static const FVertexInterface& DeclareVertexInterface()
+		{
+			using namespace BitcrusherVertexNames;
+
+			static const FVertexInterface Interface(
+				FInputVertexInterface(
+					TInputDataVertexModel<FAudioBuffer>(METASOUND_GET_PARAM_NAME_AND_TT(InputAudio)),
+					TInputDataVertexModel<float>(METASOUND_GET_PARAM_NAME_AND_TT(InputSampleRate), 8000.0f),
+					TInputDataVertexModel<float>(METASOUND_GET_PARAM_NAME_AND_TT(InputBitDepth), 8.0f)
+				),
+				FOutputVertexInterface(
+					TOutputDataVertexModel<FAudioBuffer>(METASOUND_GET_PARAM_NAME_AND_TT(OutputAudio))
+				)
+			);
+
+			return Interface;
+		}
+
+		virtual FDataReferenceCollection GetInputs() const override
+		{
+			using namespace BitcrusherVertexNames;
+
+			FDataReferenceCollection InputDataReferences;
+
+			InputDataReferences.AddDataReadReference(METASOUND_GET_PARAM_NAME(InputAudio), AudioInput);
+			InputDataReferences.AddDataReadReference(METASOUND_GET_PARAM_NAME(InputSampleRate), CrushedSampleRate);
+			InputDataReferences.AddDataReadReference(METASOUND_GET_PARAM_NAME(InputBitDepth), CrushedBitDepth);
+
+			return InputDataReferences;
+
+		}
+
+		virtual FDataReferenceCollection GetOutputs() const override
+		{
+			using namespace BitcrusherVertexNames;
+
+			FDataReferenceCollection OutputDataReferences;
+
+			OutputDataReferences.AddDataReadReference(METASOUND_GET_PARAM_NAME(OutputAudio), AudioOutput);
+
+			return OutputDataReferences;
+		}
+
+		static TUniquePtr<IOperator> CreateOperator(const FCreateOperatorParams& InParams, FBuildErrorArray& OutErrors)
+		{
+			using namespace BitcrusherVertexNames;
+			const FDataReferenceCollection& Inputs = InParams.InputDataReferences;
+			const FInputVertexInterface& InputInterface = DeclareVertexInterface().GetInputInterface();
+
+			FAudioBufferReadRef AudioIn = Inputs.GetDataReadReferenceOrConstruct<FAudioBuffer>(METASOUND_GET_PARAM_NAME(InputAudio), InParams.OperatorSettings);
+			FFloatReadRef SampleRateIn = Inputs.GetDataReadReferenceOrConstructWithVertexDefault<float>(InputInterface, METASOUND_GET_PARAM_NAME(InputSampleRate), InParams.OperatorSettings);
+			FFloatReadRef BitDepthIn = Inputs.GetDataReadReferenceOrConstructWithVertexDefault<float>(InputInterface, METASOUND_GET_PARAM_NAME(InputBitDepth), InParams.OperatorSettings);
+
+			return MakeUnique<FBitcrusherOperator>(InParams.OperatorSettings, AudioIn, SampleRateIn, BitDepthIn);
+		}
+
+		void Execute()
+		{
+			const int32 NumFrames = AudioInput->Num();
+			if (NumFrames != AudioOutput->Num())
+			{
+				return;
+			}
+
+			const float CurrSampleRate = FMath::Clamp(*CrushedSampleRate, 1.0f, 16000.0f);
+			if (!FMath::IsNearlyEqual(PrevSampleRate, CurrSampleRate))
+			{
+				Bitcrusher.SetSampleRateCrush(CurrSampleRate);
+				PrevSampleRate = CurrSampleRate;
+			}
+
+			const float CurrBitDepth = FMath::Clamp(*CrushedBitDepth, 1.0f, 16.0f);
+			if (!FMath::IsNearlyEqual(PrevBitDepth, CurrBitDepth))
+			{
+				Bitcrusher.SetBitDepthCrush(CurrBitDepth);
+				PrevBitDepth = CurrBitDepth;
+			}
+
+			Bitcrusher.ProcessAudio(AudioInput->GetData(), NumFrames, AudioOutput->GetData());
+		}
+
+
+	private:
+
+		// Audio input and output
+		FAudioBufferReadRef AudioInput;
+		FAudioBufferWriteRef AudioOutput;
+		
+		// User-set sample rate the audio will be converted to
+		FFloatReadRef CrushedSampleRate;
+		// User-set bit depth for each sample
+		FFloatReadRef CrushedBitDepth;
+		
+		// Internal DSP bitcrusher
+		Audio::FBitCrusher Bitcrusher;
+
+		// Cached input parameters
+		float PrevSampleRate;
+		float PrevBitDepth;
+	
+	};
+
+	// Node Class
+	class FBitcrusherNode : public FNodeFacade
+	{
+	public:
+		FBitcrusherNode(const FNodeInitData& InitData)
+			: FNodeFacade(InitData.InstanceName, InitData.InstanceID, TFacadeOperatorClass<FBitcrusherOperator>())
+		{
+		}
+	};
+
+	// Register node
+	METASOUND_REGISTER_NODE(FBitcrusherNode)
+
+}
+
+#undef LOCTEXT_NAMESPACE
