Fix compile error on older MSVC that can't handle defaulting class functions inside templates, no functional change

diff --git a/Engine/Plugins/Runtime/Metasound/Source/MetasoundFrontend/Public/MetasoundDataTypeRegistrationMacro.h b/Engine/Plugins/Runtime/Metasound/Source/MetasoundFrontend/Public/MetasoundDataTypeRegistrationMacro.h
--- a/Engine/Plugins/Runtime/Metasound/Source/MetasoundFrontend/Public/MetasoundDataTypeRegistrationMacro.h
+++ b/Engine/Plugins/Runtime/Metasound/Source/MetasoundFrontend/Public/MetasoundDataTypeRegistrationMacro.h
@@ -381,7 +381,7 @@ namespace Metasound
 					TGetVariableNode<TDataType> VariablePrototype(TEXT(""), FGuid(), TEXT(""), FLiteral());
 					VariableClass = Metasound::Frontend::GenerateClassDescription(VariablePrototype.GetMetadata(), EMetasoundFrontendClassType::Variable);
 				}
-				virtual ~FDataTypeRegistryEntry() = default;
+				virtual ~FDataTypeRegistryEntry() {}
 
 				virtual const Frontend::FDataTypeRegistryInfo& GetDataTypeInfo() const override
 				{
