Fixing typo in tooltip

diff --git a/Engine/Plugins/Runtime/Metasound/Source/MetasoundStandardNodes/Private/MetasoundADEnvelopeNode.cpp b/Engine/Plugins/Runtime/Metasound/Source/MetasoundStandardNodes/Private/MetasoundADEnvelopeNode.cpp
--- a/Engine/Plugins/Runtime/Metasound/Source/MetasoundStandardNodes/Private/MetasoundADEnvelopeNode.cpp
+++ b/Engine/Plugins/Runtime/Metasound/Source/MetasoundStandardNodes/Private/MetasoundADEnvelopeNode.cpp
@@ -25,7 +25,7 @@ namespace Metasound
 		METASOUND_PARAM(InputAttackTime, "Attack Time", "Attack time of the envelope.");
 		METASOUND_PARAM(InputDecayTime, "Decay Time", "Decay time of the envelope.");
 		METASOUND_PARAM(InputAttackCurve, "Attack Curve", "The exponential curve factor of the attack. 1.0 = linear growth, < 1.0 logorithmic growth, > 1.0 exponential growth.");
-		METASOUND_PARAM(InputDecayCurve, "Decay Curve", "The exponential curve factor of the decay. 1.0 = linear decay, < 1.0 exponential decay, > 1.0 logorithmic decay.");
+		METASOUND_PARAM(InputDecayCurve, "Decay Curve", "The exponential curve factor of the decay. 1.0 = linear decay, < 1.0 exponential decay, > 1.0 logarithmic decay.");
 		METASOUND_PARAM(InputLooping, "Looping", "Set to true to enable looping of the envelope. This will allow the envelope to be an LFO or wave generator.");
 
 		METASOUND_PARAM(OutputOnTrigger, "On Trigger", "Triggers when the envelope is triggered.");
