Creating metasound archetype registry and upgrade path

diff --git a/Engine/Plugins/Runtime/Metasound/Source/MetasoundEditor/Private/MetasoundEditor.cpp b/Engine/Plugins/Runtime/Metasound/Source/MetasoundEditor/Private/MetasoundEditor.cpp
--- a/Engine/Plugins/Runtime/Metasound/Source/MetasoundEditor/Private/MetasoundEditor.cpp
+++ b/Engine/Plugins/Runtime/Metasound/Source/MetasoundEditor/Private/MetasoundEditor.cpp
@@ -134,8 +134,7 @@ namespace Metasound
 
 			bool IsRequired() const
 			{
-				const FMetasoundFrontendArchetype& Archetype = GetMetasoundAssetChecked().GetArchetype();
-				return GetNodeHandle()->IsRequired(Archetype);
+				return GetNodeHandle()->IsRequired();
 			}
 
 			// FEdGraphSchemaAction interface
@@ -346,13 +345,10 @@ namespace Metasound
 			FMetasoundAssetBase* MetasoundAsset = IMetasoundUObjectRegistry::Get().GetObjectAsAssetBase(Metasound);
 			check(MetasoundAsset);
 
-			Frontend::FDocumentHandle Document = MetasoundAsset->GetDocumentHandle();
-			FName AssetName = Metasound->GetFName();
-			FString AssetPath = Metasound->GetPathName();
-			if (Frontend::FVersionDocument(AssetName, AssetPath).Transform(Document))
+			const bool bDidUpdateAsset = MetasoundAsset->VersionAsset();
+			if (bDidUpdateAsset)
 			{
 				MetasoundAsset->RegisterGraphWithFrontend();
-				Metasound->MarkPackageDirty();
 			}
 			FGraphBuilder::SynchronizeGraph(*Metasound);
 
@@ -775,46 +771,40 @@ namespace Metasound
 
 		void FEditor::Import()
 		{
-			FMetasoundAssetBase* MetasoundAsset = IMetasoundUObjectRegistry::Get().GetObjectAsAssetBase(Metasound);
-			if (MetasoundAsset)
-			{
-				const FMetasoundFrontendArchetype& Archetype = MetasoundAsset->GetArchetype();
+			// TODO: Prompt OFD and provide path from user
+			const FString InputPath = FPaths::ProjectIntermediateDir() / TEXT("MetaSounds") + FPaths::ChangeExtension(Metasound->GetPathName(), FMetasoundAssetBase::FileExtension);
+			
+			// TODO: use the same directory as the currently open MetaSound
+			const FString OutputPath = FString("/Game/ImportedMetaSound/GeneratedMetaSound");
 
-				// TODO: Prompt OFD and provide path from user
-				const FString InputPath = FPaths::ProjectIntermediateDir() / TEXT("MetaSounds") + FPaths::ChangeExtension(Metasound->GetPathName(), FMetasoundAssetBase::FileExtension);
-				
-				// TODO: use the same directory as the currently open MetaSound
-				const FString OutputPath = FString("/Game/ImportedMetaSound/GeneratedMetaSound");
+			FMetasoundFrontendDocument MetasoundDoc;
 
-				FMetasoundFrontendDocument MetasoundDoc;
+			if (Frontend::ImportJSONAssetToMetasound(InputPath, MetasoundDoc))
+			{
+				TArray<UClass*> ImportClasses = IMetasoundUObjectRegistry::Get().GetUClassesForArchetype(MetasoundDoc.ArchetypeVersion);
 
-				if (Frontend::ImportJSONAssetToMetasound(InputPath, MetasoundDoc))
+				if (ImportClasses.Num() < 1)
 				{
-					TArray<UClass*> ImportClasses = IMetasoundUObjectRegistry::Get().GetUClassesForArchetype(Archetype.Name);
-
-					if (ImportClasses.Num() < 1)
-					{
-						UE_LOG(LogMetaSound, Warning, TEXT("Cannot create UObject from MetaSound document. No UClass supports archetype \"%s\""), *Archetype.Name.ToString());
-					}
-					else
+					UE_LOG(LogMetaSound, Warning, TEXT("Cannot create UObject from MetaSound document. No UClass supports archetype \"%s\""), *MetasoundDoc.ArchetypeVersion.ToString());
+				}
+				else
+				{
+					if (ImportClasses.Num() > 1)
 					{
-						if (ImportClasses.Num() > 1)
+						for (UClass* Cls : ImportClasses)
 						{
-							for (UClass* Cls : ImportClasses)
-							{
-								// TODO: could do a modal dialog to give user choice of import type.
-								UE_LOG(LogMetaSound, Warning, TEXT("Duplicate UClass support archetype \"%s\" with UClass \"%s\""), *Archetype.Name.ToString(), *Cls->GetName());
-							}
+							// TODO: could do a modal dialog to give user choice of import type.
+							UE_LOG(LogMetaSound, Warning, TEXT("Duplicate UClass support archetype \"%s\" with UClass \"%s\""), *MetasoundDoc.ArchetypeVersion.ToString(), *Cls->GetName());
 						}
-
-						IMetasoundUObjectRegistry::Get().NewObject(ImportClasses[0], MetasoundDoc, Archetype, OutputPath);
 					}
-				}
-				else
-				{
-					UE_LOG(LogMetaSound, Warning, TEXT("Could not import MetaSound at path: %s"), *InputPath);
+
+					IMetasoundUObjectRegistry::Get().NewObject(ImportClasses[0], MetasoundDoc, OutputPath);
 				}
 			}
+			else
+			{
+				UE_LOG(LogMetaSound, Warning, TEXT("Could not import MetaSound at path: %s"), *InputPath);
+			}
 		}
 
 		void FEditor::Export()
@@ -1378,8 +1368,7 @@ namespace Metasound
 
 					FMetasoundAssetBase* MetasoundAsset = IMetasoundUObjectRegistry::Get().GetObjectAsAssetBase(Metasound);
 					check(MetasoundAsset);
-					const FMetasoundFrontendArchetype& Archetype = MetasoundAsset->GetArchetype();
-
+					
 					TSharedPtr<FMetasoundGraphNodeSchemaAction> ActionToDelete;
 					for (const TSharedPtr<FEdGraphSchemaAction>& Action : Actions)
 					{
@@ -1389,7 +1378,7 @@ namespace Metasound
 							Frontend::FNodeHandle NodeHandle = MetasoundAction->GetNodeHandle();
 							if (ensure(NodeHandle->IsValid()))
 							{
-								if (NodeHandle->IsRequired(Archetype))
+								if (NodeHandle->IsRequired())
 								{
 									if (MetasoundGraphEditor.IsValid())
 									{
diff --git a/Engine/Plugins/Runtime/Metasound/Source/MetasoundEditor/Private/MetasoundEditorGraph.cpp b/Engine/Plugins/Runtime/Metasound/Source/MetasoundEditor/Private/MetasoundEditorGraph.cpp
--- a/Engine/Plugins/Runtime/Metasound/Source/MetasoundEditor/Private/MetasoundEditorGraph.cpp
+++ b/Engine/Plugins/Runtime/Metasound/Source/MetasoundEditor/Private/MetasoundEditorGraph.cpp
@@ -143,11 +143,7 @@ bool UMetasoundEditorGraphVariable::IsRequired() const
 {
 	using namespace Metasound;
 
-	const FMetasoundAssetBase* MetasoundAsset = IMetasoundUObjectRegistry::Get().GetObjectAsAssetBase(GetOutermostObject());
-	check(MetasoundAsset);
-	const FMetasoundFrontendArchetype& Archetype = MetasoundAsset->GetArchetype();
-
-	return GetConstNodeHandle()->IsRequired(Archetype);
+	return GetConstNodeHandle()->IsRequired();
 }
 
 bool UMetasoundEditorGraphVariable::CanRename(const FText& InNewName, FText& OutError) const
diff --git a/Engine/Plugins/Runtime/Metasound/Source/MetasoundEditor/Private/MetasoundEditorGraphBuilder.cpp b/Engine/Plugins/Runtime/Metasound/Source/MetasoundEditor/Private/MetasoundEditorGraphBuilder.cpp
--- a/Engine/Plugins/Runtime/Metasound/Source/MetasoundEditor/Private/MetasoundEditorGraphBuilder.cpp
+++ b/Engine/Plugins/Runtime/Metasound/Source/MetasoundEditor/Private/MetasoundEditorGraphBuilder.cpp
@@ -852,12 +852,7 @@ namespace Metasound
 
 			// Ensure referenced MetaSound is up-to-date
 			check(MetaSoundReferencedAsset);
-			FName Name = InMetaSoundReferenced.GetFName();
-			FString Path = InMetaSoundReferenced.GetPathName();
-			if (Frontend::FVersionDocument(Name, Path).Transform(MetaSoundReferencedAsset->GetDocumentHandle()))
-			{
-				InMetaSoundReferenced.MarkPackageDirty();
-			}
+			MetaSoundReferencedAsset->VersionAsset();
 
 			// Ensure asset is registered with frontend so the class
 			// can be queried when reference is added. This may modify
@@ -869,15 +864,13 @@ namespace Metasound
 			// node's interface to the preset.
 			GraphBuilderPrivate::InitializeGraph(InMetaSoundPreset);
 
-			const FMetasoundFrontendArchetype& Archetype = MetaSoundPresetAsset->GetArchetype();
-
 			// 1a. Preset graph topology is considered "read-only", so mark as such
 			FMetasoundFrontendGraphStyle Style = PresetGraphHandle->GetGraphStyle();
 			Style.bIsGraphEditable = false;
 			PresetGraphHandle->SetGraphStyle(Style);
 
 			// 1b. Set archetype accordingly
-			const FMetasoundFrontendArchetype& RefArchetype = MetaSoundReferencedAsset->GetArchetype();
+			const FMetasoundFrontendVersion& RefArchetype = MetaSoundReferencedAsset->GetDocumentHandle()->GetArchetypeVersion();
 			if (ensure(MetaSoundPresetAsset->IsArchetypeSupported(RefArchetype)))
 			{
 				MetaSoundPresetAsset->ConformDocumentToArchetype();
@@ -895,7 +888,7 @@ namespace Metasound
 			{
 				// Required inputs are added via the prior step (archetype conformation).
 				// For these, find the corrisponding input & create an editor node only.
-				const bool bIsRequired = RefGraphInputNode->IsRequired(Archetype);
+				const bool bIsRequired = RefGraphInputNode->IsRequired();
 
 				const FText& DisplayName = RefGraphInputNode->GetDisplayName();
 				const FText& Description = RefGraphInputNode->GetDescription();
@@ -936,7 +929,7 @@ namespace Metasound
 			{
 				// Required inputs are added via the prior step (archetype conformation).
 				// For these, find the corrisponding input & create an editor node only.
-				const bool bIsRequired = RefGraphOutputNode->IsRequired(Archetype);
+				const bool bIsRequired = RefGraphOutputNode->IsRequired();
 
 				const FText& DisplayName = RefGraphOutputNode->GetDisplayName();
 				const FText& Description = RefGraphOutputNode->GetDescription();
diff --git a/Engine/Plugins/Runtime/Metasound/Source/MetasoundEditor/Private/MetasoundEditorGraphNode.cpp b/Engine/Plugins/Runtime/Metasound/Source/MetasoundEditor/Private/MetasoundEditorGraphNode.cpp
--- a/Engine/Plugins/Runtime/Metasound/Source/MetasoundEditor/Private/MetasoundEditorGraphNode.cpp
+++ b/Engine/Plugins/Runtime/Metasound/Source/MetasoundEditor/Private/MetasoundEditorGraphNode.cpp
@@ -472,11 +472,7 @@ FGuid UMetasoundEditorGraphOutputNode::GetNodeID() const
 
 bool UMetasoundEditorGraphOutputNode::CanUserDeleteNode() const
 {
-	const FMetasoundAssetBase* ConstMetasoundAsset = Metasound::IMetasoundUObjectRegistry::Get().GetObjectAsAssetBase(&GetMetasoundChecked());
-	FMetasoundAssetBase* MetasoundAsset = const_cast<FMetasoundAssetBase*>(ConstMetasoundAsset);
-	const FMetasoundFrontendArchetype& Archetype = MetasoundAsset->GetArchetype();
-
-	return !GetNodeHandle()->IsRequired(Archetype);
+	return !GetNodeHandle()->IsRequired();
 }
 
 void UMetasoundEditorGraphOutputNode::SetNodeID(FGuid InNodeID)
diff --git a/Engine/Plugins/Runtime/Metasound/Source/MetasoundEngine/Private/Metasound.cpp b/Engine/Plugins/Runtime/Metasound/Source/MetasoundEngine/Private/Metasound.cpp
--- a/Engine/Plugins/Runtime/Metasound/Source/MetasoundEngine/Private/Metasound.cpp
+++ b/Engine/Plugins/Runtime/Metasound/Source/MetasoundEngine/Private/Metasound.cpp
@@ -6,6 +6,7 @@
 #include "Internationalization/Text.h"
 #include "MetasoundAssetBase.h"
 #include "MetasoundAudioFormats.h"
+#include "MetasoundEngineArchetypes.h"
 #include "MetasoundEngineEnvironment.h"
 #include "MetasoundFrontendController.h"
 #include "MetasoundFrontendQuery.h"
@@ -101,9 +102,10 @@ Metasound::Frontend::FNodeClassInfo UMetaSound::GetAssetClassInfo() const
 	return { GetDocumentChecked().RootGraph, *GetPathName() };
 }
 
-const FMetasoundFrontendArchetype& UMetaSound::GetArchetype() const
+const FMetasoundFrontendVersion& UMetaSound::GetDefaultArchetypeVersion() const
 {
-	return GetBaseArchetype();
+	static const FMetasoundFrontendVersion DefaultArchetypeVersion = Metasound::Engine::MetasoundV1_0::GetVersion();
+	return DefaultArchetypeVersion;
 }
 
 Metasound::FOperatorSettings UMetaSound::GetOperatorSettings(Metasound::FSampleRate InSampleRate) const
@@ -124,14 +126,14 @@ Metasound::FSendAddress UMetaSound::CreateSendAddress(uint64 InInstanceID, const
 	return Address;
 }
 
-const TArray<FMetasoundFrontendArchetype>& UMetaSound::GetPreferredArchetypes() const
+const TArray<FMetasoundFrontendVersion>& UMetaSound::GetSupportedArchetypeVersions() const
 {
-	static const TArray<FMetasoundFrontendArchetype> Preferred
+	static const TArray<FMetasoundFrontendVersion> Supported
 	{
-		GetBaseArchetype()
+		Metasound::Engine::MetasoundV1_0::GetVersion()
 	};
 
-	return Preferred;
+	return Supported;
 }
 
 const FString& UMetaSound::GetAudioDeviceHandleVariableName()
@@ -140,24 +142,4 @@ const FString& UMetaSound::GetAudioDeviceHandleVariableName()
 	return AudioDeviceHandleVarName;
 }
 
-const FMetasoundFrontendArchetype& UMetaSound::GetBaseArchetype()
-{
-	auto CreateBaseArchetype = []() -> FMetasoundFrontendArchetype
-	{
-		FMetasoundFrontendArchetype Archetype;
-
-		FMetasoundFrontendEnvironmentVariable AudioDeviceHandle;
-		AudioDeviceHandle.Name = GetAudioDeviceHandleVariableName();
-		AudioDeviceHandle.Metadata.DisplayName = FText::FromString(AudioDeviceHandle.Name);
-		AudioDeviceHandle.Metadata.Description = LOCTEXT("AudioDeviceHandleToolTip", "Audio device handle");
-
-		Archetype.Interface.Environment.Add(AudioDeviceHandle);
-
-		return Archetype;
-	};
-
-	static const FMetasoundFrontendArchetype BaseArchetype = CreateBaseArchetype();
-
-	return BaseArchetype;
-}
 #undef LOCTEXT_NAMESPACE // MetaSound
diff --git a/Engine/Plugins/Runtime/Metasound/Source/MetasoundEngine/Private/MetasoundEngineArchetypes.cpp b/Engine/Plugins/Runtime/Metasound/Source/MetasoundEngine/Private/MetasoundEngineArchetypes.cpp
new file mode 100644
--- /dev/null
+++ b/Engine/Plugins/Runtime/Metasound/Source/MetasoundEngine/Private/MetasoundEngineArchetypes.cpp
@@ -0,0 +1,704 @@
+// Copyright Epic Games, Inc. All Rights Reserved.
+
+#include "MetasoundEngineArchetypes.h"
+
+#include "MetasoundAudioFormats.h"
+#include "MetasoundFrontendArchetypeRegistry.h"
+#include "MetasoundFrontendTransform.h"
+#include "MetasoundSource.h"
+
+#define LOCTEXT_NAMESPACE "MetasoundEngineArchetypes"
+
+namespace Metasound
+{
+	namespace Engine
+	{
+		namespace MetasoundEngineArchetypesPrivate
+		{
+			// Entry for registered archetype.
+			class FArchetypeRegistryEntry : public Frontend::IArchetypeRegistryEntry
+			{
+			public:
+
+				// @param InArchetype - Archetype to register.
+				// @parma InTransform - Transform to convert from prior version of archetype to this version of archetype.
+				FArchetypeRegistryEntry(const FMetasoundFrontendArchetype& InArchetype, TUniquePtr<Frontend::IDocumentTransform>&& InTransform)
+				: Archetype(InArchetype)
+				, Transform(MoveTemp(InTransform))
+				{
+				}
+
+				virtual const FMetasoundFrontendArchetype& GetArchetype() const override
+				{
+					return Archetype;
+				}
+
+				virtual bool UpdateRootGraphArchetype(Frontend::FDocumentHandle InDocument) const override
+				{
+					if (Transform.IsValid())
+					{
+						return Transform->Transform(InDocument);
+					}
+					return false;
+				}
+
+			private:
+				FMetasoundFrontendArchetype Archetype;
+				TUniquePtr<Frontend::IDocumentTransform> Transform;
+			};
+		}
+
+		// MetasoundV1_0 is a metasound without any required inputs or outputs.
+		namespace MetasoundV1_0
+		{
+			FMetasoundFrontendVersion GetVersion()
+			{
+				static const FName VersionName = FName("MetaSound");
+				return FMetasoundFrontendVersion{ VersionName, FMetasoundFrontendVersionNumber{1, 0} };
+			}
+
+			FMetasoundFrontendArchetype GetArchetype()
+			{
+				FMetasoundFrontendArchetype Archetype;
+				Archetype.Version = GetVersion();
+
+				return Archetype;
+			}
+		}
+
+		namespace MetasoundSourceMonoV1_0
+		{
+			FMetasoundFrontendVersion GetVersion()
+			{
+				static const FName VersionName = FName("MonoSource");
+				return FMetasoundFrontendVersion{ VersionName, FMetasoundFrontendVersionNumber{1, 0} };
+
+			}
+
+			const FString& GetOnPlayInputName()
+			{
+				static const FString TriggerInputName = TEXT("On Play");
+				return TriggerInputName;
+			}
+
+			const FString& GetAudioOutputName()
+			{
+				static const FString AudioOutputName = TEXT("Generated Audio");
+				return AudioOutputName;
+			}
+
+			const FString& GetIsFinishedOutputName()
+			{
+				static const FString OnFinishedOutputName = TEXT("On Finished");
+				return OnFinishedOutputName;
+			}
+
+			const FString& GetAudioDeviceHandleVariableName()
+			{
+				static const FString AudioDeviceHandleVarName = TEXT("AudioDeviceHandle");
+				return AudioDeviceHandleVarName;
+			}
+
+			const FString& GetSoundUniqueIdName()
+			{
+				static const FString SoundUniqueIdVarName = TEXT("SoundUniqueId");
+				return SoundUniqueIdVarName;
+			}
+
+			const FString& GetIsPreviewSoundName()
+			{
+				static const FString SoundIsPreviewSoundName = TEXT("IsPreviewSound");
+				return SoundIsPreviewSoundName;
+			}
+
+			FMetasoundFrontendArchetype GetArchetype()
+			{
+				FMetasoundFrontendArchetype Archetype;
+				Archetype.Version.Name = "MonoSource";
+				Archetype.Version.Number = FMetasoundFrontendVersionNumber{1, 0};
+				
+				// Inputs
+				FMetasoundFrontendClassVertex OnPlayTrigger;
+				OnPlayTrigger.Name = GetOnPlayInputName();
+
+				OnPlayTrigger.TypeName = Metasound::Frontend::GetDataTypeName<Metasound::FTrigger>();
+				OnPlayTrigger.Metadata.DisplayName = LOCTEXT("OnPlay", "On Play");
+				OnPlayTrigger.Metadata.Description = LOCTEXT("OnPlayTriggerToolTip", "Trigger executed when this source is played.");
+				OnPlayTrigger.VertexID = FGuid::NewGuid();
+
+				Archetype.Interface.Inputs.Add(OnPlayTrigger);
+
+				// Outputs 
+				FMetasoundFrontendClassVertex OnFinished;
+				OnFinished.Name = GetIsFinishedOutputName();
+
+				OnFinished.TypeName = Metasound::Frontend::GetDataTypeName<Metasound::FTrigger>();
+				OnFinished.Metadata.DisplayName = LOCTEXT("OnFinished", "On Finished");
+				OnFinished.Metadata.Description = LOCTEXT("OnFinishedToolTip", "Trigger executed to initiate stopping the source.");
+				OnFinished.VertexID = FGuid::NewGuid();
+
+				Archetype.Interface.Outputs.Add(OnFinished);
+
+				FMetasoundFrontendClassVertex GeneratedAudio;
+				GeneratedAudio.Name = GetAudioOutputName();
+				GeneratedAudio.TypeName = Metasound::Frontend::GetDataTypeName<Metasound::FMonoAudioFormat>();
+				GeneratedAudio.Metadata.DisplayName = LOCTEXT("GeneratedMono", "Audio");
+				GeneratedAudio.Metadata.Description = LOCTEXT("GeneratedAudioToolTip", "The resulting output audio from this source.");
+				GeneratedAudio.VertexID = FGuid::NewGuid();
+
+				Archetype.Interface.Outputs.Add(GeneratedAudio);
+
+				// Environment 
+				FMetasoundFrontendEnvironmentVariable AudioDeviceHandle;
+				AudioDeviceHandle.Name = GetAudioDeviceHandleVariableName();
+				AudioDeviceHandle.Metadata.DisplayName = FText::FromString(AudioDeviceHandle.Name);
+				AudioDeviceHandle.Metadata.Description = LOCTEXT("AudioDeviceHandleToolTip", "Audio device handle");
+
+				Archetype.Interface.Environment.Add(AudioDeviceHandle);
+
+				return Archetype;
+			}
+		}
+
+		namespace MetasoundSourceStereoV1_0
+		{
+			FMetasoundFrontendVersion GetVersion()
+			{
+				static const FName VersionName = FName("StereoSource");
+				return FMetasoundFrontendVersion{ VersionName, FMetasoundFrontendVersionNumber{1, 0} };
+			}
+
+			const FString& GetOnPlayInputName()
+			{
+				static const FString TriggerInputName = TEXT("On Play");
+				return TriggerInputName;
+			}
+
+			const FString& GetAudioOutputName()
+			{
+				static const FString AudioOutputName = TEXT("Generated Audio");
+				return AudioOutputName;
+			}
+
+			const FString& GetIsFinishedOutputName()
+			{
+				static const FString OnFinishedOutputName = TEXT("On Finished");
+				return OnFinishedOutputName;
+			}
+
+			const FString& GetAudioDeviceHandleVariableName()
+			{
+				static const FString AudioDeviceHandleVarName = TEXT("AudioDeviceHandle");
+				return AudioDeviceHandleVarName;
+			}
+
+			const FString& GetSoundUniqueIdName()
+			{
+				static const FString SoundUniqueIdVarName = TEXT("SoundUniqueId");
+				return SoundUniqueIdVarName;
+			}
+
+			const FString& GetIsPreviewSoundName()
+			{
+				static const FString SoundIsPreviewSoundName = TEXT("IsPreviewSound");
+				return SoundIsPreviewSoundName;
+			}
+
+			FMetasoundFrontendArchetype GetArchetype()
+			{
+				FMetasoundFrontendArchetype Archetype;
+				Archetype.Version = GetVersion();
+				
+				// Inputs
+				FMetasoundFrontendClassVertex OnPlayTrigger;
+				OnPlayTrigger.Name = GetOnPlayInputName();
+
+				OnPlayTrigger.TypeName = Metasound::Frontend::GetDataTypeName<Metasound::FTrigger>();
+				OnPlayTrigger.Metadata.DisplayName = LOCTEXT("OnPlay", "On Play");
+				OnPlayTrigger.Metadata.Description = LOCTEXT("OnPlayTriggerToolTip", "Trigger executed when this source is played.");
+				OnPlayTrigger.VertexID = FGuid::NewGuid();
+
+				Archetype.Interface.Inputs.Add(OnPlayTrigger);
+
+				// Outputs
+				FMetasoundFrontendClassVertex OnFinished;
+				OnFinished.Name = GetIsFinishedOutputName();
+
+				OnFinished.TypeName = Metasound::Frontend::GetDataTypeName<Metasound::FTrigger>();
+				OnFinished.Metadata.DisplayName = LOCTEXT("OnFinished", "On Finished");
+				OnFinished.Metadata.Description = LOCTEXT("OnFinishedToolTip", "Trigger executed to initiate stopping the source.");
+				OnFinished.VertexID = FGuid::NewGuid();
+
+				Archetype.Interface.Outputs.Add(OnFinished);
+
+				FMetasoundFrontendClassVertex GeneratedAudio;
+				GeneratedAudio.Name = GetAudioOutputName();
+				GeneratedAudio.TypeName = Metasound::Frontend::GetDataTypeName<Metasound::FStereoAudioFormat>();
+				GeneratedAudio.Metadata.DisplayName = LOCTEXT("GeneratedStereo", "Audio");
+				GeneratedAudio.Metadata.Description = LOCTEXT("GeneratedAudioToolTip", "The resulting output audio from this source.");
+				GeneratedAudio.VertexID = FGuid::NewGuid();
+
+				Archetype.Interface.Outputs.Add(GeneratedAudio);
+
+				// Environment
+				FMetasoundFrontendEnvironmentVariable AudioDeviceHandle;
+				AudioDeviceHandle.Name = GetAudioDeviceHandleVariableName();
+				AudioDeviceHandle.Metadata.DisplayName = FText::FromString(AudioDeviceHandle.Name);
+				AudioDeviceHandle.Metadata.Description = LOCTEXT("AudioDeviceHandleToolTip", "Audio device handle");
+
+				Archetype.Interface.Environment.Add(AudioDeviceHandle);
+
+				return Archetype;
+			}
+		}
+
+		namespace MetasoundSourceMonoV1_1
+		{
+			FMetasoundFrontendVersion GetVersion()
+			{
+				static const FName VersionName = FName("MonoSource");
+				return FMetasoundFrontendVersion{ VersionName, FMetasoundFrontendVersionNumber{1, 1} };
+			}
+
+			const FString& GetOnPlayInputName()
+			{
+				static const FString TriggerInputName = TEXT("On Play");
+				return TriggerInputName;
+			}
+
+			const FString& GetAudioOutputName()
+			{
+				static const FString AudioOutputName = TEXT("Audio:0");
+				return AudioOutputName;
+			}
+
+			const FString& GetIsFinishedOutputName()
+			{
+				static const FString OnFinishedOutputName = TEXT("On Finished");
+				return OnFinishedOutputName;
+			}
+
+			const FString& GetAudioDeviceHandleVariableName()
+			{
+				static const FString AudioDeviceHandleVarName = TEXT("AudioDeviceHandle");
+				return AudioDeviceHandleVarName;
+			}
+
+			const FString& GetSoundUniqueIdName()
+			{
+				static const FString SoundUniqueIdVarName = TEXT("SoundUniqueId");
+				return SoundUniqueIdVarName;
+			}
+
+			const FString& GetIsPreviewSoundName()
+			{
+				static const FString SoundIsPreviewSoundName = TEXT("IsPreviewSound");
+				return SoundIsPreviewSoundName;
+			}
+
+			FMetasoundFrontendClassVertex GetClassAudioOutput()
+			{
+				FMetasoundFrontendClassVertex GeneratedAudio;
+				GeneratedAudio.Name = GetAudioOutputName();
+				GeneratedAudio.TypeName = Metasound::Frontend::GetDataTypeName<Metasound::FAudioBuffer>();
+				GeneratedAudio.Metadata.DisplayName = LOCTEXT("GeneratedMono", "Audio");
+				GeneratedAudio.Metadata.Description = LOCTEXT("GeneratedAudioToolTip", "The resulting output audio from this source.");
+				GeneratedAudio.VertexID = FGuid::NewGuid();
+
+				return GeneratedAudio;
+			}
+
+			FMetasoundFrontendArchetype GetArchetype()
+			{
+				FMetasoundFrontendArchetype Archetype;
+				Archetype.Version = GetVersion();
+				
+				// Inputs
+				FMetasoundFrontendClassVertex OnPlayTrigger;
+				OnPlayTrigger.Name = GetOnPlayInputName();
+
+				OnPlayTrigger.TypeName = Metasound::Frontend::GetDataTypeName<Metasound::FTrigger>();
+				OnPlayTrigger.Metadata.DisplayName = LOCTEXT("OnPlay", "On Play");
+				OnPlayTrigger.Metadata.Description = LOCTEXT("OnPlayTriggerToolTip", "Trigger executed when this source is played.");
+				OnPlayTrigger.VertexID = FGuid::NewGuid();
+
+				Archetype.Interface.Inputs.Add(OnPlayTrigger);
+
+				// Outputs
+				FMetasoundFrontendClassVertex OnFinished;
+				OnFinished.Name = GetIsFinishedOutputName();
+
+				OnFinished.TypeName = Metasound::Frontend::GetDataTypeName<Metasound::FTrigger>();
+				OnFinished.Metadata.DisplayName = LOCTEXT("OnFinished", "On Finished");
+				OnFinished.Metadata.Description = LOCTEXT("OnFinishedToolTip", "Trigger executed to initiate stopping the source.");
+				OnFinished.VertexID = FGuid::NewGuid();
+
+				Archetype.Interface.Outputs.Add(OnFinished);
+
+				FMetasoundFrontendClassVertex GeneratedAudio = GetClassAudioOutput();
+				Archetype.Interface.Outputs.Add(GeneratedAudio);
+
+				// Environment
+				FMetasoundFrontendEnvironmentVariable AudioDeviceHandle;
+				AudioDeviceHandle.Name = GetAudioDeviceHandleVariableName();
+				AudioDeviceHandle.Metadata.DisplayName = FText::FromString(AudioDeviceHandle.Name);
+				AudioDeviceHandle.Metadata.Description = LOCTEXT("AudioDeviceHandleToolTip", "Audio device handle");
+
+				Archetype.Interface.Environment.Add(AudioDeviceHandle);
+
+				return Archetype;
+			}
+
+			// Update from a MetasoundSourceMonoV1_0 to MtasoundSourceMonoV1_1
+			class FUpdateArchetype : public Frontend::IDocumentTransform
+			{
+			public:
+				virtual bool Transform(Frontend::FDocumentHandle InDocument) const override
+				{
+					// Swap FMonoAudioFormat output node to an FAudioBuffer output node.
+					
+					using namespace Frontend;
+
+					FGraphHandle Graph = InDocument->GetRootGraph();
+					if (!Graph->IsValid())
+					{
+						return false;
+					}
+
+					FNodeHandle MonoFormatOutput = Graph->GetOutputNodeWithName(MetasoundSourceMonoV1_0::GetAudioOutputName());
+
+					FVector2D MonoFormatLocation;
+
+					FOutputHandle OutputToReconnect = IOutputController::GetInvalidHandle();
+					if (MonoFormatOutput->IsValid())
+					{
+						// Get the first location.
+						for (auto Location : MonoFormatOutput->GetNodeStyle().Display.Locations)
+						{
+							MonoFormatLocation = Location.Value;
+						}
+
+						// Get connections
+						TArray<FInputHandle> Inputs = MonoFormatOutput->GetInputs();
+						if (ensure(Inputs.Num() == 1))
+						{
+							OutputToReconnect = Inputs[0]->GetConnectedOutput();
+						}
+
+						Graph->RemoveOutputVertex(MetasoundSourceMonoV1_0::GetAudioOutputName());
+					}
+
+					// Create output
+					FNodeHandle BufferOutput = Graph->AddOutputVertex(GetClassAudioOutput());
+					if (ensure(BufferOutput->IsValid()))
+					{
+						FMetasoundFrontendNodeStyle Style = BufferOutput->GetNodeStyle();
+						Style.Display.Locations.Add(FGuid(), MonoFormatLocation);
+						BufferOutput->SetNodeStyle(Style);
+
+						if (OutputToReconnect->IsValid())
+						{
+							// Reconnect
+							TArray<FInputHandle> Inputs = BufferOutput->GetInputs();
+							if (ensure(Inputs.Num() == 1))
+							{
+								ensure(OutputToReconnect->Connect(*Inputs[0]));
+							}
+						}
+					}
+
+					return true;
+				}
+			};
+		}
+
+		namespace MetasoundSourceStereoV1_1
+		{
+			FMetasoundFrontendVersion GetVersion()
+			{
+				static const FName VersionName = FName("StereoSource");
+				return FMetasoundFrontendVersion{ VersionName, FMetasoundFrontendVersionNumber{1, 1} };
+			}
+
+			const FString& GetOnPlayInputName()
+			{
+				static const FString TriggerInputName = TEXT("On Play");
+				return TriggerInputName;
+			}
+
+			const FString& GetLeftAudioOutputName()
+			{
+				static const FString AudioOutputName = TEXT("Audio:0");
+				return AudioOutputName;
+			}
+
+			const FString& GetRightAudioOutputName()
+			{
+				static const FString AudioOutputName = TEXT("Audio:1");
+				return AudioOutputName;
+			}
+
+			const FString& GetIsFinishedOutputName()
+			{
+				static const FString OnFinishedOutputName = TEXT("On Finished");
+				return OnFinishedOutputName;
+			}
+
+			const FString& GetAudioDeviceHandleVariableName()
+			{
+				static const FString AudioDeviceHandleVarName = TEXT("AudioDeviceHandle");
+				return AudioDeviceHandleVarName;
+			}
+
+			const FString& GetSoundUniqueIdName()
+			{
+				static const FString SoundUniqueIdVarName = TEXT("SoundUniqueId");
+				return SoundUniqueIdVarName;
+			}
+
+			const FString& GetIsPreviewSoundName()
+			{
+				static const FString SoundIsPreviewSoundName = TEXT("IsPreviewSound");
+				return SoundIsPreviewSoundName;
+			}
+
+			FMetasoundFrontendClassVertex GetClassLeftAudioOutput()
+			{
+				FMetasoundFrontendClassVertex GeneratedLeftAudio;
+				GeneratedLeftAudio.Name = GetLeftAudioOutputName();
+				GeneratedLeftAudio.TypeName = Metasound::Frontend::GetDataTypeName<Metasound::FAudioBuffer>();
+				GeneratedLeftAudio.Metadata.DisplayName = LOCTEXT("GeneratedStereoLeft", "Left Audio");
+				GeneratedLeftAudio.Metadata.Description = LOCTEXT("GeneratedLeftAudioToolTip", "The resulting output audio from this source.");
+				GeneratedLeftAudio.VertexID = FGuid::NewGuid();
+
+				return GeneratedLeftAudio;
+			}
+
+			FMetasoundFrontendClassVertex GetClassRightAudioOutput()
+			{
+				FMetasoundFrontendClassVertex GeneratedRightAudio;
+				GeneratedRightAudio.Name = GetRightAudioOutputName();
+				GeneratedRightAudio.TypeName = Metasound::Frontend::GetDataTypeName<Metasound::FAudioBuffer>();
+				GeneratedRightAudio.Metadata.DisplayName = LOCTEXT("GeneratedStereoRight", "Right Audio");
+				GeneratedRightAudio.Metadata.Description = LOCTEXT("GeneratedRightAudioToolTip", "The resulting output audio from this source.");
+				GeneratedRightAudio.VertexID = FGuid::NewGuid();
+
+				return GeneratedRightAudio;
+			}
+
+			FMetasoundFrontendArchetype GetArchetype()
+			{
+				FMetasoundFrontendArchetype Archetype;
+				Archetype.Version = GetVersion();
+				
+				// Inputs
+				FMetasoundFrontendClassVertex OnPlayTrigger;
+				OnPlayTrigger.Name = GetOnPlayInputName();
+
+				OnPlayTrigger.TypeName = Metasound::Frontend::GetDataTypeName<Metasound::FTrigger>();
+				OnPlayTrigger.Metadata.DisplayName = LOCTEXT("OnPlay", "On Play");
+				OnPlayTrigger.Metadata.Description = LOCTEXT("OnPlayTriggerToolTip", "Trigger executed when this source is played.");
+				OnPlayTrigger.VertexID = FGuid::NewGuid();
+
+				Archetype.Interface.Inputs.Add(OnPlayTrigger);
+
+				// Outputs
+				FMetasoundFrontendClassVertex OnFinished;
+				OnFinished.Name = GetIsFinishedOutputName();
+
+				OnFinished.TypeName = Metasound::Frontend::GetDataTypeName<Metasound::FTrigger>();
+				OnFinished.Metadata.DisplayName = LOCTEXT("OnFinished", "On Finished");
+				OnFinished.Metadata.Description = LOCTEXT("OnFinishedToolTip", "Trigger executed to initiate stopping the source.");
+				OnFinished.VertexID = FGuid::NewGuid();
+
+				Archetype.Interface.Outputs.Add(OnFinished);
+
+				FMetasoundFrontendClassVertex GeneratedLeftAudio = GetClassLeftAudioOutput();
+				Archetype.Interface.Outputs.Add(GeneratedLeftAudio);
+
+				FMetasoundFrontendClassVertex GeneratedRightAudio = GetClassRightAudioOutput();
+				Archetype.Interface.Outputs.Add(GeneratedRightAudio);
+
+				// Environment
+				FMetasoundFrontendEnvironmentVariable AudioDeviceHandle;
+				AudioDeviceHandle.Name = GetAudioDeviceHandleVariableName();
+				AudioDeviceHandle.Metadata.DisplayName = FText::FromString(AudioDeviceHandle.Name);
+				AudioDeviceHandle.Metadata.Description = LOCTEXT("AudioDeviceHandleToolTip", "Audio device handle");
+				Archetype.Interface.Environment.Add(AudioDeviceHandle);
+
+				return Archetype;
+			};
+
+			// Update from MetasoundSourceStereoV1_0 to MetasoundSourceStereoV1_1
+			class FUpdateArchetype : public Frontend::IDocumentTransform
+			{
+			public:
+				virtual bool Transform(Frontend::FDocumentHandle InDocument) const override
+				{
+					using namespace Frontend;
+
+					FGraphHandle Graph = InDocument->GetRootGraph();
+
+					if (!Graph->IsValid())
+					{
+						return false;
+					}
+
+					FNodeHandle StereoFormatOutput = Graph->GetOutputNodeWithName(MetasoundSourceStereoV1_0::GetAudioOutputName());
+					FOutputHandle LeftOutputToReconnect = IOutputController::GetInvalidHandle();
+					FOutputHandle RightOutputToReconnect = IOutputController::GetInvalidHandle();
+
+					FVector2D StereoFormatLocation;
+
+					if (StereoFormatOutput->IsValid())
+					{
+						// Get the first location.
+						for (auto Location : StereoFormatOutput->GetNodeStyle().Display.Locations)
+						{
+							StereoFormatLocation = Location.Value;
+						}
+
+						TArray<FInputHandle> LeftInputs = StereoFormatOutput->GetInputsWithVertexName(TEXT("Left"));
+						if (ensure(LeftInputs.Num() == 1))
+						{
+							LeftOutputToReconnect = LeftInputs[0]->GetConnectedOutput();
+						}
+
+						TArray<FInputHandle> RightInputs = StereoFormatOutput->GetInputsWithVertexName(TEXT("Right"));
+						if (ensure(RightInputs.Num() == 1))
+						{
+							RightOutputToReconnect = RightInputs[0]->GetConnectedOutput();
+						}
+
+						Graph->RemoveOutputVertex(MetasoundSourceStereoV1_0::GetAudioOutputName());
+					}
+
+					FNodeHandle LeftBufferOutput = Graph->AddOutputVertex(GetClassLeftAudioOutput());
+					if (ensure(LeftBufferOutput->IsValid()))
+					{
+						FMetasoundFrontendNodeStyle Style = LeftBufferOutput->GetNodeStyle();
+						Style.Display.Locations.Add(FGuid(), StereoFormatLocation);
+						LeftBufferOutput->SetNodeStyle(Style);
+
+						if (LeftOutputToReconnect->IsValid())
+						{
+							TArray<FInputHandle> Inputs = LeftBufferOutput->GetInputs();
+							if (ensure(Inputs.Num() == 1))
+							{
+								ensure(LeftOutputToReconnect->Connect(*Inputs[0]));
+							}
+						}
+					}
+
+					FNodeHandle RightBufferOutput = Graph->AddOutputVertex(GetClassRightAudioOutput());
+					if (ensure(RightBufferOutput->IsValid()))
+					{
+						if (RightOutputToReconnect->IsValid())
+						{
+							FMetasoundFrontendNodeStyle Style = RightBufferOutput->GetNodeStyle();
+							Style.Display.Locations.Add(FGuid(), StereoFormatLocation + FVector2D{0, 100.f}); // TODO: How should this new output be placed?
+							RightBufferOutput->SetNodeStyle(Style);
+
+							TArray<FInputHandle> Inputs = RightBufferOutput->GetInputs();
+							if (ensure(Inputs.Num() == 1))
+							{
+								ensure(RightOutputToReconnect->Connect(*Inputs[0]));
+							}
+						}
+					}
+
+					return true;
+				}
+			};
+		}
+
+		namespace MetasoundSource
+		{
+			namespace CurrentMonoVersion = MetasoundSourceMonoV1_1;
+			namespace CurrentStereoVersion = MetasoundSourceStereoV1_1;
+
+			const FString& GetOnPlayInputName()
+			{
+				check(CurrentMonoVersion::GetOnPlayInputName() == CurrentStereoVersion::GetOnPlayInputName());
+				return CurrentMonoVersion::GetOnPlayInputName();
+			}
+
+			const FString& GetIsFinishedOutputName()
+			{
+				check(CurrentMonoVersion::GetIsFinishedOutputName() == CurrentStereoVersion::GetIsFinishedOutputName());
+				return CurrentMonoVersion::GetIsFinishedOutputName();
+			}
+
+			const FString& GetAudioDeviceHandleVariableName()
+			{
+				check(CurrentMonoVersion::GetAudioDeviceHandleVariableName() == CurrentStereoVersion::GetAudioDeviceHandleVariableName());
+				return CurrentMonoVersion::GetAudioDeviceHandleVariableName();
+			}
+
+			const FString& GetSoundUniqueIdName()
+			{
+				check(CurrentMonoVersion::GetSoundUniqueIdName() == CurrentStereoVersion::GetSoundUniqueIdName());
+				return CurrentMonoVersion::GetSoundUniqueIdName();
+			}
+
+			const FString& GetIsPreviewSoundName()
+			{
+				check(CurrentMonoVersion::GetIsPreviewSoundName() == CurrentStereoVersion::GetIsPreviewSoundName());
+				return CurrentMonoVersion::GetIsPreviewSoundName();
+			}
+		}
+
+		namespace MetasoundSourceMono
+		{
+			namespace Current = MetasoundSourceMonoV1_1;
+
+			FMetasoundFrontendVersion GetVersion()
+			{
+				return Current::GetVersion();
+			}
+
+			const FString& GetAudioOutputName()
+			{
+				return Current::GetAudioOutputName();
+			}
+		}
+
+		namespace MetasoundSourceStereo
+		{
+			namespace Current = MetasoundSourceStereoV1_1;
+
+			FMetasoundFrontendVersion GetVersion()
+			{
+				return Current::GetVersion();
+			}
+
+			const FString& GetLeftAudioOutputName()
+			{
+				return Current::GetLeftAudioOutputName();
+			}
+
+			const FString& GetRightAudioOutputName()
+			{
+				return Current::GetRightAudioOutputName();
+			}
+		}
+
+		void RegisterArchetypes()
+		{
+			using namespace MetasoundEngineArchetypesPrivate;
+
+			Frontend::IArchetypeRegistry::Get().RegisterArchetype(MakeUnique<FArchetypeRegistryEntry>(MetasoundV1_0::GetArchetype(), nullptr));
+
+			Frontend::IArchetypeRegistry::Get().RegisterArchetype(MakeUnique<FArchetypeRegistryEntry>(MetasoundSourceMonoV1_0::GetArchetype(), nullptr));
+
+			Frontend::IArchetypeRegistry::Get().RegisterArchetype(MakeUnique<FArchetypeRegistryEntry>(MetasoundSourceStereoV1_0::GetArchetype(), nullptr));
+
+			Frontend::IArchetypeRegistry::Get().RegisterArchetype(MakeUnique<FArchetypeRegistryEntry>(MetasoundSourceMonoV1_1::GetArchetype(), MakeUnique<MetasoundSourceMonoV1_1::FUpdateArchetype>()));
+
+			Frontend::IArchetypeRegistry::Get().RegisterArchetype(MakeUnique<FArchetypeRegistryEntry>(MetasoundSourceStereoV1_1::GetArchetype(), MakeUnique<MetasoundSourceStereoV1_1::FUpdateArchetype>()));
+
+		}
+	}
+}
+
+#undef LOCTEXT_NAMESPACE // MetasoundEngineArchetypes
diff --git a/Engine/Plugins/Runtime/Metasound/Source/MetasoundEngine/Private/MetasoundEngineArchetypes.h b/Engine/Plugins/Runtime/Metasound/Source/MetasoundEngine/Private/MetasoundEngineArchetypes.h
new file mode 100644
--- /dev/null
+++ b/Engine/Plugins/Runtime/Metasound/Source/MetasoundEngine/Private/MetasoundEngineArchetypes.h
@@ -0,0 +1,96 @@
+// Copyright Epic Games, Inc. All Rights Reserved.
+
+#pragma once
+
+#include "CoreMinimal.h"
+#include "MetasoundFrontendDocument.h"
+
+namespace Metasound
+{
+	namespace Engine
+	{
+		// Base metasound without any required inputs or outputs
+		namespace MetasoundV1_0
+		{
+			FMetasoundFrontendVersion GetVersion();
+		}
+
+		// V1.0 of a Metasound Mono Source. Uses FMonoAudioFormat as output.
+		namespace MetasoundSourceMonoV1_0
+		{
+			FMetasoundFrontendVersion GetVersion();
+			const FString& GetOnPlayInputName();
+			const FString& GetAudioOutputName();
+			const FString& GetIsFinishedOutputName();
+			const FString& GetAudioDeviceHandleVariableName();
+			const FString& GetSoundUniqueIdName();
+			const FString& GetIsPreviewSoundName();
+		}
+
+		// V1.0 of a Metasound Stereo Source. Uses FStereoAudioFormat as output.
+		namespace MetasoundSourceStereoV1_0
+		{
+			FMetasoundFrontendVersion GetVersion();
+			const FString& GetOnPlayInputName();
+			const FString& GetAudioOutputName();
+			const FString& GetIsFinishedOutputName();
+			const FString& GetAudioDeviceHandleVariableName();
+			const FString& GetSoundUniqueIdName();
+			const FString& GetIsPreviewSoundName();
+		}
+
+		// V1.1 of a Metasound Mono source. Uses FAudioBuffer as output.
+		namespace MetasoundSourceMonoV1_1
+		{
+			FMetasoundFrontendVersion GetVersion();
+			const FString& GetOnPlayInputName();
+			const FString& GetAudioOutputName();
+			const FString& GetIsFinishedOutputName();
+			const FString& GetAudioDeviceHandleVariableName();
+			const FString& GetSoundUniqueIdName();
+			const FString& GetIsPreviewSoundName();
+		}
+
+		// V1.1 of a Metasound Stereo source. Uses two FAudioBuffers as outputs.
+		namespace MetasoundSourceStereoV1_1
+		{
+			FMetasoundFrontendVersion GetVersion();
+			const FString& GetOnPlayInputName();
+			const FString& GetLeftAudioOutputName();
+			const FString& GetRightAudioOutputName();
+			const FString& GetIsFinishedOutputName();
+			const FString& GetAudioDeviceHandleVariableName();
+			const FString& GetSoundUniqueIdName();
+			const FString& GetIsPreviewSoundName();
+		}
+
+		// Current version of Metasound Source.
+		namespace MetasoundSource
+		{
+			const FString& GetOnPlayInputName();
+			const FString& GetIsFinishedOutputName();
+			const FString& GetAudioDeviceHandleVariableName();
+			const FString& GetSoundUniqueIdName();
+			const FString& GetIsPreviewSoundName();
+		}
+
+		// Current version of MetasoundSourceMono
+		namespace MetasoundSourceMono
+		{
+			FMetasoundFrontendVersion GetVersion();
+			const FString& GetAudioOutputName();
+		}
+
+		// Current version of MetasoundSourceStereo
+		namespace MetasoundSourceStereo
+		{
+			FMetasoundFrontendVersion GetVersion();
+			const FString& GetLeftAudioOutputName();
+			const FString& GetRightAudioOutputName();
+		}
+
+		// Register all MetasoundEngine archetypes.
+		void RegisterArchetypes();
+	}
+}
+
diff --git a/Engine/Plugins/Runtime/Metasound/Source/MetasoundEngine/Private/MetasoundEngineModule.cpp b/Engine/Plugins/Runtime/Metasound/Source/MetasoundEngine/Private/MetasoundEngineModule.cpp
--- a/Engine/Plugins/Runtime/Metasound/Source/MetasoundEngine/Private/MetasoundEngineModule.cpp
+++ b/Engine/Plugins/Runtime/Metasound/Source/MetasoundEngine/Private/MetasoundEngineModule.cpp
@@ -4,6 +4,7 @@
 #include "Metasound.h"
 #include "MetasoundDataReference.h"
 #include "MetasoundDataTypeRegistrationMacro.h"
+#include "MetasoundEngineArchetypes.h"
 #include "MetasoundSource.h"
 #include "MetasoundUObjectRegistry.h"
 #include "MetasoundWave.h"
@@ -17,8 +18,14 @@ class FMetasoundEngineModule : public IMetasoundEngineModule
 {
 	virtual void StartupModule() override
 	{
+		FModuleManager::Get().LoadModuleChecked("MetasoundFrontend");
+		FModuleManager::Get().LoadModuleChecked("AudioCodecEngine");
+
+		Metasound::Engine::RegisterArchetypes();
+
 		// If there is no archetype name, use UMetaSound
-		Metasound::IMetasoundUObjectRegistry::RegisterUClassArchetype<UMetaSound>(TEXT(""));
+		FMetasoundFrontendVersion DefaultVersion;
+		Metasound::IMetasoundUObjectRegistry::RegisterUClassArchetype<UMetaSound>(DefaultVersion);
 
 		// Register preferred archetypes
 		Metasound::IMetasoundUObjectRegistry::RegisterUClassPreferredArchetypes<UMetaSound>();
@@ -27,8 +34,6 @@ class FMetasoundEngineModule : public IMetasoundEngineModule
 		// flush node registration queue
 		FMetasoundFrontendRegistryContainer::Get()->RegisterPendingNodes();
 
-		FModuleManager::Get().LoadModuleChecked("AudioCodecEngine");
-
 		UE_LOG(LogMetasoundEngine, Log, TEXT("Metasound Engine Initialized"));
 	}
 };
diff --git a/Engine/Plugins/Runtime/Metasound/Source/MetasoundEngine/Private/MetasoundSource.cpp b/Engine/Plugins/Runtime/Metasound/Source/MetasoundEngine/Private/MetasoundSource.cpp
--- a/Engine/Plugins/Runtime/Metasound/Source/MetasoundEngine/Private/MetasoundSource.cpp
+++ b/Engine/Plugins/Runtime/Metasound/Source/MetasoundEngine/Private/MetasoundSource.cpp
@@ -6,6 +6,7 @@
 #include "Internationalization/Text.h"
 #include "MetasoundAssetBase.h"
 #include "MetasoundAudioFormats.h"
+#include "MetasoundEngineArchetypes.h"
 #include "MetasoundEngineEnvironment.h"
 #include "MetasoundEnvironment.h"
 #include "MetasoundFrontendController.h"
@@ -52,6 +53,7 @@ UMetaSoundSource::UMetaSoundSource(const FObjectInitializer& ObjectInitializer)
 }
 
 #if WITH_EDITOR
+
 void UMetaSoundSource::PreSave(FObjectPreSaveContext SaveContext)
 {
 	Super::PreSave(SaveContext);
@@ -85,17 +87,21 @@ void UMetaSoundSource::PostEditChangeProperty(FPropertyChangedEvent& InEvent)
 	{
 		ConvertFromPreset();
 
+		bool bDidModifyDocument = false;
 		switch (OutputFormat)
 		{
 			case EMetasoundSourceAudioFormat::Mono:
 			{
 				NumChannels = 1;
+				// TODO: utilize latest version `MetasoudnSourceMono`
+				bDidModifyDocument = Metasound::Frontend::FMatchRootGraphToArchetype(Metasound::Engine::MetasoundSourceMono::GetVersion()).Transform(GetDocumentHandle());
 			}
 			break;
 
 			case EMetasoundSourceAudioFormat::Stereo:
 			{
 				NumChannels = 2;
+				bDidModifyDocument = Metasound::Frontend::FMatchRootGraphToArchetype(Metasound::Engine::MetasoundSourceStereo::GetVersion()).Transform(GetDocumentHandle());
 			}
 			break;
 
@@ -106,9 +112,9 @@ void UMetaSoundSource::PostEditChangeProperty(FPropertyChangedEvent& InEvent)
 			break;
 		}
 
-		if (ensure(IsArchetypeSupported(GetArchetype())))
+		if (bDidModifyDocument)
 		{
-			ConformDocumentToArchetype();
+			MarkMetasoundDocumentDirty();
 		}
 	}
 }
@@ -165,27 +171,11 @@ bool UMetaSoundSource::SupportsSubtitles() const
 	return Super::SupportsSubtitles();
 }
 
-const FMetasoundFrontendArchetype& UMetaSoundSource::GetArchetype() const
+const FMetasoundFrontendVersion& UMetaSoundSource::GetDefaultArchetypeVersion() const
 {
-	switch (OutputFormat)
-	{
-		case EMetasoundSourceAudioFormat::Mono:
-		{
-			return GetMonoSourceArchetype();
-		}
-
-		case EMetasoundSourceAudioFormat::Stereo:
-		{
-			return GetStereoSourceArchetype();
-		}
-
-		default:
-		{
-			static_assert(static_cast<int32>(EMetasoundSourceAudioFormat::COUNT) == 2, "Possible missing format switch case coverage.");
-		}
-	}
+	static const FMetasoundFrontendVersion DefaultVersion = Metasound::Engine::MetasoundSourceMono::GetVersion();
 
-	return GetBaseArchetype();
+	return DefaultVersion;
 }
 
 float UMetaSoundSource::GetDuration()
@@ -198,6 +188,7 @@ ISoundGeneratorPtr UMetaSoundSource::CreateSoundGenerator(const FSoundGeneratorI
 {
 	using namespace Metasound;
 	using namespace Metasound::Frontend;
+	using namespace Metasound::Engine;
 
 	Duration = INDEFINITELY_LOOPING_DURATION;
 	bLooping = true;
@@ -220,11 +211,9 @@ ISoundGeneratorPtr UMetaSoundSource::CreateSoundGenerator(const FSoundGeneratorI
 			DeviceHandle = DeviceManager->GetMainAudioDeviceHandle();
 		}
 	}
-	Environment.SetValue<FAudioDeviceHandle>(GetAudioDeviceHandleVariableName(), DeviceHandle);
-
-	// Set the unique object ID as an environment variable
-	Environment.SetValue<uint32>(GetSoundUniqueIdName(), GetUniqueID());
-	Environment.SetValue<bool>(GetIsPreviewSoundName(), InParams.bIsPreviewSound);
+	Environment.SetValue<FAudioDeviceHandle>(MetasoundSource::GetAudioDeviceHandleVariableName(), DeviceHandle);
+	Environment.SetValue<uint32>(MetasoundSource::GetSoundUniqueIdName(), GetUniqueID());
+	Environment.SetValue<bool>(MetasoundSource::GetIsPreviewSoundName(), InParams.bIsPreviewSound);
 
 	const FMetasoundFrontendDocument* OriginalDoc = GetDocument().Get();
 	if (nullptr == OriginalDoc)
@@ -234,12 +223,9 @@ ISoundGeneratorPtr UMetaSoundSource::CreateSoundGenerator(const FSoundGeneratorI
 	}
 
 	// Check version and attempt to version up if out-of-date
-	Frontend::FDocumentHandle Document = GetDocumentHandle();
-	FName AssetName = GetFName();
-	FString AssetPath = GetPathName();
-	if (Frontend::FVersionDocument(AssetName, AssetPath).Transform(Document))
+	if (VersionAsset())
 	{
-		const FMetasoundFrontendDocumentMetadata& Metadata = Document->GetMetadata();
+		const FMetasoundFrontendDocumentMetadata& Metadata = GetDocumentHandle()->GetMetadata();
 		UE_LOG(LogMetaSound, Display,
 			TEXT("Dynamically versioned UMetaSoundSource asset [Name:%s] to %s (Resave to avoid update at runtime)."),
 			*GetName(),
@@ -260,13 +246,34 @@ ISoundGeneratorPtr UMetaSoundSource::CreateSoundGenerator(const FSoundGeneratorI
 		InSettings,
 		MoveTemp(DocumentWithInjectedReceives),
 		Environment,
-		NumChannels,
 		GetName(),
-		GetAudioOutputName(),
-		GetOnPlayInputName(),
-		GetIsFinishedOutputName()
+		{}, // Channel names still need to be set.
+		MetasoundSource::GetOnPlayInputName(),
+		MetasoundSource::GetIsFinishedOutputName()
 	};
 
+	// Set init params dependent upon output audio format.
+	switch (OutputFormat)
+	{
+		case EMetasoundSourceAudioFormat::Mono:
+		{
+			InitParams.AudioOutputNames = { MetasoundSourceMono::GetAudioOutputName() };
+		}
+		break;
+
+		case EMetasoundSourceAudioFormat::Stereo:
+		{
+			InitParams.AudioOutputNames = { MetasoundSourceStereo::GetLeftAudioOutputName(), MetasoundSourceStereo::GetRightAudioOutputName() };
+		}
+		break;
+
+		default:
+		{
+			static_assert(static_cast<int32>(EMetasoundSourceAudioFormat::COUNT) == 2, "Possible missing format switch case coverage.");
+		}
+		break;
+	}
+
 	return ISoundGeneratorPtr(new FMetasoundGenerator(MoveTemp(InitParams)));
 }
 
@@ -288,135 +295,15 @@ Metasound::FOperatorSettings UMetaSoundSource::GetOperatorSettings(Metasound::FS
 	return Metasound::FOperatorSettings(InSampleRate, BlockRate);
 }
 
-const TArray<FMetasoundFrontendArchetype>& UMetaSoundSource::GetPreferredArchetypes() const 
-{
-	static const TArray<FMetasoundFrontendArchetype> Preferred({GetMonoSourceArchetype(), GetStereoSourceArchetype()});
-
-	return Preferred;
-}
-
-const FString& UMetaSoundSource::GetOnPlayInputName()
-{
-	static const FString TriggerInputName = TEXT("On Play");
-	return TriggerInputName;
-}
-
-const FString& UMetaSoundSource::GetAudioOutputName()
-{
-	static const FString AudioOutputName = TEXT("Generated Audio");
-	return AudioOutputName;
-}
-
-const FString& UMetaSoundSource::GetIsFinishedOutputName()
+const TArray<FMetasoundFrontendVersion>& UMetaSoundSource::GetSupportedArchetypeVersions() const 
 {
-	static const FString OnFinishedOutputName = TEXT("On Finished");
-	return OnFinishedOutputName;
-}
-
-const FString& UMetaSoundSource::GetAudioDeviceHandleVariableName()
-{
-	static const FString AudioDeviceHandleVarName = TEXT("AudioDeviceHandle");
-	return AudioDeviceHandleVarName;
-}
-
-const FString& UMetaSoundSource::GetSoundUniqueIdName()
-{
-	static const FString SoundUniqueIdVarName = TEXT("SoundUniqueId");
-	return SoundUniqueIdVarName;
-}
-
-const FString& UMetaSoundSource::GetIsPreviewSoundName()
-{
-	static const FString SoundIsPreviewSoundName = TEXT("IsPreviewSound");
-	return SoundIsPreviewSoundName;
-}
-
-const FMetasoundFrontendArchetype& UMetaSoundSource::GetBaseArchetype()
-{
-	auto CreateBaseArchetype = []() -> FMetasoundFrontendArchetype
-	{
-		FMetasoundFrontendArchetype Archetype;
-		
-		FMetasoundFrontendClassVertex OnPlayTrigger;
-		OnPlayTrigger.Name = UMetaSoundSource::GetOnPlayInputName();
-
-		OnPlayTrigger.TypeName = Metasound::Frontend::GetDataTypeName<Metasound::FTrigger>();
-		OnPlayTrigger.Metadata.DisplayName = LOCTEXT("OnPlay", "On Play");
-		OnPlayTrigger.Metadata.Description = LOCTEXT("OnPlayTriggerToolTip", "Trigger executed when this source is played.");
-		OnPlayTrigger.VertexID = FGuid::NewGuid();
-
-		Archetype.Interface.Inputs.Add(OnPlayTrigger);
-
-		FMetasoundFrontendClassVertex OnFinished;
-		OnFinished.Name = UMetaSoundSource::GetIsFinishedOutputName();
-
-		OnFinished.TypeName = Metasound::Frontend::GetDataTypeName<Metasound::FTrigger>();
-		OnFinished.Metadata.DisplayName = LOCTEXT("OnFinished", "On Finished");
-		OnFinished.Metadata.Description = LOCTEXT("OnFinishedToolTip", "Trigger executed to initiate stopping the source.");
-		OnFinished.VertexID = FGuid::NewGuid();
-
-		Archetype.Interface.Outputs.Add(OnFinished);
-
-		FMetasoundFrontendEnvironmentVariable AudioDeviceHandle;
-		AudioDeviceHandle.Name = UMetaSoundSource::GetAudioDeviceHandleVariableName();
-		AudioDeviceHandle.Metadata.DisplayName = FText::FromString(AudioDeviceHandle.Name);
-		AudioDeviceHandle.Metadata.Description = LOCTEXT("AudioDeviceHandleToolTip", "Audio device handle");
-
-		Archetype.Interface.Environment.Add(AudioDeviceHandle);
-
-		return Archetype;
-	};
-
-	static const FMetasoundFrontendArchetype BaseArchetype = CreateBaseArchetype();
-
-	return BaseArchetype;
-}
-
-const FMetasoundFrontendArchetype& UMetaSoundSource::GetMonoSourceArchetype()
-{
-	auto CreateMonoArchetype = []() -> FMetasoundFrontendArchetype
-	{
-		FMetasoundFrontendArchetype Archetype = GetBaseArchetype();
-		Archetype.Name = "MonoSource";
-
-		FMetasoundFrontendClassVertex GeneratedAudio;
-		GeneratedAudio.Name = UMetaSoundSource::GetAudioOutputName();
-		GeneratedAudio.TypeName = Metasound::Frontend::GetDataTypeName<Metasound::FMonoAudioFormat>();
-		GeneratedAudio.Metadata.DisplayName = LOCTEXT("GeneratedMono", "Audio");
-		GeneratedAudio.Metadata.Description = LOCTEXT("GeneratedAudioToolTip", "The resulting output audio from this source.");
-		GeneratedAudio.VertexID = FGuid::NewGuid();
-
-		Archetype.Interface.Outputs.Add(GeneratedAudio);
-
-		return Archetype;
-	};
-
-	static const FMetasoundFrontendArchetype MonoArchetype = CreateMonoArchetype();
+	static const TArray<FMetasoundFrontendVersion> Supported(
+		{
+			Metasound::Engine::MetasoundSourceMono::GetVersion(),
+			Metasound::Engine::MetasoundSourceStereo::GetVersion()
+		});
 
-	return MonoArchetype;
+	return Supported;
 }
 
-const FMetasoundFrontendArchetype& UMetaSoundSource::GetStereoSourceArchetype()
-{
-	auto CreateStereoArchetype = []() -> FMetasoundFrontendArchetype
-	{
-		FMetasoundFrontendArchetype Archetype = GetBaseArchetype();
-		Archetype.Name = FName(TEXT("StereoSource"));
-
-		FMetasoundFrontendClassVertex GeneratedAudio;
-		GeneratedAudio.Name = UMetaSoundSource::GetAudioOutputName();
-		GeneratedAudio.TypeName = Metasound::Frontend::GetDataTypeName<Metasound::FStereoAudioFormat>();
-		GeneratedAudio.Metadata.DisplayName = LOCTEXT("GeneratedStereo", "Audio");
-		GeneratedAudio.Metadata.Description = LOCTEXT("GeneratedAudioToolTip", "The resulting output audio from this source.");
-		GeneratedAudio.VertexID = FGuid::NewGuid();
-
-		Archetype.Interface.Outputs.Add(GeneratedAudio);
-
-		return Archetype;
-	};
-
-	static const FMetasoundFrontendArchetype StereoArchetype = CreateStereoArchetype();
-
-	return StereoArchetype;
-}
 #undef LOCTEXT_NAMESPACE // MetaSoundSource
diff --git a/Engine/Plugins/Runtime/Metasound/Source/MetasoundEngine/Private/MetasoundUObjectRegistry.cpp b/Engine/Plugins/Runtime/Metasound/Source/MetasoundEngine/Private/MetasoundUObjectRegistry.cpp
--- a/Engine/Plugins/Runtime/Metasound/Source/MetasoundEngine/Private/MetasoundUObjectRegistry.cpp
+++ b/Engine/Plugins/Runtime/Metasound/Source/MetasoundEngine/Private/MetasoundUObjectRegistry.cpp
@@ -11,6 +11,7 @@
 #include "Engine/AssetManager.h"
 #include "Metasound.h"
 #include "MetasoundAssetBase.h"
+#include "MetasoundFrontendArchetypeRegistry.h"
 #include "MetasoundFrontendRegistries.h"
 #include "MetasoundFrontendTransform.h"
 #include "MetasoundSource.h"
@@ -151,17 +152,14 @@ void UMetaSoundAssetSubsystem::AddOrUpdateAsset(const FAssetData& InAssetData)
 
 		if (ensure(Object))
 		{
+			
 			FMetasoundAssetBase* MetaSoundAsset = Metasound::IMetasoundUObjectRegistry::Get().GetObjectAsAssetBase(Object);
 			check(MetaSoundAsset);
 
-			FDocumentHandle Document = MetaSoundAsset->GetDocumentHandle();
-
 			// Must version to ensure registration uses the correct key,
 			// which must be based off of most up-to-date document model
 			// for safety.
-			const FName AssetName = Object->GetFName();
-			const FString AssetPath = Object->GetPathName();
-			FVersionDocument(AssetName, AssetPath).Transform(Document);
+			MetaSoundAsset->VersionAsset();
 
 			MetaSoundAsset->RegisterGraphWithFrontend();
 		}
@@ -221,20 +219,20 @@ namespace Metasound
 			{
 				if (InEntry.IsValid())
 				{
-					FName ArchetypeName = InEntry->GetArchetypeName();
+					Frontend::FArchetypeRegistryKey Key = Frontend::GetArchetypeRegistryKey(InEntry->GetArchetypeVersion());
 
-					EntriesByArchetype.Add(ArchetypeName, InEntry.Get());
+					EntriesByArchetype.Add(Key, InEntry.Get());
 					Entries.Add(InEntry.Get());
 					Storage.Add(MoveTemp(InEntry));
 				}
 			}
 
-			TArray<UClass*> GetUClassesForArchetype(const FName& InArchetypeName) const override
+			TArray<UClass*> GetUClassesForArchetype(const FMetasoundFrontendVersion& InArchetypeVersion) const override
 			{
 				TArray<UClass*> Classes;
 
 				TArray<const IMetasoundUObjectRegistryEntry*> EntriesForArchetype;
-				EntriesByArchetype.MultiFind(InArchetypeName, EntriesForArchetype);
+				EntriesByArchetype.MultiFind(Frontend::GetArchetypeRegistryKey(InArchetypeVersion), EntriesForArchetype);
 
 				for (const IMetasoundUObjectRegistryEntry* Entry : EntriesForArchetype)
 				{
@@ -250,21 +248,22 @@ namespace Metasound
 				return Classes;
 			}
 
-			UObject* NewObject(UClass* InClass, const FMetasoundFrontendDocument& InDocument, const FMetasoundFrontendArchetype& InArchetype, const FString& InPath) const override
+			UObject* NewObject(UClass* InClass, const FMetasoundFrontendDocument& InDocument, const FString& InPath) const override
 			{
+				TArray<const IMetasoundUObjectRegistryEntry*> EntriesForArchetype;
+				EntriesByArchetype.MultiFind(Frontend::GetArchetypeRegistryKey(InDocument.ArchetypeVersion), EntriesForArchetype);
+
 				auto IsChildClassOfRegisteredClass = [&](const IMetasoundUObjectRegistryEntry* Entry)
 				{
 					return Entry->IsChildClass(InClass);
 				};
 
-				TArray<const IMetasoundUObjectRegistryEntry*> EntriesForClass = FindEntriesByPredicate(IsChildClassOfRegisteredClass);
+			
+				const IMetasoundUObjectRegistryEntry* const* EntryForClass = EntriesForArchetype.FindByPredicate(IsChildClassOfRegisteredClass);
 
-				for (const IMetasoundUObjectRegistryEntry* Entry : EntriesForClass)
+				if (nullptr != EntryForClass)
 				{
-					if (Entry->GetArchetypeName() == InArchetype.Name)
-					{
-						return NewObject(*Entry, InDocument, InPath);
-					}
+					return NewObject(**EntryForClass, InDocument, InPath);
 				}
 
 				return nullptr;
@@ -322,12 +321,7 @@ namespace Metasound
 				if (ensure(nullptr != NewAssetBase))
 				{
 					NewAssetBase->SetDocument(InDocument);
-
-					const FMetasoundFrontendArchetype& Archetype = NewAssetBase->GetArchetype();
-					if (ensure(NewAssetBase->IsArchetypeSupported(Archetype)))
-					{
-						NewAssetBase->ConformDocumentToArchetype();
-					}
+					NewAssetBase->ConformDocumentToArchetype();
 				}
 
 #if WITH_EDITOR
@@ -379,7 +373,7 @@ namespace Metasound
 			}
 
 			TArray<TUniquePtr<IMetasoundUObjectRegistryEntry>> Storage;
-			TMultiMap<FName, const IMetasoundUObjectRegistryEntry*> EntriesByArchetype;
+			TMultiMap<Frontend::FArchetypeRegistryKey, const IMetasoundUObjectRegistryEntry*> EntriesByArchetype;
 			TArray<const IMetasoundUObjectRegistryEntry*> Entries;
 	};
 
diff --git a/Engine/Plugins/Runtime/Metasound/Source/MetasoundEngine/Public/Metasound.h b/Engine/Plugins/Runtime/Metasound/Source/MetasoundEngine/Public/Metasound.h
--- a/Engine/Plugins/Runtime/Metasound/Source/MetasoundEngine/Public/Metasound.h
+++ b/Engine/Plugins/Runtime/Metasound/Source/MetasoundEngine/Public/Metasound.h
@@ -161,7 +161,7 @@ public:
 	// Returns Asset Metadata associated with this MetaSound
 	virtual Metasound::Frontend::FNodeClassInfo GetAssetClassInfo() const override;
 
-	virtual const FMetasoundFrontendArchetype& GetArchetype() const override;
+	virtual const FMetasoundFrontendVersion& GetDefaultArchetypeVersion() const override;
 
 	UObject* GetOwningAsset() override
 	{
@@ -176,7 +176,7 @@ public:
 	TUniquePtr<IAudioInstanceTransmitter> CreateInstanceTransmitter(const FAudioInstanceTransmitterInitParams& InParams) const;
 
 	// Get the most up to date archetype for metasound sources.
-	const TArray<FMetasoundFrontendArchetype>& GetPreferredArchetypes() const override;
+	const TArray<FMetasoundFrontendVersion>& GetSupportedArchetypeVersions() const override;
 
 protected:
 
diff --git a/Engine/Plugins/Runtime/Metasound/Source/MetasoundEngine/Public/MetasoundSource.h b/Engine/Plugins/Runtime/Metasound/Source/MetasoundEngine/Public/MetasoundSource.h
--- a/Engine/Plugins/Runtime/Metasound/Source/MetasoundEngine/Public/MetasoundSource.h
+++ b/Engine/Plugins/Runtime/Metasound/Source/MetasoundEngine/Public/MetasoundSource.h
@@ -124,7 +124,7 @@ public:
 	// Returns Asset Metadata associated with this MetaSoundSource
 	virtual Metasound::Frontend::FNodeClassInfo GetAssetClassInfo() const override;
 
-	virtual const FMetasoundFrontendArchetype& GetArchetype() const override;
+	virtual const FMetasoundFrontendVersion& GetDefaultArchetypeVersion() const override;
 
 	UObject* GetOwningAsset() override
 	{
@@ -143,11 +143,8 @@ public:
 	virtual TUniquePtr<IAudioInstanceTransmitter> CreateInstanceTransmitter(const FAudioInstanceTransmitterInitParams& InParams) const override;
 
 	// Get the most up to date archetype for metasound sources.
-	const TArray<FMetasoundFrontendArchetype>& GetPreferredArchetypes() const override;
+	const TArray<FMetasoundFrontendVersion>& GetSupportedArchetypeVersions() const override;
 
-	static const FMetasoundFrontendArchetype& GetBaseArchetype();
-	static const FMetasoundFrontendArchetype& GetMonoSourceArchetype();
-	static const FMetasoundFrontendArchetype& GetStereoSourceArchetype();
 
 protected:
 	Metasound::Frontend::FDocumentAccessPtr GetDocument() override
@@ -168,11 +165,4 @@ protected:
 
 private:
 	Metasound::FOperatorSettings GetOperatorSettings(Metasound::FSampleRate InSampleRate) const;
-
-	static const FString& GetOnPlayInputName();
-	static const FString& GetAudioOutputName();
-	static const FString& GetIsFinishedOutputName();
-	static const FString& GetAudioDeviceHandleVariableName();
-	static const FString& GetSoundUniqueIdName();
-	static const FString& GetIsPreviewSoundName();
 };
diff --git a/Engine/Plugins/Runtime/Metasound/Source/MetasoundEngine/Public/MetasoundUObjectRegistry.h b/Engine/Plugins/Runtime/Metasound/Source/MetasoundEngine/Public/MetasoundUObjectRegistry.h
--- a/Engine/Plugins/Runtime/Metasound/Source/MetasoundEngine/Public/MetasoundUObjectRegistry.h
+++ b/Engine/Plugins/Runtime/Metasound/Source/MetasoundEngine/Public/MetasoundUObjectRegistry.h
@@ -46,7 +46,7 @@ namespace Metasound
 		virtual ~IMetasoundUObjectRegistryEntry() = default;
 
 		/** Archetype name associated with this entry. */
-		virtual const FName& GetArchetypeName() const = 0;
+		virtual const FMetasoundFrontendVersion& GetArchetypeVersion() const = 0;
 
 		/** UClass associated with this entry. */
 		virtual UClass* GetUClass() const = 0;
@@ -86,16 +86,16 @@ namespace Metasound
 		static_assert(std::is_base_of<UObject, UClassType>::value, "UClass must be derived from UObject");
 
 	public:
-		TMetasoundUObjectRegistryEntry(const FName& InArchetypeName)
-		:	ArchetypeName(InArchetypeName)
+		TMetasoundUObjectRegistryEntry(const FMetasoundFrontendVersion& InArchetypeVersion)
+		:	ArchetypeVersion(InArchetypeVersion)
 		{
 		}
 
 		virtual ~TMetasoundUObjectRegistryEntry() = default;
 
-		const FName& GetArchetypeName() const
+		virtual const FMetasoundFrontendVersion& GetArchetypeVersion() const override
 		{
-			return ArchetypeName;
+			return ArchetypeVersion;
 		}
 
 		UClass* GetUClass() const override
@@ -146,7 +146,7 @@ namespace Metasound
 
 	private:
 
-		FName ArchetypeName;
+		FMetasoundFrontendVersion ArchetypeVersion;
 	};
 
 
@@ -165,49 +165,49 @@ namespace Metasound
 
 			/** Register all preferred archetypes of the UClass. 
 			 *
-			 * All Archtypes returned by the default objects GetPreferredArchetypes() will be registered. 
+			 * All Archtypes returned by the default objects GetSupportedArchetypeVersions() will be registered. 
 			 */
 			template<typename UClassType>
 			static void RegisterUClassPreferredArchetypes()
 			{
 				static_assert(std::is_base_of<FMetasoundAssetBase, UClassType>::value, "UClass must be derived from FMetasoundAssetBase");
 
-				const TArray<FMetasoundFrontendArchetype>& PreferredArchetypes = GetDefault<UClassType>()->GetPreferredArchetypes();
+				const TArray<FMetasoundFrontendVersion>& SupportedArchetypeVersions = GetDefault<UClassType>()->GetSupportedArchetypeVersions();
 
-				for (const FMetasoundFrontendArchetype& Arch : PreferredArchetypes)
+				for (const FMetasoundFrontendVersion& Version : SupportedArchetypeVersions)
 				{
-					IMetasoundUObjectRegistry::RegisterUClassArchetype<UClassType>(Arch.Name);
+					IMetasoundUObjectRegistry::RegisterUClassArchetype<UClassType>(Version);
 				}
 			}
 
 			/** Register an archetype for the UClass. 
 			 *
-			 * @param InArchetypeName - An FMetasoundFrontendDocument to associate with the UClass.
+			 * @param InArchetypeVerison - The version of the FMetasoundFrontendArchetype to associate with the UClass.
 			 */
 			template<typename UClassType>
-			static void RegisterUClassArchetype(const FName& InArchetypeName)
+			static void RegisterUClassArchetype(const FMetasoundFrontendVersion& InArchetypeVersion)
 			{
 				using FRegistryEntryType = TMetasoundUObjectRegistryEntry<UClassType>;
 
-				IMetasoundUObjectRegistry::Get().RegisterUClassArchetype(MakeUnique<FRegistryEntryType>(InArchetypeName));
+				IMetasoundUObjectRegistry::Get().RegisterUClassArchetype(MakeUnique<FRegistryEntryType>(InArchetypeVersion));
 			}
 
 			/** Adds an entry to the registry. */
 			virtual void RegisterUClassArchetype(TUniquePtr<IMetasoundUObjectRegistryEntry>&& InEntry) = 0;
 
 			/** Returns all UClasses registered to the archetype name. */
-			virtual TArray<UClass*> GetUClassesForArchetype(const FName& InArchetypeName) const = 0;
+			virtual TArray<UClass*> GetUClassesForArchetype(const FMetasoundFrontendVersion& InArchetypeVersion) const = 0;
 
 			/** Creates a new object from a metasound document.
 			 *
 			 * @param InClass - A registered UClass to create.
 			 * @param InDocument - The FMetasoundFrontendDocument to use when creating the class.
-			 * @param InArchetype - The FMetasoundFrontendArchetype to use when creating the class.
+			 * @param InArchetypeVersion - The version of the FMetasoundFrontendArchetype to use when creating the class.
 			 * @param InPath - If in editor, the created asset will be stored at this content path.
 			 *
 			 * @return A new object. A nullptr on error.
 			 */
-			virtual UObject* NewObject(UClass* InClass, const FMetasoundFrontendDocument& InDocument, const FMetasoundFrontendArchetype& InArchetype, const FString& InPath) const = 0;
+			virtual UObject* NewObject(UClass* InClass, const FMetasoundFrontendDocument& InDocument, const FString& InPath) const = 0;
 
 			/** Returns true if the InObject is of a class or child class which is registered. */
 			virtual bool IsRegisteredClass(UObject* InObject) const = 0;
diff --git a/Engine/Plugins/Runtime/Metasound/Source/MetasoundFrontend/Private/MetasoundArchetype.cpp b/Engine/Plugins/Runtime/Metasound/Source/MetasoundFrontend/Private/MetasoundArchetype.cpp
--- a/Engine/Plugins/Runtime/Metasound/Source/MetasoundFrontend/Private/MetasoundArchetype.cpp
+++ b/Engine/Plugins/Runtime/Metasound/Source/MetasoundFrontend/Private/MetasoundArchetype.cpp
@@ -81,6 +81,16 @@ namespace Metasound
 						IndexB++;
 					}
 				}
+
+				if (IndexA < SortedA.Num())
+				{
+					OutUniqueA.Append(&SortedA[IndexA], SortedA.Num() - IndexA);
+				}
+
+				if (IndexB < SortedB.Num())
+				{
+					OutUniqueB.Append(&SortedB[IndexB], SortedB.Num() - IndexB);
+				}
 			}
 
 			template<typename FrontendTypeA, typename FrontendTypeB, typename ComparisonType>
@@ -130,6 +140,19 @@ namespace Metasound
 			return bIsInputSetIncluded && bIsOutputSetIncluded && bIsEnvironmentVariableSetIncluded;
 		}
 
+		bool IsSubsetOfClass(const FMetasoundFrontendArchetype& InSubsetArchetype, const FMetasoundFrontendClass& InSupersetClass)
+		{
+			// TODO: Environment variables are ignored as they are poorly supported and classes describe which environment variables are required,
+			// not which are supported 
+			using namespace MetasoundArchetypeIntrinsics;
+
+			const bool bIsInputSetIncluded = IsSetIncluded(InSubsetArchetype.Interface.Inputs, InSupersetClass.Interface.Inputs, &IsLessThanVertex);
+
+			const bool bIsOutputSetIncluded = IsSetIncluded(InSubsetArchetype.Interface.Outputs, InSupersetClass.Interface.Outputs, &IsLessThanVertex);
+
+			return bIsInputSetIncluded && bIsOutputSetIncluded;
+		}
+
 		bool IsEquivalentArchetype(const FMetasoundFrontendArchetype& InInputArchetype, const FMetasoundFrontendArchetype& InTargetArchetype)
 		{
 			using namespace MetasoundArchetypeIntrinsics;
@@ -145,9 +168,9 @@ namespace Metasound
 
 		bool IsEqualArchetype(const FMetasoundFrontendArchetype& InInputArchetype, const FMetasoundFrontendArchetype& InTargetArchetype)
 		{
-			bool bIsMetadataEqual = InInputArchetype.Name == InTargetArchetype.Name;
-			bIsMetadataEqual &= InInputArchetype.Version.Major == InTargetArchetype.Version.Major;
-			bIsMetadataEqual &= InInputArchetype.Version.Minor== InTargetArchetype.Version.Minor;
+			bool bIsMetadataEqual = InInputArchetype.Version.Name == InTargetArchetype.Version.Name;
+			bIsMetadataEqual &= InInputArchetype.Version.Number.Major == InTargetArchetype.Version.Number.Major;
+			bIsMetadataEqual &= InInputArchetype.Version.Number.Minor == InTargetArchetype.Version.Number.Minor;
 
 			if (bIsMetadataEqual)
 			{
@@ -159,13 +182,13 @@ namespace Metasound
 			return false;
 		}
 
-		int32 InputOutputDifferenceCount(const FMetasoundFrontendDocument& InDocument, const FMetasoundFrontendArchetype& InArchetype)
+		int32 InputOutputDifferenceCount(const FMetasoundFrontendClass& InClass, const FMetasoundFrontendArchetype& InArchetype)
 		{
 			using namespace MetasoundArchetypeIntrinsics;
 
-			int32 DiffCount = SetDifferenceCount(InDocument.RootGraph.Interface.Inputs, InArchetype.Interface.Inputs, &IsLessThanVertex);
+			int32 DiffCount = SetDifferenceCount(InClass.Interface.Inputs, InArchetype.Interface.Inputs, &IsLessThanVertex);
 
-			DiffCount += SetDifferenceCount(InDocument.RootGraph.Interface.Outputs, InArchetype.Interface.Outputs, &IsLessThanVertex);
+			DiffCount += SetDifferenceCount(InClass.Interface.Outputs, InArchetype.Interface.Outputs, &IsLessThanVertex);
 
 			return DiffCount;
 		}
@@ -181,7 +204,7 @@ namespace Metasound
 			return DiffCount;
 		}
 
-		void GatherRequiredEnvironmentVariables(const FMetasoundFrontendDocument& InDocument, TArray<FMetasoundFrontendEnvironmentVariable>& OutEnvironmentVariables)
+		void GatherRequiredEnvironmentVariables(const FMetasoundFrontendGraphClass& InRootGraph, const TArray<FMetasoundFrontendClass>& InDependencies, const TArray<FMetasoundFrontendGraphClass>& InSubgraphs, TArray<FMetasoundFrontendEnvironmentVariable>& OutEnvironmentVariables)
 		{
 			auto GatherRequiredEnvironmentVariablesFromClass = [&](const FMetasoundFrontendClass& InClass)
 			{
@@ -205,12 +228,13 @@ namespace Metasound
 				}
 			};
 
-			GatherRequiredEnvironmentVariablesFromClass(InDocument.RootGraph);
-			Algo::ForEach(InDocument.Dependencies, GatherRequiredEnvironmentVariablesFromClass);
-			Algo::ForEach(InDocument.Subgraphs, GatherRequiredEnvironmentVariablesFromClass);
+			GatherRequiredEnvironmentVariablesFromClass(InRootGraph);
+			Algo::ForEach(InDependencies, GatherRequiredEnvironmentVariablesFromClass);
+			Algo::ForEach(InSubgraphs, GatherRequiredEnvironmentVariablesFromClass);
 		}
 
-		const FMetasoundFrontendArchetype* FindMostSimilarArchetypeSupportingEnvironment(const FMetasoundFrontendDocument& InDocument, const TArray<FMetasoundFrontendArchetype>& InCandidateArchetypes)
+		
+		const FMetasoundFrontendArchetype* FindMostSimilarArchetypeSupportingEnvironment(const FMetasoundFrontendGraphClass& InRootGraph, const TArray<FMetasoundFrontendClass>& InDependencies, const TArray<FMetasoundFrontendGraphClass>& InSubgraphs, const TArray<FMetasoundFrontendArchetype>& InCandidateArchetypes)
 		{
 			using namespace MetasoundArchetypeIntrinsics;
 
@@ -218,24 +242,24 @@ namespace Metasound
 
 			TArray<const FMetasoundFrontendArchetype*> CandidateArchs = MakePointerToArrayElements(InCandidateArchetypes);
 
-			TArray<FMetasoundFrontendEnvironmentVariable> AllDocumentEnvironmentVariables;
-			GatherRequiredEnvironmentVariables(InDocument, AllDocumentEnvironmentVariables);
+			TArray<FMetasoundFrontendEnvironmentVariable> AllEnvironmentVariables;
+			GatherRequiredEnvironmentVariables(InRootGraph, InDependencies, InSubgraphs, AllEnvironmentVariables);
 
 			// Remove all archetypes which do not provide all environment variables. 
 			auto DoesNotProvideAllEnvironmentVariables = [&](const FMetasoundFrontendArchetype* CandidateArch)
 			{
-				return !IsSetIncluded(AllDocumentEnvironmentVariables, CandidateArch->Interface.Environment, &IsLessThanEnvironmentVariable);
+				return !IsSetIncluded(AllEnvironmentVariables, CandidateArch->Interface.Environment, &IsLessThanEnvironmentVariable);
 			};
 
 			CandidateArchs.RemoveAll(DoesNotProvideAllEnvironmentVariables);
 
 			// Return the archetype with the least amount of differences.
-			auto DifferencesFromDocument = [&](const FMetasoundFrontendArchetype* CandidateArch) -> int32
+			auto DifferencesFromRootGraph = [&](const FMetasoundFrontendArchetype* CandidateArch) -> int32
 			{ 
-				return InputOutputDifferenceCount(InDocument, *CandidateArch); 
+				return InputOutputDifferenceCount(InRootGraph, *CandidateArch); 
 			};
 
-			Algo::SortBy(CandidateArchs, DifferencesFromDocument);
+			Algo::SortBy(CandidateArchs, DifferencesFromRootGraph);
 
 			if (0 == CandidateArchs.Num())
 			{
@@ -244,5 +268,10 @@ namespace Metasound
 
 			return CandidateArchs[0];
 		}
+		
+		const FMetasoundFrontendArchetype* FindMostSimilarArchetypeSupportingEnvironment(const FMetasoundFrontendDocument& InDocument, const TArray<FMetasoundFrontendArchetype>& InCandidateArchetypes)
+		{
+			return FindMostSimilarArchetypeSupportingEnvironment(InDocument.RootGraph, InDocument.Dependencies, InDocument.Subgraphs, InCandidateArchetypes);
+		}
 	}
 }
diff --git a/Engine/Plugins/Runtime/Metasound/Source/MetasoundFrontend/Private/MetasoundAssetBase.cpp b/Engine/Plugins/Runtime/Metasound/Source/MetasoundFrontend/Private/MetasoundAssetBase.cpp
--- a/Engine/Plugins/Runtime/Metasound/Source/MetasoundFrontend/Private/MetasoundAssetBase.cpp
+++ b/Engine/Plugins/Runtime/Metasound/Source/MetasoundFrontend/Private/MetasoundAssetBase.cpp
@@ -8,6 +8,7 @@
 #include "Internationalization/Text.h"
 #include "IStructSerializerBackend.h"
 #include "MetasoundArchetype.h"
+#include "MetasoundFrontendArchetypeRegistry.h"
 #include "MetasoundFrontendController.h"
 #include "MetasoundFrontendGraph.h"
 #include "MetasoundFrontendSearchEngine.h"
@@ -168,50 +169,109 @@ void FMetasoundAssetBase::SetMetadata(FMetasoundFrontendClassMetadata& InMetadat
 	MarkMetasoundDocumentDirty();
 }
 
-bool FMetasoundAssetBase::IsArchetypeSupported(const FMetasoundFrontendArchetype& InArchetype) const
+bool FMetasoundAssetBase::IsArchetypeSupported(const FMetasoundFrontendVersion& InArchetypeVersion) const
 {
-	auto IsEqualArchetype = [&](const FMetasoundFrontendArchetype& SupportedArchetype)
-	{
-		return Metasound::Frontend::IsEqualArchetype(InArchetype, SupportedArchetype);
-	};
-
-	return Algo::AnyOf(GetPreferredArchetypes(), IsEqualArchetype);
+	return GetSupportedArchetypeVersions().Contains(InArchetypeVersion);
 }
 
-const FMetasoundFrontendArchetype& FMetasoundAssetBase::GetPreferredArchetypes(const FMetasoundFrontendDocument& InDocument, const FMetasoundFrontendArchetype& InDefaultArchetype) const
+FMetasoundFrontendVersion FMetasoundAssetBase::GetPreferredArchetypeVersion(const FMetasoundFrontendDocument& InDocument) const
 {
+	using namespace Metasound::Frontend;
+
 	// Default to archetype provided in case it is supported.
-	if (IsArchetypeSupported(InDefaultArchetype))
+	if (IsArchetypeSupported(InDocument.ArchetypeVersion))
 	{
-		return InDefaultArchetype;
+		return InDocument.ArchetypeVersion;
 	}
 
+	// If existing archetype is not supported, check if an updated version is preferred.
+	auto IsNewVersion = [&](const FMetasoundFrontendVersion& InVersion)
+	{
+		return (InDocument.ArchetypeVersion.Name == InVersion.Name) && (InVersion.Number > InDocument.ArchetypeVersion.Number);
+	};
+	TArray<FMetasoundFrontendVersion> UpdatedVersions = GetSupportedArchetypeVersions().FilterByPredicate(IsNewVersion);
+	if (UpdatedVersions.Num() > 0)
+	{
+		UpdatedVersions.Sort();
+		return UpdatedVersions.Last();
+	}
+	
 	// If existing archetype is not supported, get the most similar that still supports the documents environment.
-	const FMetasoundFrontendArchetype* SimilarArchetype = Metasound::Frontend::FindMostSimilarArchetypeSupportingEnvironment(InDocument, GetPreferredArchetypes());
+	TArray<FMetasoundFrontendArchetype> Archetypes;
+	for (const FMetasoundFrontendVersion& Version : GetSupportedArchetypeVersions())
+	{
+		const FArchetypeRegistryKey Key = GetArchetypeRegistryKey(Version);
+		FMetasoundFrontendArchetype Archetype;
+		if (IArchetypeRegistry::Get().FindArchetype(Key, Archetype))
+		{
+			Archetypes.Add(MoveTemp(Archetype));
+		}
+	}
+	const FMetasoundFrontendArchetype* SimilarArchetype = Metasound::Frontend::FindMostSimilarArchetypeSupportingEnvironment(InDocument, Archetypes);
 
 	if (nullptr != SimilarArchetype)
 	{
-		return *SimilarArchetype;
+		return SimilarArchetype->Version;
 	}
 
 	// Nothing found. Return the existing archetype for the FMetasoundAssetBase.
-	return GetArchetype();
+	return GetDefaultArchetypeVersion();
 }
 
 void FMetasoundAssetBase::SetDocument(const FMetasoundFrontendDocument& InDocument)
 {
 	FMetasoundFrontendDocument& Document = GetDocumentChecked();
 	Document = InDocument;
+	MarkMetasoundDocumentDirty();
 }
 
 void FMetasoundAssetBase::ConformDocumentToArchetype()
 {
-	FMetasoundFrontendDocument& Document = GetDocumentChecked();
-	const FMetasoundFrontendArchetype& Archetype = GetArchetype();
-	Metasound::Frontend::FMatchRootGraphToArchetype Transform(Archetype);
-	Transform.Transform(GetDocumentHandle());
+	FMetasoundFrontendDocument& Doc = GetDocumentChecked();
 
-	MarkMetasoundDocumentDirty();
+	FMetasoundFrontendVersion PreferredArchetypeVersion = GetPreferredArchetypeVersion(Doc);
+
+	if (PreferredArchetypeVersion != Doc.ArchetypeVersion)
+	{
+		Metasound::Frontend::FMatchRootGraphToArchetype Transform(PreferredArchetypeVersion);
+		if (Transform.Transform(GetDocumentHandle()))
+		{
+			MarkMetasoundDocumentDirty();
+		}
+	}
+}
+
+bool FMetasoundAssetBase::VersionAsset()
+{
+	using namespace Metasound;
+	using namespace Metasound::Frontend;
+
+	FName AssetName;
+	FString AssetPath;
+	if (const UObject* OwningAsset = GetOwningAsset())
+	{
+		AssetName = FName(OwningAsset->GetName());
+		AssetPath = OwningAsset->GetPathName();
+	}
+
+	FDocumentHandle DocumentHandle = GetDocumentHandle();
+
+	const bool bDidUpdateDocumentVersion = FVersionDocument(AssetName, AssetPath).Transform(DocumentHandle);
+	bool bDidMatchRootGraphToArchetype = false;
+	if (FMetasoundFrontendDocument* Doc = GetDocument().Get())
+	{
+		FMetasoundFrontendVersion ArchetypeVerison = GetPreferredArchetypeVersion(*Doc);
+		bDidMatchRootGraphToArchetype = FMatchRootGraphToArchetype(ArchetypeVerison).Transform(DocumentHandle);
+	}
+
+	const bool bDidUpdate = bDidUpdateDocumentVersion || bDidMatchRootGraphToArchetype;
+
+	if (bDidUpdate)
+	{
+		MarkMetasoundDocumentDirty();
+	}
+
+	return bDidUpdate;
 }
 
 bool FMetasoundAssetBase::CopyDocumentAndInjectReceiveNodes(uint64 InInstanceID, const FMetasoundFrontendDocument& InSourceDoc, FMetasoundFrontendDocument& OutDestDoc) const
@@ -353,15 +413,28 @@ TArray<FMetasoundAssetBase::FSendInfoAndVertexName> FMetasoundAssetBase::GetSend
 TArray<FString> FMetasoundAssetBase::GetTransmittableInputVertexNames() const
 {
 	using namespace Metasound;
+	using namespace Metasound::Frontend;
+
+	FConstDocumentHandle Document = GetDocumentHandle();
+
+	// Find the archetype for the document.
+	FMetasoundFrontendArchetype Archetype;
+	FArchetypeRegistryKey ArchetypeRegistryKey = Frontend::GetArchetypeRegistryKey(Document->GetArchetypeVersion());
+	bool bFoundArchetype = IArchetypeRegistry::Get().FindArchetype(ArchetypeRegistryKey, Archetype);
+
+	if (!bFoundArchetype)
+	{
+		UE_LOG(LogMetaSound, Warning, TEXT("No registered archetype matching archetype version on document [ArchetypeVersion:%s]"), *Document->GetArchetypeVersion().ToString());
+	}
 
 	// Unused inputs are all input vertices that are not in the archetype.
 	TArray<FString> ArchetypeInputVertexNames;
-	for (const FMetasoundFrontendClassVertex& Vertex : GetArchetype().Interface.Inputs)
+	for (const FMetasoundFrontendClassVertex& Vertex : Archetype.Interface.Inputs)
 	{
 		ArchetypeInputVertexNames.Add(Vertex.Name);
 	}
 
-	Frontend::FConstGraphHandle RootGraph = GetRootGraphHandle();
+	Frontend::FConstGraphHandle RootGraph = Document->GetRootGraph();
 	TArray<FString> GraphInputVertexNames = RootGraph->GetInputVertexNames();
 
 	// Filter graph inputs by archetype inputs.
@@ -377,12 +450,6 @@ TArray<FString> FMetasoundAssetBase::GetTransmittableInputVertexNames() const
 			{
 				if (TypeInfo.bIsTransmittable)
 				{
-					// TODO: Currently values set directly on node pins are represented
-					// as input nodes in the graph. They should not be used for transmission
-					// as the number of input nodes increases quickly as more nodes
-					// are added to a graph. Connecting these input nodes to the 
-					// transmission system is relatively expensive. These undesirable input nodes are
-					// filtered out by ignoring input nodes which are not "Visible". 
 					Frontend::FConstNodeHandle InputNode = RootGraph->GetNodeWithID(ClassInput->NodeID);
 					if (InputNode->IsValid())
 					{
diff --git a/Engine/Plugins/Runtime/Metasound/Source/MetasoundFrontend/Private/MetasoundFrontend.cpp b/Engine/Plugins/Runtime/Metasound/Source/MetasoundFrontend/Private/MetasoundFrontend.cpp
--- a/Engine/Plugins/Runtime/Metasound/Source/MetasoundFrontend/Private/MetasoundFrontend.cpp
+++ b/Engine/Plugins/Runtime/Metasound/Source/MetasoundFrontend/Private/MetasoundFrontend.cpp
@@ -25,23 +25,6 @@ namespace Metasound
 {
 	namespace Frontend
 	{
-		TArray<const IRegistryTransaction*> GetRegistryTransactionsSince(FRegistryTransactionID InTransactionID, FRegistryTransactionID* OutCurrentTransactionID)
-		{
-			FMetasoundFrontendRegistryContainer* Registry = FMetasoundFrontendRegistryContainer::Get();
-			if (ensure(nullptr != Registry))
-			{
-				return Registry->GetRegistryTransactionsSince(InTransactionID, OutCurrentTransactionID);
-			}
-			else
-			{
-				if (nullptr != OutCurrentTransactionID)
-				{
-					*OutCurrentTransactionID = GetOriginRegistryTransactionID();
-				}
-				return TArray<const IRegistryTransaction*>();
-			}
-		}
-
 		// gets all metadata (name, description, author, what to say if it's missing) for a given node.
 		FMetasoundFrontendClassMetadata GenerateClassMetadata(const FNodeRegistryKey& InKey)
 		{
diff --git a/Engine/Plugins/Runtime/Metasound/Source/MetasoundFrontend/Private/MetasoundFrontendArchetypeRegistry.cpp b/Engine/Plugins/Runtime/Metasound/Source/MetasoundFrontend/Private/MetasoundFrontendArchetypeRegistry.cpp
new file mode 100644
--- /dev/null
+++ b/Engine/Plugins/Runtime/Metasound/Source/MetasoundFrontend/Private/MetasoundFrontendArchetypeRegistry.cpp
@@ -0,0 +1,120 @@
+// Copyright Epic Games, Inc. All Rights Reserved.
+
+#include "MetasoundFrontendArchetypeRegistry.h"
+#include "MetasoundFrontendRegistryTransaction.h"
+
+namespace Metasound
+{
+	namespace Frontend
+	{
+		namespace MetasoundFrontendArchetypeRegistryPrivate
+		{
+			class FArchetypeRegistry : public IArchetypeRegistry
+			{
+			public:
+				virtual ~FArchetypeRegistry() = default;
+
+				virtual FArchetypeRegistryKey RegisterArchetype(TUniquePtr<IArchetypeRegistryEntry>&& InEntry) override
+				{
+					if (InEntry.IsValid())
+					{
+						FArchetypeRegistryKey Key = GetArchetypeRegistryKey(InEntry->GetArchetype());
+						if (IsValidArchetypeRegistryKey(Key))
+						{
+							if (const IArchetypeRegistryEntry* Entry = FindArchetypeRegistryEntry(Key))
+							{
+								UE_LOG(LogMetaSound, Warning, TEXT("Registration of archetype overwriting previously registered archetype [RegistryKey: %s]"), *Key);
+								
+								FArchetypeRegistryTransaction Transaction{FArchetypeRegistryTransaction::ETransactionType::ArchetypeUnregistration, Key, Entry->GetArchetype().Version};
+								Transactions.Add(MoveTemp(Transaction));
+							}
+							
+							FArchetypeRegistryTransaction Transaction{FArchetypeRegistryTransaction::ETransactionType::ArchetypeRegistration, Key, InEntry->GetArchetype().Version};
+							Transactions.Add(MoveTemp(Transaction));
+
+							Entries.Add(Key, MoveTemp(InEntry));
+
+							return Key;
+						}
+					}
+
+					return FArchetypeRegistryKey{};
+				}
+
+				virtual const IArchetypeRegistryEntry* FindArchetypeRegistryEntry(const FArchetypeRegistryKey& InKey) const override
+				{
+					if (const TUniquePtr<IArchetypeRegistryEntry>* Entry = Entries.Find(InKey))
+					{
+						return Entry->Get();
+					}
+					return nullptr;
+				}
+
+				virtual bool FindArchetype(const FArchetypeRegistryKey& InKey, FMetasoundFrontendArchetype& OutArchetype) const override
+				{
+					if (const IArchetypeRegistryEntry* Entry = FindArchetypeRegistryEntry(InKey))
+					{
+						OutArchetype = Entry->GetArchetype();
+						return true;
+					}
+
+					return false;
+				}
+
+				virtual void ForEachRegistryTransactionSince(FRegistryTransactionID InSince, FRegistryTransactionID* OutCurrentRegistryTransactionID, TFunctionRef<void(const FArchetypeRegistryTransaction&)> InFunc) const override
+				{
+					Transactions.ForEachTransactionSince(InSince, OutCurrentRegistryTransactionID, InFunc);
+				}
+
+			private:
+
+				TMap<FArchetypeRegistryKey, TUniquePtr<IArchetypeRegistryEntry>> Entries;
+				TRegistryTransactionHistory<FArchetypeRegistryTransaction> Transactions;
+			};
+		}
+
+		bool IsValidArchetypeRegistryKey(const FArchetypeRegistryKey& InKey)
+		{
+			return !InKey.IsEmpty();
+		}
+
+		FArchetypeRegistryKey GetArchetypeRegistryKey(const FMetasoundFrontendVersion& InArchetypeVersion)
+		{
+			return FString::Format(TEXT("{0}_{1}.{2}"), { InArchetypeVersion.Name.ToString(), InArchetypeVersion.Number.Major, InArchetypeVersion.Number.Minor });
+
+		}
+
+		FArchetypeRegistryKey GetArchetypeRegistryKey(const FMetasoundFrontendArchetype& InArchetype)
+		{
+			return GetArchetypeRegistryKey(InArchetype.Version);
+		}
+
+		FArchetypeRegistryTransaction::FArchetypeRegistryTransaction(ETransactionType InType, const FArchetypeRegistryKey& InKey, const FMetasoundFrontendVersion& InArchetypeVersion)
+		: Type(InType)
+		, Key(InKey)
+		, ArchetypeVersion(InArchetypeVersion)
+		{
+		}
+
+		FArchetypeRegistryTransaction::ETransactionType FArchetypeRegistryTransaction::GetTransactionType() const
+		{
+			return Type;
+		}
+
+		const FMetasoundFrontendVersion& FArchetypeRegistryTransaction::GetArchetypeVersion() const
+		{
+			return ArchetypeVersion;
+		}
+
+		const FArchetypeRegistryKey& FArchetypeRegistryTransaction::GetArchetypeRegistryKey() const
+		{
+			return Key;
+		}
+
+		IArchetypeRegistry& IArchetypeRegistry::Get()
+		{
+			static MetasoundFrontendArchetypeRegistryPrivate::FArchetypeRegistry Registry;
+			return Registry;
+		}
+	}
+}
diff --git a/Engine/Plugins/Runtime/Metasound/Source/MetasoundFrontend/Private/MetasoundFrontendDocument.cpp b/Engine/Plugins/Runtime/Metasound/Source/MetasoundFrontend/Private/MetasoundFrontendDocument.cpp
--- a/Engine/Plugins/Runtime/Metasound/Source/MetasoundFrontend/Private/MetasoundFrontendDocument.cpp
+++ b/Engine/Plugins/Runtime/Metasound/Source/MetasoundFrontend/Private/MetasoundFrontendDocument.cpp
@@ -41,6 +41,22 @@ FMetasoundFrontendNode::FMetasoundFrontendNode(const FMetasoundFrontendClass& In
 
 }
 
+FString FMetasoundFrontendVersion::ToString() const
+{
+	return FString::Format(TEXT("{0} {1}"), { Name.ToString(), Number.ToString() });
+}
+
+bool FMetasoundFrontendVersion::IsValid() const
+{
+	return Number.IsValid() && Name.IsValid();
+}
+
+const FMetasoundFrontendVersion& FMetasoundFrontendVersion::GetInvalid()
+{
+	static const FMetasoundFrontendVersion InvalidVersion { FName(), FMetasoundFrontendVersionNumber::GetInvalid() };
+	return InvalidVersion;
+}
+
 bool FMetasoundFrontendVertex::IsFunctionalEquivalent(const FMetasoundFrontendVertex& InLHS, const FMetasoundFrontendVertex& InRHS)
 {
 	return (InLHS.Name == InRHS.Name) && (InLHS.TypeName == InRHS.TypeName);
@@ -121,5 +137,6 @@ FMetasoundFrontendDocument::FMetasoundFrontendDocument()
 {
 	RootGraph.ID = FGuid::NewGuid();
 	RootGraph.Metadata.Type = EMetasoundFrontendClassType::Graph;
+	ArchetypeVersion = FMetasoundFrontendVersion::GetInvalid();
 }
 
diff --git a/Engine/Plugins/Runtime/Metasound/Source/MetasoundFrontend/Private/MetasoundFrontendInvalidController.h b/Engine/Plugins/Runtime/Metasound/Source/MetasoundFrontend/Private/MetasoundFrontendInvalidController.h
--- a/Engine/Plugins/Runtime/Metasound/Source/MetasoundFrontend/Private/MetasoundFrontendInvalidController.h
+++ b/Engine/Plugins/Runtime/Metasound/Source/MetasoundFrontend/Private/MetasoundFrontendInvalidController.h
@@ -169,7 +169,7 @@ namespace Metasound
 
 			const FText& GetDescription() const override { return MetasoundFrontendInvalidControllerPrivate::GetInvalid<FText>(); }
 
-			bool IsRequired(const FMetasoundFrontendArchetype& InArchetype) const override { return false; }
+			bool IsRequired() const override { return false; }
 
 			TSharedRef<IGraphController> AsGraph() override;
 			TSharedRef<const IGraphController> AsGraph() const override;
@@ -334,6 +334,10 @@ namespace Metasound
 				TArray<FMetasoundFrontendClass> GetDependencies() const override { return TArray<FMetasoundFrontendClass>(); }
 				TArray<FMetasoundFrontendGraphClass> GetSubgraphs() const override { return TArray<FMetasoundFrontendGraphClass>(); }
 				TArray<FMetasoundFrontendClass> GetClasses() const override { return TArray<FMetasoundFrontendClass>(); }
+				const FMetasoundFrontendGraphClass& GetRootGraphClass() const override
+				{
+					return MetasoundFrontendInvalidControllerPrivate::GetInvalid<FMetasoundFrontendGraphClass>();
+				}
 
 				FConstClassAccessPtr FindDependencyWithID(FGuid InClassID) const override { return FConstClassAccessPtr(); }
 				FConstGraphClassAccessPtr FindSubgraphWithID(FGuid InClassID) const override { return FConstGraphClassAccessPtr(); }
@@ -345,6 +349,9 @@ namespace Metasound
 				FConstClassAccessPtr FindOrAddClass(const FMetasoundFrontendClassMetadata& InMetadata) override{ return FConstClassAccessPtr(); }
 				FGraphHandle AddDuplicateSubgraph(const IGraphController& InGraph) override { return FInvalidGraphController::GetInvalid(); }
 
+				virtual const FMetasoundFrontendVersion& GetArchetypeVersion() const override { return FMetasoundFrontendVersion::GetInvalid(); }
+				virtual void SetArchetypeVersion(const FMetasoundFrontendVersion& InVersion) override {}
+
 				void SetMetadata(const FMetasoundFrontendDocumentMetadata& InMetadata) override { }
 				const FMetasoundFrontendDocumentMetadata& GetMetadata() const override { return MetasoundFrontendInvalidControllerPrivate::GetInvalid<FMetasoundFrontendDocumentMetadata>(); }
 
diff --git a/Engine/Plugins/Runtime/Metasound/Source/MetasoundFrontend/Private/MetasoundFrontendQuerySteps.cpp b/Engine/Plugins/Runtime/Metasound/Source/MetasoundFrontend/Private/MetasoundFrontendQuerySteps.cpp
--- a/Engine/Plugins/Runtime/Metasound/Source/MetasoundFrontend/Private/MetasoundFrontendQuerySteps.cpp
+++ b/Engine/Plugins/Runtime/Metasound/Source/MetasoundFrontend/Private/MetasoundFrontendQuerySteps.cpp
@@ -18,15 +18,14 @@ namespace Metasound
 	{
 		using namespace Frontend;
 
-		const TArray<const IRegistryTransaction*> Transactions = GetRegistryTransactionsSince(CurrentTransactionID, &CurrentTransactionID);
-
-		for (const IRegistryTransaction* Transaction : Transactions)
+		auto AddEntry = [&OutEntries](const FNodeRegistryTransaction& InTransaction)
 		{
-			if (const FNodeClassInfo* Info = Transaction->GetNodeClassInfo())
-			{
-				FFrontendQueryEntry::FValue Value(TInPlaceType<const IRegistryTransaction*>(), Transaction);
-				OutEntries.Emplace(MoveTemp(Value));
-			}
+			OutEntries.Emplace(FFrontendQueryEntry::FValue(TInPlaceType<FNodeRegistryTransaction>(), InTransaction));
+		};
+		
+		if (FMetasoundFrontendRegistryContainer* Registry = FMetasoundFrontendRegistryContainer::Get())
+		{
+			Registry->ForEachNodeRegistryTransactionSince(CurrentTransactionID, &CurrentTransactionID, AddEntry);
 		}
 	}
 
@@ -41,16 +40,9 @@ namespace Metasound
 
 		FNodeRegistryKey RegistryKey;
 
-		if (ensure(InEntry.Value.IsType<const IRegistryTransaction*>()))
+		if (ensure(InEntry.Value.IsType<FNodeRegistryTransaction>()))
 		{
-			if (const IRegistryTransaction* Transaction = InEntry.Value.Get<const IRegistryTransaction*>())
-			{
-				const FNodeRegistryKey* Key = Transaction->GetNodeRegistryKey();
-				if (nullptr != Key)
-				{
-					RegistryKey = *Key;
-				}
-			}
+			RegistryKey = InEntry.Value.Get<FNodeRegistryTransaction>().GetNodeRegistryKey();
 		}
 
 		return FFrontendQueryEntry::FKey(RegistryKey);
@@ -65,21 +57,23 @@ namespace Metasound
 
 		for (FFrontendQueryEntry* Entry : InEntries)
 		{
-			if (ensure(Entry->Value.IsType<const IRegistryTransaction*>()))
+			if (ensure(Entry->Value.IsType<FNodeRegistryTransaction>()))
 			{
-				if (const IRegistryTransaction* Transaction = Entry->Value.Get<const IRegistryTransaction*>())
+				const FNodeRegistryTransaction& Transaction = Entry->Value.Get<FNodeRegistryTransaction>();
+				
+				switch (Transaction.GetTransactionType())
 				{
-					switch (Transaction->GetTransactionType())
-					{
-						case ETransactionType::Add:
-							State++;
-							FinalEntry = Entry;
-							break;
-
-						case ETransactionType::Remove:
-							State--;
-							break;
-					}
+					case FNodeRegistryTransaction::ETransactionType::NodeRegistration:
+						State++;
+						FinalEntry = Entry;
+						break;
+
+					case FNodeRegistryTransaction::ETransactionType::NodeUnregistration:
+						State--;
+						break;
+
+					default:
+						break;
 				}
 			}
 		}
@@ -96,17 +90,12 @@ namespace Metasound
 
 		FMetasoundFrontendClass FrontendClass;
 
-		if (ensure(InValue.IsType<const IRegistryTransaction*>()))
+		if (ensure(InValue.IsType<FNodeRegistryTransaction>()))
 		{
-			if (const IRegistryTransaction* Transaction = InValue.Get<const IRegistryTransaction*>())
-			{
-				const FNodeRegistryKey* Key = Transaction->GetNodeRegistryKey();
-				if (nullptr != Key)
-				{
-					bool bSuccess = FMetasoundFrontendRegistryContainer::Get()->FindFrontendClassFromRegistered(*Key, FrontendClass);
-					check(bSuccess);
-				}
-			}
+			const FNodeRegistryTransaction& Transaction = InValue.Get<FNodeRegistryTransaction>();
+			
+			bool bSuccess = FMetasoundFrontendRegistryContainer::Get()->FindFrontendClassFromRegistered(Transaction.GetNodeRegistryKey(), FrontendClass);
+			check(bSuccess);
 		}
 
 		InValue.Set<FMetasoundFrontendClass>(MoveTemp(FrontendClass));
diff --git a/Engine/Plugins/Runtime/Metasound/Source/MetasoundFrontend/Private/MetasoundFrontendRegistries.cpp b/Engine/Plugins/Runtime/Metasound/Source/MetasoundFrontend/Private/MetasoundFrontendRegistries.cpp
--- a/Engine/Plugins/Runtime/Metasound/Source/MetasoundFrontend/Private/MetasoundFrontendRegistries.cpp
+++ b/Engine/Plugins/Runtime/Metasound/Source/MetasoundFrontend/Private/MetasoundFrontendRegistries.cpp
@@ -347,6 +347,7 @@ namespace Metasound
 				FMetasoundFrontendClass FrontendClass;
 			};
 
+
 			// Registry container private implementation.
 			class FRegistryContainerImpl : public FMetasoundFrontendRegistryContainer
 			{
@@ -388,7 +389,7 @@ namespace Metasound
 
 				bool RegisterConversionNode(const FConverterNodeRegistryKey& InNodeKey, const FConverterNodeInfo& InNodeInfo) override;
 
-				TArray<const Metasound::Frontend::IRegistryTransaction*> GetRegistryTransactionsSince(Metasound::Frontend::FRegistryTransactionID InSince, Metasound::Frontend::FRegistryTransactionID* OutCurrentRegistryTransactionID) const override;
+				virtual void ForEachNodeRegistryTransactionSince(FRegistryTransactionID InSince, FRegistryTransactionID* OutCurrentRegistryTransactionID, TFunctionRef<void(const FNodeRegistryTransaction&)> InFunc) const override;
 
 				// Return any data types that can be used as a metasound input type or output type.
 				TArray<FName> GetAllValidDataTypes() override;
@@ -456,7 +457,7 @@ namespace Metasound
 				// Registry in which we keep lists of possible nodes to use to convert between two datatypes
 				TMap<FConverterNodeRegistryKey, FConverterNodeRegistryValue> ConverterNodeRegistry;
 
-				FRegistryTransactionHistory RegistryTransactionHistory;
+				TRegistryTransactionHistory<FNodeRegistryTransaction> RegistryTransactionHistory;
 			};
 
 			void FRegistryContainerImpl::RegisterPendingNodes()
@@ -489,9 +490,9 @@ namespace Metasound
 				return true;
 			}
 
-			TArray<const Metasound::Frontend::IRegistryTransaction*> FRegistryContainerImpl::GetRegistryTransactionsSince(Metasound::Frontend::FRegistryTransactionID InSince, Metasound::Frontend::FRegistryTransactionID* OutCurrentRegistryTransactionID) const 
+			void FRegistryContainerImpl::ForEachNodeRegistryTransactionSince(FRegistryTransactionID InSince, FRegistryTransactionID* OutCurrentRegistryTransactionID, TFunctionRef<void(const FNodeRegistryTransaction&)> InFunc) const 
 			{
-				return RegistryTransactionHistory.GetTransactions(InSince, OutCurrentRegistryTransactionID);
+				return RegistryTransactionHistory.ForEachTransactionSince(InSince, OutCurrentRegistryTransactionID, InFunc);
 			}
 
 			TUniquePtr<Metasound::INode> FRegistryContainerImpl::CreateInputNode(const FName& InDataType, Metasound::FInputNodeConstructorParams&& InParams)
@@ -786,7 +787,7 @@ namespace Metasound
 
 					// Store update to newly registered node in history so nodes
 					// can be queried by transaction ID
-					RegistryTransactionHistory.Add(MakeAddNodeRegistryTransaction(Key, InEntry->GetClassInfo()));
+					RegistryTransactionHistory.Add(FNodeRegistryTransaction(FNodeRegistryTransaction::ETransactionType::NodeRegistration, Key, InEntry->GetClassInfo()));
 
 					// Store registry elements in map so nodes can be queried using registry key.
 					RegisteredNodes.Add(Key, MoveTemp(InEntry));
@@ -801,7 +802,7 @@ namespace Metasound
 				{
 					if (const INodeRegistryEntry* Entry = FindNodeEntry(InKey))
 					{
-						RegistryTransactionHistory.Add(MakeRemoveNodeRegistryTransaction(InKey, Entry->GetClassInfo()));
+						RegistryTransactionHistory.Add(FNodeRegistryTransaction(FNodeRegistryTransaction::ETransactionType::NodeUnregistration, InKey, Entry->GetClassInfo()));
 
 						RegisteredNodes.Remove(InKey);
 						return true;
@@ -979,6 +980,28 @@ namespace Metasound
 			}
 		} // namespace MetasoundFrontendRegistriesPrivate
 
+		FNodeRegistryTransaction::FNodeRegistryTransaction(ETransactionType InType, const FNodeRegistryKey& InKey, const FNodeClassInfo& InNodeClassInfo)
+		: Type(InType)
+		, Key(InKey)
+		, NodeClassInfo(InNodeClassInfo)
+		{
+		}
+
+		FNodeRegistryTransaction::ETransactionType FNodeRegistryTransaction::GetTransactionType() const
+		{
+			return Type;
+		}
+
+		const FNodeClassInfo& FNodeRegistryTransaction::GetNodeClassInfo() const
+		{
+			return NodeClassInfo;
+		}
+
+		const FNodeRegistryKey& FNodeRegistryTransaction::GetNodeRegistryKey() const
+		{
+			return Key;
+		}
+
 		bool IsValidNodeRegistryKey(const FNodeRegistryKey& InKey)
 		{
 			return !InKey.IsEmpty();
diff --git a/Engine/Plugins/Runtime/Metasound/Source/MetasoundFrontend/Private/MetasoundFrontendRegistryTransaction.cpp b/Engine/Plugins/Runtime/Metasound/Source/MetasoundFrontend/Private/MetasoundFrontendRegistryTransaction.cpp
--- a/Engine/Plugins/Runtime/Metasound/Source/MetasoundFrontend/Private/MetasoundFrontendRegistryTransaction.cpp
+++ b/Engine/Plugins/Runtime/Metasound/Source/MetasoundFrontend/Private/MetasoundFrontendRegistryTransaction.cpp
@@ -13,165 +13,6 @@ namespace Metasound
 		{
 			return 0;
 		}
-
-		FRegistryTransactionHistory::FRegistryTransactionHistory()
-		: Current(GetOriginRegistryTransactionID())
-		{
-		}
-
-		FRegistryTransactionID FRegistryTransactionHistory::Add(TUniquePtr<IRegistryTransaction>&& InRegistryTransaction)
-		{
-			FScopeLock Lock(&RegistryTransactionMutex);
-
-			if (ensure(InRegistryTransaction.IsValid()))
-			{
-				Current++;
-
-				IRegistryTransaction* RegistryTransactionPointer = InRegistryTransaction.Get();
-				int32 Index = RegistryTransactions.Num();
-
-				RegistryTransactions.Add(MoveTemp(InRegistryTransaction));
-				RegistryTransactionIndexMap.Add(Current, Index);
-			}
-
-			return Current;
-		}
-
-		FRegistryTransactionID FRegistryTransactionHistory::Add(const IRegistryTransaction& InRegistryTransaction)
-		{
-			return Add(InRegistryTransaction.Clone());
-		}
-
-		FRegistryTransactionID FRegistryTransactionHistory::GetCurrent() const
-		{
-			FScopeLock Lock(&RegistryTransactionMutex);
-			{
-				return Current;
-			}
-		}
-
-		TArray<const IRegistryTransaction*> FRegistryTransactionHistory::GetTransactions(FRegistryTransactionID InSince, FRegistryTransactionID* OutCurrent) const
-		{
-			FScopeLock Lock(&RegistryTransactionMutex);
-			{
-				TArray<const IRegistryTransaction*> Result;
-
-				if (nullptr != OutCurrent)
-				{
-					*OutCurrent = Current;
-				}
-
-				int32 Start = INDEX_NONE;
-				
-				if (GetOriginRegistryTransactionID() == InSince)
-				{
-					Start = 0;
-				}
-				else if (const int32* Pos = RegistryTransactionIndexMap.Find(InSince))
-				{
-					Start = *Pos + 1;
-				}
-				
-				if (INDEX_NONE != Start)
-				{
-					const int32 Num = RegistryTransactions.Num();
-
-					if (ensure(Start <= Num))
-					{
-						const int32 OutNum = Num - Start;
-						if (OutNum > 0)
-						{
-							Result.Reserve(OutNum);
-							for (int32 i = Start; i < Num; i++)
-							{
-								Result.Add(RegistryTransactions[i].Get());
-							}
-						}
-					}
-				}
-
-				return Result;
-			}
-		}
-
-		FRegistryTransactionPtr MakeAddNodeRegistryTransaction(const FNodeRegistryKey& InKey, const FNodeClassInfo& InInfo)
-		{
-			class FNodeRegistryTransaction : public IRegistryTransaction
-			{
-			public:
-				FNodeRegistryTransaction(const FNodeRegistryKey& InKey, const FNodeClassInfo& InNodeClassInfo)
-				: NodeClassInfo(InNodeClassInfo)
-				, Key(InKey)
-				{
-				}
-
-				virtual ETransactionType GetTransactionType() const override
-				{
-					return ETransactionType::Add;
-				}
-
-				virtual TUniquePtr<IRegistryTransaction> Clone() const override
-				{
-					return MakeUnique<FNodeRegistryTransaction>(*this);
-				}
-
-				virtual const FNodeClassInfo* GetNodeClassInfo() const override
-				{
-					return &NodeClassInfo;
-				}
-
-				virtual const FNodeRegistryKey* GetNodeRegistryKey() const override
-				{
-					return &Key;
-				}
-
-			private:
-
-				FNodeClassInfo NodeClassInfo;
-				FNodeRegistryKey Key;
-			};
-
-			return MakeUnique<FNodeRegistryTransaction>(InKey, InInfo);
-		}
-
-		FRegistryTransactionPtr MakeRemoveNodeRegistryTransaction(const FNodeRegistryKey& InKey, const FNodeClassInfo& InInfo)
-		{
-			class FNodeRegistryTransaction : public IRegistryTransaction
-			{
-			public:
-				FNodeRegistryTransaction(const FNodeRegistryKey& InKey, const FNodeClassInfo& InNodeClassInfo)
-				: NodeClassInfo(InNodeClassInfo)
-				, Key(InKey)
-				{
-				}
-
-				virtual ETransactionType GetTransactionType() const override
-				{
-					return ETransactionType::Remove;
-				}
-
-				virtual TUniquePtr<IRegistryTransaction> Clone() const override
-				{
-					return MakeUnique<FNodeRegistryTransaction>(*this);
-				}
-
-				virtual const FNodeClassInfo* GetNodeClassInfo() const override
-				{
-					return &NodeClassInfo;
-				}
-
-				virtual const FNodeRegistryKey* GetNodeRegistryKey() const override
-				{
-					return &Key;
-				}
-
-			private:
-
-				FNodeClassInfo NodeClassInfo;
-				FNodeRegistryKey Key;
-			};
-
-			return MakeUnique<FNodeRegistryTransaction>(InKey, InInfo);
-		}
 	}
 }
+
diff --git a/Engine/Plugins/Runtime/Metasound/Source/MetasoundFrontend/Private/MetasoundFrontendRegistryTransaction.h b/Engine/Plugins/Runtime/Metasound/Source/MetasoundFrontend/Private/MetasoundFrontendRegistryTransaction.h
--- a/Engine/Plugins/Runtime/Metasound/Source/MetasoundFrontend/Private/MetasoundFrontendRegistryTransaction.h
+++ b/Engine/Plugins/Runtime/Metasound/Source/MetasoundFrontend/Private/MetasoundFrontendRegistryTransaction.h
@@ -5,6 +5,7 @@
 #include "CoreMinimal.h"
 #include "MetasoundFrontend.h"
 #include "MetasoundFrontendRegistries.h"
+#include "Algo/ForEach.h"
 
 namespace Metasound
 {
@@ -13,40 +14,88 @@ namespace Metasound
 		/** Returns an ID representing the beginning of the transaction history. */
 		FRegistryTransactionID GetOriginRegistryTransactionID();
 
-		/** Create a node registration transaction. */
-		FRegistryTransactionPtr MakeAddNodeRegistryTransaction(const FNodeRegistryKey& InKey, const FNodeClassInfo& InInfo);
 
-		/** Create a node unregistration transaction. */
-		FRegistryTransactionPtr MakeRemoveNodeRegistryTransaction(const FNodeRegistryKey& InKey, const FNodeClassInfo& InInfo);
-
-		/** Maintains a history of IRegistryTransactions. Calls are threadsafe (excluding
+		/** Maintains a history of TransactionTypes. Calls are threadsafe (excluding
 		 * the constructor and destructor.)
 		 */
-		class FRegistryTransactionHistory
+		template<typename TransactionType>
+		class TRegistryTransactionHistory
 		{
 		public:
-
-			FRegistryTransactionHistory();
+			TRegistryTransactionHistory()
+			: Current(GetOriginRegistryTransactionID())
+			{
+			}
 
 			/** Add a transaction to the history. Threadsafe. 
 			 *
 			 * @return The transaction ID associated with the action. */
-			FRegistryTransactionID Add(FRegistryTransactionPtr&& InRegistryTransaction);
+			FRegistryTransactionID Add(TransactionType&& InRegistryTransaction)
+			{
+				FScopeLock Lock(&RegistryTransactionMutex);
 
-			/** Add a transaction to the history. Threadsafe. 
-			 *
-			 * @return The transaction ID associated with the action. */
-			FRegistryTransactionID Add(const IRegistryTransaction& InRegistryTransaction);
+				Current++;
 
-			/** Gets the transaction ID of the most recent transaction. Threadsafe. */
-			FRegistryTransactionID GetCurrent() const;
+				int32 Index = RegistryTransactions.Num();
+
+				RegistryTransactions.Add(MoveTemp(InRegistryTransaction));
+				RegistryTransactionIndexMap.Add(Current, Index);
 
-			/** Returns an array of transaction pointers. Threadsafe. 
+				return Current;
+			}
+
+			/** Gets the transaction ID of the most recent transaction. Threadsafe. */
+			FRegistryTransactionID GetCurrent() const
+			{
+				FScopeLock Lock(&RegistryTransactionMutex);
+				{
+					return Current;
+				}
+			}
+
+			/** Invoke a function on all transactions since transaction ID. 
 			 *
-			 * @param InSince - All transactions occuring after this transaction ID will be returned.
+			 * @param InSince - All transactions occurring after this transaction ID will be returned.
 			 * @param OutCurrent - If not null, will be set to the most recent transaction ID returned. 
+			 * @param InCallable - A callable of the form Callable(const TransactionType&)
 			 */
-			TArray<const IRegistryTransaction*> GetTransactions(FRegistryTransactionID InSince, FRegistryTransactionID* OutCurrent) const;
+			template<typename CallableType>
+			void ForEachTransactionSince(FRegistryTransactionID InSince, FRegistryTransactionID* OutCurrent, CallableType InCallable) const
+			{
+				FScopeLock Lock(&RegistryTransactionMutex);
+				{
+					if (nullptr != OutCurrent)
+					{
+						*OutCurrent = Current;
+					}
+
+					int32 Start = INDEX_NONE;
+					
+					if (GetOriginRegistryTransactionID() == InSince)
+					{
+						Start = 0;
+					}
+					else if (const int32* Pos = RegistryTransactionIndexMap.Find(InSince))
+					{
+						Start = *Pos + 1;
+					}
+					
+					if (INDEX_NONE != Start)
+					{
+						const int32 Num = RegistryTransactions.Num();
+
+						if (ensure(Start <= Num))
+						{
+							const int32 OutNum = Num - Start;
+							if (OutNum > 0)
+							{
+								TArrayView<const TransactionType> TransactionsSince = MakeArrayView(&RegistryTransactions[Start], OutNum);
+								Algo::ForEach(TransactionsSince, InCallable);
+							}
+						}
+					}
+				}
+			}
 
 		private:
 
@@ -55,7 +104,7 @@ namespace Metasound
 			FRegistryTransactionID Current;
 
 			TMap<FRegistryTransactionID, int32> RegistryTransactionIndexMap;
-			TArray<FRegistryTransactionPtr> RegistryTransactions;
+			TArray<TransactionType> RegistryTransactions;
 		};
 	}
 }
diff --git a/Engine/Plugins/Runtime/Metasound/Source/MetasoundFrontend/Private/MetasoundFrontendSearchEngine.cpp b/Engine/Plugins/Runtime/Metasound/Source/MetasoundFrontend/Private/MetasoundFrontendSearchEngine.cpp
--- a/Engine/Plugins/Runtime/Metasound/Source/MetasoundFrontend/Private/MetasoundFrontendSearchEngine.cpp
+++ b/Engine/Plugins/Runtime/Metasound/Source/MetasoundFrontend/Private/MetasoundFrontendSearchEngine.cpp
@@ -4,8 +4,10 @@
 
 #include "Algo/Transform.h"
 #include "HAL/IConsoleManager.h"
+#include "MetasoundFrontendArchetypeRegistry.h"
 #include "MetasoundFrontendQuery.h"
 #include "MetasoundFrontendQuerySteps.h"
+#include "MetasoundFrontendRegistryTransaction.h"
 
 static int32 ClearMetaSoundFrontendSearchEngineCacheCVar = 0;
 FAutoConsoleVariableRef CVarClearMetaSoundFrontendSearchEngineCache(
@@ -20,6 +22,241 @@ namespace Metasound
 {
 	namespace Frontend
 	{
+		namespace SearchEngineQuerySteps
+		{
+			class FArchetypeRegistryTransactionSource : public IFrontendQuerySource
+			{
+			public:
+				FArchetypeRegistryTransactionSource()
+				: CurrentTransactionID(GetOriginRegistryTransactionID())
+				{
+				}
+
+				void Stream(TArray<FFrontendQueryEntry>& OutEntries) override
+				{
+					auto AddEntry = [&OutEntries](const FArchetypeRegistryTransaction& InTransaction)
+					{
+						OutEntries.Emplace(FFrontendQueryEntry::FValue(TInPlaceType<FArchetypeRegistryTransaction>(), InTransaction));
+					};
+					
+					IArchetypeRegistry::Get().ForEachRegistryTransactionSince(CurrentTransactionID, &CurrentTransactionID, AddEntry);
+				}
+
+				void Reset() override
+				{
+					CurrentTransactionID = GetOriginRegistryTransactionID();
+				}
+
+			private:
+				FRegistryTransactionID CurrentTransactionID;
+			};
+
+			class FMapArchetypeRegistryTransactionsToArchetypeRegistryKeys : public IFrontendQueryMapStep
+			{
+			public:
+				FFrontendQueryEntry::FKey Map(const FFrontendQueryEntry& InEntry) const override
+				{
+					FArchetypeRegistryKey RegistryKey;
+					if (ensure(InEntry.Value.IsType<FArchetypeRegistryTransaction>()))
+					{
+						RegistryKey = InEntry.Value.Get<FArchetypeRegistryTransaction>().GetArchetypeRegistryKey();
+					}
+					return FFrontendQueryEntry::FKey{RegistryKey};
+				}
+			};
+
+			class FReduceArchetypeRegistryTransactionsToCurrentStatus : public IFrontendQueryReduceStep
+			{
+			public:
+				void Reduce(FFrontendQueryEntry::FKey InKey, TArrayView<FFrontendQueryEntry * const>& InEntries, FReduceOutputView& OutResult) const override
+				{
+					int32 State = 0;
+					FFrontendQueryEntry* FinalEntry = nullptr;
+
+					for (FFrontendQueryEntry* Entry : InEntries)
+					{
+						if (ensure(Entry->Value.IsType<FArchetypeRegistryTransaction>()))
+						{
+							const FArchetypeRegistryTransaction& Transaction = Entry->Value.Get<FArchetypeRegistryTransaction>();
+							
+							switch (Transaction.GetTransactionType())
+							{
+								case FArchetypeRegistryTransaction::ETransactionType::ArchetypeRegistration:
+									State++;
+									FinalEntry = Entry;
+									break;
+
+								case FArchetypeRegistryTransaction::ETransactionType::ArchetypeUnregistration:
+									State--;
+									break;
+
+								default:
+									break;
+							}
+						}
+					}
+
+					if ((nullptr != FinalEntry) && (State > 0))
+					{
+						OutResult.Add(*FinalEntry);
+					}
+				}
+			};
+
+			class FTransformArchetypeRegistryTransactionToArchetype : public IFrontendQueryTransformStep
+			{
+			public:
+				virtual void Transform(FFrontendQueryEntry::FValue& InValue) const override
+				{
+					FMetasoundFrontendArchetype Archetype;
+					if (ensure(InValue.IsType<FArchetypeRegistryTransaction>()))
+					{
+						FArchetypeRegistryKey RegistryKey = InValue.Get<FArchetypeRegistryTransaction>().GetArchetypeRegistryKey();
+						IArchetypeRegistry::Get().FindArchetype(RegistryKey, Archetype);
+					}
+					InValue.Set<FMetasoundFrontendArchetype>(MoveTemp(Archetype));
+				}
+			};
+
+			class FMapArchetypeToArchetypeNameAndMajorVersion : public IFrontendQueryMapStep
+			{
+			public:
+				FFrontendQueryEntry::FKey Map(const FFrontendQueryEntry& InEntry) const override
+				{
+					FString Key;
+					if (ensure(InEntry.Value.IsType<FMetasoundFrontendArchetype>()))
+					{
+						const FMetasoundFrontendArchetype& Arch = InEntry.Value.Get<FMetasoundFrontendArchetype>();
+						Key = FString::Format(TEXT("{0}_{1}"), { Arch.Version.Name.ToString(), Arch.Version.Number.Major });
+					}
+					return FFrontendQueryEntry::FKey{Key};
+				}
+			};
+
+			class FReduceArchetypesToHighestVersion : public IFrontendQueryReduceStep
+			{
+			public:
+				void Reduce(FFrontendQueryEntry::FKey InKey, TArrayView<FFrontendQueryEntry * const>& InEntries, FReduceOutputView& OutResult) const override
+				{
+					FFrontendQueryEntry* HighestVersionEntry = nullptr;
+					FMetasoundFrontendVersionNumber HighestVersion = FMetasoundFrontendVersionNumber::GetInvalid();
+
+					for (FFrontendQueryEntry* Entry : InEntries)
+					{
+						if (ensure(Entry->Value.IsType<FMetasoundFrontendArchetype>()))
+						{
+							const FMetasoundFrontendVersionNumber& VersionNumber = Entry->Value.Get<FMetasoundFrontendArchetype>().Version.Number;
+
+							if (VersionNumber > HighestVersion)
+							{
+								HighestVersionEntry = Entry;
+								HighestVersion = VersionNumber;
+							}
+						}
+					}
+
+					if (HighestVersionEntry)
+					{
+						OutResult.Add(*HighestVersionEntry);
+					}
+				}
+			};
+
+			class FFilterArchetypeRegistryTransactionsByName : public IFrontendQueryFilterStep
+			{
+			public:
+				FFilterArchetypeRegistryTransactionsByName(const FName& InName)
+				: Name(InName)
+				{
+				}
+
+				bool Filter(const FFrontendQueryEntry& InEntry) const override
+				{
+					if (ensure(InEntry.Value.IsType<FArchetypeRegistryTransaction>()))
+					{
+						return Name == InEntry.Value.Get<FArchetypeRegistryTransaction>().GetArchetypeVersion().Name;
+					}
+					return false;
+				}
+
+			private:
+				FName Name;
+			};
+
+			class FFilterArchetypeRegistryTransactionsByNameAndMajorVersion : public IFrontendQueryFilterStep
+			{
+			public:
+				FFilterArchetypeRegistryTransactionsByNameAndMajorVersion(const FName& InName, int32 InMajorVersion)
+				: Name(InName)
+				, MajorVersion(InMajorVersion)
+				{
+				}
+
+				bool Filter(const FFrontendQueryEntry& InEntry) const override
+				{
+					if (ensure(InEntry.Value.IsType<FArchetypeRegistryTransaction>()))
+					{
+						const FMetasoundFrontendVersion& Version = InEntry.Value.Get<FArchetypeRegistryTransaction>().GetArchetypeVersion();
+
+						return (Name == Version.Name) && (MajorVersion == Version.Number.Major);
+					}
+					return false;
+				}
+
+			private:
+				FName Name;
+				int32 MajorVersion;
+			};
+
+			class FMapArchetypeToArchetypeName : public IFrontendQueryMapStep
+			{
+			public:
+				FFrontendQueryEntry::FKey Map(const FFrontendQueryEntry& InEntry) const override
+				{
+					FString Key;
+					if (ensure(InEntry.Value.IsType<FMetasoundFrontendArchetype>()))
+					{
+						Key = InEntry.Value.Get<FMetasoundFrontendArchetype>().Version.Name.ToString();
+					}
+					return FFrontendQueryEntry::FKey{Key};
+				}
+			};
+
+			class FFilterRegistryTransactionsByName : public IFrontendQueryFilterStep
+			{
+			public:
+				FFilterRegistryTransactionsByName(const FName& InName)
+				: Name(InName)
+				{
+				}
+
+				bool Filter(const FFrontendQueryEntry& InEntry) const override
+				{
+					if (ensure(InEntry.Value.IsType<FArchetypeRegistryTransaction>()))
+					{
+						return Name == InEntry.Value.Get<FArchetypeRegistryTransaction>().GetArchetypeVersion().Name;
+					}
+					return false;
+				}
+			private:
+				FName Name;
+			};
+
+			class FTransformArchetypeRegistryTransactionToArchetypeVersion : public IFrontendQueryTransformStep
+			{
+			public:
+				virtual void Transform(FFrontendQueryEntry::FValue& InValue) const override
+				{
+					FMetasoundFrontendVersion Version;
+					if (ensure(InValue.IsType<FArchetypeRegistryTransaction>()))
+					{
+						Version = InValue.Get<FArchetypeRegistryTransaction>().GetArchetypeVersion();
+					}
+					InValue.Set<FMetasoundFrontendVersion>(Version);
+				}
+			};
+		}
+
 		class FSearchEngine : public ISearchEngine
 		{
 				FSearchEngine(const FSearchEngine&) = delete;
@@ -37,6 +274,14 @@ namespace Metasound
 				bool FindClassWithHighestVersion(const FNodeClassName& InName, FMetasoundFrontendClass& OutClass) override;
 				bool FindClassWithMajorVersion(const FNodeClassName& InName, int32 InMajorVersion, FMetasoundFrontendClass& OutClass) override;
 
+				virtual TArray<FMetasoundFrontendArchetype> FindAllArchetypes(bool bInIncludeDeprecated) override;
+
+				virtual TArray<FMetasoundFrontendVersion> FindAllRegisteredArchetypesWithName(const FName& InArchetypeName) override;
+
+				virtual bool FindArchetypeWithHighestVersion(const FName& InArchetypeName, FMetasoundFrontendArchetype& OutArchetype) override;
+
+				virtual bool FindArchetypeWithMajorVersion(const FName& InArchetypeName, int32 InMajorVersion, FMetasoundFrontendArchetype& OutArchetype) override;
+
 			private:
 				FFrontendQuery* FindQuery(const FString& InQueryName);
 				FFrontendQuery* AddQuery(const FString& InQueryName, TUniquePtr<FFrontendQuery>&& InQuery);
@@ -184,6 +429,7 @@ namespace Metasound
 		{
 			const FString QueryName = FString(TEXT("FindClassWithMajorVersion")) + InName.GetFullName().ToString() + FString::FromInt(InMajorVersion);
 			FFrontendQuery* Query = FindQuery(QueryName);
+
 			if (!Query)
 			{
 				TUniquePtr<FFrontendQuery> NewQuery = MakeUnique<FFrontendQuery>();
@@ -220,6 +466,169 @@ namespace Metasound
 			return false;
 		}
 
+		TArray<FMetasoundFrontendArchetype> FSearchEngine::FindAllArchetypes(bool bInIncludeDeprecated)
+		{
+			using namespace SearchEngineQuerySteps;
+
+			const FString QueryName = FString::Format(TEXT("FindAllArchetypes_IncludeDeprecated:{0}"), { bInIncludeDeprecated });
+
+			FFrontendQuery* Query = FindQuery(QueryName);
+			if (!Query)
+			{
+				TUniquePtr<FFrontendQuery> NewQuery = MakeUnique<FFrontendQuery>();
+
+				NewQuery->AddStep<FArchetypeRegistryTransactionSource>()
+					.AddStep<FMapArchetypeRegistryTransactionsToArchetypeRegistryKeys>()
+					.AddStep<FReduceArchetypeRegistryTransactionsToCurrentStatus>()
+					.AddStep<FTransformArchetypeRegistryTransactionToArchetype>();
+
+				if (!bInIncludeDeprecated)
+				{
+					NewQuery->AddStep<FMapArchetypeToArchetypeNameAndMajorVersion>()
+					.AddStep<FReduceArchetypesToHighestVersion>();
+				}
+
+				Query = AddQuery(QueryName, MoveTemp(NewQuery));
+			}
+
+			TArray<FMetasoundFrontendArchetype> Result;
+
+			if (ensure(Query))
+			{
+				FFrontendQuerySelectionView Selection = Query->Execute();
+				TArrayView<const FFrontendQueryEntry* const> Entries = Selection.GetSelection();
+
+				for (const FFrontendQueryEntry* Entry : Selection.GetSelection())
+				{
+					if (ensure(Entry->Value.IsType<FMetasoundFrontendArchetype>()))
+					{
+						Result.Add(Entry->Value.Get<FMetasoundFrontendArchetype>());
+					}
+				}
+			}
+
+			return Result;
+		}
+
+		TArray<FMetasoundFrontendVersion> FSearchEngine::FindAllRegisteredArchetypesWithName(const FName& InArchetypeName)
+		{
+			using namespace SearchEngineQuerySteps;
+
+			const FString QueryName = FString::Format(TEXT("FindAllRegisteredArchetypesWithName_Name:{0}"), { InArchetypeName.ToString() });
+
+			FFrontendQuery* Query = FindQuery(QueryName);
+			if (!Query)
+			{
+				TUniquePtr<FFrontendQuery> NewQuery = MakeUnique<FFrontendQuery>();
+
+				NewQuery->AddStep<FArchetypeRegistryTransactionSource>()
+					.AddStep<FFilterRegistryTransactionsByName>(InArchetypeName)
+					.AddStep<FMapArchetypeRegistryTransactionsToArchetypeRegistryKeys>()
+					.AddStep<FReduceArchetypeRegistryTransactionsToCurrentStatus>()
+					.AddStep<FTransformArchetypeRegistryTransactionToArchetypeVersion>();
+
+				Query = AddQuery(QueryName, MoveTemp(NewQuery));
+			}
+
+			TArray<FMetasoundFrontendVersion> Result;
+
+			if (ensure(Query))
+			{
+				FFrontendQuerySelectionView Selection = Query->Execute();
+				TArrayView<const FFrontendQueryEntry* const> Entries = Selection.GetSelection();
+
+				for (const FFrontendQueryEntry* Entry : Selection.GetSelection())
+				{
+					if (ensure(Entry->Value.IsType<FMetasoundFrontendVersion>()))
+					{
+						Result.Add(Entry->Value.Get<FMetasoundFrontendVersion>());
+					}
+				}
+			}
+
+			return Result;
+		}
+
+		bool FSearchEngine::FindArchetypeWithHighestVersion(const FName& InArchetypeName, FMetasoundFrontendArchetype& OutArchetype)
+		{
+			using namespace SearchEngineQuerySteps;
+			const FString QueryName = FString::Format(TEXT("FindArchetypeWithHighestVersion_Name:{0}"), { InArchetypeName.ToString() });
+
+			FFrontendQuery* Query = FindQuery(QueryName);
+			if (!Query)
+			{
+				TUniquePtr<FFrontendQuery> NewQuery = MakeUnique<FFrontendQuery>();
+
+				NewQuery->AddStep<FArchetypeRegistryTransactionSource>()
+					.AddStep<FFilterArchetypeRegistryTransactionsByName>(InArchetypeName)
+					.AddStep<FMapArchetypeRegistryTransactionsToArchetypeRegistryKeys>()
+					.AddStep<FReduceArchetypeRegistryTransactionsToCurrentStatus>()
+					.AddStep<FTransformArchetypeRegistryTransactionToArchetype>()
+					.AddStep<FMapArchetypeToArchetypeName>()
+					.AddStep<FReduceArchetypesToHighestVersion>();
+
+				Query = AddQuery(QueryName, MoveTemp(NewQuery));
+			}
+
+			if (ensure(Query))
+			{
+				FFrontendQuerySelectionView Selection = Query->Execute();
+				TArrayView<const FFrontendQueryEntry* const> Entries = Selection.GetSelection();
+				if (Entries.IsEmpty())
+				{
+					return false;
+				}
+
+				if (ensure(Entries.Num() == 1))
+				{
+					OutArchetype = Entries[0]->Value.Get<FMetasoundFrontendArchetype>();
+					return true;
+				}
+			}
+
+			return false;
+		}
+
+		bool FSearchEngine::FindArchetypeWithMajorVersion(const FName& InArchetypeName, int32 InMajorVersion, FMetasoundFrontendArchetype& OutArchetype)
+		{
+			using namespace SearchEngineQuerySteps;
+			const FString QueryName = FString::Format(TEXT("FindArchetypeWithMajorVersion_Name:{0}_MajorVersion:{1}"), { InArchetypeName.ToString(), InMajorVersion });
+
+			FFrontendQuery* Query = FindQuery(QueryName);
+			if (!Query)
+			{
+				TUniquePtr<FFrontendQuery> NewQuery = MakeUnique<FFrontendQuery>();
+
+				NewQuery->AddStep<FArchetypeRegistryTransactionSource>()
+					.AddStep<FFilterArchetypeRegistryTransactionsByNameAndMajorVersion>(InArchetypeName, InMajorVersion)
+					.AddStep<FMapArchetypeRegistryTransactionsToArchetypeRegistryKeys>()
+					.AddStep<FReduceArchetypeRegistryTransactionsToCurrentStatus>()
+					.AddStep<FTransformArchetypeRegistryTransactionToArchetype>()
+					.AddStep<FMapArchetypeToArchetypeName>()
+					.AddStep<FReduceArchetypesToHighestVersion>();
+
+				Query = AddQuery(QueryName, MoveTemp(NewQuery));
+			}
+
+			if (ensure(Query))
+			{
+				FFrontendQuerySelectionView Selection = Query->Execute();
+				TArrayView<const FFrontendQueryEntry* const> Entries = Selection.GetSelection();
+				if (Entries.IsEmpty())
+				{
+					return false;
+				}
+
+				if (ensure(Entries.Num() == 1))
+				{
+					OutArchetype = Entries[0]->Value.Get<FMetasoundFrontendArchetype>();
+					return true;
+				}
+			}
+
+			return false;
+		}
+
 		FFrontendQuery* FSearchEngine::FindQuery(const FString& InQueryName)
 		{
 			if (ClearMetaSoundFrontendSearchEngineCacheCVar)
diff --git a/Engine/Plugins/Runtime/Metasound/Source/MetasoundFrontend/Private/MetasoundFrontendStandardController.cpp b/Engine/Plugins/Runtime/Metasound/Source/MetasoundFrontend/Private/MetasoundFrontendStandardController.cpp
--- a/Engine/Plugins/Runtime/Metasound/Source/MetasoundFrontend/Private/MetasoundFrontendStandardController.cpp
+++ b/Engine/Plugins/Runtime/Metasound/Source/MetasoundFrontend/Private/MetasoundFrontendStandardController.cpp
@@ -7,6 +7,7 @@
 #include "CoreMinimal.h"
 #include "HAL/FileManager.h"
 #include "MetasoundAccessPtr.h"
+#include "MetasoundFrontendArchetypeRegistry.h"
 #include "MetasoundFrontendController.h"
 #include "MetasoundFrontendDocument.h"
 #include "MetasoundFrontendDocumentAccessPtr.h"
@@ -1397,7 +1398,7 @@ namespace Metasound
 			return FInvalidInputController::GetInvalid();
 		}
 
-		bool FBaseNodeController::IsRequired(const FMetasoundFrontendArchetype& InArchetype) const
+		bool FBaseNodeController::IsRequired() const
 		{
 			return false;
 		}
@@ -1771,19 +1772,39 @@ namespace Metasound
 			return OutputDisplayTitle;
 		}
 
-		bool FOutputNodeController::IsRequired(const FMetasoundFrontendArchetype& InArchetype) const
+		bool FOutputNodeController::IsRequired() const
 		{
-			if (NodePtr.IsValid() && OwningGraph->IsValid())
+			FConstDocumentHandle OwningDocument = OwningGraph->GetOwningDocument();
+			FConstGraphHandle RootGraph = OwningDocument->GetRootGraph();
+
+			// Test if this node exists on the document's root graph.
+			const bool bIsNodeOnRootGraph = OwningGraph->IsValid() && (RootGraph->GetClassID() == OwningGraph->GetClassID());
+
+			if (bIsNodeOnRootGraph)
 			{
-				const FString& Name = NodePtr->Name;
-				const TArray<FMetasoundFrontendClassVertex>& RequiredOutputs = InArchetype.Interface.Outputs;
-				for (const FMetasoundFrontendClassVertex& OutputVertex : RequiredOutputs)
+				// If the node is on the root graph, test if it is in the archetypes
+				// required inputs or outputs. 
+				FMetasoundFrontendArchetype Archetype;
+
+				FArchetypeRegistryKey ArchetypeKey = GetArchetypeRegistryKey(OwningDocument->GetArchetypeVersion());
+
+				bool bFoundArchetype = IArchetypeRegistry::Get().FindArchetype(ArchetypeKey, Archetype);
+				if (bFoundArchetype)
 				{
-					if (OutputVertex.Name == Name)
+					if (const FMetasoundFrontendNode* Node = NodePtr.Get())
 					{
-						return true;
+						const FString& Name = Node->Name;
+						auto IsVertexWithSameName = [&Name](const FMetasoundFrontendClassVertex& InVertex)
+						{
+							return InVertex.Name == Name;
+						};
+						return Archetype.Interface.Outputs.ContainsByPredicate(IsVertexWithSameName);
 					}
 				}
+				else
+				{
+					UE_LOG(LogMetaSound, Warning, TEXT("Document using unregistered archetype [ArchetypeVersion:%s]"), *OwningDocument->GetArchetypeVersion().ToString());
+				}
 			}
 
 			return false;
@@ -1925,19 +1946,39 @@ namespace Metasound
 			return InputDisplayTitle;
 		}
 
-		bool FInputNodeController::IsRequired(const FMetasoundFrontendArchetype& InArchetype) const
+		bool FInputNodeController::IsRequired() const
 		{
-			if (NodePtr.IsValid() && OwningGraph->IsValid())
+			FConstDocumentHandle OwningDocument = OwningGraph->GetOwningDocument();
+			FConstGraphHandle RootGraph = OwningDocument->GetRootGraph();
+
+			// Test if this node exists on the document's root graph.
+			const bool bIsNodeOnRootGraph = OwningGraph->IsValid() && (RootGraph->GetClassID() == OwningGraph->GetClassID());
+
+			if (bIsNodeOnRootGraph)
 			{
-				const FString& Name = NodePtr->Name;
-				const TArray<FMetasoundFrontendClassVertex>& RequiredInputs = InArchetype.Interface.Inputs;
-				for (const FMetasoundFrontendClassVertex& InputVertex : RequiredInputs)
+				// If the node is on the root graph, test if it is in the archetypes
+				// required inputs or outputs. 
+				FMetasoundFrontendArchetype Archetype;
+
+				FArchetypeRegistryKey ArchetypeKey = GetArchetypeRegistryKey(OwningDocument->GetArchetypeVersion());
+
+				bool bFoundArchetype = IArchetypeRegistry::Get().FindArchetype(ArchetypeKey, Archetype);
+				if (bFoundArchetype)
 				{
-					if (InputVertex.Name == Name)
+					if (const FMetasoundFrontendNode* Node = NodePtr.Get())
 					{
-						return true;
+						const FString& Name = Node->Name;
+						auto IsVertexWithSameName = [&Name](const FMetasoundFrontendClassVertex& InVertex)
+						{
+							return InVertex.Name == Name;
+						};
+						return Archetype.Interface.Inputs.ContainsByPredicate(IsVertexWithSameName);
 					}
 				}
+				else
+				{
+					UE_LOG(LogMetaSound, Warning, TEXT("Document using unregistered archetype [ArchetypeVersion:%s]"), *OwningDocument->GetArchetypeVersion().ToString());
+				}
 			}
 
 			return false;
@@ -3326,6 +3367,15 @@ namespace Metasound
 			return Classes;
 		}
 
+		const FMetasoundFrontendGraphClass& FDocumentController::GetRootGraphClass() const
+		{
+			if (const FMetasoundFrontendDocument* Doc = DocumentPtr.Get())
+			{
+				return Doc->RootGraph;
+			}
+			return FrontendControllerIntrinsics::GetInvalidValueConstRef<FMetasoundFrontendGraphClass>();
+		}
+
 		bool FDocumentController::AddDuplicateSubgraph(const FMetasoundFrontendGraphClass& InGraphToCopy, const FMetasoundFrontendDocument& InOtherDocument)
 		{
 			FMetasoundFrontendDocument* ThisDocument = DocumentPtr.Get();
@@ -3389,6 +3439,23 @@ namespace Metasound
 			return bSuccess;
 		}
 
+		const FMetasoundFrontendVersion& FDocumentController::GetArchetypeVersion() const
+		{
+			if (const FMetasoundFrontendDocument* Doc = DocumentPtr.Get())
+			{
+				return Doc->ArchetypeVersion;
+			}
+			return FMetasoundFrontendVersion::GetInvalid();
+		}
+
+		void FDocumentController::SetArchetypeVersion(const FMetasoundFrontendVersion& InVersion)
+		{
+			if (FMetasoundFrontendDocument* Doc = DocumentPtr.Get())
+			{
+				Doc->ArchetypeVersion = InVersion;
+			}
+		}
+
 		FGraphHandle FDocumentController::AddDuplicateSubgraph(const IGraphController& InGraph)
 		{
 			// TODO: class IDs have issues.. 
diff --git a/Engine/Plugins/Runtime/Metasound/Source/MetasoundFrontend/Private/MetasoundFrontendStandardController.h b/Engine/Plugins/Runtime/Metasound/Source/MetasoundFrontend/Private/MetasoundFrontendStandardController.h
--- a/Engine/Plugins/Runtime/Metasound/Source/MetasoundFrontend/Private/MetasoundFrontendStandardController.h
+++ b/Engine/Plugins/Runtime/Metasound/Source/MetasoundFrontend/Private/MetasoundFrontendStandardController.h
@@ -425,7 +425,7 @@ namespace Metasound
 			TArray<FOutputHandle> GetOutputsWithVertexName(const FString& InName) override;
 			TArray<FConstOutputHandle> GetConstOutputsWithVertexName(const FString& InName) const override;
 
-			bool IsRequired(const FMetasoundFrontendArchetype& InArchetype) const override;
+			bool IsRequired() const override;
 
 			/** Returns an input with the given id. 
 			 *
@@ -590,7 +590,7 @@ namespace Metasound
 			void SetDescription(const FText& InDescription) override;
 			void SetDisplayName(const FText& InText) override;
 			const FText& GetDisplayTitle() const override;
-			bool IsRequired(const FMetasoundFrontendArchetype& InArchetype) const override;
+			bool IsRequired() const override;
 
 		protected:
 
@@ -649,7 +649,7 @@ namespace Metasound
 			const FText& GetDescription() const override;
 			const FText& GetDisplayName() const override;
 			const FText& GetDisplayTitle() const override;
-			bool IsRequired(const FMetasoundFrontendArchetype& InArchetype) const override;
+			bool IsRequired() const override;
 			bool IsValid() const override;
 			void SetDescription(const FText& InDescription) override;
 			void SetDisplayName(const FText& InText) override;
@@ -913,6 +913,7 @@ namespace Metasound
 			TArray<FMetasoundFrontendClass> GetDependencies() const override;
 			TArray<FMetasoundFrontendGraphClass> GetSubgraphs() const override;
 			TArray<FMetasoundFrontendClass> GetClasses() const override;
+			const FMetasoundFrontendGraphClass& GetRootGraphClass() const override;
 
 			FConstClassAccessPtr FindDependencyWithID(FGuid InClassID) const override;
 			FConstGraphClassAccessPtr FindSubgraphWithID(FGuid InClassID) const override;
@@ -926,6 +927,9 @@ namespace Metasound
 
 			virtual FGraphHandle AddDuplicateSubgraph(const IGraphController& InGraph) override;
 
+			virtual const FMetasoundFrontendVersion& GetArchetypeVersion() const override;
+			virtual void SetArchetypeVersion(const FMetasoundFrontendVersion& InVersion) override;
+
 			void SetMetadata(const FMetasoundFrontendDocumentMetadata& InMetadata) override;
 			const FMetasoundFrontendDocumentMetadata& GetMetadata() const override;
 
diff --git a/Engine/Plugins/Runtime/Metasound/Source/MetasoundFrontend/Private/MetasoundFrontendTransform.cpp b/Engine/Plugins/Runtime/Metasound/Source/MetasoundFrontend/Private/MetasoundFrontendTransform.cpp
--- a/Engine/Plugins/Runtime/Metasound/Source/MetasoundFrontend/Private/MetasoundFrontendTransform.cpp
+++ b/Engine/Plugins/Runtime/Metasound/Source/MetasoundFrontend/Private/MetasoundFrontendTransform.cpp
@@ -1,82 +1,328 @@
 // Copyright Epic Games, Inc. All Rights Reserved.
 #include "MetasoundFrontendTransform.h"
 
+#include "Algo/Transform.h"
 #include "CoreMinimal.h"
+#include "MetasoundArchetype.h"
+#include "MetasoundFrontendArchetypeRegistry.h"
 #include "MetasoundFrontendDocument.h"
+#include "MetasoundFrontendSearchEngine.h"
 #include "MetasoundLog.h"
 #include "Misc/App.h"
 
-
 namespace Metasound
 {
 	namespace Frontend
 	{
-		bool FMatchRootGraphToArchetype::Transform(FDocumentHandle InDocument) const
+		FSwapGraphArchetype::FSwapGraphArchetype(const FMetasoundFrontendArchetype& InFromArchetype, const FMetasoundFrontendArchetype& InToArchetype) 
 		{
-			if (!InDocument->IsValid())
+			InputsToAdd = InToArchetype.Interface.Inputs;
+			InputsToRemove = InFromArchetype.Interface.Inputs;
+			OutputsToAdd = InToArchetype.Interface.Outputs;
+			OutputsToRemove = InFromArchetype.Interface.Outputs;
+
+			// Iterate in reverse to allow removal from `InputsToAdd`
+			for (int32 AddIndex = InputsToAdd.Num() - 1; AddIndex >= 0; AddIndex--)
 			{
-				return false;
+				const FMetasoundFrontendClassVertex& VertexToAdd = InputsToAdd[AddIndex];
+
+				int32 RemoveIndex = InputsToRemove.IndexOfByPredicate([&](const FMetasoundFrontendClassVertex& VertexToRemove)
+					{
+						return FMetasoundFrontendClassVertex::IsFunctionalEquivalent(VertexToAdd, VertexToRemove);
+					});
+
+				if (INDEX_NONE != RemoveIndex)
+				{
+					PairedInputs.Add(FVertexPair{InputsToRemove[RemoveIndex], InputsToAdd[AddIndex]});
+					InputsToRemove.RemoveAtSwap(RemoveIndex);
+					InputsToAdd.RemoveAtSwap(AddIndex);
+				}
 			}
 
+			// Iterate in reverse to allow removal from `OutputsToAdd`
+			for (int32 AddIndex = OutputsToAdd.Num() - 1; AddIndex >= 0; AddIndex--)
+			{
+				const FMetasoundFrontendClassVertex& VertexToAdd = OutputsToAdd[AddIndex];
+
+				int32 RemoveIndex = OutputsToRemove.IndexOfByPredicate([&](const FMetasoundFrontendClassVertex& VertexToRemove)
+					{
+						return FMetasoundFrontendClassVertex::IsFunctionalEquivalent(VertexToAdd, VertexToRemove);
+					});
+
+				if (INDEX_NONE != RemoveIndex)
+				{
+					PairedOutputs.Add(FVertexPair{OutputsToRemove[RemoveIndex], OutputsToAdd[AddIndex]});
+					OutputsToRemove.RemoveAtSwap(RemoveIndex);
+					OutputsToAdd.RemoveAtSwap(AddIndex);
+				}
+			}
+		}
+
+		bool FSwapGraphArchetype::Transform(FGraphHandle InGraph) const 
+		{
 			bool bDidEdit = false;
-			FGraphHandle Graph = InDocument->GetRootGraph();
 
-			const TArray<FMetasoundFrontendClassVertex>& RequiredInputs = Archetype.Interface.Inputs;
-			const TArray<FMetasoundFrontendClassVertex>& RequiredOutputs = Archetype.Interface.Outputs;
+			// Remove unsupported inputs
+			for (const FMetasoundFrontendClassVertex& InputToRemove : InputsToRemove)
+			{
+				if (const FMetasoundFrontendClassInput* ClassInput = InGraph->FindClassInputWithName(InputToRemove.Name).Get())
+				{
+					if (FMetasoundFrontendClassInput::IsFunctionalEquivalent(*ClassInput, InputToRemove))
+					{
+						bDidEdit = true;
+						InGraph->RemoveInputVertex(InputToRemove.Name);
+					}
+				}
+			}
 
-			// Go through each input and add or swap if something is missing.
-			for (const FMetasoundFrontendClassVertex& RequiredInput : RequiredInputs)
+			// Remove unrequired outputs
+			for (const FMetasoundFrontendClassVertex& OutputToRemove : OutputsToRemove)
 			{
-				if (const FMetasoundFrontendClassInput* ClassInput = Graph->FindClassInputWithName(RequiredInput.Name).Get())
+				if (const FMetasoundFrontendClassOutput* ClassOutput = InGraph->FindClassOutputWithName(OutputToRemove.Name).Get())
 				{
-					if (!FMetasoundFrontendClassVertex::IsFunctionalEquivalent(RequiredInput, *ClassInput))
+					if (FMetasoundFrontendClassOutput::IsFunctionalEquivalent(*ClassOutput, OutputToRemove))
 					{
 						bDidEdit = true;
+						InGraph->RemoveOutputVertex(OutputToRemove.Name);
+					}
+				}
+			}
 
-						// Cache off node locations to push to new node
-						FNodeHandle InputNode = Graph->GetInputNodeWithName(ClassInput->Name);
-						const TMap<FGuid, FVector2D> Locations = InputNode->GetNodeStyle().Display.Locations;
+			// Add missing inputs
+			for (const FMetasoundFrontendClassVertex& InputToAdd : InputsToAdd)
+			{
+				bDidEdit = true;
+				InGraph->AddInputVertex(InputToAdd);
+			}
 
-						Graph->RemoveInputVertex(RequiredInput.Name);
-						InputNode = Graph->AddInputVertex(RequiredInput);
+			// Add missing outputs
+			for (const FMetasoundFrontendClassVertex& OutputToAdd : OutputsToAdd)
+			{
+				bDidEdit = true;
+				FNodeHandle NewOutputNode = InGraph->AddOutputVertex(OutputToAdd);
+				FMetasoundFrontendNodeStyle Style = NewOutputNode->GetNodeStyle();
+				// TODO: Create generic "FindGoodLocation()" for nodes. 
+				// Inputs are on the left,
+				// Outputs are on the right,
+				// Place near nodes with similar pin types.
+				// Requires a "NodeSize" approximation
+				Style.Display.Locations.Add(FGuid(), FVector2D{ 0.f, 0.f });
+				NewOutputNode->SetNodeStyle(Style);
+			}
 
-						FMetasoundFrontendNodeStyle Style = InputNode->GetNodeStyle();
-						Style.Display.Locations = Locations;
-						InputNode->SetNodeStyle(Style);
+			// Swap paired inputs.
+			for (const FVertexPair& InputPair : PairedInputs)
+			{
+				bDidEdit = true;
+
+				const FMetasoundFrontendClassVertex& OriginalVertex = InputPair.Get<0>();
+				FMetasoundFrontendClassInput NewVertex = InputPair.Get<1>();
+
+				// Cache off node locations and connections to push to new node
+				TMap<FGuid, FVector2D> Locations;
+				TArray<FInputHandle> ConnectedInputs;
+				if (const FMetasoundFrontendClassInput* ClassInput = InGraph->FindClassInputWithName(OriginalVertex.Name).Get())
+				{
+					if (FMetasoundFrontendVertex::IsFunctionalEquivalent(*ClassInput, OriginalVertex))
+					{
+						NewVertex.DefaultLiteral = ClassInput->DefaultLiteral;
+						FNodeHandle OriginalInputNode = InGraph->GetInputNodeWithName(OriginalVertex.Name);
+						Locations = OriginalInputNode->GetNodeStyle().Display.Locations;
+
+						TArray<FOutputHandle> Outputs = OriginalInputNode->GetOutputsWithVertexName(OriginalVertex.Name);
+						if (Outputs.Num() == 1)
+						{
+							ConnectedInputs = Outputs[0]->GetConnectedInputs();
+						}
+						InGraph->RemoveInputVertex(OriginalVertex.Name);
 					}
 				}
-				else
+
+				FNodeHandle NewInputNode = InGraph->AddInputVertex(NewVertex);
+				// Copy prior node locations
+				if (Locations.Num() > 0)
 				{
-					bDidEdit = true;
-					Graph->AddInputVertex(RequiredInput);
+					// TODO: copy entire style.
+					FMetasoundFrontendNodeStyle Style = NewInputNode->GetNodeStyle();
+					Style.Display.Locations = Locations;
+					NewInputNode->SetNodeStyle(Style);
+				}
+
+				// Copy prior node connections
+				TArray<FOutputHandle> OutputHandles= NewInputNode->GetOutputsWithVertexName(NewVertex.Name);
+				if (OutputHandles.Num() == 1)
+				{
+					FOutputHandle OutputHandle = OutputHandles[0];
+					for (FInputHandle& ConnectedInput : ConnectedInputs)
+					{
+						OutputHandle->Connect(*ConnectedInput);
+					}
 				}
 			}
 
-			// Go through each output and add or swap if something is missing.
-			for (const FMetasoundFrontendClassVertex& RequiredOutput : RequiredOutputs)
+			// Swap paired outputs.
+			for (const FVertexPair& OutputPair : PairedOutputs)
 			{
-				if (const FMetasoundFrontendClassOutput* ClassOutput = Graph->FindClassOutputWithName(RequiredOutput.Name).Get())
+				bDidEdit = true;
+
+				const FMetasoundFrontendClassVertex& OriginalVertex = OutputPair.Get<0>();
+				const FMetasoundFrontendClassVertex& NewVertex = OutputPair.Get<1>();
+
+				// Cache off node locations to push to new node
+				TMap<FGuid, FVector2D> Locations;
+				FOutputHandle ConnectedOutput = IOutputController::GetInvalidHandle();
+
+				// Default add output node to origin.
+				Locations.Add(FGuid(), FVector2D{0.f, 0.f});
+				if (const FMetasoundFrontendClassOutput* ClassOutput = InGraph->FindClassOutputWithName(OriginalVertex.Name).Get())
 				{
-					if (!FMetasoundFrontendClassVertex::IsFunctionalEquivalent(RequiredOutput, *ClassOutput))
+					if (FMetasoundFrontendVertex::IsFunctionalEquivalent(*ClassOutput, OriginalVertex))
 					{
-						bDidEdit = true;
+						FNodeHandle OriginalOutputNode = InGraph->GetOutputNodeWithName(OriginalVertex.Name);
+						Locations = OriginalOutputNode->GetNodeStyle().Display.Locations;
+						TArray<FInputHandle> Inputs = OriginalOutputNode->GetInputsWithVertexName(OriginalVertex.Name);
+						if (Inputs.Num() == 1)
+						{
+							ConnectedOutput = Inputs[0]->GetConnectedOutput();
+						}
+						InGraph->RemoveOutputVertex(OriginalVertex.Name);
+					}
+				}
 
-						// Cache off node locations to push to new node
-						FNodeHandle OutputNode = Graph->GetOutputNodeWithName(ClassOutput->Name);
-						const TMap<FGuid, FVector2D> Locations = OutputNode->GetNodeStyle().Display.Locations;
+				FNodeHandle NewOutputNode = InGraph->AddOutputVertex(NewVertex);
+
+				if (Locations.Num() > 0)
+				{
+					FMetasoundFrontendNodeStyle Style = NewOutputNode->GetNodeStyle();
+					Style.Display.Locations = Locations;
+					NewOutputNode->SetNodeStyle(Style);
+				}
+
+				// Copy prior node connections
+				TArray<FInputHandle> InputHandles= NewOutputNode->GetInputsWithVertexName(NewVertex.Name);
+				if (InputHandles.Num() == 1)
+				{
+					ConnectedOutput->Connect(*InputHandles[0]);
+				}
+			}
+
+			return bDidEdit;
+		}
+
+		bool FMatchRootGraphToArchetype::Transform(FDocumentHandle InDocument) const
+		{
+			bool bDidEdit = false;
+
+			if (!InDocument->IsValid())
+			{
+				return bDidEdit;
+			}
 
-						Graph->RemoveOutputVertex(RequiredOutput.Name);
-						OutputNode = Graph->AddOutputVertex(RequiredOutput);
+			// Find registered target archetype.
+			FMetasoundFrontendArchetype TargetArchetype;
+			bool bFoundTargetArchetype = IArchetypeRegistry::Get().FindArchetype(GetArchetypeRegistryKey(ArchetypeVersion), TargetArchetype);
 
-						FMetasoundFrontendNodeStyle Style = OutputNode->GetNodeStyle();
-						Style.Display.Locations = Locations;
-						OutputNode->SetNodeStyle(Style);
+			if (!bFoundTargetArchetype)
+			{
+				UE_LOG(LogMetaSound, Error, TEXT("Target archetype is not registered [ArchetypeVersion:%s]"), *ArchetypeVersion.ToString());
+				return false;
+			}
+
+			// Get current archetype version on document.
+			const FMetasoundFrontendVersion InitialArchetypeVersion = InDocument->GetArchetypeVersion();
+
+			// Attempt to upgrade
+			TArray<const IArchetypeRegistryEntry*> UpgradePath;
+			GetUpgradePathForDocument(InitialArchetypeVersion, ArchetypeVersion, UpgradePath);
+
+			if (UpgradeDocumentArchetype(UpgradePath, InDocument))
+			{
+				bDidEdit = true;
+			}
+
+			// Force archetype to conform
+			if (ConformDocumentToArchetype(TargetArchetype, InDocument))
+			{
+				bDidEdit = true;
+			}
+
+			return bDidEdit;
+		}
+
+		void FMatchRootGraphToArchetype::GetUpgradePathForDocument(const FMetasoundFrontendVersion& InCurrentVersion, const FMetasoundFrontendVersion& InTargetVersion, TArray<const IArchetypeRegistryEntry*>& OutUpgradePath) const
+		{
+			if (InCurrentVersion.Name == InTargetVersion.Name)
+			{
+				// Get all associated registered archetypes
+				TArray<FMetasoundFrontendVersion> RegisteredVersions = ISearchEngine::Get().FindAllRegisteredArchetypesWithName(InTargetVersion.Name);
+
+				// Filter registry entries that exist between current version and target version
+				auto FilterRegistryEntries = [&InCurrentVersion, &InTargetVersion](const FMetasoundFrontendVersion& InVersion)
+				{
+					const bool bIsGreaterThanCurrent = InVersion.Number > InCurrentVersion.Number;
+					const bool bIsLessThanOrEqualToTarget = InVersion.Number <= InTargetVersion.Number;
+
+					return bIsGreaterThanCurrent && bIsLessThanOrEqualToTarget;
+				};
+				RegisteredVersions = RegisteredVersions.FilterByPredicate(FilterRegistryEntries);
+
+				// sort registry entries to create an ordered upgrade path.
+				RegisteredVersions.Sort();
+
+				// Get registry entries from registry keys.
+				auto GetRegistryEntry = [](const FMetasoundFrontendVersion& InVersion)
+				{
+					FArchetypeRegistryKey Key = GetArchetypeRegistryKey(InVersion);
+					return IArchetypeRegistry::Get().FindArchetypeRegistryEntry(Key);
+				};
+				Algo::Transform(RegisteredVersions, OutUpgradePath, GetRegistryEntry);
+			}
+		}
+
+		bool FMatchRootGraphToArchetype::UpgradeDocumentArchetype(const TArray<const IArchetypeRegistryEntry*>& InUpgradePath, FDocumentHandle InDocument) const
+		{
+			bool bDidEdit = false;
+			for (const IArchetypeRegistryEntry* Entry : InUpgradePath)
+			{
+				if (ensure(nullptr != Entry))
+				{
+					if (Entry->UpdateRootGraphArchetype(InDocument))
+					{
+						bDidEdit = true;
 					}
 				}
-				else
+			}
+			return bDidEdit;
+		}
+
+		bool FMatchRootGraphToArchetype::ConformDocumentToArchetype(const FMetasoundFrontendArchetype& InTargetArchetype, FDocumentHandle InDocument) const
+		{
+			bool bDidEdit = false;
+
+			FMetasoundFrontendArchetype CurrentArchetype;
+			FArchetypeRegistryKey CurrentRegistryKey = GetArchetypeRegistryKey(InDocument->GetArchetypeVersion());
+			const bool bFoundCurrentArchetype = IArchetypeRegistry::Get().FindArchetype(CurrentRegistryKey, CurrentArchetype);
+
+			if (!bFoundCurrentArchetype)
+			{
+				UE_LOG(LogMetaSound, Warning, TEXT("Failed to find current archetype on document [ArchetypeVersion:%s]"), *InDocument->GetArchetypeVersion().ToString());
+			}
+
+			const bool bIsEqualArchetypeVersion = InDocument->GetArchetypeVersion() == InTargetArchetype.Version;
+			const bool bRequiredInterfaceExists = IsSubsetOfClass(InTargetArchetype, InDocument->GetRootGraphClass());
+
+			if (!bIsEqualArchetypeVersion)
+			{
+				InDocument->SetArchetypeVersion(InTargetArchetype.Version);
+				bDidEdit = true;
+			}
+
+			if (!(bRequiredInterfaceExists && bIsEqualArchetypeVersion))
+			{
+				FGraphHandle Graph = InDocument->GetRootGraph();
+				if (FSwapGraphArchetype(CurrentArchetype, InTargetArchetype).Transform(Graph))
 				{
 					bDidEdit = true;
-					Graph->AddOutputVertex(RequiredOutput);
 				}
 			}
 
@@ -199,7 +445,7 @@ namespace Metasound
 			}
 		};
 
-		/** Versions document from 1.1 to 1.3. */
+		/** Versions document from 1.2 to 1.3. */
 		class FVersionDocument_1_3 : public FVersionDocumentTransform
 		{
 		public:
@@ -209,7 +455,7 @@ namespace Metasound
 
 			FMetasoundFrontendVersionNumber GetTargetVersion() const override
 			{
-				return FVersionDocument::GetMaxVersion();
+				return {1, 3};
 			}
 
 			void TransformInternal(FDocumentHandle InDocument) const override
@@ -222,6 +468,54 @@ namespace Metasound
 			}
 		};
 
+		/** Versions document from 1.3 to 1.4. */
+		class FVersionDocument_1_4 : public FVersionDocumentTransform
+		{
+		public:
+			FVersionDocument_1_4()
+			{
+			}
+
+			FMetasoundFrontendVersionNumber GetTargetVersion() const override
+			{
+				return {1, 4};
+			}
+
+			void TransformInternal(FDocumentHandle InDocument) const override
+			{
+				
+				check(InDocument->GetMetadata().Version.Number.Major == 1);
+				check(InDocument->GetMetadata().Version.Number.Minor == 3);
+
+				FMetasoundFrontendVersion ArchetypeVersion = InDocument->GetArchetypeVersion();
+
+				// Version 1.3 did not have an "ArchetypeVersion" property on the 
+				// document, so any document that is being updated should start off
+				// with an "Invalid" archetype version.
+				if (ensure(!ArchetypeVersion.IsValid()))
+				{
+					constexpr bool bIncludeDeprecatedArchetypes = true;
+					TArray<FMetasoundFrontendArchetype> AllArchetypes = ISearchEngine::Get().FindAllArchetypes(bIncludeDeprecatedArchetypes);
+
+					const FMetasoundFrontendGraphClass& RootGraph = InDocument->GetRootGraphClass();
+					const TArray<FMetasoundFrontendClass>& Dependencies = InDocument->GetDependencies();
+					const TArray<FMetasoundFrontendGraphClass>& Subgraphs = InDocument->GetSubgraphs();
+
+
+					if (const FMetasoundFrontendArchetype* Arch = FindMostSimilarArchetypeSupportingEnvironment(RootGraph, Dependencies, Subgraphs, AllArchetypes))
+					{
+						UE_LOG(LogMetaSound, Display, TEXT("Assigned archetype [ArchetypeVersion:%s] to document [RootGraphClassName:%s]"), *Arch->Version.ToString(), *RootGraph.Metadata.ClassName.ToString());
+
+						InDocument->SetArchetypeVersion(Arch->Version);
+					}
+					else
+					{
+						UE_LOG(LogMetaSound, Warning, TEXT("Failed to find archetype for document [RootGraphClassName:%s]"), *RootGraph.Metadata.ClassName.ToString());
+					}
+				}
+			}
+		};
+
 		FVersionDocument::FVersionDocument(FName InName, const FString& InPath)
 			: Name(InName)
 			, Path(InPath)
@@ -241,6 +535,7 @@ namespace Metasound
 			bWasUpdated |= FVersionDocument_1_1().Transform(InDocument);
 			bWasUpdated |= FVersionDocument_1_2(Name, Path).Transform(InDocument);
 			bWasUpdated |= FVersionDocument_1_3().Transform(InDocument);
+			bWasUpdated |= FVersionDocument_1_4().Transform(InDocument);
 
 			return bWasUpdated;
 		}
diff --git a/Engine/Plugins/Runtime/Metasound/Source/MetasoundFrontend/Public/MetasoundArchetype.h b/Engine/Plugins/Runtime/Metasound/Source/MetasoundFrontend/Public/MetasoundArchetype.h
--- a/Engine/Plugins/Runtime/Metasound/Source/MetasoundFrontend/Public/MetasoundArchetype.h
+++ b/Engine/Plugins/Runtime/Metasound/Source/MetasoundFrontend/Public/MetasoundArchetype.h
@@ -8,15 +8,22 @@ namespace Metasound
 {
 	namespace Frontend
 	{
+
 		METASOUNDFRONTEND_API bool IsSubsetOfArchetype(const FMetasoundFrontendArchetype& InSubsetArchetype, const FMetasoundFrontendArchetype& InSupersetArchetype);
 
+		METASOUNDFRONTEND_API bool IsSubsetOfClass(const FMetasoundFrontendArchetype& InSubsetArchetype, const FMetasoundFrontendClass& InSupersetClass);
+
 		METASOUNDFRONTEND_API bool IsEquivalentArchetype(const FMetasoundFrontendArchetype& InInputArchetype, const FMetasoundFrontendArchetype& InTargetArchetype);
 
 		METASOUNDFRONTEND_API bool IsEqualArchetype(const FMetasoundFrontendArchetype& InInputArchetype, const FMetasoundFrontendArchetype& InTargetArchetype);
 
+		METASOUNDFRONTEND_API int32 InputOutputDifferenceCount(const FMetasoundFrontendClass& InClass, const FMetasoundFrontendArchetype& InArchetype);
+
 		METASOUNDFRONTEND_API int32 InputOutputDifferenceCount(const FMetasoundFrontendArchetype& InArchetypeA, const FMetasoundFrontendArchetype& InArchetypeB);
 
-		METASOUNDFRONTEND_API void GatherRequiredEnvironmentVariables(const FMetasoundFrontendDocument& InDocument, TArray<FMetasoundFrontendEnvironmentVariable>& OutEnvironmentVariables);
+		METASOUNDFRONTEND_API void GatherRequiredEnvironmentVariables(const FMetasoundFrontendClass& InClass, TArray<FMetasoundFrontendEnvironmentVariable>& OutEnvironmentVariables);
+
+		METASOUNDFRONTEND_API const FMetasoundFrontendArchetype* FindMostSimilarArchetypeSupportingEnvironment(const FMetasoundFrontendGraphClass& InRootGraph, const TArray<FMetasoundFrontendClass>& InDependencies, const TArray<FMetasoundFrontendGraphClass>& InSubgraphs, const TArray<FMetasoundFrontendArchetype>& InCandidateArchetypes);
 
 		METASOUNDFRONTEND_API const FMetasoundFrontendArchetype* FindMostSimilarArchetypeSupportingEnvironment(const FMetasoundFrontendDocument& InDocument, const TArray<FMetasoundFrontendArchetype>& InCandidateArchetypes);
 	}
diff --git a/Engine/Plugins/Runtime/Metasound/Source/MetasoundFrontend/Public/MetasoundAssetBase.h b/Engine/Plugins/Runtime/Metasound/Source/MetasoundFrontend/Public/MetasoundAssetBase.h
--- a/Engine/Plugins/Runtime/Metasound/Source/MetasoundFrontend/Public/MetasoundAssetBase.h
+++ b/Engine/Plugins/Runtime/Metasound/Source/MetasoundFrontend/Public/MetasoundAssetBase.h
@@ -75,16 +75,16 @@ public:
 	virtual void SetMetadata(FMetasoundFrontendClassMetadata& InMetadata);
 
 	// Returns  a description of the required inputs and outputs for this metasound UClass.
-	virtual const FMetasoundFrontendArchetype& GetArchetype() const = 0;
+	virtual const FMetasoundFrontendVersion& GetDefaultArchetypeVersion() const = 0;
 
 	// Returns true if the archetype is supported by this object.
-	virtual bool IsArchetypeSupported(const FMetasoundFrontendArchetype& InArchetype) const;
+	virtual bool IsArchetypeSupported(const FMetasoundFrontendVersion& InArchetypeVersion) const;
 
 	// Returns an array of archetypes preferred for this class.
-	virtual const TArray<FMetasoundFrontendArchetype>& GetPreferredArchetypes() const = 0;
+	virtual const TArray<FMetasoundFrontendVersion>& GetSupportedArchetypeVersions() const = 0;
 
 	// Returns the preferred archetype for the given document.
-	virtual const FMetasoundFrontendArchetype& GetPreferredArchetypes(const FMetasoundFrontendDocument& InDocument, const FMetasoundFrontendArchetype& InDefaultArchetype) const;
+	virtual FMetasoundFrontendVersion GetPreferredArchetypeVersion(const FMetasoundFrontendDocument& InDocument) const;
 
 	// TODO:
 	//virtual void OnMetaSoundDependencyAdded(const FMetasoundFrontendClassMetadata& InMetadata) = 0;
@@ -121,6 +121,8 @@ public:
 	// This must be called on UObject::PostLoad, as well as in this asset's UFactory, to fix up the root document based on the most recent version of the archetype.
 	void ConformDocumentToArchetype();
 
+	bool VersionAsset();
+
 	// Calls the outermost package and marks it dirty. 
 	bool MarkMetasoundDocumentDirty() const;
 
diff --git a/Engine/Plugins/Runtime/Metasound/Source/MetasoundFrontend/Public/MetasoundFrontend.h b/Engine/Plugins/Runtime/Metasound/Source/MetasoundFrontend/Public/MetasoundFrontend.h
--- a/Engine/Plugins/Runtime/Metasound/Source/MetasoundFrontend/Public/MetasoundFrontend.h
+++ b/Engine/Plugins/Runtime/Metasound/Source/MetasoundFrontend/Public/MetasoundFrontend.h
@@ -18,17 +18,6 @@ namespace Metasound
 		class IRegistryTransaction;
 		using FRegistryTransactionID = int32;
 
-		METASOUNDFRONTEND_API FRegistryTransactionID GetCurrentRegistryTransactionID();
-
-		/** Return all registry transactions that have occurred since the given transaction ID. 
-		 *
-		 * @param InTransactionID - Initial transaction ID. 
-		 * @param OutCurrentTransactionID - If not null, this will be set to the current transaction ID.
-		 *
-		 * @return An array of IRegistryTransactions that occurred in the range (InTransactionID, OutCurrentTransactionID].
-		 */
-		METASOUNDFRONTEND_API TArray<const IRegistryTransaction*> GetRegistryTransactionsSince(FRegistryTransactionID InTransactionID, FRegistryTransactionID* OutCurrentTransactionID);
-
 		/** Returns all metadata (name, description, author, what to say if it's missing) for a given node.
 		 *
 		 * @param InInfo - Class info for a already registered external node.
diff --git a/Engine/Plugins/Runtime/Metasound/Source/MetasoundFrontend/Public/MetasoundFrontendArchetypeRegistry.h b/Engine/Plugins/Runtime/Metasound/Source/MetasoundFrontend/Public/MetasoundFrontendArchetypeRegistry.h
new file mode 100644
--- /dev/null
+++ b/Engine/Plugins/Runtime/Metasound/Source/MetasoundFrontend/Public/MetasoundFrontendArchetypeRegistry.h
@@ -0,0 +1,66 @@
+// Copyright Epic Games, Inc. All Rights Reserved.
+
+#pragma once
+
+#include "CoreMinimal.h"
+#include "MetasoundFrontendController.h"
+#include "MetasoundFrontendDocument.h"
+
+namespace Metasound
+{
+	namespace Frontend
+	{
+		using FArchetypeRegistryKey = FString;
+		using FRegistryTransactionID = int32;
+
+		METASOUNDFRONTEND_API bool IsValidArchetypeRegistryKey(const FArchetypeRegistryKey& InKey);
+		METASOUNDFRONTEND_API FArchetypeRegistryKey GetArchetypeRegistryKey(const FMetasoundFrontendVersion& InArchetypeVersion);
+		METASOUNDFRONTEND_API FArchetypeRegistryKey GetArchetypeRegistryKey(const FMetasoundFrontendArchetype& InArchetype);
+
+		class IArchetypeRegistryEntry
+		{
+		public:
+			virtual ~IArchetypeRegistryEntry() = default;
+
+			virtual const FMetasoundFrontendArchetype& GetArchetype() const = 0;
+			virtual bool UpdateRootGraphArchetype(FDocumentHandle InDocument) const = 0;
+		};
+
+		class METASOUNDFRONTEND_API FArchetypeRegistryTransaction 
+		{
+		public:
+			/** Describes the type of transaction. */
+			enum class ETransactionType : uint8
+			{
+				ArchetypeRegistration,     //< Something was added to the registry.
+				ArchetypeUnregistration,  //< Something was removed from the registry.
+				Invalid
+			};
+
+			FArchetypeRegistryTransaction(ETransactionType InType, const FArchetypeRegistryKey& InKey, const FMetasoundFrontendVersion& InArchetypeVersion);
+
+			ETransactionType GetTransactionType() const;
+			const FMetasoundFrontendVersion& GetArchetypeVersion() const;
+			const FArchetypeRegistryKey& GetArchetypeRegistryKey() const;
+
+		private:
+
+			ETransactionType Type;
+			FArchetypeRegistryKey Key;
+			FMetasoundFrontendVersion ArchetypeVersion;
+		};
+
+		class METASOUNDFRONTEND_API IArchetypeRegistry
+		{
+		public:
+			static IArchetypeRegistry& Get();
+
+			virtual ~IArchetypeRegistry() = default;
+
+			virtual FArchetypeRegistryKey RegisterArchetype(TUniquePtr<IArchetypeRegistryEntry>&& InEntry) = 0;
+			virtual const IArchetypeRegistryEntry* FindArchetypeRegistryEntry(const FArchetypeRegistryKey& InKey) const = 0;
+			virtual bool FindArchetype(const FArchetypeRegistryKey& InKey, FMetasoundFrontendArchetype& OutArchetype) const = 0;
+			virtual void ForEachRegistryTransactionSince(FRegistryTransactionID InSince, FRegistryTransactionID* OutCurrentRegistryTransactionID, TFunctionRef<void(const FArchetypeRegistryTransaction&)> InFunc) const = 0;
+		};
+	}
+}
diff --git a/Engine/Plugins/Runtime/Metasound/Source/MetasoundFrontend/Public/MetasoundFrontendController.h b/Engine/Plugins/Runtime/Metasound/Source/MetasoundFrontend/Public/MetasoundFrontendController.h
--- a/Engine/Plugins/Runtime/Metasound/Source/MetasoundFrontend/Public/MetasoundFrontendController.h
+++ b/Engine/Plugins/Runtime/Metasound/Source/MetasoundFrontend/Public/MetasoundFrontendController.h
@@ -333,7 +333,7 @@ namespace Metasound
 			virtual TArray<FConstOutputHandle> GetConstOutputsWithVertexName(const FString& InName) const = 0;
 
 			/** Returns true if node is required to satisfy the document archetype. */
-			virtual bool IsRequired(const FMetasoundFrontendArchetype& InArchetype) const = 0;
+			virtual bool IsRequired() const = 0;
 
 			/** Returns an input with the given id.
 			 *
@@ -696,6 +696,7 @@ namespace Metasound
 			virtual TArray<FMetasoundFrontendClass> GetDependencies() const = 0;
 			virtual TArray<FMetasoundFrontendGraphClass> GetSubgraphs() const = 0;
 			virtual TArray<FMetasoundFrontendClass> GetClasses() const = 0;
+			virtual const FMetasoundFrontendGraphClass& GetRootGraphClass() const = 0;
 
 			virtual FConstClassAccessPtr FindDependencyWithID(FGuid InClassID) const = 0;
 			virtual FConstGraphClassAccessPtr FindSubgraphWithID(FGuid InClassID) const = 0;
@@ -745,6 +746,9 @@ namespace Metasound
 			 */
 			virtual FGraphHandle AddDuplicateSubgraph(const IGraphController& InGraph) = 0;
 
+			virtual const FMetasoundFrontendVersion& GetArchetypeVersion() const = 0;
+			virtual void SetArchetypeVersion(const FMetasoundFrontendVersion& InVersion) = 0;
+
 			/** Removes all dependencies which are no longer referenced by any graphs within this document
 			  * and updates dependency Metadata where necessary with that found in the registry.  */
 			virtual void SynchronizeDependencies() = 0;
diff --git a/Engine/Plugins/Runtime/Metasound/Source/MetasoundFrontend/Public/MetasoundFrontendDocument.h b/Engine/Plugins/Runtime/Metasound/Source/MetasoundFrontend/Public/MetasoundFrontendDocument.h
--- a/Engine/Plugins/Runtime/Metasound/Source/MetasoundFrontend/Public/MetasoundFrontendDocument.h
+++ b/Engine/Plugins/Runtime/Metasound/Source/MetasoundFrontend/Public/MetasoundFrontendDocument.h
@@ -127,7 +127,7 @@ struct METASOUNDFRONTEND_API FMetasoundFrontendVersionNumber
 
 // General purpose version info for Metasound Frontend objects.
 USTRUCT()
-struct FMetasoundFrontendVersion
+struct METASOUNDFRONTEND_API FMetasoundFrontendVersion
 {
 	GENERATED_BODY()
 
@@ -138,6 +138,62 @@ struct FMetasoundFrontendVersion
 	// Version number.
 	UPROPERTY()
 	FMetasoundFrontendVersionNumber Number;
+
+	FString ToString() const;
+
+	bool IsValid() const;
+
+	static const FMetasoundFrontendVersion& GetInvalid();
+
+	friend bool operator==(const FMetasoundFrontendVersion& InLHS, const FMetasoundFrontendVersion& InRHS)
+	{
+		return InLHS.Name == InRHS.Name && InLHS.Number == InRHS.Number;
+	}
+
+	friend bool operator!=(const FMetasoundFrontendVersion& InLHS, const FMetasoundFrontendVersion& InRHS)
+	{
+		return !(InLHS == InRHS);
+	}
+
+	friend bool operator>(const FMetasoundFrontendVersion& InLHS, const FMetasoundFrontendVersion& InRHS)
+	{
+		if (InRHS.Name.FastLess(InLHS.Name))
+		{
+			return true;
+		}
+
+		if (InLHS.Name == InRHS.Name)
+		{
+			return InLHS.Number > InRHS.Number;
+		}
+
+		return false;
+	}
+
+	friend bool operator>=(const FMetasoundFrontendVersion& InLHS, const FMetasoundFrontendVersion& InRHS)
+	{
+		return InLHS == InRHS || InLHS > InRHS;
+	}
+
+	friend bool operator<(const FMetasoundFrontendVersion& InLHS, const FMetasoundFrontendVersion& InRHS)
+	{
+		if (InLHS.Name.FastLess(InRHS.Name))
+		{
+			return true;
+		}
+
+		if (InLHS.Name == InRHS.Name)
+		{
+			return InLHS.Number < InRHS.Number;
+		}
+
+		return false;
+	}
+
+	friend bool operator<=(const FMetasoundFrontendVersion& InLHS, const FMetasoundFrontendVersion& InRHS)
+	{
+		return InLHS == InRHS || InLHS < InRHS;
+	}
 };
 
 
@@ -784,20 +840,18 @@ struct FMetasoundFrontendArchetypeInterface
 	TArray<FMetasoundFrontendEnvironmentVariable> Environment;
 };
 
-// This is used to describe the required inputs and outputs for a metasound, and is used to make sure we can use a metasound graph for specific applications.
-// For example, a UMetaSoundSource needs to generate audio, so its RequiredOutputs will contain "MainAudioOutput"
+/** Describes the expected minimum interface of a MetaSound root graph. This allows
+ * systems to describe the minimum interface for a MetaSound to work within them. */
 USTRUCT()
 struct FMetasoundFrontendArchetype
 {
 	GENERATED_BODY()
 
-	// Name of the archetype we're using.
+	/** Name and version number of the archetype. */
 	UPROPERTY()
-	FName Name;
-
-	UPROPERTY()
-	FMetasoundFrontendVersionNumber Version;
+	FMetasoundFrontendVersion Version;
 
+	/** Expected minimum interface for MetaSounds satisfying archetype. */
 	UPROPERTY()
 	FMetasoundFrontendArchetypeInterface Interface;
 };
@@ -808,7 +862,6 @@ struct METASOUNDFRONTEND_API FMetasoundFrontendDocument
 {
 	GENERATED_BODY()
 
-	//Metasound::Frontend::TAccessPoint<FMetasoundFrontendDocument> AccessPoint;
 	Metasound::Frontend::FAccessPoint AccessPoint;
 
 	FMetasoundFrontendDocument();
@@ -816,15 +869,18 @@ struct METASOUNDFRONTEND_API FMetasoundFrontendDocument
 	UPROPERTY(EditAnywhere, Category = Metadata)
 	FMetasoundFrontendDocumentMetadata Metadata;
 
+	UPROPERTY()
+	FMetasoundFrontendEditorData EditorData;
+
+	UPROPERTY()
+	FMetasoundFrontendVersion ArchetypeVersion;
+
 	UPROPERTY(EditAnywhere, Category = CustomView)
 	FMetasoundFrontendGraphClass RootGraph;
 
 	UPROPERTY()
 	TArray<FMetasoundFrontendGraphClass> Subgraphs;
 
-	UPROPERTY()
-	FMetasoundFrontendEditorData EditorData;
-
 	UPROPERTY()
 	TArray<FMetasoundFrontendClass> Dependencies;
 };
diff --git a/Engine/Plugins/Runtime/Metasound/Source/MetasoundFrontend/Public/MetasoundFrontendQuery.h b/Engine/Plugins/Runtime/Metasound/Source/MetasoundFrontend/Public/MetasoundFrontendQuery.h
--- a/Engine/Plugins/Runtime/Metasound/Source/MetasoundFrontend/Public/MetasoundFrontendQuery.h
+++ b/Engine/Plugins/Runtime/Metasound/Source/MetasoundFrontend/Public/MetasoundFrontendQuery.h
@@ -6,11 +6,12 @@
 #include "Math/NumericLimits.h"
 #include "MetasoundFrontendDocument.h"
 #include "MetasoundFrontendRegistries.h"
+#include "MetasoundFrontendArchetypeRegistry.h"
 #include "Misc/TVariant.h"
 
 namespace Metasound
 {
-	/** FFronendQueryKey allows entries to be partitioned by their key. A key
+	/** FFrontendQueryKey allows entries to be partitioned by their key. A key
 	 * can be created by default constructor, int32, FString or FName.
 	 */
 	struct METASOUNDFRONTEND_API FFrontendQueryKey
@@ -45,7 +46,7 @@ namespace Metasound
 	 */
 	struct METASOUNDFRONTEND_API FFrontendQueryEntry
 	{
-		using FValue = TVariant<const Frontend::IRegistryTransaction*, FMetasoundFrontendDocument, FMetasoundFrontendClass, FMetasoundFrontendNode>;
+		using FValue = TVariant<Frontend::FNodeRegistryTransaction, FMetasoundFrontendVersion, FMetasoundFrontendClass, Frontend::FArchetypeRegistryTransaction, FMetasoundFrontendArchetype>;
 		using FKey = FFrontendQueryKey;
 
 		FKey Key;
diff --git a/Engine/Plugins/Runtime/Metasound/Source/MetasoundFrontend/Public/MetasoundFrontendRegistries.h b/Engine/Plugins/Runtime/Metasound/Source/MetasoundFrontend/Public/MetasoundFrontendRegistries.h
--- a/Engine/Plugins/Runtime/Metasound/Source/MetasoundFrontend/Public/MetasoundFrontendRegistries.h
+++ b/Engine/Plugins/Runtime/Metasound/Source/MetasoundFrontend/Public/MetasoundFrontendRegistries.h
@@ -10,6 +10,7 @@
 #include "MetasoundNodeInterface.h"
 #include "MetasoundOperatorInterface.h"
 #include "MetasoundRouter.h"
+#include "Templates/Function.h"
 
 
 namespace Metasound
@@ -328,39 +329,31 @@ namespace Metasound
 			TArray<FConverterNodeInfo> PotentialConverterNodes;
 		};
 
-
 		using FRegistryTransactionID = int32;
 
-		/** Describes the type of transaction. */
-		enum class ETransactionType : uint8
-		{
-			// TODO: Rename to be more explicit (eg. NodeRegistration, NodeUnregistration)
-			Add,     //< Something was added to the registry.
-			Remove,  //< Something was removed from the registry.
-			Invalid
-		};
-
-		/** Interface for a registry transaction. */
-		class IRegistryTransaction
+		class METASOUNDFRONTEND_API FNodeRegistryTransaction 
 		{
 		public:
-			virtual ~IRegistryTransaction() = default;
-			virtual TUniquePtr<IRegistryTransaction> Clone() const = 0;
+			/** Describes the type of transaction. */
+			enum class ETransactionType : uint8
+			{
+				NodeRegistration,     //< Something was added to the registry.
+				NodeUnregistration,  //< Something was removed from the registry.
+				Invalid
+			};
 
-			/** Returns the type of transaction */
-			virtual ETransactionType GetTransactionType() const = 0;
+			FNodeRegistryTransaction(ETransactionType InType, const FNodeRegistryKey& InKey, const FNodeClassInfo& InNodeClassInfo);
 
-			/** If a FNodeClassInfo added during the transaction, this will return
-			 * a non-null pointer to the FNodeClassInfo. */
-			virtual const FNodeClassInfo* GetNodeClassInfo() const = 0;
-			
-			/** If a node was registered or unregistered during this transaction,
-			 * this returns a valid FNodeRegistryKey. */
-			virtual const FNodeRegistryKey* GetNodeRegistryKey() const = 0;
-		};
+			ETransactionType GetTransactionType() const;
+			const FNodeClassInfo& GetNodeClassInfo() const;
+			const FNodeRegistryKey& GetNodeRegistryKey() const;
 
-		using FRegistryTransactionPtr = TUniquePtr<IRegistryTransaction>;
+		private:
 
+			ETransactionType Type;
+			FNodeRegistryKey Key;
+			FNodeClassInfo NodeClassInfo;
+		};
 	} // namespace Frontend
 } // namespace Metasound
 
@@ -409,8 +402,8 @@ public:
 	// This is called on module startup. This invokes any registration commands enqueued by our registration macros.
 	virtual void RegisterPendingNodes() = 0;
 
-	/** Return registry transactions since a given transaction id. */
-	virtual TArray<const Metasound::Frontend::IRegistryTransaction*> GetRegistryTransactionsSince(Metasound::Frontend::FRegistryTransactionID InSince, Metasound::Frontend::FRegistryTransactionID* OutCurrentRegistryTransactionID) const = 0;
+	/** Perform function for each registry transaction since a given transaction ID. */
+	virtual void ForEachNodeRegistryTransactionSince(Metasound::Frontend::FRegistryTransactionID InSince, Metasound::Frontend::FRegistryTransactionID* OutCurrentRegistryTransactionID, TFunctionRef<void(const Metasound::Frontend::FNodeRegistryTransaction&)> InFunc) const = 0;
 
 	/** Register a data type
 	 * @param InName - Name of data type.
diff --git a/Engine/Plugins/Runtime/Metasound/Source/MetasoundFrontend/Public/MetasoundFrontendSearchEngine.h b/Engine/Plugins/Runtime/Metasound/Source/MetasoundFrontend/Public/MetasoundFrontendSearchEngine.h
--- a/Engine/Plugins/Runtime/Metasound/Source/MetasoundFrontend/Public/MetasoundFrontendSearchEngine.h
+++ b/Engine/Plugins/Runtime/Metasound/Source/MetasoundFrontend/Public/MetasoundFrontendSearchEngine.h
@@ -2,10 +2,10 @@
 
 #pragma once
 
+#include "MetasoundFrontendArchetypeRegistry.h"
 #include "MetasoundFrontendDocument.h"
 #include "MetasoundNodeInterface.h"
 
-
 namespace Metasound
 {
 	namespace Frontend
@@ -36,6 +36,14 @@ namespace Metasound
 				/** Find the class with the given ClassName & Major Version. Returns false if not found, true if found. */
 				virtual bool FindClassWithMajorVersion(const FNodeClassName& InName, int32 InMajorVersion, FMetasoundFrontendClass& OutClass) = 0;
 
+				virtual TArray<FMetasoundFrontendArchetype> FindAllArchetypes(bool bInIncludeDeprecated=false) = 0;
+
+				virtual TArray<FMetasoundFrontendVersion> FindAllRegisteredArchetypesWithName(const FName& InArchetypeName) = 0;
+
+				virtual bool FindArchetypeWithHighestVersion(const FName& InArchetypeName, FMetasoundFrontendArchetype& OutArchetype) = 0;
+
+				virtual bool FindArchetypeWithMajorVersion(const FName& InArchetypeName, int32 InMajorVersion, FMetasoundFrontendArchetype& OutArchetype) = 0;
+
 				virtual ~ISearchEngine() = default;
 
 			protected:
diff --git a/Engine/Plugins/Runtime/Metasound/Source/MetasoundFrontend/Public/MetasoundFrontendTransform.h b/Engine/Plugins/Runtime/Metasound/Source/MetasoundFrontend/Public/MetasoundFrontendTransform.h
--- a/Engine/Plugins/Runtime/Metasound/Source/MetasoundFrontend/Public/MetasoundFrontendTransform.h
+++ b/Engine/Plugins/Runtime/Metasound/Source/MetasoundFrontend/Public/MetasoundFrontendTransform.h
@@ -1,6 +1,7 @@
 // Copyright Epic Games, Inc. All Rights Reserved.
 #pragma once
 
+#include "MetasoundFrontendArchetypeRegistry.h"
 #include "MetasoundFrontendController.h"
 #include "MetasoundFrontendDocument.h"
 
@@ -12,24 +13,57 @@ namespace Metasound
 		/** Interface for transforms applied to documents. */
 		class IDocumentTransform
 		{
-			public:
-				virtual ~IDocumentTransform() = default;
-				virtual bool Transform(FDocumentHandle InDocument) const = 0;
+		public:
+			virtual ~IDocumentTransform() = default;
+			/** Return true if InDocument was modified, false otherwise. */
+			virtual bool Transform(FDocumentHandle InDocument) const = 0;
+		};
+
+		/** Interface for transforms applied to a graph. */
+		class IGraphTransform
+		{
+		public:
+			virtual ~IGraphTransform() = default;
+			/** Return true if InGraph was modified, false otherwise. */
+			virtual bool Transform(FGraphHandle InGraph) const = 0;
+		};
+
+		/** Adds or swaps graph inputs and outputs to match the archetype. */
+		class METASOUNDFRONTEND_API FSwapGraphArchetype : public IGraphTransform
+		{
+		public:
+			FSwapGraphArchetype(const FMetasoundFrontendArchetype& InFromArchetype, const FMetasoundFrontendArchetype& InToArchetype);
+
+			bool Transform(FGraphHandle InGraph) const override;
+
+		private:
+			using FVertexPair = TTuple<FMetasoundFrontendClassVertex, FMetasoundFrontendClassVertex>;
+			TArray<FVertexPair> PairedInputs;
+			TArray<FVertexPair> PairedOutputs;
+			TArray<FMetasoundFrontendClassVertex> InputsToAdd;
+			TArray<FMetasoundFrontendClassVertex> InputsToRemove;
+			TArray<FMetasoundFrontendClassVertex> OutputsToAdd;
+			TArray<FMetasoundFrontendClassVertex> OutputsToRemove;
+
 		};
 
 		/** Adds or swaps root graph inputs and outputs to match the document's archetype. */
 		class METASOUNDFRONTEND_API FMatchRootGraphToArchetype : public IDocumentTransform
 		{
-			public:
-				FMatchRootGraphToArchetype(const FMetasoundFrontendArchetype& InArchetype)
-					: Archetype(InArchetype)
-				{
-				}
+		public:
+			FMatchRootGraphToArchetype(const FMetasoundFrontendVersion& InArchetypeVersion)
+				: ArchetypeVersion(InArchetypeVersion)
+			{
+			}
+
+			bool Transform(FDocumentHandle InDocument) const override;
 
-				bool Transform(FDocumentHandle InDocument) const override;
+		private:
+			void GetUpgradePathForDocument(const FMetasoundFrontendVersion& InCurrentVersion, const FMetasoundFrontendVersion& InTargetVersion, TArray<const IArchetypeRegistryEntry*>& OutUpgradePath) const;
+			bool UpgradeDocumentArchetype(const TArray<const IArchetypeRegistryEntry*>& InUpgradePath, FDocumentHandle InDocument) const;
+			bool ConformDocumentToArchetype(const FMetasoundFrontendArchetype& InTargetArchetype, FDocumentHandle InDocument) const;
 
-			private:
-				FMetasoundFrontendArchetype Archetype;
+			FMetasoundFrontendVersion ArchetypeVersion;
 		};
 
 		/** Base class for versioning a document. */
@@ -38,15 +72,15 @@ namespace Metasound
 			const FName Name;
 			const FString& Path;
 
-			public:
-				static FMetasoundFrontendVersionNumber GetMaxVersion()
-				{
-					return FMetasoundFrontendVersionNumber { 1, 3 };
-				}
+		public:
+			static FMetasoundFrontendVersionNumber GetMaxVersion()
+			{
+				return FMetasoundFrontendVersionNumber { 1, 4 };
+			}
 
-				FVersionDocument(FName InName, const FString& InPath);
+			FVersionDocument(FName InName, const FString& InPath);
 
-				bool Transform(FDocumentHandle InDocument) const override;
+			bool Transform(FDocumentHandle InDocument) const override;
 		};
 	}
 }
diff --git a/Engine/Plugins/Runtime/Metasound/Source/MetasoundGenerator/Private/MetasoundGenerator.cpp b/Engine/Plugins/Runtime/Metasound/Source/MetasoundGenerator/Private/MetasoundGenerator.cpp
--- a/Engine/Plugins/Runtime/Metasound/Source/MetasoundGenerator/Private/MetasoundGenerator.cpp
+++ b/Engine/Plugins/Runtime/Metasound/Source/MetasoundGenerator/Private/MetasoundGenerator.cpp
@@ -50,26 +50,18 @@ namespace Metasound
 			FDataReferenceCollection Outputs = GraphOperator->GetOutputs();
 
 			// Get output audio buffers.
-			if (InitParams.NumOutputChannels == 2)
+			for (const FString& AudioOutputName : InitParams.AudioOutputNames)
 			{
-				if (!Outputs.ContainsDataReadReference<FStereoAudioFormat>(InitParams.OutputName))
+				if (!Outputs.ContainsDataReadReference<FAudioBuffer>(AudioOutputName))
 				{
-					UE_LOG(LogMetaSound, Warning, TEXT("MetasoundSource [%s] does not contain stereo output [%s] in output"), *InitParams.MetaSoundName, *InitParams.OutputName);
+					UE_LOG(LogMetaSound, Warning, TEXT("MetasoundSource [%s] does not contain audio output [%s] in output"), *InitParams.MetaSoundName, *AudioOutputName);
 				}
-				OutputBuffers = Outputs.GetDataReadReferenceOrConstruct<FStereoAudioFormat>(InitParams.OutputName, InitParams.OperatorSettings)->GetBuffers();
-			}
-			else if (InitParams.NumOutputChannels == 1)
-			{
-				if (!Outputs.ContainsDataReadReference<FMonoAudioFormat>(InitParams.OutputName))
-				{
-					UE_LOG(LogMetaSound, Warning, TEXT("MetasoundSource [%s] does not contain mono output [%s] in output"), *InitParams.MetaSoundName, *InitParams.OutputName);
-				}
-				OutputBuffers = Outputs.GetDataReadReferenceOrConstruct<FMonoAudioFormat>(InitParams.OutputName, InitParams.OperatorSettings)->GetBuffers();
+				OutputBuffers.Add(Outputs.GetDataReadReferenceOrConstruct<FAudioBuffer>(AudioOutputName, InitParams.OperatorSettings));
 			}
 
 			// References must be cached before moving the operator to the InitParams
 			FDataReferenceCollection Inputs = GraphOperator->GetInputs();
-			FTriggerWriteRef PlayTrigger = Inputs.GetDataWriteReferenceOrConstruct<FTrigger>(InitParams.InputName, InitParams.OperatorSettings, false);
+			FTriggerWriteRef PlayTrigger = Inputs.GetDataWriteReferenceOrConstruct<FTrigger>(InitParams.OnPlayInputName, InitParams.OperatorSettings, false);
 			FTriggerReadRef FinishTrigger = Outputs.GetDataReadReferenceOrConstruct<FTrigger>(InitParams.IsFinishedOutputName, InitParams.OperatorSettings, false);
 
 			FMetasoundGeneratorData GeneratorData
@@ -95,7 +87,7 @@ namespace Metasound
 		, bPendingGraphTrigger(true)
 		, bIsNewGraphPending(true)
 	{
-		NumChannels = InParams.NumOutputChannels;
+		NumChannels = InParams.AudioOutputNames.Num();
 		NumFramesPerExecute = InParams.OperatorSettings.GetNumFramesPerBlock();
 		NumSamplesPerExecute = NumChannels * NumFramesPerExecute;
 
diff --git a/Engine/Plugins/Runtime/Metasound/Source/MetasoundGenerator/Public/MetasoundGenerator.h b/Engine/Plugins/Runtime/Metasound/Source/MetasoundGenerator/Public/MetasoundGenerator.h
--- a/Engine/Plugins/Runtime/Metasound/Source/MetasoundGenerator/Public/MetasoundGenerator.h
+++ b/Engine/Plugins/Runtime/Metasound/Source/MetasoundGenerator/Public/MetasoundGenerator.h
@@ -21,10 +21,9 @@ namespace Metasound
 		FOperatorSettings OperatorSettings;
 		FMetasoundFrontendDocument DocumentCopy;
 		FMetasoundEnvironment Environment;
-		int32 NumOutputChannels = 0;
 		FString MetaSoundName;
-		FString OutputName;
-		FString InputName;
+		TArray<FString> AudioOutputNames;
+		FString OnPlayInputName;
 		FString IsFinishedOutputName;
 	};
 
