Fix MetaSound initial node placement

diff --git a/Engine/Plugins/Runtime/Metasound/Source/MetasoundEditor/Private/MetasoundEditorGraphBuilder.cpp b/Engine/Plugins/Runtime/Metasound/Source/MetasoundEditor/Private/MetasoundEditorGraphBuilder.cpp
--- a/Engine/Plugins/Runtime/Metasound/Source/MetasoundEditor/Private/MetasoundEditorGraphBuilder.cpp
+++ b/Engine/Plugins/Runtime/Metasound/Source/MetasoundEditor/Private/MetasoundEditorGraphBuilder.cpp
@@ -821,7 +821,9 @@ namespace Metasound
 					ExternalNodeLocation += NodeLayout::BufferY;
 				}
 				FMetasoundFrontendNodeStyle Style = NodeHandle->GetNodeStyle();
-				Style.Display.Locations.Add(FGuid::NewGuid(), NewLocation);
+				// TODO: Find consistent location for controlling node locations.
+				// Currently it is split between MetasoundEditor and MetasoundFrontend modules.
+				Style.Display.Locations = {{FGuid::NewGuid(), NewLocation}};
 				NodeHandle->SetNodeStyle(Style);
 			}
 
