Fix for ensure when attempting to mark package dirty when building generator/updating via AssetSubsystem

diff --git a/Engine/Plugins/Runtime/Metasound/Source/MetasoundEngine/Private/MetasoundSource.cpp b/Engine/Plugins/Runtime/Metasound/Source/MetasoundEngine/Private/MetasoundSource.cpp
--- a/Engine/Plugins/Runtime/Metasound/Source/MetasoundEngine/Private/MetasoundSource.cpp
+++ b/Engine/Plugins/Runtime/Metasound/Source/MetasoundEngine/Private/MetasoundSource.cpp
@@ -223,7 +223,7 @@ ISoundGeneratorPtr UMetaSoundSource::CreateSoundGenerator(const FSoundGeneratorI
 	}
 
 	// Check version and attempt to version up if out-of-date
-	if (VersionAsset())
+	if (VersionAsset(false /* bMarkDirty */))
 	{
 		const FMetasoundFrontendDocumentMetadata& Metadata = GetDocumentHandle()->GetMetadata();
 		UE_LOG(LogMetaSound, Display,
diff --git a/Engine/Plugins/Runtime/Metasound/Source/MetasoundEngine/Private/MetasoundUObjectRegistry.cpp b/Engine/Plugins/Runtime/Metasound/Source/MetasoundEngine/Private/MetasoundUObjectRegistry.cpp
--- a/Engine/Plugins/Runtime/Metasound/Source/MetasoundEngine/Private/MetasoundUObjectRegistry.cpp
+++ b/Engine/Plugins/Runtime/Metasound/Source/MetasoundEngine/Private/MetasoundUObjectRegistry.cpp
@@ -159,7 +159,7 @@ void UMetaSoundAssetSubsystem::AddOrUpdateAsset(const FAssetData& InAssetData)
 			// Must version to ensure registration uses the correct key,
 			// which must be based off of most up-to-date document model
 			// for safety.
-			MetaSoundAsset->VersionAsset();
+			MetaSoundAsset->VersionAsset(false /* bMarkDirty */);
 
 			MetaSoundAsset->RegisterGraphWithFrontend();
 		}
diff --git a/Engine/Plugins/Runtime/Metasound/Source/MetasoundFrontend/Private/MetasoundAssetBase.cpp b/Engine/Plugins/Runtime/Metasound/Source/MetasoundFrontend/Private/MetasoundAssetBase.cpp
--- a/Engine/Plugins/Runtime/Metasound/Source/MetasoundFrontend/Private/MetasoundAssetBase.cpp
+++ b/Engine/Plugins/Runtime/Metasound/Source/MetasoundFrontend/Private/MetasoundAssetBase.cpp
@@ -241,7 +241,7 @@ void FMetasoundAssetBase::ConformDocumentToArchetype()
 	}
 }
 
-bool FMetasoundAssetBase::VersionAsset()
+bool FMetasoundAssetBase::VersionAsset(bool bInMarkDirty)
 {
 	using namespace Metasound;
 	using namespace Metasound::Frontend;
@@ -266,7 +266,7 @@ bool FMetasoundAssetBase::VersionAsset()
 
 	const bool bDidUpdate = bDidUpdateDocumentVersion || bDidMatchRootGraphToArchetype;
 
-	if (bDidUpdate)
+	if (bInMarkDirty && bDidUpdate)
 	{
 		MarkMetasoundDocumentDirty();
 	}
diff --git a/Engine/Plugins/Runtime/Metasound/Source/MetasoundFrontend/Public/MetasoundAssetBase.h b/Engine/Plugins/Runtime/Metasound/Source/MetasoundFrontend/Public/MetasoundAssetBase.h
--- a/Engine/Plugins/Runtime/Metasound/Source/MetasoundFrontend/Public/MetasoundAssetBase.h
+++ b/Engine/Plugins/Runtime/Metasound/Source/MetasoundFrontend/Public/MetasoundAssetBase.h
@@ -121,7 +121,7 @@ public:
 	// This must be called on UObject::PostLoad, as well as in this asset's UFactory, to fix up the root document based on the most recent version of the archetype.
 	void ConformDocumentToArchetype();
 
-	bool VersionAsset();
+	bool VersionAsset(bool bInMarkDirty = true);
 
 	// Calls the outermost package and marks it dirty. 
 	bool MarkMetasoundDocumentDirty() const;
