Fix crash on MetaSound Wave writer vertex interface reading invalid data

diff --git a/Engine/Plugins/Runtime/Metasound/Source/MetasoundStandardNodes/Private/MetasoundWaveWriterNode.cpp b/Engine/Plugins/Runtime/Metasound/Source/MetasoundStandardNodes/Private/MetasoundWaveWriterNode.cpp
--- a/Engine/Plugins/Runtime/Metasound/Source/MetasoundStandardNodes/Private/MetasoundWaveWriterNode.cpp
+++ b/Engine/Plugins/Runtime/Metasound/Source/MetasoundStandardNodes/Private/MetasoundWaveWriterNode.cpp
@@ -210,7 +210,7 @@ namespace Metasound
 			return OutputDataReferences;
 		}
 
-		static FVertexInterface DeclareVertexInterface() 
+		static const FVertexInterface& DeclareVertexInterface() 
 		{
 			static const FVertexInterface Interface(
 				FInputVertexInterface(
