Added Trigger Coin node to Metasounds

diff --git a/Engine/Plugins/Runtime/Metasound/Source/MetasoundStandardNodes/Private/MetasoundTriggerCoinNode.cpp b/Engine/Plugins/Runtime/Metasound/Source/MetasoundStandardNodes/Private/MetasoundTriggerCoinNode.cpp
new file mode 100644
--- /dev/null
+++ b/Engine/Plugins/Runtime/Metasound/Source/MetasoundStandardNodes/Private/MetasoundTriggerCoinNode.cpp
@@ -0,0 +1,252 @@
+// Copyright Epic Games, Inc. All Rights Reserved.
+
+#include "CoreMinimal.h"
+#include "MetasoundBuilderInterface.h"
+#include "MetasoundDataReferenceCollection.h"
+#include "MetasoundExecutableOperator.h"
+#include "MetasoundFacade.h"
+#include "MetasoundNode.h"
+#include "MetasoundNodeInterface.h"
+#include "MetasoundNodeRegistrationMacro.h"
+#include "MetasoundOperatorInterface.h"
+#include "MetasoundPrimitives.h"
+#include "MetasoundSampleCounter.h"
+#include "MetasoundStandardNodesNames.h"
+#include "MetasoundStandardNodesCategories.h"
+#include "MetasoundParamHelper.h"
+
+#define LOCTEXT_NAMESPACE "MetasoundStandardNodes_TriggerCoinNode"
+
+namespace Metasound
+{
+
+	namespace TriggerCoinVertexNames
+	{
+		METASOUND_PARAM(InputTrigger, "Trigger", "Input trigger.");
+		METASOUND_PARAM(InputReset, "Reset", "Trigger to reset the random sequence with the supplied seed. Useful to get randomized repetition.");
+		METASOUND_PARAM(InputSeed, "Seed", "The seed value to use for the random node. Set to -1 to use a random seed.");
+		METASOUND_PARAM(InputProbability, "Probability", "Probability that the output trigger will get called. 0.0 is always tails, 1.0 is always heads, and 0.5 is a 50% chance for either.");
+		METASOUND_PARAM(OutputTrueTrigger, "Heads", "The first possible output trigger.");
+		METASOUND_PARAM(OutputFalseTrigger, "Tails", "The second possible output trigger.");
+	}
+
+	class FTriggerCoinOperator : public TExecutableOperator<FTriggerCoinOperator>
+	{
+	public:
+		static const FNodeClassMetadata& GetNodeInfo();
+		static const FVertexInterface& GetVertexInterface();
+		static TUniquePtr<IOperator> CreateOperator(const FCreateOperatorParams& InParams, FBuildErrorArray& OutErrors);
+
+		FTriggerCoinOperator(const FOperatorSettings& InSettings, 
+			const FTriggerReadRef& InTrigger, 
+			const FTriggerReadRef& InTriggerReset, 
+			const FInt32ReadRef& InSeed, 
+			const FFloatReadRef& InProbability);
+
+		virtual FDataReferenceCollection GetInputs() const override;
+
+		virtual FDataReferenceCollection GetOutputs() const override;
+
+		void Execute();
+
+	private:
+		// Generate new random number
+		FTriggerReadRef TriggerInput;
+		// Reset seed
+		FTriggerReadRef TriggerResetInput;
+		// RNG Seed
+		FInt32ReadRef SeedInput;
+		
+		// Probability that output is triggered on TriggerInput
+		FFloatReadRef ProbabilityInput;
+		// Trigger output
+		FTriggerWriteRef TriggerTrueOutput;
+		FTriggerWriteRef TriggerFalseOutput;
+
+		// Random number generator
+		FRandomStream RandomStream;
+		// Internal bools to keep track of seed changes
+		bool bIsDefaultSeeded = false;
+		bool bIsRandomStreamInitialized = false;
+
+		// Handle situations where the seed is changed
+		void EvaluateSeedChanges();
+
+	};
+
+	FTriggerCoinOperator::FTriggerCoinOperator(const FOperatorSettings& InSettings,
+		const FTriggerReadRef& InTrigger,
+		const FTriggerReadRef& InTriggerReset,
+		const FInt32ReadRef& InSeed,
+		const FFloatReadRef& InProbability)
+		: TriggerInput(InTrigger)
+		, TriggerResetInput(InTriggerReset)
+		, SeedInput(InSeed)
+		, ProbabilityInput(InProbability)
+		, TriggerTrueOutput(FTriggerWriteRef::CreateNew(InSettings))
+		, TriggerFalseOutput(FTriggerWriteRef::CreateNew(InSettings))
+		, bIsDefaultSeeded(*SeedInput == -1)
+		, bIsRandomStreamInitialized(false)
+	{
+		EvaluateSeedChanges();
+		RandomStream.Reset();
+	}
+
+	FDataReferenceCollection FTriggerCoinOperator::GetInputs() const
+	{
+		using namespace TriggerCoinVertexNames;
+
+		FDataReferenceCollection InputDataReferences;
+		InputDataReferences.AddDataReadReference(METASOUND_GET_PARAM_NAME(InputTrigger), TriggerInput);
+		InputDataReferences.AddDataReadReference(METASOUND_GET_PARAM_NAME(InputReset), TriggerResetInput);
+		InputDataReferences.AddDataReadReference(METASOUND_GET_PARAM_NAME(InputSeed), SeedInput);
+		InputDataReferences.AddDataReadReference(METASOUND_GET_PARAM_NAME(InputProbability), ProbabilityInput);
+		return InputDataReferences;
+	}
+
+	FDataReferenceCollection FTriggerCoinOperator::GetOutputs() const
+	{
+		using namespace TriggerCoinVertexNames;
+
+		FDataReferenceCollection OutputDataReferences;
+		OutputDataReferences.AddDataReadReference(METASOUND_GET_PARAM_NAME(OutputTrueTrigger), TriggerTrueOutput);
+		OutputDataReferences.AddDataReadReference(METASOUND_GET_PARAM_NAME(OutputFalseTrigger), TriggerFalseOutput);
+
+		return OutputDataReferences;
+	}
+
+	void FTriggerCoinOperator::Execute()
+	{
+		TriggerTrueOutput->AdvanceBlock();
+		TriggerFalseOutput->AdvanceBlock();
+		
+		// Reset the seed
+		TriggerResetInput->ExecuteBlock(
+			[](int32, int32)
+			{
+			},
+			[this](int32 StartFrame, int32 EndFrame)
+			{
+				EvaluateSeedChanges();
+				RandomStream.Reset();
+			}
+		);
+
+		// Generate new random number
+		TriggerInput->ExecuteBlock(
+			[](int32, int32) 
+			{
+			},
+			[this](int32 StartFrame, int32 EndFrame)
+			{
+				// Use min float so that 0 never triggers output
+				const float val = RandomStream.FRandRange(TNumericLimits<float>::Min(), 1.0f);
+
+				// If number < probability threshold, send output trigger.
+				if (val <= *ProbabilityInput)
+				{
+					TriggerTrueOutput->TriggerFrame(StartFrame);
+				}
+				else
+				{
+					TriggerFalseOutput->TriggerFrame(StartFrame);
+				}
+			}
+		);		
+
+	}
+
+
+	TUniquePtr<IOperator> FTriggerCoinOperator::CreateOperator(const FCreateOperatorParams& InParams, FBuildErrorArray& OutErrors)
+	{
+		using namespace TriggerCoinVertexNames;
+
+		const FInputVertexInterface& InputInterface = GetVertexInterface().GetInputInterface();
+
+		FTriggerReadRef TriggerIn = InParams.InputDataReferences.GetDataReadReferenceOrConstruct<FTrigger>(METASOUND_GET_PARAM_NAME(InputTrigger), InParams.OperatorSettings);
+		FTriggerReadRef TriggerResetIn = InParams.InputDataReferences.GetDataReadReferenceOrConstruct<FTrigger>(METASOUND_GET_PARAM_NAME(InputReset), InParams.OperatorSettings);
+		FInt32ReadRef SeedIn = InParams.InputDataReferences.GetDataReadReferenceOrConstructWithVertexDefault<int32>(InputInterface, METASOUND_GET_PARAM_NAME(InputSeed), InParams.OperatorSettings);
+		FFloatReadRef ProbabilityIn = InParams.InputDataReferences.GetDataReadReferenceOrConstructWithVertexDefault<float>(InputInterface, METASOUND_GET_PARAM_NAME(InputProbability), InParams.OperatorSettings);
+
+		return MakeUnique<FTriggerCoinOperator>(InParams.OperatorSettings, TriggerIn, TriggerResetIn, SeedIn, ProbabilityIn);
+	}
+
+	const FVertexInterface& FTriggerCoinOperator::GetVertexInterface()
+	{
+		using namespace TriggerCoinVertexNames;
+
+		static const FVertexInterface Interface(
+			FInputVertexInterface(
+				TInputDataVertexModel<FTrigger>(METASOUND_GET_PARAM_NAME_AND_TT(InputTrigger)),
+				TInputDataVertexModel<FTrigger>(METASOUND_GET_PARAM_NAME_AND_TT(InputReset)),
+				TInputDataVertexModel<int32>(METASOUND_GET_PARAM_NAME_AND_TT(InputSeed), -1),
+				TInputDataVertexModel<float>(METASOUND_GET_PARAM_NAME_AND_TT(InputProbability), 1.0f)
+			),
+			FOutputVertexInterface(
+				TOutputDataVertexModel<FTrigger>(METASOUND_GET_PARAM_NAME_AND_TT(OutputTrueTrigger)),
+				TOutputDataVertexModel<FTrigger>(METASOUND_GET_PARAM_NAME_AND_TT(OutputFalseTrigger))
+			)
+		);
+
+		return Interface;
+	}
+
+	const FNodeClassMetadata& FTriggerCoinOperator::GetNodeInfo()
+	{
+		auto InitNodeInfo = []() -> FNodeClassMetadata
+		{
+			FNodeClassMetadata Info;
+			Info.ClassName = { Metasound::StandardNodes::Namespace, TEXT("Trigger Filter"), TEXT("") };
+			Info.MajorVersion = 1;
+			Info.MinorVersion = 0;
+			Info.DisplayName = LOCTEXT("Metasound_TriggerCoinNodeDisplayName", "Trigger Filter");
+			Info.Description = LOCTEXT("Metasound_TriggerCoinNodeDescription", "When triggered, randomly sends one of two output triggers based on a given probability.");
+			Info.Author = PluginAuthor;
+			Info.PromptIfMissing = PluginNodeMissingPrompt;
+			Info.DefaultInterface = GetVertexInterface();
+			Info.CategoryHierarchy.Emplace(StandardNodes::TriggerUtils);
+
+			return Info;
+		};
+
+		static const FNodeClassMetadata Info = InitNodeInfo();
+
+		return Info;
+	}
+
+	void FTriggerCoinOperator::EvaluateSeedChanges()
+	{
+		// if we have a non-random seed
+		if (*SeedInput != -1)
+		{
+			// If we were previously randomly-seeded OR our seed has changed
+			if (bIsDefaultSeeded || !bIsRandomStreamInitialized || *SeedInput != RandomStream.GetInitialSeed())
+			{
+				bIsRandomStreamInitialized = true;
+				bIsDefaultSeeded = false;
+				RandomStream.Initialize(*SeedInput);
+			}
+		}
+		// If we are randomly-seeded now BUT were previously not, we need to randomize our seed
+		else if (!bIsDefaultSeeded || !bIsRandomStreamInitialized)
+		{
+			bIsRandomStreamInitialized = true;
+			bIsDefaultSeeded = true;
+			RandomStream.Initialize(FPlatformTime::Cycles());
+		}
+	}
+
+	// Node Class
+	class FTriggerCoinNode : public FNodeFacade
+	{
+	public:
+		FTriggerCoinNode(const FNodeInitData& InitData)
+			: FNodeFacade(InitData.InstanceName, InitData.InstanceID, TFacadeOperatorClass<FTriggerCoinOperator>())
+		{
+		}
+	};
+
+	METASOUND_REGISTER_NODE(FTriggerCoinNode)
+}
+
+#undef LOCTEXT_NAMESPACE
