Additive synth Metasound node using new wavetable DSP object

diff --git a/Engine/Plugins/Runtime/Metasound/Source/MetasoundStandardNodes/Private/MetasoundAdditiveSynthNode.cpp b/Engine/Plugins/Runtime/Metasound/Source/MetasoundStandardNodes/Private/MetasoundAdditiveSynthNode.cpp
new file mode 100644
--- /dev/null
+++ b/Engine/Plugins/Runtime/Metasound/Source/MetasoundStandardNodes/Private/MetasoundAdditiveSynthNode.cpp
@@ -0,0 +1,310 @@
+// Copyright Epic Games, Inc. All Rights Reserved.
+
+#include "Internationalization/Text.h"
+#include "MetasoundFacade.h"
+#include "MetasoundExecutableOperator.h"
+#include "MetasoundNodeRegistrationMacro.h"
+#include "MetasoundPrimitives.h"
+#include "MetasoundStandardNodesNames.h"
+#include "MetasoundStandardNodesCategories.h"
+#include "MetasoundParamHelper.h"
+#include "DSP/BufferVectorOperations.h"
+#include "DSP/SineWaveTableOsc.h"
+
+#define LOCTEXT_NAMESPACE "MetasoundAdditiveSynthNode"
+
+namespace Metasound
+{
+	namespace AdditiveSynthVertexNames
+	{
+		METASOUND_PARAM(BaseFrequency, "Base Frequency", "Sinusoid frequency that harmonics are based on.");
+		METASOUND_PARAM(HarmonicMultipliers, "HarmonicMultipliers", "Array of float harmonic multipliers that will be applied to the base frequency. Number of sinusoids used depends on the size of this array.");
+		METASOUND_PARAM(Amplitudes, "Amplitudes", "Array of sinusoid amplitudes, clamped to 0.0 - 1.0");
+		METASOUND_PARAM(Phases, "Phases", "Array of sinusoid phases, clamped to 0.0 - 1.0.");
+		METASOUND_PARAM(PanAmounts, "Pan Amounts", "Array of pan amounts (using equal power law; -1.0 is full left, 1.0 is full right).");
+
+		METASOUND_PARAM(LeftAudioOut, "Out Left Audio", "Left output synthesized audio.");
+		METASOUND_PARAM(RightAudioOut, "Out Right Audio", "Right output synthesized audio.");
+	}
+
+	class FAdditiveSynthOperator : public TExecutableOperator<FAdditiveSynthOperator>
+	{
+	public:
+		using FInputFloatArrayReadRef = TDataReadReference<TArray<float>>;
+		using FInputFloatArrayType = TArray<float>;
+
+		static const FNodeClassMetadata& GetNodeInfo()
+		{
+			auto InitNodeInfo = []() -> FNodeClassMetadata
+			{
+				const FName OperatorName = TEXT("Additive Synth");
+				const FText NodeDisplayName = LOCTEXT("Metasound_AdditiveSynthNodeDisplayName", "Additive Synth");
+				const FText NodeDescription = LOCTEXT("Metasound_AdditiveSynthNodeDescription", "Synthesizes audio output given input array of sinusoids.");
+
+				FNodeClassMetadata Info;
+				Info.ClassName = { Metasound::StandardNodes::Namespace, OperatorName, TEXT("") };
+				Info.MajorVersion = 1;
+				Info.MinorVersion = 0;
+				Info.DisplayName = NodeDisplayName;
+				Info.Description = NodeDescription;
+				Info.Author = PluginAuthor;
+				Info.PromptIfMissing = PluginNodeMissingPrompt;
+				Info.DefaultInterface = GetVertexInterface();
+
+				return Info;
+			};
+
+			static const FNodeClassMetadata Info = InitNodeInfo();
+			return Info;
+		}
+
+		static const FVertexInterface& GetVertexInterface()
+		{
+			using namespace AdditiveSynthVertexNames;
+
+			static const FVertexInterface Interface(
+				FInputVertexInterface(
+					TInputDataVertexModel<float>(METASOUND_GET_PARAM_NAME_AND_TT(BaseFrequency), 440.0f),
+					TInputDataVertexModel<FInputFloatArrayType>(METASOUND_GET_PARAM_NAME_AND_TT(HarmonicMultipliers)),
+					TInputDataVertexModel<FInputFloatArrayType>(METASOUND_GET_PARAM_NAME_AND_TT(Amplitudes)),
+					TInputDataVertexModel<FInputFloatArrayType>(METASOUND_GET_PARAM_NAME_AND_TT(Phases)),
+					TInputDataVertexModel<FInputFloatArrayType>(METASOUND_GET_PARAM_NAME_AND_TT(PanAmounts))
+				),
+				FOutputVertexInterface(
+					TOutputDataVertexModel<FAudioBuffer>(METASOUND_GET_PARAM_NAME_AND_TT(LeftAudioOut)),
+					TOutputDataVertexModel<FAudioBuffer>(METASOUND_GET_PARAM_NAME_AND_TT(RightAudioOut))
+				)
+			);
+
+			return Interface;
+		}
+
+		static TUniquePtr<IOperator> CreateOperator(const FCreateOperatorParams& InParams, FBuildErrorArray& OutErrors)
+		{
+			using namespace AdditiveSynthVertexNames;
+
+			const FDataReferenceCollection& InputCollection = InParams.InputDataReferences;
+			const FInputVertexInterface& InputInterface = GetVertexInterface().GetInputInterface();
+
+			FFloatReadRef InputBaseFrequency = InputCollection.GetDataReadReferenceOrConstruct<float>(METASOUND_GET_PARAM_NAME(BaseFrequency));
+			FInputFloatArrayReadRef InputHarmonicMultipliers = InputCollection.GetDataReadReferenceOrConstruct<FInputFloatArrayType>(METASOUND_GET_PARAM_NAME(HarmonicMultipliers));
+			FInputFloatArrayReadRef InputAmplitudes = InputCollection.GetDataReadReferenceOrConstruct<FInputFloatArrayType>(METASOUND_GET_PARAM_NAME(Amplitudes));
+			FInputFloatArrayReadRef InputPhases = InputCollection.GetDataReadReferenceOrConstruct<FInputFloatArrayType>(METASOUND_GET_PARAM_NAME(Phases));
+			FInputFloatArrayReadRef InputPanAmounts = InputCollection.GetDataReadReferenceOrConstruct<FInputFloatArrayType>(METASOUND_GET_PARAM_NAME(PanAmounts));
+
+			return MakeUnique<FAdditiveSynthOperator>(InParams.OperatorSettings, InputBaseFrequency, InputHarmonicMultipliers, InputAmplitudes, InputPhases, InputPanAmounts);
+		}
+
+		FAdditiveSynthOperator(const FOperatorSettings& InSettings,
+			const FFloatReadRef& InputBaseFrequency,
+			const FInputFloatArrayReadRef& InputHarmonicMultipliers,
+			const FInputFloatArrayReadRef& InputAmplitudes,
+			const FInputFloatArrayReadRef& InputPhases,
+			const FInputFloatArrayReadRef& InputPanAmounts)
+			: BaseFrequency(InputBaseFrequency)
+			, HarmonicMultipliers(InputHarmonicMultipliers)
+			, Amplitudes(InputAmplitudes)
+			, Phases(InputPhases)
+			, PanAmounts(InputPanAmounts)
+			, LeftAudioOut(FAudioBufferWriteRef::CreateNew(InSettings))
+			, RightAudioOut(FAudioBufferWriteRef::CreateNew(InSettings))
+			, SampleRate(InSettings.GetSampleRate())
+		{
+			CreateSinusoids();
+		}
+
+		FDataReferenceCollection GetInputs() const
+		{
+			using namespace AdditiveSynthVertexNames;
+
+			FDataReferenceCollection InputDataReferences;
+			InputDataReferences.AddDataReadReference(METASOUND_GET_PARAM_NAME(BaseFrequency), BaseFrequency);
+			InputDataReferences.AddDataReadReference(METASOUND_GET_PARAM_NAME(HarmonicMultipliers), HarmonicMultipliers);
+			InputDataReferences.AddDataReadReference(METASOUND_GET_PARAM_NAME(Amplitudes), Amplitudes);
+			InputDataReferences.AddDataReadReference(METASOUND_GET_PARAM_NAME(Phases), Phases);
+			InputDataReferences.AddDataReadReference(METASOUND_GET_PARAM_NAME(PanAmounts), PanAmounts);
+			return InputDataReferences;
+		}
+
+		virtual FDataReferenceCollection GetOutputs() const
+		{
+			using namespace AdditiveSynthVertexNames;
+
+			FDataReferenceCollection OutputDataReferences;
+			OutputDataReferences.AddDataReadReference(METASOUND_GET_PARAM_NAME(LeftAudioOut), LeftAudioOut);
+			OutputDataReferences.AddDataReadReference(METASOUND_GET_PARAM_NAME(RightAudioOut), RightAudioOut);
+			return OutputDataReferences;
+		}
+
+		void Execute()
+		{
+			float Frequency = *BaseFrequency;
+			int32 NumSamples = LeftAudioOut->Num();
+
+			// Zero out output buffers since we'll sum directly into them
+			FMemory::Memzero(LeftAudioOut->GetData(), sizeof(float) * NumSamples);
+			FMemory::Memzero(RightAudioOut->GetData(), sizeof(float) * NumSamples);
+			// Clear temporary buffers
+			if (ScratchBuffer.Num() != NumSamples)
+			{
+				ScratchBuffer.Reset();
+				ScratchBuffer.Init(0, NumSamples);
+			}
+
+			int32 NumSinusoids = HarmonicMultipliers->Num();
+			
+			// Number of elements in parallel arrays doesn't match the current number of oscillators
+			if (NumSinusoids != SinusoidData.Num())
+			{
+				CreateSinusoids();
+			}
+
+			bool UpdateFrequency = false;
+			// Flag to update frequency (which depends on node base frequency and indiviudal harmonic multiplier) if necessary
+			if (!FMath::IsNearlyEqual(*BaseFrequency, PreviousBaseFrequency))
+			{
+				PreviousBaseFrequency = *BaseFrequency;
+				UpdateFrequency = true;
+			}
+
+			// For each sinusoid, update parameters and generate audio
+			for (int32 SinusoidIndex = 0; SinusoidIndex < NumSinusoids; ++SinusoidIndex)
+			{
+				Audio::FSineWaveTableOsc& WaveTableOsc = SinusoidData[SinusoidIndex];
+
+				// Update harmonic multiplier and resulting frequency if necessary
+				const float HarmonicMultiplier = (*HarmonicMultipliers)[SinusoidIndex];
+				if (!FMath::IsNearlyEqual(HarmonicMultiplier, PreviousHarmonicMultipliers[SinusoidIndex]))
+				{
+					PreviousHarmonicMultipliers[SinusoidIndex] = HarmonicMultiplier;
+					UpdateFrequency = true;
+				}
+				
+				if (UpdateFrequency)
+				{
+					WaveTableOsc.SetFrequencyHz((*BaseFrequency) * HarmonicMultiplier);
+				}
+
+				// Update phase if necessary, set to 0 if no phases are given
+				// If the number of elements in Phases doesn't equal the number of sinusoids, 
+				// modulus on the Phases array (i.e. if 2 elements, phases will alternate)
+				float Phase = 0.0f;
+				if (Phases->Num() == 0)
+				{
+					if (!FMath::IsNearlyEqual(Phase, PreviousPhases[SinusoidIndex]))
+					{
+						PreviousPhases[SinusoidIndex] = Phase;
+						WaveTableOsc.SetPhase(Phase);
+					}
+				}
+				else
+				{
+					int32 PhaseModIndex = SinusoidIndex % Phases->Num();
+					Phase = FMath::Clamp((*Phases)[PhaseModIndex], 0.0f, 1.0f);
+					if (!FMath::IsNearlyEqual(Phase, PreviousPhases[PhaseModIndex]))
+					{
+						PreviousPhases[PhaseModIndex] = Phase;
+						WaveTableOsc.SetPhase(Phase);
+					}
+				}
+
+				// Get samples from wavetable
+				WaveTableOsc.Generate(ScratchBuffer.GetData(), NumSamples);
+
+				// Apply pan
+				float LeftPanGain = 1.0f;
+				float RightPanGain = 1.0f;
+				// If PanAmounts is empty, will default to 0 
+				// If the number of elements in PanAmounts doesn't equal the number of sinusoids, 
+				// modulus on the PanAmounts array (i.e. if 2 elements, pan amounts will alternate)
+				if (PanAmounts->Num() > 0)
+				{
+					float Fraction = 0.5f * (FMath::Clamp((*PanAmounts)[SinusoidIndex % PanAmounts->Num()], -1.0f, 1.0f) + 1.0f);
+					// Using equal power panning law 
+					FMath::SinCos(&RightPanGain, &LeftPanGain, 0.5f * PI * Fraction);
+				}
+
+				// Clamp given amplitude
+				float BaseGain = 1.0f;
+				if (Amplitudes->Num() > 0)
+				{
+					 BaseGain = FMath::Clamp((*Amplitudes)[SinusoidIndex % Amplitudes->Num()], 0.0f, 1.0f);
+				}
+
+				// Multiply by gain (both amplitude and pan)
+				float NewLeftGain = LeftPanGain * BaseGain;
+				float NewRightGain = RightPanGain * BaseGain;
+
+				// Mix into output buffer
+				Audio::MixInBufferFast(ScratchBuffer.GetData(), LeftAudioOut->GetData(), NumSamples, NewLeftGain);
+				Audio::MixInBufferFast(ScratchBuffer.GetData(), RightAudioOut->GetData(), NumSamples, NewRightGain);
+			}
+		}
+
+	private:
+		void CreateSinusoids()
+		{
+			int32 NumSinusoids = HarmonicMultipliers->Num();
+			SinusoidData.Reset(NumSinusoids);
+			PreviousHarmonicMultipliers.Reset(NumSinusoids);
+			PreviousPhases.Reset(NumSinusoids);
+
+			// Cache information so we don't have to update the oscillator every time, pad to NumSinusoids 
+			PreviousHarmonicMultipliers = *HarmonicMultipliers;
+			PreviousPhases = *Phases;
+			if (PreviousPhases.Num() < NumSinusoids)
+			{
+				PreviousPhases.AddZeroed(NumSinusoids - PreviousPhases.Num());
+			}
+			PreviousBaseFrequency = *BaseFrequency;
+
+			for (int32 SinusoidIndex = 0; SinusoidIndex < NumSinusoids; ++SinusoidIndex)
+			{
+				// Create wavetable oscillators for each sinusoid
+				Audio::FSineWaveTableOsc WaveTableOsc = Audio::FSineWaveTableOsc();
+				float Phase = Phases->Num() > 0 ? (*Phases)[SinusoidIndex % Phases->Num()] : 0.0f;
+				WaveTableOsc.Init(SampleRate, (*HarmonicMultipliers)[SinusoidIndex] * (*BaseFrequency), Phase);
+				SinusoidData.Add(WaveTableOsc);
+			}
+		}
+
+		TArray<Audio::FSineWaveTableOsc> SinusoidData;
+
+		FFloatReadRef BaseFrequency;
+		FInputFloatArrayReadRef HarmonicMultipliers;
+		FInputFloatArrayReadRef Amplitudes;
+		FInputFloatArrayReadRef Phases;
+		FInputFloatArrayReadRef PanAmounts;
+
+		// Cached copies of sinusoid parameters to avoid unnecessary updates
+		TArray<float> PreviousHarmonicMultipliers;
+		TArray<float> PreviousPhases;
+		float PreviousBaseFrequency = 440.0f;
+
+		// Scratch aligned buffer
+		Audio::FAlignedFloatBuffer ScratchBuffer;
+
+		// Audio output
+		FAudioBufferWriteRef LeftAudioOut;
+		FAudioBufferWriteRef RightAudioOut;
+
+		// Sample rate, will be set in constructor
+		float SampleRate = 48000.0f;
+	};
+
+	class FAdditiveSynthNode : public FNodeFacade
+	{
+	public:
+		/**
+		 * Constructor used by the Metasound Frontend.
+		 */
+		FAdditiveSynthNode(const FNodeInitData& InitData)
+			: FNodeFacade(InitData.InstanceName, InitData.InstanceID, TFacadeOperatorClass<FAdditiveSynthOperator>())
+		{
+		}
+	};
+
+	METASOUND_REGISTER_NODE(FAdditiveSynthNode)
+}
+
+#undef LOCTEXT_NAMESPACE
diff --git a/Engine/Source/Runtime/SignalProcessing/Private/SineWaveTableOsc.cpp b/Engine/Source/Runtime/SignalProcessing/Private/SineWaveTableOsc.cpp
new file mode 100644
--- /dev/null
+++ b/Engine/Source/Runtime/SignalProcessing/Private/SineWaveTableOsc.cpp
@@ -0,0 +1,99 @@
+// Copyright Epic Games, Inc. All Rights Reserved.
+
+#include "DSP/SineWaveTableOsc.h"
+
+namespace Audio
+{
+	static const int32 SINE_WAVE_TABLE_SIZE = 4096;
+
+	FSineWaveTableOsc::FSineWaveTableOsc()
+	{
+		SetFrequencyHz(FrequencyHz);
+	}
+
+	FSineWaveTableOsc::~FSineWaveTableOsc()
+	{
+	}
+
+	void FSineWaveTableOsc::Init(const float InSampleRate, const float InFrequencyHz, const float InPhase)
+	{
+		SampleRate = InSampleRate;
+		FrequencyHz = InFrequencyHz;
+		InitialPhase = FMath::Clamp(InPhase, 0.0f, 1.0f);
+
+		Reset();
+		UpdatePhaseIncrement();
+	}
+
+	void FSineWaveTableOsc::SetSampleRate(const float InSampleRate)
+	{
+		SampleRate = InSampleRate;
+		UpdatePhaseIncrement();
+	}
+
+	void FSineWaveTableOsc::Reset()
+	{
+		ReadIndex = InitialPhase * (WaveTableBuffer.Num() - 1.0f);		
+	}
+
+	void FSineWaveTableOsc::SetFrequencyHz(const float InFrequencyHz)
+	{
+		FrequencyHz = InFrequencyHz;
+		UpdatePhaseIncrement();
+	}
+
+	void FSineWaveTableOsc::SetPhase(const float InPhase)
+	{
+		InitialPhase = FMath::Clamp(InPhase, 0.0f, 1.0f);
+		Reset();
+	}
+
+	void FSineWaveTableOsc::UpdatePhaseIncrement()
+	{
+		PhaseIncrement = (float)WaveTableBuffer.Num() * FrequencyHz / (float)SampleRate;
+	}
+
+	void FSineWaveTableOsc::Generate(float* OutSample)
+	{
+		const int32 ReadIndexPrev = (int32)ReadIndex;
+		const float Alpha = ReadIndex - (float) ReadIndexPrev;
+
+		const int32 ReadIndexNext = (ReadIndexPrev + 1) % WaveTableBuffer.Num();
+		*OutSample = FMath::Lerp(WaveTableBuffer[ReadIndexPrev], WaveTableBuffer[ReadIndexNext], Alpha);
+		
+		// Increment ReadIndex and wrap around if necessary
+		ReadIndex += PhaseIncrement;
+		while (ReadIndex >= WaveTableBuffer.Num())
+		{
+			ReadIndex -= WaveTableBuffer.Num();
+		}
+	}
+
+	void FSineWaveTableOsc::Generate(float* OutBuffer, const int32 NumSamples)
+	{
+		for (int32 SampleIndex = 0; SampleIndex < NumSamples; ++SampleIndex)
+		{
+			Generate(&OutBuffer[SampleIndex]);
+		}
+	}
+
+	const TArray<float>& FSineWaveTableOsc::GetWaveTable()
+	{
+		auto MakeSineTable = []() -> const TArray<float>
+		{
+			// Generate the table
+			TArray<float> WaveTable;
+			WaveTable.AddUninitialized(SINE_WAVE_TABLE_SIZE);
+			float* WaveTableData = WaveTable.GetData();
+			for (int32 i = 0; i < SINE_WAVE_TABLE_SIZE; ++i)
+			{
+				float Phase = (float)i / SINE_WAVE_TABLE_SIZE;
+				WaveTableData[i] = FMath::Sin(Phase * 2.f * PI);
+			}
+			return WaveTable;
+		};
+
+		static const TArray<float> SineWaveTable = MakeSineTable();
+		return SineWaveTable;
+	}
+}
diff --git a/Engine/Source/Runtime/SignalProcessing/Public/DSP/SineWaveTableOsc.h b/Engine/Source/Runtime/SignalProcessing/Public/DSP/SineWaveTableOsc.h
new file mode 100644
--- /dev/null
+++ b/Engine/Source/Runtime/SignalProcessing/Public/DSP/SineWaveTableOsc.h
@@ -0,0 +1,69 @@
+// Copyright Epic Games, Inc. All Rights Reserved.
+
+#pragma once
+
+#include "CoreMinimal.h"
+
+namespace Audio
+{
+	class FWaveTableOsc;
+
+
+	// A sinusoidal wave table oscillator class
+	class SIGNALPROCESSING_API FSineWaveTableOsc
+	{
+	public:
+		// Constructor
+		FSineWaveTableOsc();
+
+		// Virtual Destructor
+		virtual ~FSineWaveTableOsc();
+
+		// Initialize the wave table oscillator
+		void Init(const float InSampleRate, const float InFrequencyHz, const float InPhase);
+
+		// Sets the sample rate of the oscillator.
+		void SetSampleRate(const float InSampleRate);
+
+		// Resets the wave table read indices.
+		void Reset();
+
+		// Sets the frequency of the wave table oscillator.
+		void SetFrequencyHz(const float InFrequencyHz);
+
+		// Returns the frequency of the wave table oscillator.
+		float GetFrequencyHz() const { return FrequencyHz; }
+		
+		// Sets the phase of the wave table oscillator.
+		void SetPhase(const float InPhase);
+
+		// Processes the wave table and outputs a sample 
+		void Generate(float* OutSample);
+		// Processes the wave table and fills a buffer
+		void Generate(float* OutBuffer, const int32 NumSamples);
+
+		// The static sinusoidal wave table 
+		static const TArray<float>& GetWaveTable();
+
+	protected:
+		void UpdatePhaseIncrement();
+
+		// The wave table buffer
+		const TArray<float>& WaveTableBuffer = GetWaveTable();
+		
+		// The frequency of the output (given the sample rate)
+		float FrequencyHz = 440.0f;
+
+		// The sample rate of the oscillator
+		float SampleRate = 48000.0f;
+
+		// Read index
+		float ReadIndex = 0.0f;
+
+		// The phase increment (based on frequency)
+		float PhaseIncrement = 0.0f;
+
+		// The initial phase, will be clamped to range 0.0 - 1.0f
+		float InitialPhase = 0.0f;
+	};
+}
