Fix missing set of MetaSound archetype version

diff --git a/Engine/Plugins/Runtime/Metasound/Source/MetasoundFrontend/Private/MetasoundFrontendTransform.cpp b/Engine/Plugins/Runtime/Metasound/Source/MetasoundFrontend/Private/MetasoundFrontendTransform.cpp
--- a/Engine/Plugins/Runtime/Metasound/Source/MetasoundFrontend/Private/MetasoundFrontendTransform.cpp
+++ b/Engine/Plugins/Runtime/Metasound/Source/MetasoundFrontend/Private/MetasoundFrontendTransform.cpp
@@ -289,6 +289,7 @@ namespace Metasound
 					if (Entry->UpdateRootGraphArchetype(InDocument))
 					{
 						bDidEdit = true;
+						InDocument->SetArchetypeVersion(Entry->GetArchetype().Version);
 					}
 				}
 			}
