IWYU Fix

diff --git a/Engine/Plugins/Runtime/Metasound/Source/MetasoundEngine/Public/MetasoundUObjectRegistry.h b/Engine/Plugins/Runtime/Metasound/Source/MetasoundEngine/Public/MetasoundUObjectRegistry.h
--- a/Engine/Plugins/Runtime/Metasound/Source/MetasoundEngine/Public/MetasoundUObjectRegistry.h
+++ b/Engine/Plugins/Runtime/Metasound/Source/MetasoundEngine/Public/MetasoundUObjectRegistry.h
@@ -2,6 +2,8 @@
 #pragma once
 
 #include "CoreMinimal.h"
+
+#include "Engine/Engine.h"
 #include "MetasoundAssetBase.h"
 #include "Subsystems/EngineSubsystem.h"
 #include "UObject/Object.h"
