MetaSound Frontend Speedups
- Cached array indices for TAccessPtr point to an object in an array
- Reduced number of call sites to TAccessPtr::Get()
- Avoid creating registry keys when possible.

diff --git a/Engine/Plugins/Runtime/Metasound/Source/MetasoundFrontend/Private/MetasoundAssetBase.cpp b/Engine/Plugins/Runtime/Metasound/Source/MetasoundFrontend/Private/MetasoundAssetBase.cpp
--- a/Engine/Plugins/Runtime/Metasound/Source/MetasoundFrontend/Private/MetasoundAssetBase.cpp
+++ b/Engine/Plugins/Runtime/Metasound/Source/MetasoundFrontend/Private/MetasoundAssetBase.cpp
@@ -10,6 +10,7 @@
 #include "MetasoundArchetype.h"
 #include "MetasoundFrontendArchetypeRegistry.h"
 #include "MetasoundFrontendController.h"
+#include "MetasoundFrontendDocument.h"
 #include "MetasoundFrontendGraph.h"
 #include "MetasoundFrontendSearchEngine.h"
 #include "MetasoundFrontendTransform.h"
@@ -109,9 +110,9 @@ void FMetasoundAssetBase::RegisterGraphWithFrontend()
 	};
 
 	FNodeClassInfo ClassInfo = GetAssetClassInfo();
-	FNodeRegistryKey RegistryKey = FMetasoundFrontendRegistryContainer::Get()->GetRegistryKey(ClassInfo);
+	FNodeRegistryKey RegistryKey = NodeRegistryKey::CreateKey(ClassInfo);
 
-	if (IsValidNodeRegistryKey(RegistryKey))
+	if (NodeRegistryKey::IsValid(RegistryKey))
 	{
 		// Unregister prior version if it exists.
 		if (bHasRegistered)
@@ -125,7 +126,7 @@ void FMetasoundAssetBase::RegisterGraphWithFrontend()
 		RegistryKey = FMetasoundFrontendRegistryContainer::Get()->RegisterNode(MakeUnique<FNodeRegistryEntry>(AssetName, *Doc, ClassInfo.AssetPath));
 	}
 
-	if (IsValidNodeRegistryKey(RegistryKey))
+	if (NodeRegistryKey::IsValid(RegistryKey))
 	{
 		bHasRegistered = true;
 
@@ -545,8 +546,8 @@ Metasound::Frontend::FConstGraphHandle FMetasoundAssetBase::GetRootGraphHandle()
 
 bool FMetasoundAssetBase::ImportFromJSON(const FString& InJSON)
 {
-	Metasound::Frontend::FDocumentAccessPtr Document = GetDocument();
-	if (ensure(Document.IsValid()))
+	FMetasoundFrontendDocument* Document = GetDocument().Get();
+	if (ensure(nullptr != Document))
 	{
 		bool bSuccess = Metasound::Frontend::ImportJSONToMetasound(InJSON, *Document);
 
@@ -562,8 +563,8 @@ bool FMetasoundAssetBase::ImportFromJSON(const FString& InJSON)
 
 bool FMetasoundAssetBase::ImportFromJSONAsset(const FString& InAbsolutePath)
 {
-	Metasound::Frontend::FDocumentAccessPtr Document = GetDocument();
-	if (ensure(Document.IsValid()))
+	Metasound::Frontend::FDocumentAccessPtr DocumentPtr = GetDocument();
+	if (FMetasoundFrontendDocument* Document = DocumentPtr.Get())
 	{
 		bool bSuccess = Metasound::Frontend::ImportJSONAssetToMetasound(InAbsolutePath, *Document);
 
@@ -579,17 +580,16 @@ bool FMetasoundAssetBase::ImportFromJSONAsset(const FString& InAbsolutePath)
 
 FMetasoundFrontendDocument& FMetasoundAssetBase::GetDocumentChecked()
 {
-	Metasound::Frontend::FDocumentAccessPtr DocAccessPtr = GetDocument();
-
-	check(DocAccessPtr.IsValid());
-	return *DocAccessPtr;
+	FMetasoundFrontendDocument* Document = GetDocument().Get();
+	check(nullptr != Document);
+	return *Document;
 }
 
 const FMetasoundFrontendDocument& FMetasoundAssetBase::GetDocumentChecked() const
 {
-	Metasound::Frontend::FConstDocumentAccessPtr DocAccessPtr = GetDocument();
+	const FMetasoundFrontendDocument* Document = GetDocument().Get();
 
-	check(DocAccessPtr.IsValid());
-	return *DocAccessPtr;
+	check(nullptr != Document);
+	return *Document;
 }
 #undef LOCTEXT_NAMESPACE // "MetaSound"
diff --git a/Engine/Plugins/Runtime/Metasound/Source/MetasoundFrontend/Private/MetasoundFrontendDocumentAccessPtr.cpp b/Engine/Plugins/Runtime/Metasound/Source/MetasoundFrontend/Private/MetasoundFrontendDocumentAccessPtr.cpp
--- a/Engine/Plugins/Runtime/Metasound/Source/MetasoundFrontend/Private/MetasoundFrontendDocumentAccessPtr.cpp
+++ b/Engine/Plugins/Runtime/Metasound/Source/MetasoundFrontend/Private/MetasoundFrontendDocumentAccessPtr.cpp
@@ -9,284 +9,591 @@ namespace Metasound
 	{
 		namespace MetasoundFrontendDocumentAccessPtrPrivate
 		{
-			using FRegistry = FMetasoundFrontendRegistryContainer;
+			/** TFindInArray is a callable object for finding child objects with a parent 
+			 * objects member array. When searching for a child object, it caches the index
+			 * of the child object and will test that index first on subsequent searches.
+			 *
+			 * @tparam ParentType - The type of the parent object.
+			 * @tparam ChildType - The type of the child object.
+			 * @tparam ArrayAccessType - The type of callable object used to access the parent's array.
+			 * @tparam TestType - The type of callable objet used to test for finding the desired child instance.
+			 */
+			template<typename ParentType, typename ChildType, typename ArrayAccessType, typename TestType>
+			class TFindInArray
+			{
+			public:	
+				/** Create a TFindInArray
+				 *
+				 * @param InAccessArray - A callable object which accepts a reference to the parent
+				 *                        object and returns the array containing child objects.
+				 * @param InTest - A callable object which accepts a const reference to a child object
+				 *                 and returns true if it is the desired child object.
+				 */
+				TFindInArray(ArrayAccessType InAccessArray, TestType InTest)
+				: AccessArray(InAccessArray)
+				, Test(InTest)
+				{
+				}
 
-			bool IsMatchingMetasoundClass(const FMetasoundFrontendClassMetadata& InMetadataA, const FMetasoundFrontendClassMetadata& InMetadataB) 
-			{
-				if (InMetadataA.Type == InMetadataB.Type)
+				TFindInArray(const TFindInArray&) = default;
+				TFindInArray& operator=(const TFindInArray&) = default;
+				TFindInArray(TFindInArray&&) = default;
+				TFindInArray& operator=(TFindInArray&&) = default;
+
+				/** Find child object within parent object.
+				 *
+				 * @param InParent - The parent object containing the desired child object.
+				 *
+				 * @return If found, A pointer to the child object. Otherwise a nullptr.
+				 */
+				ChildType* operator()(ParentType& InParent) const
 				{
-					if (InMetadataA.ClassName == InMetadataB.ClassName)
+					auto& Array = AccessArray(InParent);
+
+					if (INDEX_NONE != CachedIndex)
+					{
+						if (CachedIndex < Array.Num())
+						{
+							if (Test(Array[CachedIndex]))
+							{
+								return &Array[CachedIndex];
+							}
+						}
+					}
+
+					CachedIndex = Array.IndexOfByPredicate(Test);
+					if (INDEX_NONE != CachedIndex)
+					{
+						return &Array[CachedIndex];
+					}
+
+					return nullptr;
+				}
+
+			private:
+
+				mutable int32 CachedIndex = INDEX_NONE;
+				ArrayAccessType AccessArray;
+				TestType Test;
+			};
+
+			/** TFindClassInDocument is a callable object for finding a class within
+			 * a document. It searches both the document's Dependencies and Subgraphs.
+			 * When searching for a class, it caches the location of the class and 
+			 * will test that location first on subsequent searches.
+			 * 
+			 * @tparam ClassType type of class to tind.
+			 * @tparam TestType - The type of callable objet used to test for finding the desired class.
+			 */
+			template<typename ClassType, typename TestType>
+			struct TFindClassInDocument
+			{
+			public:
+				/** Create a TFindClassInDocument
+				 *
+				 * @param InTest - A callable object which accepts a const reference to a FMetasoundFrontendClass
+				 *                 and returns true if it is the desired class and false otherwise.
+				 */
+				TFindClassInDocument(TestType InTest)
+				: Test(InTest)
+				{
+				}
+
+				/** Find class in document.
+				 *
+				 * @tparam DocumentType - Type of document to search.
+				 *
+				 * @param InDoc - The document to search.
+				 * 
+				 * @return If found, a pointer to the desired class. Otherwise, a nullptr.
+				 */
+				template<typename DocumentType>
+				ClassType*  operator()(DocumentType& InDoc) const
+				{
+					// Search in cached location.
+					if (ClassType* FoundClass = FindClassInCachedLocation<DocumentType>(InDoc))
+					{
+						return FoundClass;
+					}
+
+					// Search in dependencies
+					CachedIndex = InDoc.Dependencies.IndexOfByPredicate(Test);
+					if (INDEX_NONE != CachedIndex)
 					{
-						return FRegistry::GetRegistryKey(InMetadataA) == FRegistry::GetRegistryKey(InMetadataB);
+						CachedSource = ESource::Dependencies;
+						return &InDoc.Dependencies[CachedIndex];
 					}
+
+					// Search in subgraphs.
+					CachedIndex = InDoc.Subgraphs.IndexOfByPredicate(Test);
+					if (INDEX_NONE != CachedIndex)
+					{
+						CachedSource = ESource::Subgraphs;
+						return &InDoc.Subgraphs[CachedIndex];
+					}
+
+					// Could not find class.
+					CachedSource = ESource::None;
+					return nullptr;
+				}
+
+			private:
+
+				template<typename DocumentType>
+				ClassType*  FindClassInCachedLocation(DocumentType& InDoc) const
+				{
+					if (CachedIndex != INDEX_NONE)
+					{
+						switch (CachedSource)
+						{
+						case ESource::Dependencies:
+							{
+								if (CachedIndex < InDoc.Dependencies.Num())
+								{
+									if (Test(InDoc.Dependencies[CachedIndex]))
+									{
+										return &InDoc.Dependencies[CachedIndex];
+									}
+								}
+							}
+							break;
+
+						case ESource::Subgraphs:
+							{
+								if (CachedIndex < InDoc.Subgraphs.Num())
+								{
+									if (Test(InDoc.Subgraphs[CachedIndex]))
+									{
+										return &InDoc.Subgraphs[CachedIndex];
+									}
+								}
+							}
+							break;
+						default:
+						break;
+						}
+					}
+
+					return nullptr;
+				}
+
+				enum class ESource : uint8 { None, Dependencies, Subgraphs };
+
+				mutable int32 CachedIndex = INDEX_NONE;
+				mutable ESource CachedSource = ESource::None;
+				TestType Test;
+			};
+
+			/** A callable object for accessing arrays of inputs from a node. */
+			struct FGetNodeInputArray 
+			{
+			 	TArray<FMetasoundFrontendVertex>& operator()(FMetasoundFrontendNode& InNode) const { return InNode.Interface.Inputs; }
+			 	const TArray<FMetasoundFrontendVertex>& operator()(const FMetasoundFrontendNode& InNode) const { return InNode.Interface.Inputs; }
+			};
+
+			/** A callable object for accessing arrays of outputs from a node. */
+			struct FGetNodeOutputArray 
+			{
+			 	TArray<FMetasoundFrontendVertex>& operator()(FMetasoundFrontendNode& InNode) const { return InNode.Interface.Outputs; }
+			 	const TArray<FMetasoundFrontendVertex>& operator()(const FMetasoundFrontendNode& InNode) const { return InNode.Interface.Outputs; }
+			};
+
+			/** A callable object for accessing arrays of inputs from a class. */
+			struct FGetClassInputArray 
+			{
+			 	TArray<FMetasoundFrontendClassInput>& operator()(FMetasoundFrontendClass& InClass) const { return InClass.Interface.Inputs; }
+			 	const TArray<FMetasoundFrontendClassInput>& operator()(const FMetasoundFrontendClass& InClass) const { return InClass.Interface.Inputs; }
+			};
+
+			/** A callable object for accessing arrays of outputs from a class. */
+			struct FGetClassOutputArray 
+			{
+			 	TArray<FMetasoundFrontendClassOutput>& operator()(FMetasoundFrontendClass& InClass) const { return InClass.Interface.Outputs; }
+			 	const TArray<FMetasoundFrontendClassOutput>& operator()(const FMetasoundFrontendClass& InClass) const { return InClass.Interface.Outputs; }
+			};
+
+			/** A callable object for accessing arrays of nodes from a graph class. */
+			struct FGetGraphClassNodes
+			{
+			 	TArray<FMetasoundFrontendNode>& operator()(FMetasoundFrontendGraphClass& InGraphClass) const { return InGraphClass.Graph.Nodes; }
+			 	const TArray<FMetasoundFrontendNode>& operator()(const FMetasoundFrontendGraphClass& InGraphClass) const { return InGraphClass.Graph.Nodes; }
+			};
+
+			/** A callable object for accessing arrays of subgraphs from a document. */
+			struct FGetDocumentSubgraphs
+			{
+			 	TArray<FMetasoundFrontendGraphClass>& operator()(FMetasoundFrontendDocument& InDocument) const { return InDocument.Subgraphs; }
+			 	const TArray<FMetasoundFrontendGraphClass>& operator()(const FMetasoundFrontendDocument& InDocument) const { return InDocument.Subgraphs; }
+			};
+
+			/** A callable object for accessing arrays of dependencies from a document. */
+			struct FGetDocumentDependencies
+			{
+			 	TArray<FMetasoundFrontendClass>& operator()(FMetasoundFrontendDocument& InDocument) const { return InDocument.Dependencies; }
+			 	const TArray<FMetasoundFrontendClass>& operator()(const FMetasoundFrontendDocument& InDocument) const { return InDocument.Dependencies; }
+			};
+
+			/** A callable object for testing whether a vertex contains a matching name. */
+			struct FIsVertexWithName
+			{
+				FIsVertexWithName(const FString& InName)
+				: Name(InName)
+				{
+				}
+
+				bool operator() (const FMetasoundFrontendVertex& InVertex) const { return InVertex.Name == Name; }
+			private:
+				FString Name;
+			};
+
+			/** A callable object for testing whether a vertex contains a matching ID. */
+			struct FIsVertexWithID
+			{
+				FIsVertexWithID(const FGuid& InID)
+				: ID(InID)
+				{
+				}
+
+				bool operator() (const FMetasoundFrontendVertex& InVertex) const { return InVertex.VertexID == ID; }
+			private:
+				FGuid ID;
+			};
+
+			/** A callable object for testing whether a class input contains a matching node ID. */
+			struct FIsClassInputWithNodeID
+			{
+				FIsClassInputWithNodeID(const FGuid& InNodeID)
+				: NodeID(InNodeID)
+				{
+				}
+
+				bool operator() (const FMetasoundFrontendClassInput& InClassInput) const { return InClassInput.NodeID == NodeID; }
+			private:
+				FGuid NodeID;
+			};
+
+			/** A callable object for testing whether a class output contains a matching node ID. */
+			struct FIsClassOutputWithNodeID
+			{
+				FIsClassOutputWithNodeID(const FGuid& InNodeID)
+				: NodeID(InNodeID)
+				{
 				}
-				return false;
-			}
 
-			bool IsMatchingMetasoundClass(const FNodeClassInfo& InNodeClass, const FMetasoundFrontendClassMetadata& InMetadata) 
+				bool operator() (const FMetasoundFrontendClassOutput& InClassOutput) const { return InClassOutput.NodeID == NodeID; }
+			private:
+				FGuid NodeID;
+			};
+
+			/** A callable object for testing whether a node contains a matching ID. */
+			struct FIsNodeWithID
 			{
-				return FRegistry::GetRegistryKey(InNodeClass) == FRegistry::GetRegistryKey(InMetadata);
-			}
+				FIsNodeWithID(const FGuid& InNodeID)
+				: NodeID(InNodeID)
+				{
+				}
+
+				bool operator() (const FMetasoundFrontendNode& InNode) const { return InNode.ID == NodeID; }
+			private:
+				FGuid NodeID;
+			};
 
-			bool IsMatchingMetasoundClass(const FNodeRegistryKey& InKey, const FMetasoundFrontendClassMetadata& InMetadata) 
+			/** A callable object for testing whether a class contains a matching ID. */
+			struct FIsClassWithID
 			{
-				return InKey == FRegistry::GetRegistryKey(InMetadata);
-			}
+				FIsClassWithID(const FGuid& InClassID)
+				: ClassID(InClassID)
+				{
+				}
+
+				bool operator() (const FMetasoundFrontendClass& InClass) const { return InClass.ID == ClassID; }
+			private:
+				FGuid ClassID;
+			};
+
+			/** A callable object for testing whether a class contains a matching Metadata. */
+			class FIsClassWithMetadata
+			{
+			public:
+				FIsClassWithMetadata(const FMetasoundFrontendClassMetadata& InMetadata)
+				: Metadata(InMetadata)
+				{
+				}
+
+				bool operator()(const FMetasoundFrontendClass& InClass) const
+				{ 
+					return NodeRegistryKey::IsEqual(Metadata, InClass.Metadata);
+				};
+			private:
+				FMetasoundFrontendClassMetadata Metadata;
+			};
+
+			/** A callable object for testing whether a class contains a matching node info. */
+			class FIsClassWithNodeInfo
+			{
+			public:
+				FIsClassWithNodeInfo(const FNodeClassInfo& InNodeInfo)
+				: NodeInfo(InNodeInfo)
+				{
+				}
+
+				bool operator()(const FMetasoundFrontendClass& InClass) const
+				{ 
+					return NodeRegistryKey::IsEqual(NodeInfo, InClass.Metadata);
+				};
+			private:
+				FNodeClassInfo NodeInfo;
+			};
+
+			/** A callable object for testing whether a class contains a matching registry key. */
+			class FIsClassWithNodeRegistryKey
+			{
+			public:
+				FIsClassWithNodeRegistryKey(const FNodeRegistryKey& InRegistryKey)
+				: RegistryKey(InRegistryKey)
+				{
+				}
+
+				bool operator()(const FMetasoundFrontendClass& InClass) const
+				{ 
+					return NodeRegistryKey::IsEqual(RegistryKey, NodeRegistryKey::CreateKey(InClass.Metadata));
+				};
+			private:
+				FNodeRegistryKey RegistryKey;
+			};
 		}
 
 		FVertexAccessPtr FNodeAccessPtr::GetInputWithName(const FString& InName)
 		{
-			auto FindInputWithName = [Name=InName](FMetasoundFrontendNode& InNode) -> FMetasoundFrontendVertex*
-			{
-				auto IsVertexWithName = [&](const FMetasoundFrontendVertex& Vertex) { return Vertex.Name == Name; };
-				return InNode.Interface.Inputs.FindByPredicate(IsVertexWithName);
-			};
+			using namespace MetasoundFrontendDocumentAccessPtrPrivate;
+
+			using FFinder= TFindInArray<FMetasoundFrontendNode, FMetasoundFrontendVertex, FGetNodeInputArray, FIsVertexWithName>;
+			FFinder Finder{FGetNodeInputArray(), FIsVertexWithName(InName)};
 
-			return GetMemberAccessPtr<FVertexAccessPtr>(FindInputWithName);
+			return GetMemberAccessPtr<FVertexAccessPtr>(Finder);
 		}
 
 		FVertexAccessPtr FNodeAccessPtr::GetInputWithVertexID(const FGuid& InID)
 		{
-			auto FindInputWithID = [ID=InID](FMetasoundFrontendNode& InNode) -> FMetasoundFrontendVertex*
-			{
-				auto IsVertexWithVertexID = [=](const FMetasoundFrontendVertex& Vertex) { return Vertex.VertexID == ID; };
-				return InNode.Interface.Inputs.FindByPredicate(IsVertexWithVertexID);
-			};
-			return GetMemberAccessPtr<FVertexAccessPtr>(FindInputWithID);
+			using namespace MetasoundFrontendDocumentAccessPtrPrivate;
+
+			using FFinder= TFindInArray<FMetasoundFrontendNode, FMetasoundFrontendVertex, FGetNodeInputArray, FIsVertexWithID>;
+			FFinder Finder{FGetNodeInputArray(), FIsVertexWithID(InID)};
+
+			return GetMemberAccessPtr<FVertexAccessPtr>(Finder);
 		}
 
 		FVertexAccessPtr FNodeAccessPtr::GetOutputWithName(const FString& InName)
 		{
-			auto FindOutputWithName = [Name=InName](FMetasoundFrontendNode& InNode) -> FMetasoundFrontendVertex*
-			{
-				auto IsVertexWithName = [&](const FMetasoundFrontendVertex& Vertex) { return Vertex.Name == Name; };
-				return InNode.Interface.Outputs.FindByPredicate(IsVertexWithName);
-			};
+			using namespace MetasoundFrontendDocumentAccessPtrPrivate;
 
-			return GetMemberAccessPtr<FVertexAccessPtr>(FindOutputWithName);
+			using FFinder= TFindInArray<FMetasoundFrontendNode, FMetasoundFrontendVertex, FGetNodeOutputArray, FIsVertexWithName>;
+			FFinder Finder{FGetNodeOutputArray(), FIsVertexWithName(InName)};
+
+			return GetMemberAccessPtr<FVertexAccessPtr>(Finder);
 		}
 
 		FVertexAccessPtr FNodeAccessPtr::GetOutputWithVertexID(const FGuid& InID)
 		{
-			auto FindOutputWithID = [ID=InID](FMetasoundFrontendNode& InNode) -> FMetasoundFrontendVertex*
-			{
-				auto IsVertexWithVertexID = [=](const FMetasoundFrontendVertex& Vertex) { return Vertex.VertexID == ID; };
-				return InNode.Interface.Outputs.FindByPredicate(IsVertexWithVertexID);
-			};
-			return GetMemberAccessPtr<FVertexAccessPtr>(FindOutputWithID);
+			using namespace MetasoundFrontendDocumentAccessPtrPrivate;
+
+			using FFinder= TFindInArray<FMetasoundFrontendNode, FMetasoundFrontendVertex, FGetNodeOutputArray, FIsVertexWithID>;
+			FFinder Finder{FGetNodeOutputArray(), FIsVertexWithID(InID)};
+
+			return GetMemberAccessPtr<FVertexAccessPtr>(Finder);
 		}
 
 		FConstVertexAccessPtr FNodeAccessPtr::GetInputWithName(const FString& InName) const
 		{
-			auto FindInputWithName = [Name=InName](const FMetasoundFrontendNode& InNode) -> const FMetasoundFrontendVertex*
-			{
-				auto IsVertexWithName = [&](const FMetasoundFrontendVertex& Vertex) { return Vertex.Name == Name; };
-				return InNode.Interface.Inputs.FindByPredicate(IsVertexWithName);
-			};
+			using namespace MetasoundFrontendDocumentAccessPtrPrivate;
 
-			return GetMemberAccessPtr<FConstVertexAccessPtr>(FindInputWithName);
+			using FFinder= TFindInArray<const FMetasoundFrontendNode, const FMetasoundFrontendVertex, FGetNodeInputArray, FIsVertexWithName>;
+			FFinder Finder{FGetNodeInputArray(), FIsVertexWithName(InName)};
+
+			return GetMemberAccessPtr<FConstVertexAccessPtr>(Finder);
 		}
 
 		FConstVertexAccessPtr FNodeAccessPtr::GetInputWithVertexID(const FGuid& InID) const
 		{
-			auto FindInputWithID = [ID=InID](const FMetasoundFrontendNode& InNode) -> const FMetasoundFrontendVertex*
-			{
-				auto IsVertexWithVertexID = [=](const FMetasoundFrontendVertex& Vertex) { return Vertex.VertexID == ID; };
-				return InNode.Interface.Inputs.FindByPredicate(IsVertexWithVertexID);
-			};
-			return GetMemberAccessPtr<FConstVertexAccessPtr>(FindInputWithID);
+			using namespace MetasoundFrontendDocumentAccessPtrPrivate;
+
+			using FFinder= TFindInArray<const FMetasoundFrontendNode, const FMetasoundFrontendVertex, FGetNodeInputArray, FIsVertexWithID>;
+			FFinder Finder{FGetNodeInputArray(), FIsVertexWithID(InID)};
+
+			return GetMemberAccessPtr<FConstVertexAccessPtr>(Finder);
 		}
 
 		FConstVertexAccessPtr FNodeAccessPtr::GetOutputWithName(const FString& InName) const
 		{
-			auto FindOutputWithName = [Name=InName](const FMetasoundFrontendNode& InNode) -> const FMetasoundFrontendVertex*
-			{
-				auto IsVertexWithName = [&](const FMetasoundFrontendVertex& Vertex) { return Vertex.Name == Name; };
-				return InNode.Interface.Outputs.FindByPredicate(IsVertexWithName);
-			};
+			using namespace MetasoundFrontendDocumentAccessPtrPrivate;
 
-			return GetMemberAccessPtr<FConstVertexAccessPtr>(FindOutputWithName);
+			using FFinder= TFindInArray<const FMetasoundFrontendNode, const FMetasoundFrontendVertex, FGetNodeOutputArray, FIsVertexWithName>;
+			FFinder Finder{FGetNodeOutputArray(), FIsVertexWithName(InName)};
+
+			return GetMemberAccessPtr<FConstVertexAccessPtr>(Finder);
 		}
 
 		FConstVertexAccessPtr FNodeAccessPtr::GetOutputWithVertexID(const FGuid& InID) const
 		{
-			auto FindOutputWithID = [ID=InID](const FMetasoundFrontendNode& InNode) -> const FMetasoundFrontendVertex*
-			{
-				auto IsVertexWithVertexID = [=](const FMetasoundFrontendVertex& Vertex) { return Vertex.VertexID == ID; };
-				return InNode.Interface.Outputs.FindByPredicate(IsVertexWithVertexID);
-			};
-			return GetMemberAccessPtr<FConstVertexAccessPtr>(FindOutputWithID);
+			using namespace MetasoundFrontendDocumentAccessPtrPrivate;
+
+			using FFinder= TFindInArray<const FMetasoundFrontendNode, const FMetasoundFrontendVertex, FGetNodeOutputArray, FIsVertexWithID>;
+			FFinder Finder{FGetNodeOutputArray(), FIsVertexWithID(InID)};
+
+			return GetMemberAccessPtr<FConstVertexAccessPtr>(Finder);
 		}
 
 		FConstVertexAccessPtr FConstNodeAccessPtr::GetInputWithName(const FString& InName) const
 		{
-			auto FindInputWithName = [Name=InName](const FMetasoundFrontendNode& InNode) -> const FMetasoundFrontendVertex*
-			{
-				auto IsVertexWithName = [&](const FMetasoundFrontendVertex& Vertex) { return Vertex.Name == Name; };
-				return InNode.Interface.Inputs.FindByPredicate(IsVertexWithName);
-			};
+			using namespace MetasoundFrontendDocumentAccessPtrPrivate;
+
+			using FFinder= TFindInArray<const FMetasoundFrontendNode, const FMetasoundFrontendVertex, FGetNodeInputArray, FIsVertexWithName>;
+			FFinder Finder{FGetNodeInputArray(), FIsVertexWithName(InName)};
 
-			return GetMemberAccessPtr<FConstVertexAccessPtr>(FindInputWithName);
+			return GetMemberAccessPtr<FConstVertexAccessPtr>(Finder);
 		}
 
 		FConstVertexAccessPtr FConstNodeAccessPtr::GetInputWithVertexID(const FGuid& InID) const
 		{
-			auto FindInputWithID = [ID=InID](const FMetasoundFrontendNode& InNode) -> const FMetasoundFrontendVertex*
-			{
-				auto IsVertexWithVertexID = [=](const FMetasoundFrontendVertex& Vertex) { return Vertex.VertexID == ID; };
-				return InNode.Interface.Inputs.FindByPredicate(IsVertexWithVertexID);
-			};
-			return GetMemberAccessPtr<FConstVertexAccessPtr>(FindInputWithID);
+			using namespace MetasoundFrontendDocumentAccessPtrPrivate;
+
+			using FFinder= TFindInArray<const FMetasoundFrontendNode, const FMetasoundFrontendVertex, FGetNodeInputArray, FIsVertexWithID>;
+			FFinder Finder{FGetNodeInputArray(), FIsVertexWithID(InID)};
+
+			return GetMemberAccessPtr<FConstVertexAccessPtr>(Finder);
 		}
 
 		FConstVertexAccessPtr FConstNodeAccessPtr::GetOutputWithName(const FString& InName) const
 		{
-			auto FindOutputWithName = [Name=InName](const FMetasoundFrontendNode& InNode) -> const FMetasoundFrontendVertex*
-			{
-				auto IsVertexWithName = [&](const FMetasoundFrontendVertex& Vertex) { return Vertex.Name == Name; };
-				return InNode.Interface.Outputs.FindByPredicate(IsVertexWithName);
-			};
+			using namespace MetasoundFrontendDocumentAccessPtrPrivate;
+
+			using FFinder= TFindInArray<const FMetasoundFrontendNode, const FMetasoundFrontendVertex, FGetNodeOutputArray, FIsVertexWithName>;
+			FFinder Finder{FGetNodeOutputArray(), FIsVertexWithName(InName)};
 
-			return GetMemberAccessPtr<FConstVertexAccessPtr>(FindOutputWithName);
+			return GetMemberAccessPtr<FConstVertexAccessPtr>(Finder);
 		}
 
 		FConstVertexAccessPtr FConstNodeAccessPtr::GetOutputWithVertexID(const FGuid& InID) const
 		{
-			auto FindOutputWithID = [ID=InID](const FMetasoundFrontendNode& InNode) -> const FMetasoundFrontendVertex*
-			{
-				auto IsVertexWithVertexID = [=](const FMetasoundFrontendVertex& Vertex) { return Vertex.VertexID == ID; };
-				return InNode.Interface.Outputs.FindByPredicate(IsVertexWithVertexID);
-			};
-			return GetMemberAccessPtr<FConstVertexAccessPtr>(FindOutputWithID);
-		}
+			using namespace MetasoundFrontendDocumentAccessPtrPrivate;
 
+			using FFinder= TFindInArray<const FMetasoundFrontendNode, const FMetasoundFrontendVertex, FGetNodeOutputArray, FIsVertexWithID>;
+			FFinder Finder{FGetNodeOutputArray(), FIsVertexWithID(InID)};
+
+			return GetMemberAccessPtr<FConstVertexAccessPtr>(Finder);
+		}
 
 		// FConstClassAccessPtr
 		FClassInputAccessPtr FClassAccessPtr::GetInputWithName(const FString& InName)
 		{
-			auto FindClassInputWithName = [Name=InName](FMetasoundFrontendClass& InClass)
-			{
-				auto IsClassInputWithName = [&](const FMetasoundFrontendClassInput& ClassInput) { return ClassInput.Name == Name; };
-				return InClass.Interface.Inputs.FindByPredicate(IsClassInputWithName);
+			using namespace MetasoundFrontendDocumentAccessPtrPrivate;
 
-			};
+			using FFinder= TFindInArray<FMetasoundFrontendClass, FMetasoundFrontendClassInput, FGetClassInputArray, FIsVertexWithName>;
+			FFinder Finder{FGetClassInputArray(), FIsVertexWithName(InName)};
 
-			return GetMemberAccessPtr<FClassInputAccessPtr>(FindClassInputWithName);
+			return GetMemberAccessPtr<FClassInputAccessPtr>(Finder);
 		}
 
 		FClassOutputAccessPtr FClassAccessPtr::GetOutputWithName(const FString& InName)
 		{
-			auto FindClassOutputWithName = [Name=InName](FMetasoundFrontendClass& InClass)
-			{
-				auto IsClassOutputWithName = [&](const FMetasoundFrontendClassOutput& ClassOutput) { return ClassOutput.Name == Name; };
-				return InClass.Interface.Outputs.FindByPredicate(IsClassOutputWithName);
+			using namespace MetasoundFrontendDocumentAccessPtrPrivate;
 
-			};
+			using FFinder= TFindInArray<FMetasoundFrontendClass, FMetasoundFrontendClassOutput, FGetClassOutputArray, FIsVertexWithName>;
+			FFinder Finder{FGetClassOutputArray(), FIsVertexWithName(InName)};
 
-			return GetMemberAccessPtr<FClassOutputAccessPtr>(FindClassOutputWithName);
+			return GetMemberAccessPtr<FClassOutputAccessPtr>(Finder);
 		}
 
 		FConstClassInputAccessPtr FClassAccessPtr::GetInputWithName(const FString& InName) const
 		{
-			auto FindClassInputWithName = [Name=InName](const FMetasoundFrontendClass& InClass)
-			{
-				auto IsClassInputWithName = [&](const FMetasoundFrontendClassInput& ClassInput) { return ClassInput.Name == Name; };
-				return InClass.Interface.Inputs.FindByPredicate(IsClassInputWithName);
+			using namespace MetasoundFrontendDocumentAccessPtrPrivate;
 
-			};
+			using FFinder= TFindInArray<const FMetasoundFrontendClass, const FMetasoundFrontendClassInput, FGetClassInputArray, FIsVertexWithName>;
+			FFinder Finder{FGetClassInputArray(), FIsVertexWithName(InName)};
 
-			return GetMemberAccessPtr<FConstClassInputAccessPtr>(FindClassInputWithName);
+			return GetMemberAccessPtr<FConstClassInputAccessPtr>(Finder);
 		}
 
 		FConstClassOutputAccessPtr FClassAccessPtr::GetOutputWithName(const FString& InName) const
 		{
-			auto FindClassOutputWithName = [Name=InName](const FMetasoundFrontendClass& InClass)
-			{
-				auto IsClassOutputWithName = [&](const FMetasoundFrontendClassOutput& ClassOutput) { return ClassOutput.Name == Name; };
-				return InClass.Interface.Outputs.FindByPredicate(IsClassOutputWithName);
+			using namespace MetasoundFrontendDocumentAccessPtrPrivate;
 
-			};
+			using FFinder= TFindInArray<const FMetasoundFrontendClass, const FMetasoundFrontendClassOutput, FGetClassOutputArray, FIsVertexWithName>;
+			FFinder Finder{FGetClassOutputArray(), FIsVertexWithName(InName)};
 
-			return GetMemberAccessPtr<FConstClassOutputAccessPtr>(FindClassOutputWithName);
+			return GetMemberAccessPtr<FConstClassOutputAccessPtr>(Finder);
 		}
 
 		// FConstClassAccessPtr
 		FConstClassInputAccessPtr FConstClassAccessPtr::GetInputWithName(const FString& InName) const
 		{
-			auto FindClassInputWithName = [Name=InName](const FMetasoundFrontendClass& InClass)
-			{
-				auto IsClassInputWithName = [&](const FMetasoundFrontendClassInput& ClassInput) { return ClassInput.Name == Name; };
-				return InClass.Interface.Inputs.FindByPredicate(IsClassInputWithName);
+			using namespace MetasoundFrontendDocumentAccessPtrPrivate;
 
-			};
+			using FFinder= TFindInArray<const FMetasoundFrontendClass, const FMetasoundFrontendClassInput, FGetClassInputArray, FIsVertexWithName>;
+			FFinder Finder{FGetClassInputArray(), FIsVertexWithName(InName)};
 
-			return GetMemberAccessPtr<FConstClassInputAccessPtr>(FindClassInputWithName);
+			return GetMemberAccessPtr<FConstClassInputAccessPtr>(Finder);
 		}
 
 		FConstClassOutputAccessPtr FConstClassAccessPtr::GetOutputWithName(const FString& InName) const
 		{
-			auto FindClassOutputWithName = [Name=InName](const FMetasoundFrontendClass& InClass)
-			{
-				auto IsClassOutputWithName = [&](const FMetasoundFrontendClassOutput& ClassOutput) { return ClassOutput.Name == Name; };
-				return InClass.Interface.Outputs.FindByPredicate(IsClassOutputWithName);
+			using namespace MetasoundFrontendDocumentAccessPtrPrivate;
 
-			};
+			using FFinder= TFindInArray<const FMetasoundFrontendClass, const FMetasoundFrontendClassOutput, FGetClassOutputArray, FIsVertexWithName>;
+			FFinder Finder{FGetClassOutputArray(), FIsVertexWithName(InName)};
 
-			return GetMemberAccessPtr<FConstClassOutputAccessPtr>(FindClassOutputWithName);
+			return GetMemberAccessPtr<FConstClassOutputAccessPtr>(Finder);
 		}
 
 		// FGraphClassAccessPtr
 		FClassInputAccessPtr FGraphClassAccessPtr::GetInputWithName(const FString& InName)
 		{
-			auto FindClassInputWithName = [Name=InName](FMetasoundFrontendGraphClass& InClass)
-			{
-				auto IsClassInputWithName = [&](const FMetasoundFrontendClassInput& ClassInput) { return ClassInput.Name == Name; };
-				return InClass.Interface.Inputs.FindByPredicate(IsClassInputWithName);
+			using namespace MetasoundFrontendDocumentAccessPtrPrivate;
 
-			};
+			using FFinder= TFindInArray<FMetasoundFrontendClass, FMetasoundFrontendClassInput, FGetClassInputArray, FIsVertexWithName>;
+			FFinder Finder{FGetClassInputArray(), FIsVertexWithName(InName)};
 
-			return GetMemberAccessPtr<FClassInputAccessPtr>(FindClassInputWithName);
+			return GetMemberAccessPtr<FClassInputAccessPtr>(Finder);
 		}
 
 		FClassInputAccessPtr FGraphClassAccessPtr::GetInputWithNodeID(const FGuid& InNodeID)
 		{
-			auto FindClassInputWithNodeID = [NodeID=InNodeID](FMetasoundFrontendGraphClass& InGraphClass) -> FMetasoundFrontendClassInput*
-			{
-				return InGraphClass.Interface.Inputs.FindByPredicate([&](const FMetasoundFrontendClassInput& Desc) { return Desc.NodeID == NodeID; });
-			};
-			return GetMemberAccessPtr<FClassInputAccessPtr>(FindClassInputWithNodeID);
+			using namespace MetasoundFrontendDocumentAccessPtrPrivate;
+
+			using FFinder= TFindInArray<FMetasoundFrontendGraphClass, FMetasoundFrontendClassInput, FGetClassInputArray, FIsClassInputWithNodeID>;
+			FFinder Finder{FGetClassInputArray(), FIsClassInputWithNodeID(InNodeID)};
+
+			return GetMemberAccessPtr<FClassInputAccessPtr>(Finder);
 		}
 
 		FClassOutputAccessPtr FGraphClassAccessPtr::GetOutputWithName(const FString& InName)
 		{
-			auto FindClassOutputWithName = [Name=InName](FMetasoundFrontendGraphClass& InClass)
-			{
-				auto IsClassOutputWithName = [&](const FMetasoundFrontendClassOutput& ClassOutput) { return ClassOutput.Name == Name; };
-				return InClass.Interface.Outputs.FindByPredicate(IsClassOutputWithName);
+			using namespace MetasoundFrontendDocumentAccessPtrPrivate;
 
-			};
+			using FFinder= TFindInArray<FMetasoundFrontendClass, FMetasoundFrontendClassOutput, FGetClassOutputArray, FIsVertexWithName>;
+			FFinder Finder{FGetClassOutputArray(), FIsVertexWithName(InName)};
 
-			return GetMemberAccessPtr<FClassOutputAccessPtr>(FindClassOutputWithName);
+			return GetMemberAccessPtr<FClassOutputAccessPtr>(Finder);
 		}
 
 		FClassOutputAccessPtr FGraphClassAccessPtr::GetOutputWithNodeID(const FGuid& InNodeID)
 		{
-			auto FindClassOutputWithNodeID = [NodeID=InNodeID](FMetasoundFrontendGraphClass& InGraphClass) -> FMetasoundFrontendClassOutput*
-			{
-				return InGraphClass.Interface.Outputs.FindByPredicate([&](const FMetasoundFrontendClassOutput& Desc) { return Desc.NodeID == NodeID; });
-			};
-			return GetMemberAccessPtr<FClassOutputAccessPtr>(FindClassOutputWithNodeID);
+			using namespace MetasoundFrontendDocumentAccessPtrPrivate;
+
+			using FFinder= TFindInArray<FMetasoundFrontendGraphClass, FMetasoundFrontendClassOutput, FGetClassOutputArray, FIsClassOutputWithNodeID>;
+			FFinder Finder{FGetClassOutputArray(), FIsClassOutputWithNodeID(InNodeID)};
+
+			return GetMemberAccessPtr<FClassOutputAccessPtr>(Finder);
 		}
 
 		FNodeAccessPtr FGraphClassAccessPtr::GetNodeWithNodeID(const FGuid& InNodeID)
 		{
-			auto FindNodeWithNodeID = [ID=InNodeID](FMetasoundFrontendGraphClass& InGraphClass) -> FMetasoundFrontendNode*
-			{
-				auto IsNodeWithID = [&](const FMetasoundFrontendNode& InNode) { return InNode.ID == ID; };
-				return InGraphClass.Graph.Nodes.FindByPredicate(IsNodeWithID);
-			};
+			using namespace MetasoundFrontendDocumentAccessPtrPrivate;
 
-			return GetMemberAccessPtr<FNodeAccessPtr>(FindNodeWithNodeID);
+			using FFinder= TFindInArray<FMetasoundFrontendGraphClass, FMetasoundFrontendNode, FGetGraphClassNodes, FIsNodeWithID>;
+			FFinder Finder{FGetGraphClassNodes(), FIsNodeWithID(InNodeID)};
+
+			return GetMemberAccessPtr<FNodeAccessPtr>(Finder);
 		}
 
 		FGraphAccessPtr FGraphClassAccessPtr::GetGraph()
@@ -302,56 +609,53 @@ namespace Metasound
 
 		FConstClassInputAccessPtr FGraphClassAccessPtr::GetInputWithName(const FString& InName) const
 		{
-			auto FindClassInputWithName = [Name=InName](const FMetasoundFrontendGraphClass& InClass)
-			{
-				auto IsClassInputWithName = [&](const FMetasoundFrontendClassInput& ClassInput) { return ClassInput.Name == Name; };
-				return InClass.Interface.Inputs.FindByPredicate(IsClassInputWithName);
+			using namespace MetasoundFrontendDocumentAccessPtrPrivate;
 
-			};
+			using FFinder= TFindInArray<const FMetasoundFrontendClass, const FMetasoundFrontendClassInput, FGetClassInputArray, FIsVertexWithName>;
+			FFinder Finder{FGetClassInputArray(), FIsVertexWithName(InName)};
 
-			return GetMemberAccessPtr<FConstClassInputAccessPtr>(FindClassInputWithName);
+			return GetMemberAccessPtr<FConstClassInputAccessPtr>(Finder);
 		}
 
 
 		FConstClassInputAccessPtr FGraphClassAccessPtr::GetInputWithNodeID(const FGuid& InNodeID) const
 		{
-			auto FindClassInputWithNodeID = [NodeID=InNodeID](const FMetasoundFrontendGraphClass& InGraphClass) -> const FMetasoundFrontendClassInput*
-			{
-				return InGraphClass.Interface.Inputs.FindByPredicate([&](const FMetasoundFrontendClassInput& Desc) { return Desc.NodeID == NodeID; });
-			};
-			return GetMemberAccessPtr<FConstClassInputAccessPtr>(FindClassInputWithNodeID);
+			using namespace MetasoundFrontendDocumentAccessPtrPrivate;
+
+			using FFinder= TFindInArray<const FMetasoundFrontendGraphClass, const FMetasoundFrontendClassInput, FGetClassInputArray, FIsClassInputWithNodeID>;
+			FFinder Finder{FGetClassInputArray(), FIsClassInputWithNodeID(InNodeID)};
+
+			return GetMemberAccessPtr<FConstClassInputAccessPtr>(Finder);
 		}
 
 		FConstClassOutputAccessPtr FGraphClassAccessPtr::GetOutputWithName(const FString& InName) const
 		{
-			auto FindClassOutputWithName = [Name=InName](const FMetasoundFrontendGraphClass& InClass)
-			{
-				auto IsClassOutputWithName = [&](const FMetasoundFrontendClassOutput& ClassOutput) { return ClassOutput.Name == Name; };
-				return InClass.Interface.Outputs.FindByPredicate(IsClassOutputWithName);
+			using namespace MetasoundFrontendDocumentAccessPtrPrivate;
 
-			};
+			using FFinder= TFindInArray<const FMetasoundFrontendClass, const FMetasoundFrontendClassOutput, FGetClassOutputArray, FIsVertexWithName>;
+			FFinder Finder{FGetClassOutputArray(), FIsVertexWithName(InName)};
 
-			return GetMemberAccessPtr<FConstClassOutputAccessPtr>(FindClassOutputWithName);
+			return GetMemberAccessPtr<FConstClassOutputAccessPtr>(Finder);
 		}
 
 		FConstClassOutputAccessPtr FGraphClassAccessPtr::GetOutputWithNodeID(const FGuid& InNodeID) const
 		{
-			auto FindClassOutputWithNodeID = [NodeID=InNodeID](const FMetasoundFrontendGraphClass& InGraphClass) -> const FMetasoundFrontendClassOutput*
-			{
-				return InGraphClass.Interface.Outputs.FindByPredicate([&](const FMetasoundFrontendClassOutput& Desc) { return Desc.NodeID == NodeID; });
-			};
-			return GetMemberAccessPtr<FConstClassOutputAccessPtr>(FindClassOutputWithNodeID);
+			using namespace MetasoundFrontendDocumentAccessPtrPrivate;
+
+			using FFinder= TFindInArray<const FMetasoundFrontendGraphClass, const FMetasoundFrontendClassOutput, FGetClassOutputArray, FIsClassOutputWithNodeID>;
+			FFinder Finder{FGetClassOutputArray(), FIsClassOutputWithNodeID(InNodeID)};
+
+			return GetMemberAccessPtr<FConstClassOutputAccessPtr>(Finder);
 		}
 
 		FConstNodeAccessPtr FGraphClassAccessPtr::GetNodeWithNodeID(const FGuid& InNodeID) const
 		{
-			auto FindNodeWithNodeID = [ID=InNodeID](const FMetasoundFrontendGraphClass& InGraphClass) -> const FMetasoundFrontendNode*
-			{
-				auto IsNodeWithID = [&](const FMetasoundFrontendNode& InNode) { return InNode.ID == ID; };
-				return InGraphClass.Graph.Nodes.FindByPredicate(IsNodeWithID);
-			};
+			using namespace MetasoundFrontendDocumentAccessPtrPrivate;
+
+			using FFinder= TFindInArray<const FMetasoundFrontendGraphClass, const FMetasoundFrontendNode, FGetGraphClassNodes, FIsNodeWithID>;
+			FFinder Finder{FGetGraphClassNodes(), FIsNodeWithID(InNodeID)};
 
-			return GetMemberAccessPtr<FConstNodeAccessPtr>(FindNodeWithNodeID);
+			return GetMemberAccessPtr<FConstNodeAccessPtr>(Finder);
 		}
 
 		FConstGraphAccessPtr FGraphClassAccessPtr::GetGraph() const
@@ -367,55 +671,53 @@ namespace Metasound
 		// FConstGraphClassAccessPtr
 		FConstClassInputAccessPtr FConstGraphClassAccessPtr::GetInputWithName(const FString& InName) const
 		{
-			auto FindClassInputWithName = [Name=InName](const FMetasoundFrontendGraphClass& InClass)
-			{
-				auto IsClassInputWithName = [&](const FMetasoundFrontendClassInput& ClassInput) { return ClassInput.Name == Name; };
-				return InClass.Interface.Inputs.FindByPredicate(IsClassInputWithName);
+			using namespace MetasoundFrontendDocumentAccessPtrPrivate;
 
-			};
+			using FFinder= TFindInArray<const FMetasoundFrontendClass, const FMetasoundFrontendClassInput, FGetClassInputArray, FIsVertexWithName>;
+			FFinder Finder{FGetClassInputArray(), FIsVertexWithName(InName)};
 
-			return GetMemberAccessPtr<FConstClassInputAccessPtr>(FindClassInputWithName);
+			return GetMemberAccessPtr<FConstClassInputAccessPtr>(Finder);
 		}
 
+
 		FConstClassInputAccessPtr FConstGraphClassAccessPtr::GetInputWithNodeID(const FGuid& InNodeID) const
 		{
-			auto FindClassInputWithNodeID = [NodeID=InNodeID](const FMetasoundFrontendGraphClass& InGraphClass) -> const FMetasoundFrontendClassInput*
-			{
-				return InGraphClass.Interface.Inputs.FindByPredicate([&](const FMetasoundFrontendClassInput& Desc) { return Desc.NodeID == NodeID; });
-			};
-			return GetMemberAccessPtr<FConstClassInputAccessPtr>(FindClassInputWithNodeID);
+			using namespace MetasoundFrontendDocumentAccessPtrPrivate;
+
+			using FFinder= TFindInArray<const FMetasoundFrontendGraphClass, const FMetasoundFrontendClassInput, FGetClassInputArray, FIsClassInputWithNodeID>;
+			FFinder Finder{FGetClassInputArray(), FIsClassInputWithNodeID(InNodeID)};
+
+			return GetMemberAccessPtr<FConstClassInputAccessPtr>(Finder);
 		}
 
 		FConstClassOutputAccessPtr FConstGraphClassAccessPtr::GetOutputWithName(const FString& InName) const
 		{
-			auto FindClassOutputWithName = [Name=InName](const FMetasoundFrontendGraphClass& InClass)
-			{
-				auto IsClassOutputWithName = [&](const FMetasoundFrontendClassOutput& ClassOutput) { return ClassOutput.Name == Name; };
-				return InClass.Interface.Outputs.FindByPredicate(IsClassOutputWithName);
+			using namespace MetasoundFrontendDocumentAccessPtrPrivate;
 
-			};
+			using FFinder= TFindInArray<const FMetasoundFrontendClass, const FMetasoundFrontendClassOutput, FGetClassOutputArray, FIsVertexWithName>;
+			FFinder Finder{FGetClassOutputArray(), FIsVertexWithName(InName)};
 
-			return GetMemberAccessPtr<FConstClassOutputAccessPtr>(FindClassOutputWithName);
+			return GetMemberAccessPtr<FConstClassOutputAccessPtr>(Finder);
 		}
 
 		FConstClassOutputAccessPtr FConstGraphClassAccessPtr::GetOutputWithNodeID(const FGuid& InNodeID) const
 		{
-			auto FindClassOutputWithNodeID = [NodeID=InNodeID](const FMetasoundFrontendGraphClass& InGraphClass) -> const FMetasoundFrontendClassOutput*
-			{
-				return InGraphClass.Interface.Outputs.FindByPredicate([&](const FMetasoundFrontendClassOutput& Desc) { return Desc.NodeID == NodeID; });
-			};
-			return GetMemberAccessPtr<FConstClassOutputAccessPtr>(FindClassOutputWithNodeID);
+			using namespace MetasoundFrontendDocumentAccessPtrPrivate;
+
+			using FFinder= TFindInArray<const FMetasoundFrontendGraphClass, const FMetasoundFrontendClassOutput, FGetClassOutputArray, FIsClassOutputWithNodeID>;
+			FFinder Finder{FGetClassOutputArray(), FIsClassOutputWithNodeID(InNodeID)};
+
+			return GetMemberAccessPtr<FConstClassOutputAccessPtr>(Finder);
 		}
 
 		FConstNodeAccessPtr FConstGraphClassAccessPtr::GetNodeWithNodeID(const FGuid& InNodeID) const
 		{
-			auto FindNodeWithNodeID = [ID=InNodeID](const FMetasoundFrontendGraphClass& InGraphClass) -> const FMetasoundFrontendNode*
-			{
-				auto IsNodeWithID = [&](const FMetasoundFrontendNode& InNode) { return InNode.ID == ID; };
-				return InGraphClass.Graph.Nodes.FindByPredicate(IsNodeWithID);
-			};
+			using namespace MetasoundFrontendDocumentAccessPtrPrivate;
+
+			using FFinder= TFindInArray<const FMetasoundFrontendGraphClass, const FMetasoundFrontendNode, FGetGraphClassNodes, FIsNodeWithID>;
+			FFinder Finder{FGetGraphClassNodes(), FIsNodeWithID(InNodeID)};
 
-			return GetMemberAccessPtr<FConstNodeAccessPtr>(FindNodeWithNodeID);
+			return GetMemberAccessPtr<FConstNodeAccessPtr>(Finder);
 		}
 
 		FConstGraphAccessPtr FConstGraphClassAccessPtr::GetGraph() const
@@ -429,7 +731,6 @@ namespace Metasound
 		}
 
 		// FDocumentAccessPtr
-		
 		FGraphClassAccessPtr FDocumentAccessPtr::GetRootGraph()
 		{
 			return GetMemberAccessPtr<FGraphClassAccessPtr>([](FMetasoundFrontendDocument& InDoc) { return &InDoc.RootGraph; });
@@ -437,97 +738,63 @@ namespace Metasound
 
 		FGraphClassAccessPtr FDocumentAccessPtr::GetSubgraphWithID(const FGuid& InID)
 		{
-			auto FindSubgraphWithID = [ID=InID](FMetasoundFrontendDocument& InDoc) -> FMetasoundFrontendGraphClass*
-			{
-				return InDoc.Subgraphs.FindByPredicate([&](const FMetasoundFrontendGraphClass& InSubgraph) { return InSubgraph.ID == ID; });
-			};
+			using namespace MetasoundFrontendDocumentAccessPtrPrivate;
 
-			return GetMemberAccessPtr<FGraphClassAccessPtr>(FindSubgraphWithID);
+			using FFinder= TFindInArray<FMetasoundFrontendDocument, FMetasoundFrontendGraphClass, FGetDocumentSubgraphs, FIsClassWithID>;
+			FFinder Finder{FGetDocumentSubgraphs(), FIsClassWithID(InID)};
+
+			return GetMemberAccessPtr<FGraphClassAccessPtr>(Finder);
 		}
 
 		FClassAccessPtr FDocumentAccessPtr::GetDependencyWithID(const FGuid& InID)
 		{
-			auto FindDependencyWithID = [ID=InID](FMetasoundFrontendDocument& InDoc) -> FMetasoundFrontendClass*
-			{
-				return InDoc.Dependencies.FindByPredicate([&](const FMetasoundFrontendClass& InDependency) { return InDependency.ID == ID; });
-			};
-			return GetMemberAccessPtr<FClassAccessPtr>(FindDependencyWithID);
+			using namespace MetasoundFrontendDocumentAccessPtrPrivate;
+
+			using FFinder= TFindInArray<FMetasoundFrontendDocument, FMetasoundFrontendClass, FGetDocumentDependencies, FIsClassWithID>;
+			FFinder Finder{FGetDocumentDependencies(), FIsClassWithID(InID)};
+
+			return GetMemberAccessPtr<FClassAccessPtr>(Finder);
 		}
 
 
 		FClassAccessPtr FDocumentAccessPtr::GetClassWithID(const FGuid& InID) 
 		{
-			auto FindClassWithID = [ID=InID](FMetasoundFrontendDocument& InDoc) -> FMetasoundFrontendClass*
-			{
-				FMetasoundFrontendClass* MetasoundClass = InDoc.Dependencies.FindByPredicate([&](const FMetasoundFrontendClass& InDependency) { return InDependency.ID == ID; });
+			using namespace MetasoundFrontendDocumentAccessPtrPrivate;
 
-				if (nullptr == MetasoundClass)
-				{
-					MetasoundClass = InDoc.Subgraphs.FindByPredicate([&](const FMetasoundFrontendGraphClass& InSubgraph) { return InSubgraph.ID == ID; });
-				}
+			using FFinder = TFindClassInDocument<FMetasoundFrontendClass, FIsClassWithID>;
+			FFinder Finder(InID);
 
-				return MetasoundClass;
-			};
-			return GetMemberAccessPtr<FClassAccessPtr>(FindClassWithID);
+			return GetMemberAccessPtr<FClassAccessPtr>(Finder);
 		}
 
 		FClassAccessPtr FDocumentAccessPtr::GetClassWithMetadata(const FMetasoundFrontendClassMetadata& InMetadata)
 		{
-			auto FindClassWithMetadata = [Metadata=InMetadata](FMetasoundFrontendDocument& InDoc)
-			{
-				auto IsClassWithMetadata = [=](const FMetasoundFrontendClass& InClass) 
-				{ 
-					using namespace MetasoundFrontendDocumentAccessPtrPrivate;
+			using namespace MetasoundFrontendDocumentAccessPtrPrivate;
 
-					return IsMatchingMetasoundClass(Metadata, InClass.Metadata);
-				};
-
-				FMetasoundFrontendClass* MetasoundClass = InDoc.Dependencies.FindByPredicate(IsClassWithMetadata);
-				if (nullptr == MetasoundClass)
-				{
-					MetasoundClass = InDoc.Subgraphs.FindByPredicate(IsClassWithMetadata);
-				}
-				return MetasoundClass;
-			};
+			using FFinder = TFindClassInDocument<FMetasoundFrontendClass, FIsClassWithMetadata>;
+			FFinder Finder(InMetadata);
 
-			return GetMemberAccessPtr<FClassAccessPtr>(FindClassWithMetadata);
+			return GetMemberAccessPtr<FClassAccessPtr>(Finder);
 		}
 
 		FClassAccessPtr FDocumentAccessPtr::GetClassWithInfo(const FNodeClassInfo& InInfo)
 		{
-			auto FindClassWithInfo = [Info=InInfo](FMetasoundFrontendDocument& InDoc)
-			{
-				auto IsClassWithInfo = [=](const FMetasoundFrontendClass& InClass) 
-				{ 
-					using namespace MetasoundFrontendDocumentAccessPtrPrivate;
-					return IsMatchingMetasoundClass(Info, InClass.Metadata);
-				};
+			using namespace MetasoundFrontendDocumentAccessPtrPrivate;
 
-				FMetasoundFrontendClass* MetasoundClass = InDoc.Dependencies.FindByPredicate(IsClassWithInfo);
-				if (nullptr == MetasoundClass)
-				{
-					MetasoundClass = InDoc.Subgraphs.FindByPredicate(IsClassWithInfo);
-				}
-				return MetasoundClass;
-			};
+			using FFinder = TFindClassInDocument<FMetasoundFrontendClass, FIsClassWithNodeInfo>;
+			FFinder Finder(InInfo);
 
-			return GetMemberAccessPtr<FClassAccessPtr>(FindClassWithInfo);
+			return GetMemberAccessPtr<FClassAccessPtr>(Finder);
 		}
 
 		FClassAccessPtr FDocumentAccessPtr::GetClassWithRegistryKey(const FNodeRegistryKey& InKey)
 		{
-			auto FindClassWithRegistryKey = [Key=InKey](FMetasoundFrontendDocument& InDoc)
-			{
-				auto IsClassWithRegistryKey = [=](const FMetasoundFrontendClass& InClass) 
-				{ 
-					using namespace MetasoundFrontendDocumentAccessPtrPrivate;
-					return IsMatchingMetasoundClass(Key, InClass.Metadata);
-				};
+			using namespace MetasoundFrontendDocumentAccessPtrPrivate;
 
-				return InDoc.Dependencies.FindByPredicate(IsClassWithRegistryKey);
-			};
+			using FFinder = TFindClassInDocument<FMetasoundFrontendClass, FIsClassWithNodeRegistryKey>;
+			FFinder Finder(InKey);
 
-			return GetMemberAccessPtr<FClassAccessPtr>(FindClassWithRegistryKey);
+			return GetMemberAccessPtr<FClassAccessPtr>(Finder);
 		}
 
 		FConstGraphClassAccessPtr FDocumentAccessPtr::GetRootGraph() const
@@ -537,95 +804,63 @@ namespace Metasound
 
 		FConstGraphClassAccessPtr FDocumentAccessPtr::GetSubgraphWithID(const FGuid& InID) const
 		{
-			auto FindSubgraphWithID = [ID=InID](const FMetasoundFrontendDocument& InDoc)
-			{
-				return InDoc.Subgraphs.FindByPredicate([&](const FMetasoundFrontendGraphClass& InSubgraph) { return InSubgraph.ID == ID; });
-			};
-			return GetMemberAccessPtr<FConstGraphClassAccessPtr>(FindSubgraphWithID);
+			using namespace MetasoundFrontendDocumentAccessPtrPrivate;
+
+			using FFinder= TFindInArray<const FMetasoundFrontendDocument, const FMetasoundFrontendGraphClass, FGetDocumentSubgraphs, FIsClassWithID>;
+			FFinder Finder{FGetDocumentSubgraphs(), FIsClassWithID(InID)};
+
+			return GetMemberAccessPtr<FConstGraphClassAccessPtr>(Finder);
 		}
 
 		FConstClassAccessPtr FDocumentAccessPtr::GetDependencyWithID(const FGuid& InID) const
 		{
-			auto FindDependencyWithID = [ID=InID](const FMetasoundFrontendDocument& InDoc)
-			{
-				return InDoc.Dependencies.FindByPredicate([&](const FMetasoundFrontendClass& InDependency) { return InDependency.ID == ID; });
-			};
-			return GetMemberAccessPtr<FConstClassAccessPtr>(FindDependencyWithID);
+			using namespace MetasoundFrontendDocumentAccessPtrPrivate;
+
+			using FFinder= TFindInArray<const FMetasoundFrontendDocument, const FMetasoundFrontendClass, FGetDocumentDependencies, FIsClassWithID>;
+			FFinder Finder{FGetDocumentDependencies(), FIsClassWithID(InID)};
+
+			return GetMemberAccessPtr<FConstClassAccessPtr>(Finder);
 		}
 
 
 		FConstClassAccessPtr FDocumentAccessPtr::GetClassWithID(const FGuid& InID) const
 		{
-			auto FindClassWithID = [ID=InID](const FMetasoundFrontendDocument& InDoc) -> const FMetasoundFrontendClass*
-			{
-				const FMetasoundFrontendClass* MetasoundClass = InDoc.Dependencies.FindByPredicate([&](const FMetasoundFrontendClass& InDependency) { return InDependency.ID == ID; });
+			using namespace MetasoundFrontendDocumentAccessPtrPrivate;
 
-				if (nullptr == MetasoundClass)
-				{
-					MetasoundClass = InDoc.Subgraphs.FindByPredicate([&](const FMetasoundFrontendGraphClass& InSubgraph) { return InSubgraph.ID == ID; });
-				}
+			using FFinder = TFindClassInDocument<const FMetasoundFrontendClass, FIsClassWithID>;
+			FFinder Finder(InID);
 
-				return MetasoundClass;
-			};
-			return GetMemberAccessPtr<FConstClassAccessPtr>(FindClassWithID);
+			return GetMemberAccessPtr<FConstClassAccessPtr>(Finder);
 		}
 
 		FConstClassAccessPtr FDocumentAccessPtr::GetClassWithMetadata(const FMetasoundFrontendClassMetadata& InMetadata) const
 		{
-			auto FindClassWithMetadata = [Metadata=InMetadata](const FMetasoundFrontendDocument& InDoc)
-			{
-				auto IsClassWithMetadata = [=](const FMetasoundFrontendClass& InClass) 
-				{ 
-					using namespace MetasoundFrontendDocumentAccessPtrPrivate;
-					return IsMatchingMetasoundClass(Metadata, InClass.Metadata);
-				};
+			using namespace MetasoundFrontendDocumentAccessPtrPrivate;
 
-				const FMetasoundFrontendClass* MetasoundClass = InDoc.Dependencies.FindByPredicate(IsClassWithMetadata);
-				if (nullptr == MetasoundClass)
-				{
-					MetasoundClass = InDoc.Subgraphs.FindByPredicate(IsClassWithMetadata);
-				}
-				return MetasoundClass;
-			};
+			using FFinder = TFindClassInDocument<const FMetasoundFrontendClass, FIsClassWithMetadata>;
+			FFinder Finder(InMetadata);
 
-			return GetMemberAccessPtr<FConstClassAccessPtr>(FindClassWithMetadata);
+			return GetMemberAccessPtr<FConstClassAccessPtr>(Finder);
 		}
 
 		FConstClassAccessPtr FDocumentAccessPtr::GetClassWithInfo(const FNodeClassInfo& InInfo) const
 		{
-			auto FindClassWithInfo = [Info=InInfo](const FMetasoundFrontendDocument& InDoc)
-			{
-				auto IsClassWithInfo = [=](const FMetasoundFrontendClass& InClass) 
-				{ 
-					using namespace MetasoundFrontendDocumentAccessPtrPrivate;
-					return IsMatchingMetasoundClass(Info, InClass.Metadata);
-				};
+			using namespace MetasoundFrontendDocumentAccessPtrPrivate;
 
-				const FMetasoundFrontendClass* MetasoundClass = InDoc.Dependencies.FindByPredicate(IsClassWithInfo);
-				if (nullptr == MetasoundClass)
-				{
-					MetasoundClass = InDoc.Subgraphs.FindByPredicate(IsClassWithInfo);
-				}
-				return MetasoundClass;
-			};
+			using FFinder = TFindClassInDocument<const FMetasoundFrontendClass, FIsClassWithNodeInfo>;
+			FFinder Finder(InInfo);
 
-			return GetMemberAccessPtr<FConstClassAccessPtr>(FindClassWithInfo);
+			return GetMemberAccessPtr<FConstClassAccessPtr>(Finder);
 		}
 
 		FConstClassAccessPtr FDocumentAccessPtr::GetClassWithRegistryKey(const FNodeRegistryKey& InKey) const
 		{
-			auto FindClassWithRegistryKey = [Key=InKey](const FMetasoundFrontendDocument& InDoc)
-			{
-				auto IsClassWithRegistryKey = [=](const FMetasoundFrontendClass& InClass) 
-				{ 
-					using namespace MetasoundFrontendDocumentAccessPtrPrivate;
-					return IsMatchingMetasoundClass(Key, InClass.Metadata);
-				};
+			using namespace MetasoundFrontendDocumentAccessPtrPrivate;
 
-				return InDoc.Dependencies.FindByPredicate(IsClassWithRegistryKey);
-			};
+			using FFinder = TFindClassInDocument<const FMetasoundFrontendClass, FIsClassWithNodeRegistryKey>;
+			FFinder Finder(InKey);
 
-			return GetMemberAccessPtr<FConstClassAccessPtr>(FindClassWithRegistryKey);
+			return GetMemberAccessPtr<FConstClassAccessPtr>(Finder);
 		}
 
 		FConstGraphClassAccessPtr FConstDocumentAccessPtr::GetRootGraph() const
@@ -635,96 +870,64 @@ namespace Metasound
 
 		FConstGraphClassAccessPtr FConstDocumentAccessPtr::GetSubgraphWithID(const FGuid& InID) const
 		{
-			auto FindSubgraphWithID = [ID=InID](const FMetasoundFrontendDocument& InDoc)
-			{
-				return InDoc.Subgraphs.FindByPredicate([&](const FMetasoundFrontendGraphClass& InSubgraph) { return InSubgraph.ID == ID; });
-			};
-			return GetMemberAccessPtr<FConstGraphClassAccessPtr>(FindSubgraphWithID);
+			using namespace MetasoundFrontendDocumentAccessPtrPrivate;
+
+			using FFinder= TFindInArray<const FMetasoundFrontendDocument, const FMetasoundFrontendGraphClass, FGetDocumentSubgraphs, FIsClassWithID>;
+			FFinder Finder{FGetDocumentSubgraphs(), FIsClassWithID(InID)};
+
+			return GetMemberAccessPtr<FConstGraphClassAccessPtr>(Finder);
 		}
 
 		FConstClassAccessPtr FConstDocumentAccessPtr::GetDependencyWithID(const FGuid& InID) const
 		{
-			auto FindDependencyWithID = [ID=InID](const FMetasoundFrontendDocument& InDoc)
-			{
-				return InDoc.Dependencies.FindByPredicate([&](const FMetasoundFrontendClass& InDependency) { return InDependency.ID == ID; });
-			};
-			return GetMemberAccessPtr<FConstClassAccessPtr>(FindDependencyWithID);
+			using namespace MetasoundFrontendDocumentAccessPtrPrivate;
+
+			using FFinder= TFindInArray<const FMetasoundFrontendDocument, const FMetasoundFrontendClass, FGetDocumentDependencies, FIsClassWithID>;
+			FFinder Finder{FGetDocumentDependencies(), FIsClassWithID(InID)};
+
+			return GetMemberAccessPtr<FConstClassAccessPtr>(Finder);
 		}
 
 
 		FConstClassAccessPtr FConstDocumentAccessPtr::GetClassWithID(const FGuid& InID) const
 		{
-			auto FindClassWithID = [ID=InID](const FMetasoundFrontendDocument& InDoc)
-			{
-				const FMetasoundFrontendClass* MetasoundClass = InDoc.Dependencies.FindByPredicate([&](const FMetasoundFrontendClass& InDependency) { return InDependency.ID == ID; });
+			using namespace MetasoundFrontendDocumentAccessPtrPrivate;
 
-				if (nullptr == MetasoundClass)
-				{
-					MetasoundClass = InDoc.Subgraphs.FindByPredicate([&](const FMetasoundFrontendGraphClass& InSubgraph) { return InSubgraph.ID == ID; });
-				}
+			using FFinder = TFindClassInDocument<const FMetasoundFrontendClass, FIsClassWithID>;
+			FFinder Finder(InID);
 
-				return MetasoundClass;
-			};
-			return GetMemberAccessPtr<FConstClassAccessPtr>(FindClassWithID);
+			return GetMemberAccessPtr<FConstClassAccessPtr>(Finder);
 		}
 
 		FConstClassAccessPtr FConstDocumentAccessPtr::GetClassWithMetadata(const FMetasoundFrontendClassMetadata& InMetadata) const
 		{
-			auto FindClassWithMetadata = [Metadata=InMetadata](const FMetasoundFrontendDocument& InDoc)
-			{
-				auto IsClassWithMetadata = [=](const FMetasoundFrontendClass& InClass) 
-				{ 
-					using namespace MetasoundFrontendDocumentAccessPtrPrivate;
-					return IsMatchingMetasoundClass(Metadata, InClass.Metadata);
-				};
+			using namespace MetasoundFrontendDocumentAccessPtrPrivate;
 
-				const FMetasoundFrontendClass* MetasoundClass = InDoc.Dependencies.FindByPredicate(IsClassWithMetadata);
-				if (nullptr == MetasoundClass)
-				{
-					MetasoundClass = InDoc.Subgraphs.FindByPredicate(IsClassWithMetadata);
-				}
-				return MetasoundClass;
-			};
+			using FFinder = TFindClassInDocument<const FMetasoundFrontendClass, FIsClassWithMetadata>;
+			FFinder Finder(InMetadata);
 
-			return GetMemberAccessPtr<FConstClassAccessPtr>(FindClassWithMetadata);
+			return GetMemberAccessPtr<FConstClassAccessPtr>(Finder);
 		}
 
 		FConstClassAccessPtr FConstDocumentAccessPtr::GetClassWithInfo(const FNodeClassInfo& InInfo) const
 		{
-			auto FindClassWithInfo = [Info=InInfo](const FMetasoundFrontendDocument& InDoc)
-			{
-				auto IsClassWithInfo = [=](const FMetasoundFrontendClass& InClass) 
-				{ 
-					using namespace MetasoundFrontendDocumentAccessPtrPrivate;
-					return IsMatchingMetasoundClass(Info, InClass.Metadata);
-				};
+			using namespace MetasoundFrontendDocumentAccessPtrPrivate;
 
-				const FMetasoundFrontendClass* MetasoundClass = InDoc.Dependencies.FindByPredicate(IsClassWithInfo);
-				if (nullptr == MetasoundClass)
-				{
-					MetasoundClass = InDoc.Subgraphs.FindByPredicate(IsClassWithInfo);
-				}
-				return MetasoundClass;
-			};
+			using FFinder = TFindClassInDocument<const FMetasoundFrontendClass, FIsClassWithNodeInfo>;
+			FFinder Finder(InInfo);
 
-			return GetMemberAccessPtr<FConstClassAccessPtr>(FindClassWithInfo);
+			return GetMemberAccessPtr<FConstClassAccessPtr>(Finder);
 		}
 
 
 		FConstClassAccessPtr FConstDocumentAccessPtr::GetClassWithRegistryKey(const FNodeRegistryKey& InKey) const
 		{
-			auto FindClassWithRegistryKey = [Key=InKey](const FMetasoundFrontendDocument& InDoc)
-			{
-				auto IsClassWithRegistryKey = [=](const FMetasoundFrontendClass& InClass) 
-				{ 
-					using namespace MetasoundFrontendDocumentAccessPtrPrivate;
-					return IsMatchingMetasoundClass(Key, InClass.Metadata);
-				};
+			using namespace MetasoundFrontendDocumentAccessPtrPrivate;
 
-				return InDoc.Dependencies.FindByPredicate(IsClassWithRegistryKey);
-			};
+			using FFinder = TFindClassInDocument<const FMetasoundFrontendClass, FIsClassWithNodeRegistryKey>;
+			FFinder Finder(InKey);
 
-			return GetMemberAccessPtr<FConstClassAccessPtr>(FindClassWithRegistryKey);
+			return GetMemberAccessPtr<FConstClassAccessPtr>(Finder);
 		}
 	}
 }
diff --git a/Engine/Plugins/Runtime/Metasound/Source/MetasoundFrontend/Private/MetasoundFrontendRegistries.cpp b/Engine/Plugins/Runtime/Metasound/Source/MetasoundFrontend/Private/MetasoundFrontendRegistries.cpp
--- a/Engine/Plugins/Runtime/Metasound/Source/MetasoundFrontend/Private/MetasoundFrontendRegistries.cpp
+++ b/Engine/Plugins/Runtime/Metasound/Source/MetasoundFrontend/Private/MetasoundFrontendRegistries.cpp
@@ -53,12 +53,6 @@ namespace Metasound
 				}
 			}
 
-			// All registry keys should be created through this function to ensure consistency.
-			Frontend::FNodeRegistryKey GetRegistryKey(EMetasoundFrontendClassType InType, const FString& InFullClassName, int32 InMajorVersion, int32 InMinorVersion)
-			{
-				FString RegistryKey = FString::Format(TEXT("{0}_{1}_{2}.{3}"), {*GetClassTypeString(InType), *InFullClassName, InMajorVersion, InMinorVersion});
-				return RegistryKey;
-			}
 
 			// Return the compatible literal with the most descriptive type.
 			// TODO: Currently TIsParsable<> allows for implicit conversion of
@@ -775,7 +769,7 @@ namespace Metasound
 
 				if (InEntry.IsValid())
 				{
-					Key = FMetasoundFrontendRegistryContainer::GetRegistryKey(InEntry->GetClassInfo());
+					Key = NodeRegistryKey::CreateKey(InEntry->GetClassInfo());
 
 					// check to see if an identical node was already registered, and log
 					ensureAlwaysMsgf(
@@ -904,7 +898,7 @@ namespace Metasound
 			{
 				if (const IDataTypeRegistryEntry* Entry = FindDataTypeEntry(InDataTypeName))
 				{
-					OutKey = GetRegistryKey(Entry->GetFrontendInputClass().Metadata);
+					OutKey = NodeRegistryKey::CreateKey(Entry->GetFrontendInputClass().Metadata);
 					return true;
 				}
 				return false;
@@ -914,7 +908,7 @@ namespace Metasound
 			{
 				if (const IDataTypeRegistryEntry* Entry = FindDataTypeEntry(InDataTypeName))
 				{
-					OutKey = GetRegistryKey(Entry->GetFrontendVariableClass().Metadata);
+					OutKey = NodeRegistryKey::CreateKey(Entry->GetFrontendVariableClass().Metadata);
 					return true;
 				}
 				return false;
@@ -924,7 +918,7 @@ namespace Metasound
 			{
 				if (const IDataTypeRegistryEntry* Entry = FindDataTypeEntry(InDataTypeName))
 				{
-					OutKey = GetRegistryKey(Entry->GetFrontendOutputClass().Metadata);
+					OutKey = NodeRegistryKey::CreateKey(Entry->GetFrontendOutputClass().Metadata);
 					return true;
 				}
 				return false;
@@ -1002,9 +996,77 @@ namespace Metasound
 			return Key;
 		}
 
+		namespace NodeRegistryKey
+		{
+			// All registry keys should be created through this function to ensure consistency.
+			FNodeRegistryKey CreateKey(EMetasoundFrontendClassType InType, const FString& InFullClassName, int32 InMajorVersion, int32 InMinorVersion)
+			{
+				using namespace MetasoundFrontendRegistryPrivate;
+
+				FString RegistryKey = FString::Format(TEXT("{0}_{1}_{2}.{3}"), {*GetClassTypeString(InType), *InFullClassName, InMajorVersion, InMinorVersion});
+
+				return RegistryKey;
+			}
+
+			bool IsValid(const FNodeRegistryKey& InKey)
+			{
+				return !InKey.IsEmpty();
+			}
+
+			bool IsEqual(const FNodeRegistryKey& InLHS, const FNodeRegistryKey& InRHS)
+			{
+				return InLHS == InRHS;
+			}
+
+			bool IsEqual(const FMetasoundFrontendClassMetadata& InLHS, const FMetasoundFrontendClassMetadata& InRHS)
+			{
+				if (InLHS.ClassName == InRHS.ClassName)
+				{
+					if (InLHS.Type == InRHS.Type)
+					{
+						if (InLHS.Version == InRHS.Version)
+						{
+							return true;
+						}
+					}
+				}
+				return false;
+			}
+
+			bool IsEqual(const FNodeClassInfo& InLHS, const FMetasoundFrontendClassMetadata& InRHS)
+			{
+				if (InLHS.ClassName == InRHS.ClassName)
+				{
+					if (InLHS.Type == InRHS.Type)
+					{
+						if (InLHS.Version == InRHS.Version)
+						{
+							return true;
+						}
+					}
+				}
+				return false;
+			}
+
+			FNodeRegistryKey CreateKey(const FNodeClassMetadata& InNodeMetadata)
+			{
+				return CreateKey(EMetasoundFrontendClassType::External, InNodeMetadata.ClassName.GetFullName().ToString(), InNodeMetadata.MajorVersion, InNodeMetadata.MinorVersion);
+			}
+
+			FNodeRegistryKey CreateKey(const FMetasoundFrontendClassMetadata& InNodeMetadata)
+			{
+				return CreateKey(InNodeMetadata.Type, InNodeMetadata.ClassName.GetFullName().ToString(), InNodeMetadata.Version.Major, InNodeMetadata.Version.Minor);
+			}
+
+			FNodeRegistryKey CreateKey(const FNodeClassInfo& InClassInfo)
+			{
+				return CreateKey(InClassInfo.Type, InClassInfo.ClassName.GetFullName().ToString(), InClassInfo.Version.Major, InClassInfo.Version.Minor);
+			}
+		}
+
 		bool IsValidNodeRegistryKey(const FNodeRegistryKey& InKey)
 		{
-			return !InKey.IsEmpty();
+			return NodeRegistryKey::IsValid(InKey);
 		}
 
 		FNodeClassInfo::FNodeClassInfo(const FMetasoundFrontendClassMetadata& InMetadata)
@@ -1059,17 +1121,17 @@ void FMetasoundFrontendRegistryContainer::ShutdownMetasoundFrontend()
 
 Metasound::Frontend::FNodeRegistryKey FMetasoundFrontendRegistryContainer::GetRegistryKey(const FNodeClassMetadata& InNodeMetadata)
 {
-	return Metasound::Frontend::MetasoundFrontendRegistryPrivate::GetRegistryKey(EMetasoundFrontendClassType::External, InNodeMetadata.ClassName.GetFullName().ToString(), InNodeMetadata.MajorVersion, InNodeMetadata.MinorVersion);
+	return Metasound::Frontend::NodeRegistryKey::CreateKey(InNodeMetadata);
 }
 
 Metasound::Frontend::FNodeRegistryKey FMetasoundFrontendRegistryContainer::GetRegistryKey(const FMetasoundFrontendClassMetadata& InNodeMetadata)
 {
-	return Metasound::Frontend::MetasoundFrontendRegistryPrivate::GetRegistryKey(InNodeMetadata.Type, InNodeMetadata.ClassName.GetFullName().ToString(), InNodeMetadata.Version.Major, InNodeMetadata.Version.Minor);
+	return Metasound::Frontend::NodeRegistryKey::CreateKey(InNodeMetadata);
 }
 
 Metasound::Frontend::FNodeRegistryKey FMetasoundFrontendRegistryContainer::GetRegistryKey(const FNodeClassInfo& InClassInfo)
 {
-	return Metasound::Frontend::MetasoundFrontendRegistryPrivate::GetRegistryKey(InClassInfo.Type, InClassInfo.ClassName.GetFullName().ToString(), InClassInfo.Version.Major, InClassInfo.Version.Minor);
+	return Metasound::Frontend::NodeRegistryKey::CreateKey(InClassInfo);
 }
 
 bool FMetasoundFrontendRegistryContainer::GetFrontendClassFromRegistered(const FNodeRegistryKey& InKey, FMetasoundFrontendClass& OutClass)
diff --git a/Engine/Plugins/Runtime/Metasound/Source/MetasoundFrontend/Private/MetasoundFrontendStandardController.cpp b/Engine/Plugins/Runtime/Metasound/Source/MetasoundFrontend/Private/MetasoundFrontendStandardController.cpp
--- a/Engine/Plugins/Runtime/Metasound/Source/MetasoundFrontend/Private/MetasoundFrontendStandardController.cpp
+++ b/Engine/Plugins/Runtime/Metasound/Source/MetasoundFrontend/Private/MetasoundFrontendStandardController.cpp
@@ -84,7 +84,7 @@ namespace Metasound
 
 		bool FBaseOutputController::IsValid() const
 		{
-			return OwningNode->IsValid() && NodeVertexPtr.IsValid() && GraphPtr.IsValid();
+			return OwningNode->IsValid() && (nullptr != NodeVertexPtr.Get()) && (nullptr != GraphPtr.Get());
 		}
 
 		FGuid FBaseOutputController::GetID() const
@@ -94,18 +94,18 @@ namespace Metasound
 
 		const FName& FBaseOutputController::GetDataType() const
 		{
-			if (NodeVertexPtr.IsValid())
+			if (const FMetasoundFrontendVertex* Vertex = NodeVertexPtr.Get())
 			{
-				return NodeVertexPtr->TypeName;
+				return Vertex->TypeName;
 			}
 			return FrontendControllerIntrinsics::GetInvalidValueConstRef<FName>();
 		}
 
 		const FString& FBaseOutputController::GetName() const
 		{
-			if (NodeVertexPtr.IsValid())
+			if (const FMetasoundFrontendVertex* Vertex = NodeVertexPtr.Get())
 			{
-				return NodeVertexPtr->Name;
+				return Vertex->Name;
 			}
 			return FrontendControllerIntrinsics::GetInvalidValueConstRef<FString>();
 		}
@@ -207,7 +207,7 @@ namespace Metasound
 
 		TArray<FMetasoundFrontendEdge> FBaseOutputController::FindEdges() const
 		{
-			if (GraphPtr.IsValid())
+			if (const FMetasoundFrontendGraph* Graph = GraphPtr.Get())
 			{
 				const FGuid NodeID = GetOwningNodeID();
 				FGuid VertexID = GetID();
@@ -217,7 +217,7 @@ namespace Metasound
 					return (Edge.FromNodeID == NodeID) && (Edge.FromVertexID == VertexID);
 				};
 
-				return GraphPtr->Edges.FilterByPredicate(EdgeHasMatchingSource);
+				return Graph->Edges.FilterByPredicate(EdgeHasMatchingSource);
 			}
 
 			return TArray<FMetasoundFrontendEdge>();
@@ -225,9 +225,9 @@ namespace Metasound
 
 		const FText& FBaseOutputController::GetDisplayName() const
 		{
-			if (ClassOutputPtr.IsValid())
+			if (const FMetasoundFrontendClassOutput* Output = ClassOutputPtr.Get())
 			{
-				return ClassOutputPtr->Metadata.DisplayName;
+				return Output->Metadata.DisplayName;
 			}
 
 			return FrontendControllerIntrinsics::GetInvalidValueConstRef<FText>();
@@ -235,9 +235,9 @@ namespace Metasound
 
 		const FText& FBaseOutputController::GetTooltip() const
 		{
-			if (ClassOutputPtr.IsValid())
+			if (const FMetasoundFrontendClassOutput* Output = ClassOutputPtr.Get())
 			{
-				return ClassOutputPtr->Metadata.Description;
+				return Output->Metadata.Description;
 			}
 
 			return FrontendControllerIntrinsics::GetInvalidValueConstRef<FText>();
@@ -245,9 +245,9 @@ namespace Metasound
 
 		const FMetasoundFrontendVertexMetadata& FBaseOutputController::GetMetadata() const
 		{
-			if (ClassOutputPtr.IsValid())
+			if (const FMetasoundFrontendClassOutput* Output = ClassOutputPtr.Get())
 			{
-				return ClassOutputPtr->Metadata;
+				return Output->Metadata;
 			}
 
 			return FrontendControllerIntrinsics::GetInvalidValueConstRef<FMetasoundFrontendVertexMetadata>();
@@ -287,7 +287,7 @@ namespace Metasound
 
 		bool FInputNodeOutputController::IsValid() const
 		{
-			return FBaseOutputController::IsValid() && OwningGraphClassInputPtr.IsValid();
+			return FBaseOutputController::IsValid() && (nullptr != OwningGraphClassInputPtr.Get());
 		}
 
 		const FText& FInputNodeOutputController::GetDisplayName() const
@@ -311,9 +311,9 @@ namespace Metasound
 
 		const FText& FInputNodeOutputController::GetTooltip() const
 		{
-			if (OwningGraphClassInputPtr.IsValid())
+			if (const FMetasoundFrontendClassInput* Input = OwningGraphClassInputPtr.Get())
 			{
-				return OwningGraphClassInputPtr->Metadata.Description;
+				return Input->Metadata.Description;
 			}
 
 			return FrontendControllerIntrinsics::GetInvalidValueConstRef<FText>();
@@ -321,9 +321,9 @@ namespace Metasound
 
 		const FMetasoundFrontendVertexMetadata& FInputNodeOutputController::GetMetadata() const
 		{
-			if (OwningGraphClassInputPtr.IsValid())
+			if (const FMetasoundFrontendClassInput* Input = OwningGraphClassInputPtr.Get())
 			{
-				return OwningGraphClassInputPtr->Metadata;
+				return Input->Metadata;
 			}
 
 			return FrontendControllerIntrinsics::GetInvalidValueConstRef<FMetasoundFrontendVertexMetadata>();
@@ -358,7 +358,7 @@ namespace Metasound
 
 		bool FOutputNodeOutputController::IsValid() const
 		{
-			return FBaseOutputController::IsValid() && OwningGraphClassOutputPtr.IsValid();
+			return FBaseOutputController::IsValid() && (nullptr != OwningGraphClassOutputPtr.Get());
 		}
 
 		const FText& FOutputNodeOutputController::GetDisplayName() const
@@ -373,9 +373,9 @@ namespace Metasound
 
 		const FText& FOutputNodeOutputController::GetTooltip() const
 		{
-			if (OwningGraphClassOutputPtr.IsValid())
+			if (const FMetasoundFrontendClassOutput* OwningOutput = OwningGraphClassOutputPtr.Get())
 			{
-				return OwningGraphClassOutputPtr->Metadata.Description;
+				return OwningOutput->Metadata.Description;
 			}
 
 			return FrontendControllerIntrinsics::GetInvalidValueConstRef<FText>();
@@ -383,9 +383,9 @@ namespace Metasound
 
 		const FMetasoundFrontendVertexMetadata& FOutputNodeOutputController::GetMetadata() const
 		{
-			if (OwningGraphClassOutputPtr.IsValid())
+			if (const FMetasoundFrontendClassOutput* OwningOutput = OwningGraphClassOutputPtr.Get())
 			{
-				return OwningGraphClassOutputPtr->Metadata;
+				return OwningOutput->Metadata;
 			}
 
 			return FrontendControllerIntrinsics::GetInvalidValueConstRef<FMetasoundFrontendVertexMetadata>();
@@ -423,7 +423,7 @@ namespace Metasound
 
 		bool FBaseInputController::IsValid() const 
 		{
-			return OwningNode->IsValid() && NodeVertexPtr.IsValid() &&  GraphPtr.IsValid();
+			return OwningNode->IsValid() && (nullptr != NodeVertexPtr.Get()) &&  (nullptr != GraphPtr.Get());
 		}
 
 		FGuid FBaseInputController::GetID() const
@@ -433,9 +433,9 @@ namespace Metasound
 
 		const FName& FBaseInputController::GetDataType() const
 		{
-			if (NodeVertexPtr.IsValid())
+			if (const FMetasoundFrontendVertex* Vertex = NodeVertexPtr.Get())
 			{
-				return NodeVertexPtr->TypeName;
+				return Vertex->TypeName;
 			}
 			
 			return FrontendControllerIntrinsics::GetInvalidValueConstRef<FName>();
@@ -443,9 +443,9 @@ namespace Metasound
 
 		const FString& FBaseInputController::GetName() const
 		{
-			if (NodeVertexPtr.IsValid())
+			if (const FMetasoundFrontendVertex* Vertex = NodeVertexPtr.Get())
 			{
-				return NodeVertexPtr->Name;
+				return Vertex->Name;
 			}
 			
 			return FrontendControllerIntrinsics::GetInvalidValueConstRef<FString>();
@@ -463,10 +463,9 @@ namespace Metasound
 
 		const FMetasoundFrontendLiteral* FBaseInputController::GetLiteral() const
 		{
-			if (NodeVertexPtr.IsValid() && OwningNode->IsValid())
+			if (const FMetasoundFrontendVertex* Vertex = NodeVertexPtr.Get())
 			{
-				const FGuid& VertexID = NodeVertexPtr->VertexID;
-				return OwningNode->GetInputLiteral(VertexID);
+				return OwningNode->GetInputLiteral(Vertex->VertexID);
 			}
 
 			return nullptr;
@@ -474,20 +473,19 @@ namespace Metasound
 
 		void FBaseInputController::SetLiteral(const FMetasoundFrontendLiteral& InLiteral)
 		{
-			if (NodeVertexPtr.IsValid() && OwningNode->IsValid())
+			if (const FMetasoundFrontendVertex* Vertex = NodeVertexPtr.Get())
 			{
-				const FGuid& VertexID = NodeVertexPtr->VertexID;
 				if (const FMetasoundFrontendLiteral* ClassLiteral = GetClassDefaultLiteral())
 				{
 					// Clear if equivalent to class default as fallback is the class default
 					if (ClassLiteral->IsEquivalent(InLiteral))
 					{
-						OwningNode->ClearInputLiteral(VertexID);
+						OwningNode->ClearInputLiteral(Vertex->VertexID);
 						return;
 					}
 				}
 
-				OwningNode->SetInputLiteral(FMetasoundFrontendVertexLiteral{ VertexID, InLiteral });
+				OwningNode->SetInputLiteral(FMetasoundFrontendVertexLiteral{ Vertex->VertexID, InLiteral });
 			}
 		}
 
@@ -571,12 +569,15 @@ namespace Metasound
 			FConnectability OutConnectability;
 			OutConnectability.Connectable = FConnectability::EConnectable::No;
 
-			if (!(InController.IsValid() && IsValid()))
+			const FName& DataType = GetDataType();
+			const FName& OtherDataType = InController.GetDataType();
+
+			if (!DataType.IsValid())
 			{
 				return OutConnectability;
 			}
 
-			if (InController.GetDataType() == GetDataType())
+			if (OtherDataType == DataType)
 			{
 				// If data types are equal, connection can happen.
 				OutConnectability.Connectable = FConnectability::EConnectable::Yes;
@@ -585,7 +586,7 @@ namespace Metasound
 
 			// If data types are not equal, check for converter nodes which could
 			// convert data type.
-			OutConnectability.PossibleConverterNodeClasses = FRegistry::Get()->GetPossibleConverterNodes(InController.GetDataType(), GetDataType());
+			OutConnectability.PossibleConverterNodeClasses = FRegistry::Get()->GetPossibleConverterNodes(OtherDataType, DataType);
 
 			if (OutConnectability.PossibleConverterNodeClasses.Num() > 0)
 			{
@@ -598,27 +599,33 @@ namespace Metasound
 
 		bool FBaseInputController::Connect(IOutputController& InController)
 		{
-			if (!IsValid() || !InController.IsValid())
+			const FName& DataType = GetDataType();
+			const FName& OtherDataType = InController.GetDataType();
+
+			if (!DataType.IsValid())
 			{
 				return false;
 			}
 
-			if (ensureAlwaysMsgf(InController.GetDataType() == GetDataType(), TEXT("Cannot connect incompatible types.")))
+			if (FMetasoundFrontendGraph* Graph = GraphPtr.Get())
 			{
-				// Overwrite an existing connection if it exists.
-				FMetasoundFrontendEdge* Edge = FindEdge();
-
-				if (!Edge)
+				if (ensureAlwaysMsgf(OtherDataType == DataType, TEXT("Cannot connect incompatible types.")))
 				{
-					Edge = &GraphPtr->Edges.AddDefaulted_GetRef();
-					Edge->ToNodeID = GetOwningNodeID();
-					Edge->ToVertexID = GetID();
-				}
+					// Overwrite an existing connection if it exists.
+					FMetasoundFrontendEdge* Edge = FindEdge();
+
+					if (!Edge)
+					{
+						Edge = &Graph->Edges.AddDefaulted_GetRef();
+						Edge->ToNodeID = GetOwningNodeID();
+						Edge->ToVertexID = GetID();
+					}
 
-				Edge->FromNodeID = InController.GetOwningNodeID();
-				Edge->FromVertexID = InController.GetID();
+					Edge->FromNodeID = InController.GetOwningNodeID();
+					Edge->FromVertexID = InController.GetID();
 
-				return true;
+					return true;
+				}
 			}
 
 			return false;
@@ -626,10 +633,10 @@ namespace Metasound
 
 		bool FBaseInputController::ConnectWithConverterNode(IOutputController& InController, const FConverterNodeInfo& InConverterInfo)
 		{
-			FGraphHandle Graph = OwningNode->GetOwningGraph();
+			FGraphHandle OwningGraph = OwningNode->GetOwningGraph();
 
 			// Generate the converter node.
-			FNodeHandle ConverterNode = Graph->AddNode(InConverterInfo.NodeKey);
+			FNodeHandle ConverterNode = OwningGraph->AddNode(InConverterInfo.NodeKey);
 
 			TArray<FInputHandle> ConverterInputs = ConverterNode->GetInputsWithVertexName(InConverterInfo.PreferredConverterInputPin);
 			TArray<FOutputHandle> ConverterOutputs = ConverterNode->GetOutputsWithVertexName(InConverterInfo.PreferredConverterOutputPin);
@@ -660,7 +667,7 @@ namespace Metasound
 
 		bool FBaseInputController::Disconnect(IOutputController& InController) 
 		{
-			if (GraphPtr.IsValid())
+			if (FMetasoundFrontendGraph* Graph = GraphPtr.Get())
 			{
 				FGuid FromNodeID = InController.GetOwningNodeID();
 				FGuid FromVertexID = InController.GetID();
@@ -672,7 +679,7 @@ namespace Metasound
 					return (Edge.FromNodeID == FromNodeID) && (Edge.FromVertexID == FromVertexID) && (Edge.ToNodeID == ToNodeID) && (Edge.ToVertexID == ToVertexID);
 				};
 
-				int32 NumRemoved = GraphPtr->Edges.RemoveAllSwap(IsMatchingEdge);
+				int32 NumRemoved = Graph->Edges.RemoveAllSwap(IsMatchingEdge);
 				return NumRemoved > 0;
 			}
 
@@ -681,7 +688,7 @@ namespace Metasound
 
 		bool FBaseInputController::Disconnect()
 		{
-			if (GraphPtr.IsValid())
+			if (FMetasoundFrontendGraph* Graph = GraphPtr.Get())
 			{
 				const FGuid NodeID = GetOwningNodeID();
 				FGuid VertexID = GetID();
@@ -691,7 +698,7 @@ namespace Metasound
 					return (Edge.ToNodeID == NodeID) && (Edge.ToVertexID == VertexID);
 				};
 
-				int32 NumRemoved = GraphPtr->Edges.RemoveAllSwap(EdgeHasMatchingDestination);
+				int32 NumRemoved = Graph->Edges.RemoveAllSwap(EdgeHasMatchingDestination);
 				return NumRemoved > 0;
 			}
 
@@ -700,7 +707,7 @@ namespace Metasound
 
 		const FMetasoundFrontendEdge* FBaseInputController::FindEdge() const
 		{
-			if (GraphPtr.IsValid())
+			if (const FMetasoundFrontendGraph* Graph = GraphPtr.Get())
 			{
 				const FGuid NodeID = GetOwningNodeID();
 				FGuid VertexID = GetID();
@@ -710,7 +717,7 @@ namespace Metasound
 					return (Edge.ToNodeID == NodeID) && (Edge.ToVertexID == VertexID);
 				};
 
-				return GraphPtr->Edges.FindByPredicate(EdgeHasMatchingDestination);
+				return Graph->Edges.FindByPredicate(EdgeHasMatchingDestination);
 			}
 
 			return nullptr;
@@ -718,7 +725,7 @@ namespace Metasound
 
 		FMetasoundFrontendEdge* FBaseInputController::FindEdge()
 		{
-			if (GraphPtr.IsValid())
+			if (FMetasoundFrontendGraph* Graph = GraphPtr.Get())
 			{
 				const FGuid NodeID = GetOwningNodeID();
 				FGuid VertexID = GetID();
@@ -728,7 +735,7 @@ namespace Metasound
 					return (Edge.ToNodeID == NodeID) && (Edge.ToVertexID == VertexID);
 				};
 
-				return GraphPtr->Edges.FindByPredicate(EdgeHasMatchingDestination);
+				return Graph->Edges.FindByPredicate(EdgeHasMatchingDestination);
 			}
 
 			return nullptr;
@@ -769,7 +776,7 @@ namespace Metasound
 
 		bool FOutputNodeInputController::IsValid() const
 		{
-			return FBaseInputController::IsValid() && OwningGraphClassOutputPtr.IsValid();
+			return FBaseInputController::IsValid() && (nullptr != OwningGraphClassOutputPtr.Get());
 		}
 
 		const FText& FOutputNodeInputController::GetDisplayName() const
@@ -794,9 +801,9 @@ namespace Metasound
 
 		const FText& FOutputNodeInputController::GetTooltip() const
 		{
-			if (OwningGraphClassOutputPtr.IsValid())
+			if (const FMetasoundFrontendClassOutput* OwningOutput = OwningGraphClassOutputPtr.Get())
 			{
-				return OwningGraphClassOutputPtr->Metadata.Description;
+				return OwningOutput->Metadata.Description;
 			}
 			
 			return FrontendControllerIntrinsics::GetInvalidValueConstRef<FText>();
@@ -804,9 +811,9 @@ namespace Metasound
 
 		const FMetasoundFrontendVertexMetadata& FOutputNodeInputController::GetMetadata() const
 		{
-			if (OwningGraphClassOutputPtr.IsValid())
+			if (const FMetasoundFrontendClassOutput* OwningOutput = OwningGraphClassOutputPtr.Get())
 			{
-				return OwningGraphClassOutputPtr->Metadata;
+				return OwningOutput->Metadata;
 			}
 
 			return FrontendControllerIntrinsics::GetInvalidValueConstRef<FMetasoundFrontendVertexMetadata>();
@@ -842,7 +849,7 @@ namespace Metasound
 
 		bool FInputNodeInputController::IsValid() const
 		{
-			return FBaseInputController::IsValid() && OwningGraphClassInputPtr.IsValid();
+			return FBaseInputController::IsValid() && (nullptr != OwningGraphClassInputPtr.Get());
 		}
 
 		const FText& FInputNodeInputController::GetDisplayName() const
@@ -899,19 +906,22 @@ namespace Metasound
 		, ClassPtr(InParams.ClassPtr)
 		, OwningGraph(InParams.OwningGraph)
 		{
-			if (NodePtr.IsValid() && ClassPtr.IsValid())
+			if (FMetasoundFrontendNode* Node = NodePtr.Get())
 			{
-				if (NodePtr->ClassID != ClassPtr->ID)
+				if (const FMetasoundFrontendClass* Class = ClassPtr.Get())
 				{
-					UE_LOG(LogMetaSound, Warning, TEXT("Changing node's class id from [ClassID:%s] to [ClassID:%s]"), *NodePtr->ClassID.ToString(), *ClassPtr->ID.ToString());
-					NodePtr->ClassID = ClassPtr->ID;
+					if (Node->ClassID != Class->ID)
+					{
+						UE_LOG(LogMetaSound, Warning, TEXT("Changing node's class id from [ClassID:%s] to [ClassID:%s]"), *Node->ClassID.ToString(), *Class->ID.ToString());
+						Node->ClassID = Class->ID;
+					}
 				}
 			}
 		}
 
 		bool FBaseNodeController::IsValid() const
 		{
-			return OwningGraph->IsValid() && NodePtr.IsValid() && ClassPtr.IsValid();
+			return OwningGraph->IsValid() && (nullptr != NodePtr.Get()) && (nullptr != ClassPtr.Get());
 		}
 
 		FGuid FBaseNodeController::GetOwningGraphClassID() const
@@ -931,27 +941,27 @@ namespace Metasound
 
 		FGuid FBaseNodeController::GetID() const
 		{
-			if (NodePtr.IsValid())
+			if (const FMetasoundFrontendNode* Node = NodePtr.Get())
 			{
-				return NodePtr->ID;
+				return Node->ID;
 			}
 			return Metasound::FrontendInvalidID;
 		}
 
 		FGuid FBaseNodeController::GetClassID() const
 		{
-			if (ClassPtr.IsValid())
+			if (const FMetasoundFrontendClass* Class = ClassPtr.Get())
 			{
-				return ClassPtr->ID;
+				return Class->ID;
 			}
 			return Metasound::FrontendInvalidID;
 		}
 
 		const FMetasoundFrontendLiteral* FBaseNodeController::GetInputLiteral(const FGuid& InVertexID) const
 		{
-			if (NodePtr.IsValid())
+			if (const FMetasoundFrontendNode* Node = NodePtr.Get())
 			{
-				for (const FMetasoundFrontendVertexLiteral& VertexLiteral : NodePtr->InputLiterals)
+				for (const FMetasoundFrontendVertexLiteral& VertexLiteral : Node->InputLiterals)
 				{
 					if (VertexLiteral.VertexID == InVertexID)
 					{
@@ -965,20 +975,20 @@ namespace Metasound
 
 		void FBaseNodeController::SetInputLiteral(const FMetasoundFrontendVertexLiteral& InVertexLiteral)
 		{
-			if (NodePtr.IsValid())
+			if (FMetasoundFrontendNode* Node = NodePtr.Get())
 			{
 				auto IsInputVertex = [InVertexLiteral] (const FMetasoundFrontendVertex& Vertex)
 				{
 					return InVertexLiteral.VertexID == Vertex.VertexID;
 				};
 
-				FMetasoundFrontendNodeInterface& NodeInterface = NodePtr->Interface;
+				FMetasoundFrontendNodeInterface& NodeInterface = Node->Interface;
 				if (!ensure(NodeInterface.Inputs.ContainsByPredicate(IsInputVertex)))
 				{
 					return;
 				}
 
-				for (FMetasoundFrontendVertexLiteral& VertexLiteral : NodePtr->InputLiterals)
+				for (FMetasoundFrontendVertexLiteral& VertexLiteral : Node->InputLiterals)
 				{
 					if (VertexLiteral.VertexID == InVertexLiteral.VertexID)
 					{
@@ -990,20 +1000,20 @@ namespace Metasound
 					}
 				}
 
-				NodePtr->InputLiterals.Add(InVertexLiteral);
+				Node->InputLiterals.Add(InVertexLiteral);
 			}
 		}
 
 		bool FBaseNodeController::ClearInputLiteral(FGuid InVertexID)
 		{
-			if (NodePtr.IsValid())
+			if (FMetasoundFrontendNode* Node = NodePtr.Get())
 			{
 				auto IsInputVertex = [InVertexID](const FMetasoundFrontendVertexLiteral& VertexLiteral)
 				{
 					return InVertexID == VertexLiteral.VertexID;
 				};
 
-				return NodePtr->InputLiterals.RemoveAllSwap(IsInputVertex, false) > 0;
+				return Node->InputLiterals.RemoveAllSwap(IsInputVertex, false) > 0;
 			}
 
 			return false;
@@ -1011,27 +1021,27 @@ namespace Metasound
 
 		const FMetasoundFrontendClassInterface& FBaseNodeController::GetClassInterface() const
 		{
-			if (ClassPtr.IsValid())
+			if (const FMetasoundFrontendClass* Class = ClassPtr.Get())
 			{
-				return ClassPtr->Interface;
+				return Class->Interface;
 			}
 			return FrontendControllerIntrinsics::GetInvalidValueConstRef<FMetasoundFrontendClassInterface>();
 		}
 
 		const FMetasoundFrontendClassMetadata& FBaseNodeController::GetClassMetadata() const
 		{
-			if (ClassPtr.IsValid())
+			if (const FMetasoundFrontendClass* Class = ClassPtr.Get())
 			{
-				return ClassPtr->Metadata;
+				return Class->Metadata;
 			}
 			return FrontendControllerIntrinsics::GetInvalidValueConstRef<FMetasoundFrontendClassMetadata>();
 		}
 
 		const FMetasoundFrontendInterfaceStyle& FBaseNodeController::GetInputStyle() const
 		{
-			if (ClassPtr.IsValid())
+			if (const FMetasoundFrontendClass* Class = ClassPtr.Get())
 			{
-				return ClassPtr->Interface.InputStyle;
+				return Class->Interface.InputStyle;
 			}
 
 			static const FMetasoundFrontendInterfaceStyle Invalid;
@@ -1040,9 +1050,9 @@ namespace Metasound
 
 		const FMetasoundFrontendInterfaceStyle& FBaseNodeController::GetOutputStyle() const
 		{
-			if (ClassPtr.IsValid())
+			if (const FMetasoundFrontendClass* Class = ClassPtr.Get())
 			{
-				return ClassPtr->Interface.OutputStyle;
+				return Class->Interface.OutputStyle;
 			}
 
 			static const FMetasoundFrontendInterfaceStyle Invalid;
@@ -1051,9 +1061,9 @@ namespace Metasound
 
 		const FMetasoundFrontendClassStyle& FBaseNodeController::GetClassStyle() const
 		{
-			if (ClassPtr.IsValid())
+			if (const FMetasoundFrontendClass* Class = ClassPtr.Get())
 			{
-				return ClassPtr->Style;
+				return Class->Style;
 			}
 
 			static const FMetasoundFrontendClassStyle Invalid;
@@ -1062,9 +1072,9 @@ namespace Metasound
 
 		const FMetasoundFrontendNodeStyle& FBaseNodeController::GetNodeStyle() const
 		{
-			if (NodePtr.IsValid())
+			if (const FMetasoundFrontendNode* Node = NodePtr.Get())
 			{
-				return NodePtr->Style;
+				return Node->Style;
 			}
 
 			static const FMetasoundFrontendNodeStyle Invalid;
@@ -1073,26 +1083,26 @@ namespace Metasound
 
 		void FBaseNodeController::SetNodeStyle(const FMetasoundFrontendNodeStyle& InStyle)
 		{
-			if (NodePtr.IsValid())
+			if (FMetasoundFrontendNode* Node = NodePtr.Get())
 			{
-				NodePtr->Style = InStyle;
+				Node->Style = InStyle;
 			}
 		}
 
 		const FText& FBaseNodeController::GetDescription() const
 		{
-			if (ClassPtr.IsValid())
+			if (const FMetasoundFrontendClass* Class = ClassPtr.Get())
 			{
-				return ClassPtr->Metadata.Description;
+				return Class->Metadata.Description;
 			}
 			return FrontendControllerIntrinsics::GetInvalidValueConstRef<FText>();
 		}
 
 		const FString& FBaseNodeController::GetNodeName() const
 		{
-			if (NodePtr.IsValid())
+			if (const FMetasoundFrontendNode* Node = NodePtr.Get())
 			{
-				return NodePtr->Name;
+				return Node->Name;
 			}
 			return FrontendControllerIntrinsics::GetInvalidValueConstRef<FString>();
 		}
@@ -1157,9 +1167,9 @@ namespace Metasound
 
 		int32 FBaseNodeController::GetNumInputs() const
 		{
-			if (NodePtr.IsValid())
+			if (const FMetasoundFrontendNode* Node = NodePtr.Get())
 			{
-				return NodePtr->Interface.Inputs.Num();
+				return Node->Interface.Inputs.Num();
 			}
 
 			return 0;
@@ -1197,9 +1207,9 @@ namespace Metasound
 
 		int32 FBaseNodeController::GetNumOutputs() const
 		{
-			if (NodePtr.IsValid())
+			if (const FMetasoundFrontendNode* Node = NodePtr.Get())
 			{
-				return NodePtr->Interface.Outputs.Num();
+				return Node->Interface.Outputs.Num();
 			}
 
 			return 0;
@@ -1245,9 +1255,9 @@ namespace Metasound
 
 		const FText& FBaseNodeController::GetDisplayTitle() const
 		{
-			if (ClassPtr.IsValid())
+			if (const FMetasoundFrontendClass* Class = ClassPtr.Get())
 			{
-				return ClassPtr->Metadata.DisplayName;
+				return Class->Metadata.DisplayName;
 			}
 
 			return FrontendControllerIntrinsics::GetInvalidValueConstRef<FText>();
@@ -1255,9 +1265,9 @@ namespace Metasound
 
 		const FText& FBaseNodeController::GetDisplayName() const
 		{
-			if (ClassPtr.IsValid())
+			if (const FMetasoundFrontendClass* Class = ClassPtr.Get())
 			{
-				return ClassPtr->Metadata.DisplayName;
+				return Class->Metadata.DisplayName;
 			}
 			return FText::GetEmpty();
 		}
@@ -1549,9 +1559,9 @@ namespace Metasound
 			// TODO: consider adding support for external graph owned in another document.
 			// Will require lookup support for external subgraphs..
 			
-			if (ClassPtr.IsValid())
+			if (const FMetasoundFrontendClass* Class = ClassPtr.Get())
 			{
-				return GetOwningGraph()->GetOwningDocument()->GetSubgraphWithClassID(ClassPtr->ID);
+				return GetOwningGraph()->GetOwningDocument()->GetSubgraphWithClassID(Class->ID);
 			}
 
 			return FInvalidGraphController::GetInvalid();
@@ -1561,9 +1571,9 @@ namespace Metasound
 		{
 			// TODO: add support for graph owned in another asset.
 			// Will require lookup support for external subgraphs..
-			if (ClassPtr.IsValid())
+			if (const FMetasoundFrontendClass* Class = ClassPtr.Get())
 			{
-				return GetOwningGraph()->GetOwningDocument()->GetSubgraphWithClassID(ClassPtr->ID);
+				return GetOwningGraph()->GetOwningDocument()->GetSubgraphWithClassID(Class->ID);
 			}
 
 			return FInvalidGraphController::GetInvalid();
@@ -1603,16 +1613,19 @@ namespace Metasound
 
 		FNodeHandle FNodeController::CreateNodeHandle(const FNodeController::FInitParams& InParams)
 		{
-			if (InParams.NodePtr.IsValid() && InParams.ClassPtr.IsValid())
+			if (const FMetasoundFrontendNode* Node = InParams.NodePtr.Get())
 			{
-				// Cannot make a valid node handle if the node description and class description differ
-				if (InParams.NodePtr->ClassID == InParams.ClassPtr->ID)
+				if (const FMetasoundFrontendClass* Class = InParams.ClassPtr.Get())
 				{
-					return MakeShared<FNodeController>(EPrivateToken::Token, InParams);
-				}
-				else
-				{
-					UE_LOG(LogMetaSound, Warning, TEXT("Frontend Node [NodeID:%s, ClassID:%s] is not of expected class class [ClassID:%s]"), *InParams.NodePtr->ID.ToString(), *InParams.NodePtr->ClassID.ToString(), *InParams.ClassPtr->ID.ToString());
+					// Cannot make a valid node handle if the node description and class description differ
+					if (Node->ClassID == Class->ID)
+					{
+						return MakeShared<FNodeController>(EPrivateToken::Token, InParams);
+					}
+					else
+					{
+						UE_LOG(LogMetaSound, Warning, TEXT("Frontend Node [NodeID:%s, ClassID:%s] is not of expected class class [ClassID:%s]"), *Node->ID.ToString(), *Node->ClassID.ToString(), *Class->ID.ToString());
+					}
 				}
 			}
 			return FInvalidNodeController::GetInvalid();
@@ -1620,16 +1633,19 @@ namespace Metasound
 
 		FConstNodeHandle FNodeController::CreateConstNodeHandle(const FNodeController::FInitParams& InParams)
 		{
-			if (InParams.NodePtr.IsValid() && InParams.ClassPtr.IsValid())
+			if (const FMetasoundFrontendNode* Node = InParams.NodePtr.Get())
 			{
-				// Cannot make a valid node handle if the node description and class description differ
-				if (InParams.NodePtr->ClassID == InParams.ClassPtr->ID)
+				if (const FMetasoundFrontendClass* Class = InParams.ClassPtr.Get())
 				{
-					return MakeShared<const FNodeController>(EPrivateToken::Token, InParams);
-				}
-				else
-				{
-					UE_LOG(LogMetaSound, Warning, TEXT("Frontend Node [NodeID:%s, ClassID:%s] is not of expected class class [ClassID:%s]"), *InParams.NodePtr->ID.ToString(), *InParams.NodePtr->ClassID.ToString(), *InParams.ClassPtr->ID.ToString());
+					// Cannot make a valid node handle if the node description and class description differ
+					if (Node->ClassID == Class->ID)
+					{
+						return MakeShared<const FNodeController>(EPrivateToken::Token, InParams);
+					}
+					else
+					{
+						UE_LOG(LogMetaSound, Warning, TEXT("Frontend Node [NodeID:%s, ClassID:%s] is not of expected class class [ClassID:%s]"), *Node->ID.ToString(), *Node->ClassID.ToString(), *Class->ID.ToString());
+					}
 				}
 			}
 			return FInvalidNodeController::GetInvalid();
@@ -1637,7 +1653,7 @@ namespace Metasound
 
 		bool FNodeController::IsValid() const
 		{
-			return FBaseNodeController::IsValid() && GraphPtr.IsValid();
+			return FBaseNodeController::IsValid() && (nullptr != GraphPtr.Get());
 		}
 
 		FInputHandle FNodeController::CreateInputController(FGuid InVertexID, FConstVertexAccessPtr InNodeVertexPtr, FConstClassInputAccessPtr InClassInputPtr, FNodeHandle InOwningNode) const
@@ -1682,23 +1698,26 @@ namespace Metasound
 
 		FNodeHandle FOutputNodeController::CreateOutputNodeHandle(const FOutputNodeController::FInitParams& InParams)
 		{
-			if (InParams.NodePtr.IsValid() && InParams.ClassPtr.IsValid())
+			if (const FMetasoundFrontendNode* Node = InParams.NodePtr.Get())
 			{
-				if (EMetasoundFrontendClassType::Output == InParams.ClassPtr->Metadata.Type)
+				if (const FMetasoundFrontendClass* Class = InParams.ClassPtr.Get())
 				{
-					if (InParams.ClassPtr->ID == InParams.NodePtr->ClassID)
+					if (EMetasoundFrontendClassType::Output == Class->Metadata.Type)
 					{
-						return MakeShared<FOutputNodeController>(EPrivateToken::Token, InParams);
+						if (Class->ID == Node->ClassID)
+						{
+							return MakeShared<FOutputNodeController>(EPrivateToken::Token, InParams);
+						}
+						else
+						{
+							UE_LOG(LogMetaSound, Warning, TEXT("Frontend Node [NodeID:%s, ClassID:%s] is not of expected class class [ClassID:%s]"), *Node->ID.ToString(), *Node->ClassID.ToString(), *Class->ID.ToString());
+						}
 					}
 					else
 					{
-						UE_LOG(LogMetaSound, Warning, TEXT("Frontend Node [NodeID:%s, ClassID:%s] is not of expected class class [ClassID:%s]"), *InParams.NodePtr->ID.ToString(), *InParams.NodePtr->ClassID.ToString(), *InParams.ClassPtr->ID.ToString());
+						UE_LOG(LogMetaSound, Warning, TEXT("Frontend Class of incorrect class type [ClassID:%s] while creating output node.. Must be EMetasoundFrontendClassType::Output."), *Class->ID.ToString());
 					}
 				}
-				else
-				{
-					UE_LOG(LogMetaSound, Warning, TEXT("Frontend Class of incorrect class type [ClassID:%s] while creating output node.. Must be EMetasoundFrontendClassType::Output."), *InParams.ClassPtr->ID.ToString());
-				}
 			}
 
 			return FInvalidNodeController::GetInvalid();
@@ -1706,9 +1725,9 @@ namespace Metasound
 
 		const FText& FOutputNodeController::GetDescription() const
 		{
-			if (OwningGraphClassOutputPtr.IsValid())
+			if (const FMetasoundFrontendClassOutput* OwningOutput = OwningGraphClassOutputPtr.Get())
 			{
-				return OwningGraphClassOutputPtr->Metadata.Description;
+				return OwningOutput->Metadata.Description;
 			}
 
 			return FrontendControllerIntrinsics::GetInvalidValueConstRef<FText>();
@@ -1716,9 +1735,9 @@ namespace Metasound
 
 		const FText& FOutputNodeController::GetDisplayName() const
 		{
-			if (OwningGraphClassOutputPtr.IsValid())
+			if (const FMetasoundFrontendClassOutput* OwningOutput = OwningGraphClassOutputPtr.Get())
 			{
-				return OwningGraphClassOutputPtr->Metadata.DisplayName;
+				return OwningOutput->Metadata.DisplayName;
 			}
 
 			return FrontendControllerIntrinsics::GetInvalidValueConstRef<FText>();
@@ -1744,23 +1763,26 @@ namespace Metasound
 
 		FConstNodeHandle FOutputNodeController::CreateConstOutputNodeHandle(const FOutputNodeController::FInitParams& InParams)
 		{
-			if (InParams.NodePtr.IsValid() && InParams.ClassPtr.IsValid())
+			if (const FMetasoundFrontendNode* Node = InParams.NodePtr.Get())
 			{
-				if (EMetasoundFrontendClassType::Output == InParams.ClassPtr->Metadata.Type)
+				if (const FMetasoundFrontendClass* Class = InParams.ClassPtr.Get())
 				{
-					if (InParams.ClassPtr->ID == InParams.NodePtr->ClassID)
+					if (EMetasoundFrontendClassType::Output == Class->Metadata.Type)
 					{
-						return MakeShared<const FOutputNodeController>(EPrivateToken::Token, InParams);
+						if (Class->ID == Node->ClassID)
+						{
+							return MakeShared<const FOutputNodeController>(EPrivateToken::Token, InParams);
+						}
+						else
+						{
+							UE_LOG(LogMetaSound, Warning, TEXT("Frontend Node [NodeID:%s, ClassID:%s] is not of expected class class [ClassID:%s]"), *Node->ID.ToString(), *Node->ClassID.ToString(), *Class->ID.ToString());
+						}
 					}
 					else
 					{
-						UE_LOG(LogMetaSound, Warning, TEXT("Frontend Node [NodeID:%s, ClassID:%s] is not of expected class class [ClassID:%s]"), *InParams.NodePtr->ID.ToString(), *InParams.NodePtr->ClassID.ToString(), *InParams.ClassPtr->ID.ToString());
+						UE_LOG(LogMetaSound, Warning, TEXT("Frontend Class of incorrect class type [ClassID:%s] while creating output node.. Must be EMetasoundFrontendClassType::Output."), *Class->ID.ToString());
 					}
 				}
-				else
-				{
-					UE_LOG(LogMetaSound, Warning, TEXT("Frontend Class of incorrect class type [ClassID:%s] while creating output node. Must be EMetasoundFrontendClassType::Output."), *InParams.ClassPtr->ID.ToString());
-				}
 			}
 
 			return FInvalidNodeController::GetInvalid();
@@ -1812,7 +1834,7 @@ namespace Metasound
 
 		bool FOutputNodeController::IsValid() const
 		{
-			return OwningGraph->IsValid() && OwningGraphClassOutputPtr.IsValid() && GraphPtr.IsValid();
+			return OwningGraph->IsValid() && (nullptr != OwningGraphClassOutputPtr.Get()) && (nullptr != GraphPtr.Get());
 		}
 
 		FInputHandle FOutputNodeController::CreateInputController(FGuid InVertexID, FConstVertexAccessPtr InNodeVertexPtr, FConstClassInputAccessPtr InClassInputPtr, FNodeHandle InOwningNode) const
@@ -1859,23 +1881,26 @@ namespace Metasound
 
 		FNodeHandle FInputNodeController::CreateInputNodeHandle(const FInputNodeController::FInitParams& InParams)
 		{
-			if (InParams.NodePtr.IsValid() && InParams.ClassPtr.IsValid())
+			if (const FMetasoundFrontendNode* Node = InParams.NodePtr.Get())
 			{
-				if (EMetasoundFrontendClassType::Input == InParams.ClassPtr->Metadata.Type)
+				if (const FMetasoundFrontendClass* Class = InParams.ClassPtr.Get())
 				{
-					if (InParams.ClassPtr->ID == InParams.NodePtr->ClassID)
+					if (EMetasoundFrontendClassType::Input == Class->Metadata.Type)
 					{
-						return MakeShared<FInputNodeController>(EPrivateToken::Token, InParams);
+						if (Class->ID == Node->ClassID)
+						{
+							return MakeShared<FInputNodeController>(EPrivateToken::Token, InParams);
+						}
+						else
+						{
+							UE_LOG(LogMetaSound, Warning, TEXT("Frontend Node [NodeID:%s, ClassID:%s] is not of expected class class [ClassID:%s]"), *Node->ID.ToString(), *Node->ClassID.ToString(), *Class->ID.ToString());
+						}
 					}
 					else
 					{
-						UE_LOG(LogMetaSound, Warning, TEXT("Frontend Node [NodeID:%s, ClassID:%s] is not of expected class class [ClassID:%s]"), *InParams.NodePtr->ID.ToString(), *InParams.NodePtr->ClassID.ToString(), *InParams.ClassPtr->ID.ToString());
+						UE_LOG(LogMetaSound, Warning, TEXT("Frontend Class of incorrect class type [ClassID:%s] while creating input node. Must be EMetasoundFrontendClassType::Input."), *Class->ID.ToString());
 					}
 				}
-				else
-				{
-					UE_LOG(LogMetaSound, Warning, TEXT("Frontend Class of incorrect class type [ClassID:%s] while creating input node. Must be EMetasoundFrontendClassType::Input."), *InParams.ClassPtr->ID.ToString());
-				}
 			}
 
 			return FInvalidNodeController::GetInvalid();
@@ -1883,23 +1908,26 @@ namespace Metasound
 
 		FConstNodeHandle FInputNodeController::CreateConstInputNodeHandle(const FInputNodeController::FInitParams& InParams)
 		{
-			if (InParams.NodePtr.IsValid() && InParams.ClassPtr.IsValid())
+			if (const FMetasoundFrontendNode* Node = InParams.NodePtr.Get())
 			{
-				if (EMetasoundFrontendClassType::Input == InParams.ClassPtr->Metadata.Type)
+				if (const FMetasoundFrontendClass* Class = InParams.ClassPtr.Get())
 				{
-					if (InParams.ClassPtr->ID == InParams.NodePtr->ClassID)
+					if (EMetasoundFrontendClassType::Input == Class->Metadata.Type)
 					{
-						return MakeShared<const FInputNodeController>(EPrivateToken::Token, InParams);
+						if (Class->ID == Node->ClassID)
+						{
+							return MakeShared<const FInputNodeController>(EPrivateToken::Token, InParams);
+						}
+						else
+						{
+							UE_LOG(LogMetaSound, Warning, TEXT("Frontend Node [NodeID:%s, ClassID:%s] is not of expected class class [ClassID:%s]"), *Node->ID.ToString(), *Node->ClassID.ToString(), *Class->ID.ToString());
+						}
 					}
 					else
 					{
-						UE_LOG(LogMetaSound, Warning, TEXT("Frontend Node [NodeID:%s, ClassID:%s] is not of expected class class [ClassID:%s]"), *InParams.NodePtr->ID.ToString(), *InParams.NodePtr->ClassID.ToString(), *InParams.ClassPtr->ID.ToString());
+						UE_LOG(LogMetaSound, Warning, TEXT("Frontend Class of incorrect class type [ClassID:%s] while creating input node. Must be EMetasoundFrontendClassType::Input."), *Class->ID.ToString());
 					}
 				}
-				else
-				{
-					UE_LOG(LogMetaSound, Warning, TEXT("Frontend Class of incorrect class type [ClassID:%s] while creating input node. Must be EMetasoundFrontendClassType::Input."), *InParams.ClassPtr->ID.ToString());
-				}
 			}
 
 			return FInvalidNodeController::GetInvalid();
@@ -1907,7 +1935,7 @@ namespace Metasound
 
 		bool FInputNodeController::IsValid() const
 		{
-			return OwningGraph->IsValid() && OwningGraphClassInputPtr.IsValid() && GraphPtr.IsValid();
+			return OwningGraph->IsValid() && (nullptr != OwningGraphClassInputPtr.Get()) && (nullptr != GraphPtr.Get());
 		}
 
 		FInputHandle FInputNodeController::CreateInputController(FGuid InVertexID, FConstVertexAccessPtr InNodeVertexPtr, FConstClassInputAccessPtr InClassInputPtr, FNodeHandle InOwningNode) const
@@ -1922,9 +1950,9 @@ namespace Metasound
 
 		const FText& FInputNodeController::GetDescription() const
 		{
-			if (OwningGraphClassInputPtr.IsValid())
+			if (const FMetasoundFrontendClassInput* OwningInput = OwningGraphClassInputPtr.Get())
 			{
-				return OwningGraphClassInputPtr->Metadata.Description;
+				return OwningInput->Metadata.Description;
 			}
 
 			return FrontendControllerIntrinsics::GetInvalidValueConstRef<FText>();
@@ -1932,9 +1960,9 @@ namespace Metasound
 
 		const FText& FInputNodeController::GetDisplayName() const
 		{
-			if (OwningGraphClassInputPtr.IsValid())
+			if (const FMetasoundFrontendClassInput* OwningInput = OwningGraphClassInputPtr.Get())
 			{
-				return OwningGraphClassInputPtr->Metadata.DisplayName;
+				return OwningInput->Metadata.DisplayName;
 			}
 
 			return FrontendControllerIntrinsics::GetInvalidValueConstRef<FText>();
@@ -2034,15 +2062,15 @@ namespace Metasound
 
 		FGraphHandle FGraphController::CreateGraphHandle(const FGraphController::FInitParams& InParams)
 		{
-			if (InParams.GraphClassPtr.IsValid())
+			if (FMetasoundFrontendGraphClass* GraphClass = InParams.GraphClassPtr.Get())
 			{
-				if (InParams.GraphClassPtr->Metadata.Type == EMetasoundFrontendClassType::Graph)
+				if (GraphClass->Metadata.Type == EMetasoundFrontendClassType::Graph)
 				{
 					return MakeShared<FGraphController>(EPrivateToken::Token, InParams);
 				}
 				else
 				{
-					UE_LOG(LogMetaSound, Warning, TEXT("Failed to make graph controller [ClassID:%s]. Class must be EMeatsoundFrontendClassType::Graph."), *InParams.GraphClassPtr->ID.ToString())
+					UE_LOG(LogMetaSound, Warning, TEXT("Failed to make graph controller [ClassID:%s]. Class must be EMeatsoundFrontendClassType::Graph."), *GraphClass->ID.ToString())
 				}
 			}
 			return FInvalidGraphController::GetInvalid();
@@ -2050,15 +2078,15 @@ namespace Metasound
 
 		FConstGraphHandle FGraphController::CreateConstGraphHandle(const FGraphController::FInitParams& InParams)
 		{
-			if (InParams.GraphClassPtr.IsValid())
+			if (const FMetasoundFrontendGraphClass* GraphClass = InParams.GraphClassPtr.Get())
 			{
-				if (InParams.GraphClassPtr->Metadata.Type == EMetasoundFrontendClassType::Graph)
+				if (GraphClass->Metadata.Type == EMetasoundFrontendClassType::Graph)
 				{
 					return MakeShared<FGraphController>(EPrivateToken::Token, InParams);
 				}
 				else
 				{
-					UE_LOG(LogMetaSound, Warning, TEXT("Failed to make graph controller [ClassID:%s]. Class must be EMeatsoundFrontendClassType::Graph."), *InParams.GraphClassPtr->ID.ToString())
+					UE_LOG(LogMetaSound, Warning, TEXT("Failed to make graph controller [ClassID:%s]. Class must be EMeatsoundFrontendClassType::Graph."), *GraphClass->ID.ToString())
 				}
 			}
 			return FInvalidGraphController::GetInvalid();
@@ -2066,14 +2094,14 @@ namespace Metasound
 
 		bool FGraphController::IsValid() const
 		{
-			return GraphClassPtr.IsValid() && OwningDocument->IsValid();
+			return (nullptr != GraphClassPtr.Get()) && OwningDocument->IsValid();
 		}
 
 		FGuid FGraphController::GetClassID() const
 		{
-			if (GraphClassPtr.IsValid())
+			if (const FMetasoundFrontendGraphClass* GraphClass = GraphClassPtr.Get())
 			{
-				return GraphClassPtr->ID;
+				return GraphClass->ID;
 			}
 
 			return Metasound::FrontendInvalidID;
@@ -2093,9 +2121,9 @@ namespace Metasound
 		{
 			TArray<FString> Names;
 
-			if (GraphClassPtr.IsValid())
+			if (const FMetasoundFrontendGraphClass* GraphClass = GraphClassPtr.Get())
 			{
-				for (const FMetasoundFrontendClassInput& Input : GraphClassPtr->Interface.Inputs)
+				for (const FMetasoundFrontendClassInput& Input : GraphClass->Interface.Inputs)
 				{
 					Names.Add(Input.Name);
 				}
@@ -2108,9 +2136,9 @@ namespace Metasound
 		{
 			TArray<FString> Names;
 
-			if (GraphClassPtr.IsValid())
+			if (const FMetasoundFrontendGraphClass* GraphClass = GraphClassPtr.Get())
 			{
-				for (const FMetasoundFrontendClassOutput& Output : GraphClassPtr->Interface.Outputs)
+				for (const FMetasoundFrontendClassOutput& Output : GraphClass->Interface.Outputs)
 				{
 					Names.Add(Output.Name);
 				}
@@ -2177,9 +2205,9 @@ namespace Metasound
 
 		const FMetasoundFrontendGraphStyle& FGraphController::GetGraphStyle() const
 		{
-			if (GraphClassPtr.IsValid())
+			if (const FMetasoundFrontendGraphClass* GraphClass = GraphClassPtr.Get())
 			{
-				return GraphClassPtr->Graph.Style;
+				return GraphClass->Graph.Style;
 			}
 
 			return FrontendControllerIntrinsics::GetInvalidValueConstRef<FMetasoundFrontendGraphStyle>();
@@ -2187,9 +2215,9 @@ namespace Metasound
 
 		void FGraphController::SetGraphStyle(const FMetasoundFrontendGraphStyle& InStyle)
 		{
-			if (GraphClassPtr.IsValid())
+			if (FMetasoundFrontendGraphClass* GraphClass = GraphClassPtr.Get())
 			{
-				GraphClassPtr->Graph.Style = InStyle;
+				GraphClass->Graph.Style = InStyle;
 			}
 		}
 
@@ -2213,9 +2241,9 @@ namespace Metasound
 
 		void FGraphController::IterateNodes(TUniqueFunction<void(FNodeHandle)> InFunction, EMetasoundFrontendClassType InClassType)
 		{
-			if (GraphClassPtr.IsValid())
+			if (FMetasoundFrontendGraphClass* GraphClass = GraphClassPtr.Get())
 			{
-				for (const FMetasoundFrontendNode& Node : GraphClassPtr->Graph.Nodes)
+				for (const FMetasoundFrontendNode& Node : GraphClass->Graph.Nodes)
 				{
 					FConstClassAccessPtr NodeClassPtr = OwningDocument->FindClassWithID(Node.ClassID);
 					if (const FMetasoundFrontendClass* NodeClass = NodeClassPtr.Get())
@@ -2237,9 +2265,9 @@ namespace Metasound
 
 		void FGraphController::IterateConstNodes(TUniqueFunction<void(FConstNodeHandle)> InFunction, EMetasoundFrontendClassType InClassType) const
 		{
-			if (GraphClassPtr.IsValid())
+			if (FMetasoundFrontendGraphClass* GraphClass = GraphClassPtr.Get())
 			{
-				for (const FMetasoundFrontendNode& Node : GraphClassPtr->Graph.Nodes)
+				for (const FMetasoundFrontendNode& Node : GraphClass->Graph.Nodes)
 				{
 					FConstClassAccessPtr NodeClassPtr = OwningDocument->FindClassWithID(Node.ClassID);
 					if (const FMetasoundFrontendClass* NodeClass = NodeClassPtr.Get())
@@ -2349,7 +2377,8 @@ namespace Metasound
 					FNodeRegistryKey Key;
 					if (FRegistry::GetInputNodeRegistryKeyForDataType(InClassInput.TypeName, Key))
 					{
-						if (FConstClassAccessPtr InputClass = OwningDocument->FindOrAddClass(Key))
+						FConstClassAccessPtr InputClassPtr = OwningDocument->FindOrAddClass(Key);
+						if (const FMetasoundFrontendClass* InputClass = InputClassPtr.Get())
 						{
 							FMetasoundFrontendNode& Node = GraphClass->Graph.Nodes.Emplace_GetRef(*InputClass);
 
@@ -2375,11 +2404,11 @@ namespace Metasound
 								OutputVertex->Name = InClassInput.Name;
 							}
 
-							FMetasoundFrontendClassInput& NewInput = GraphClassPtr->Interface.Inputs.Add_GetRef(InClassInput);
+							FMetasoundFrontendClassInput& NewInput = GraphClass->Interface.Inputs.Add_GetRef(InClassInput);
 							NewInput.NodeID = Node.ID;
 
 							FNodeAccessPtr NodePtr = GraphClassPtr.GetNodeWithNodeID(Node.ID);
-							return GetNodeHandle(FGraphController::FNodeAndClass{NodePtr, InputClass});
+							return GetNodeHandle(FGraphController::FNodeAndClass{NodePtr, InputClassPtr});
 						}
 					}
 					else 
@@ -2389,7 +2418,7 @@ namespace Metasound
 				}
 				else
 				{
-					UE_LOG(LogMetaSound, Display, TEXT("Failed to add input. Input with same name \"%s\" exists in class [ClassID:%s]"), *InClassInput.Name, *GraphClassPtr->ID.ToString());
+					UE_LOG(LogMetaSound, Display, TEXT("Failed to add input. Input with same name \"%s\" exists in class [ClassID:%s]"), *InClassInput.Name, *GraphClass->ID.ToString());
 				}
 			}
 			return FInvalidNodeController::GetInvalid();
@@ -2397,16 +2426,16 @@ namespace Metasound
 
 		bool FGraphController::RemoveInputVertex(const FString& InName)
 		{
-			if (GraphClassPtr.IsValid())
-			{
-				auto IsInputNodeWithSameName = [&](const FMetasoundFrontendClass& InClass, const FMetasoundFrontendNode& InNode) 
-				{ 
-					return (InClass.Metadata.Type == EMetasoundFrontendClassType::Input) && (InNode.Name == InName); 
-				};
+			auto IsInputNodeWithSameName = [&](const FMetasoundFrontendClass& InClass, const FMetasoundFrontendNode& InNode) 
+			{ 
+				return (InClass.Metadata.Type == EMetasoundFrontendClassType::Input) && (InNode.Name == InName); 
+			};
 
-				for (const FNodeAndClass& NodeAndClass : GetNodesAndClassesByPredicate(IsInputNodeWithSameName))
+			for (const FNodeAndClass& NodeAndClass : GetNodesAndClassesByPredicate(IsInputNodeWithSameName))
+			{
+				if (const FMetasoundFrontendNode* Node = NodeAndClass.Node.Get())
 				{
-					return RemoveInput(*NodeAndClass.Node);
+					return RemoveInput(*Node);
 				}
 			}
 
@@ -2423,7 +2452,8 @@ namespace Metasound
 					FNodeRegistryKey Key;
 					if (FRegistry::GetOutputNodeRegistryKeyForDataType(InClassOutput.TypeName, Key))
 					{
-						if (FConstClassAccessPtr OutputClass = OwningDocument->FindOrAddClass(Key))
+						FConstClassAccessPtr OutputClassPtr = OwningDocument->FindOrAddClass(Key);
+						if (const FMetasoundFrontendClass* OutputClass = OutputClassPtr.Get())
 						{
 							FMetasoundFrontendNode& Node = GraphClass->Graph.Nodes.Emplace_GetRef(*OutputClass);
 
@@ -2450,11 +2480,11 @@ namespace Metasound
 								UE_LOG(LogMetaSound, Error, TEXT("Output node [TypeName:%s] does not contain output vertex with type [TypeName:%s]"), *InClassOutput.TypeName.ToString(), *InClassOutput.TypeName.ToString());
 							}
 
-							FMetasoundFrontendClassOutput& NewOutput = GraphClassPtr->Interface.Outputs.Add_GetRef(InClassOutput);
+							FMetasoundFrontendClassOutput& NewOutput = GraphClass->Interface.Outputs.Add_GetRef(InClassOutput);
 							NewOutput.NodeID = Node.ID;
 
 							FNodeAccessPtr NodePtr = GraphClassPtr.GetNodeWithNodeID(Node.ID);
-							return GetNodeHandle(FGraphController::FNodeAndClass { NodePtr, OutputClass });
+							return GetNodeHandle(FGraphController::FNodeAndClass { NodePtr, OutputClassPtr });
 						}
 					}
 					else 
@@ -2464,7 +2494,7 @@ namespace Metasound
 				}
 				else
 				{
-					UE_LOG(LogMetaSound, Display, TEXT("Failed to add output. Output with same name \"%s\" exists in class [ClassID:%s]"), *InClassOutput.Name, *GraphClassPtr->ID.ToString());
+					UE_LOG(LogMetaSound, Display, TEXT("Failed to add output. Output with same name \"%s\" exists in class [ClassID:%s]"), *InClassOutput.Name, *GraphClass->ID.ToString());
 				}
 			}
 
@@ -2473,16 +2503,16 @@ namespace Metasound
 
 		bool FGraphController::RemoveOutputVertex(const FString& InName)
 		{
-			if (GraphClassPtr.IsValid())
-			{
-				auto IsOutputNodeWithSameName = [&](const FMetasoundFrontendClass& InClass, const FMetasoundFrontendNode& InNode) 
-				{ 
-					return (InClass.Metadata.Type == EMetasoundFrontendClassType::Output) && (InNode.Name == InName); 
-				};
+			auto IsOutputNodeWithSameName = [&](const FMetasoundFrontendClass& InClass, const FMetasoundFrontendNode& InNode) 
+			{ 
+				return (InClass.Metadata.Type == EMetasoundFrontendClassType::Output) && (InNode.Name == InName); 
+			};
 
-				for (const FNodeAndClass& NodeAndClass : GetNodesAndClassesByPredicate(IsOutputNodeWithSameName))
+			for (const FNodeAndClass& NodeAndClass : GetNodesAndClassesByPredicate(IsOutputNodeWithSameName))
+			{
+				if (const FMetasoundFrontendNode* Node = NodeAndClass.Node.Get())
 				{
-					return RemoveOutput(*NodeAndClass.Node);
+					return RemoveOutput(*Node);
 				}
 			}
 
@@ -2617,83 +2647,87 @@ namespace Metasound
 
 		FNodeHandle FGraphController::AddNode(const FNodeRegistryKey& InKey)
 		{
-			if (IsValid())
+			// Construct a FNodeClassInfo from this lookup key.
+			FConstClassAccessPtr Class = OwningDocument->FindOrAddClass(InKey);
+			const bool bIsValidClass = (nullptr != Class.Get());
+
+			if (bIsValidClass)
 			{
-				// Construct a FNodeClassInfo from this lookup key.
-				if (FConstClassAccessPtr DependencyDescription = OwningDocument->FindOrAddClass(InKey))
-				{
-					return AddNode(DependencyDescription);
-				}
-				else
-				{
-					UE_LOG(LogMetaSound, Warning, TEXT("Failed to find node class info with registry key [Key:%s]"), *InKey);
-					return INodeController::GetInvalidHandle();
-				}
+				return AddNode(Class);
 			}
 
-			return FInvalidNodeController::GetInvalid();
+			UE_LOG(LogMetaSound, Warning, TEXT("Failed to find or add node class info with registry key [Key:%s]"), *InKey);
+			return INodeController::GetInvalidHandle();
 		}
 
 		FNodeHandle FGraphController::AddNode(const FMetasoundFrontendClassMetadata& InClassMetadata)
 		{
-			return AddNode(FRegistry::GetRegistryKey(InClassMetadata));
+			return AddNode(NodeRegistryKey::CreateKey(InClassMetadata));
 		}
 
 		FNodeHandle FGraphController::AddDuplicateNode(const INodeController& InNode)
 		{
+			// TODO: will need to copy node interface when dynamic pins exist.
 			const FMetasoundFrontendClassMetadata& ClassMetadata = InNode.GetClassMetadata();
 
+			FConstClassAccessPtr ClassPtr;
+
 			if (EMetasoundFrontendClassType::Graph == ClassMetadata.Type)
 			{
 				// Add subgraph and dependencies if needed
-				if (!OwningDocument->FindClass(ClassMetadata).IsValid())
+				ClassPtr = OwningDocument->FindClass(ClassMetadata);
+				const bool bIsClassMissing = (nullptr == ClassPtr.Get());
+
+				if (bIsClassMissing)
 				{
 					// Class does not exist, need to add the subgraph
 					OwningDocument->AddDuplicateSubgraph(*(InNode.AsGraph()));
+					ClassPtr = OwningDocument->FindClass(ClassMetadata);
 				}
 			}
-
-			// TODO: will need to copy node interface when dynamic pins exist.
-			if (FConstClassAccessPtr DependencyDescription = OwningDocument->FindOrAddClass(ClassMetadata))
+			else
 			{
-				return AddNode(DependencyDescription);
+				ClassPtr = OwningDocument->FindOrAddClass(ClassMetadata);
 			}
 
-			return INodeController::GetInvalidHandle();
+			return AddNode(ClassPtr);
 		}
 
 		// Remove the node corresponding to this node handle.
 		// On success, invalidates the received node handle.
 		bool FGraphController::RemoveNode(INodeController& InNode)
 		{
-			FGuid NodeID = InNode.GetID();
-			auto IsNodeWithSameID = [&](const FMetasoundFrontendNode& InDesc) { return InDesc.ID == NodeID; };
-			if (const FMetasoundFrontendNode* Desc = GraphClassPtr->Graph.Nodes.FindByPredicate(IsNodeWithSameID))
+			if (FMetasoundFrontendGraphClass* GraphClass = GraphClassPtr.Get())
 			{
-				switch(InNode.GetClassMetadata().Type)
+				FGuid NodeID = InNode.GetID();
+				auto IsNodeWithSameID = [&](const FMetasoundFrontendNode& InDesc) { return InDesc.ID == NodeID; };
+				if (const FMetasoundFrontendNode* Desc = GraphClass->Graph.Nodes.FindByPredicate(IsNodeWithSameID))
 				{
-					case EMetasoundFrontendClassType::Input:
+					switch(InNode.GetClassMetadata().Type)
 					{
-						return RemoveInput(*Desc);
-					}
+						case EMetasoundFrontendClassType::Input:
+						{
+							return RemoveInput(*Desc);
+						}
 
-					case EMetasoundFrontendClassType::Output:
-					{
-						return RemoveOutput(*Desc);
-					}
+						case EMetasoundFrontendClassType::Output:
+						{
+							return RemoveOutput(*Desc);
+						}
 
-					case EMetasoundFrontendClassType::Variable:
-					case EMetasoundFrontendClassType::External:
-					case EMetasoundFrontendClassType::Graph:
-					{
-						return RemoveNode(*Desc);
-					}
+						case EMetasoundFrontendClassType::Variable:
+						case EMetasoundFrontendClassType::External:
+						case EMetasoundFrontendClassType::Graph:
+						{
+							return RemoveNode(*Desc);
+						}
 
-					default:
-					case EMetasoundFrontendClassType::Invalid:
-					{
-						static_assert(static_cast<int32>(EMetasoundFrontendClassType::Invalid) == 5, "Possible missing switch case coverage for EMetasoundFrontendClassType.");
-						checkNoEntry();
+						default:
+						case EMetasoundFrontendClassType::Invalid:
+						{
+							static_assert(static_cast<int32>(EMetasoundFrontendClassType::Invalid) == 5, "Possible missing switch case coverage for EMetasoundFrontendClassType.");
+							checkNoEntry();
+						}
 					}
 				}
 			}
@@ -2704,68 +2738,68 @@ namespace Metasound
 		// Returns the metadata for the current graph, including the name, description and author.
 		const FMetasoundFrontendClassMetadata& FGraphController::GetGraphMetadata() const
 		{
-			if (GraphClassPtr.IsValid())
+			if (const FMetasoundFrontendGraphClass* GraphClass = GraphClassPtr.Get())
 			{
-				return GraphClassPtr->Metadata;
+				return GraphClass->Metadata;
 			}
 			return FrontendControllerIntrinsics::GetInvalidValueConstRef<FMetasoundFrontendClassMetadata>();
 		}
 
 		void FGraphController::SetGraphMetadata(const FMetasoundFrontendClassMetadata& InMetadata)
 		{
-			if (GraphClassPtr.IsValid())
+			if (FMetasoundFrontendGraphClass* GraphClass = GraphClassPtr.Get())
 			{
-				GraphClassPtr->Metadata = InMetadata;
+				GraphClass->Metadata = InMetadata;
 			}
 		}
 
 		FNodeHandle FGraphController::CreateEmptySubgraph(const FMetasoundFrontendClassMetadata& InInfo)
 		{
-			if (IsValid())
+			if (InInfo.Type == EMetasoundFrontendClassType::Graph)
 			{
-				if (InInfo.Type == EMetasoundFrontendClassType::Graph)
+				if (const FMetasoundFrontendClass* ExistingDependency = OwningDocument->FindClass(InInfo).Get())
 				{
-					if (const FMetasoundFrontendClass* ExistingDependency = OwningDocument->FindClass(InInfo).Get())
-					{
-						UE_LOG(LogMetaSound, Error, TEXT("Cannot add new subgraph. Metasound class already exists with matching metadata Name: \"%s\", Version %d.%d"), *(ExistingDependency->Metadata.ClassName.GetFullName().ToString()), ExistingDependency->Metadata.Version.Major, ExistingDependency->Metadata.Version.Minor);
-					}
-					else if (FConstClassAccessPtr DependencyDescription = OwningDocument->FindOrAddClass(InInfo))
-					{
-						return AddNode(DependencyDescription);
-					}
+					UE_LOG(LogMetaSound, Error, TEXT("Cannot add new subgraph. Metasound class already exists with matching metadata Name: \"%s\", Version %d.%d"), *(ExistingDependency->Metadata.ClassName.GetFullName().ToString()), ExistingDependency->Metadata.Version.Major, ExistingDependency->Metadata.Version.Minor);
 				}
-				else
+				else 
 				{
-					UE_LOG(LogMetaSound, Warning, TEXT("Incompatible Metasound NodeType encountered when attempting to create an empty subgraph.  NodeType must equal EMetasoundFrontendClassType::Graph"));
+					return AddNode(OwningDocument->FindOrAddClass(InInfo));
 				}
 			}
+			else
+			{
+				UE_LOG(LogMetaSound, Warning, TEXT("Incompatible Metasound NodeType encountered when attempting to create an empty subgraph.  NodeType must equal EMetasoundFrontendClassType::Graph"));
+			}
 			
 			return FInvalidNodeController::GetInvalid();
 		}
 
 		TUniquePtr<IOperator> FGraphController::BuildOperator(const FOperatorSettings& InSettings, const FMetasoundEnvironment& InEnvironment, TArray<IOperatorBuilder::FBuildErrorPtr>& OutBuildErrors) const
 		{
-			if (!IsValid())
+			if (const FMetasoundFrontendGraphClass* GraphClass = GraphClassPtr.Get())
 			{
-				return TUniquePtr<IOperator>(nullptr);
-			}
 
-			// TODO: Implement subgraph inflation step here.
+				// TODO: Implement subgraph inflation step here.
 
-			// TODO: bubble up errors. 
-			TArray<FMetasoundFrontendGraphClass> Subgraphs = OwningDocument->GetSubgraphs();
-			TArray<FMetasoundFrontendClass> Dependencies = OwningDocument->GetDependencies();
+				// TODO: bubble up errors. 
+				TArray<FMetasoundFrontendGraphClass> Subgraphs = OwningDocument->GetSubgraphs();
+				TArray<FMetasoundFrontendClass> Dependencies = OwningDocument->GetDependencies();
 
-			TUniquePtr<FFrontendGraph> Graph = FFrontendGraphBuilder::CreateGraph(*GraphClassPtr, Subgraphs, Dependencies);
+				TUniquePtr<FFrontendGraph> Graph = FFrontendGraphBuilder::CreateGraph(*GraphClass, Subgraphs, Dependencies);
 
-			if (!Graph.IsValid())
+				if (!Graph.IsValid())
+				{
+					return TUniquePtr<IOperator>(nullptr);
+				}
+
+				FOperatorBuilder OperatorBuilder(FOperatorBuilderSettings::GetDefaultSettings());
+				FBuildGraphParams BuildParams{*Graph, InSettings, FDataReferenceCollection{}, InEnvironment};
+				return OperatorBuilder.BuildGraphOperator(BuildParams, OutBuildErrors);
+			}
+			else
 			{
 				return TUniquePtr<IOperator>(nullptr);
 			}
-
-			FOperatorBuilder OperatorBuilder(FOperatorBuilderSettings::GetDefaultSettings());
-			FBuildGraphParams BuildParams{*Graph, InSettings, FDataReferenceCollection{}, InEnvironment};
-			return OperatorBuilder.BuildGraphOperator(BuildParams, OutBuildErrors);
 		}
 
 		FDocumentHandle FGraphController::GetOwningDocument()
@@ -2784,7 +2818,7 @@ namespace Metasound
 			{
 				if (const FMetasoundFrontendClass* NodeClass = InExistingDependency.Get())
 				{
-					FMetasoundFrontendNode& Node = GraphClass->Graph.Nodes.Emplace_GetRef(*InExistingDependency);
+					FMetasoundFrontendNode& Node = GraphClass->Graph.Nodes.Emplace_GetRef(*NodeClass);
 
 					Node.ID = FGuid::NewGuid();
 
@@ -2798,16 +2832,16 @@ namespace Metasound
 
 		bool FGraphController::RemoveNode(const FMetasoundFrontendNode& InDesc)
 		{
-			if (GraphClassPtr.IsValid())
+			if (FMetasoundFrontendGraphClass* GraphClass = GraphClassPtr.Get())
 			{
 				auto IsEdgeForThisNode = [&](const FMetasoundFrontendEdge& ConDesc) { return (ConDesc.FromNodeID == InDesc.ID) || (ConDesc.ToNodeID == InDesc.ID); };
 
 				// Remove any reference connections
-				int32 NumRemoved = GraphClassPtr->Graph.Edges.RemoveAll(IsEdgeForThisNode);
+				int32 NumRemoved = GraphClass->Graph.Edges.RemoveAll(IsEdgeForThisNode);
 
 				auto IsNodeWithID = [&](const FMetasoundFrontendNode& Desc) { return InDesc.ID == Desc.ID; };
 
-				NumRemoved += GraphClassPtr->Graph.Nodes.RemoveAll(IsNodeWithID);
+				NumRemoved += GraphClass->Graph.Nodes.RemoveAll(IsNodeWithID);
 
 				OwningDocument->SynchronizeDependencies();
 
@@ -2818,11 +2852,11 @@ namespace Metasound
 
 		bool FGraphController::RemoveInput(const FMetasoundFrontendNode& InNode)
 		{
-			if (GraphClassPtr.IsValid())
+			if (FMetasoundFrontendGraphClass* GraphClass = GraphClassPtr.Get())
 			{
 				auto IsInputWithSameNodeID = [&](const FMetasoundFrontendClassInput& ClassInput) { return ClassInput.NodeID == InNode.ID; };
 
-				int32 NumInputsRemoved = GraphClassPtr->Interface.Inputs.RemoveAll(IsInputWithSameNodeID);
+				int32 NumInputsRemoved = GraphClass->Interface.Inputs.RemoveAll(IsInputWithSameNodeID);
 
 				bool bDidRemoveNode = RemoveNode(InNode);
 
@@ -2834,11 +2868,11 @@ namespace Metasound
 
 		bool FGraphController::RemoveOutput(const FMetasoundFrontendNode& InNode)
 		{
-			if (GraphClassPtr.IsValid())
+			if (FMetasoundFrontendGraphClass* GraphClass = GraphClassPtr.Get())
 			{
 				auto IsOutputWithSameNodeID = [&](const FMetasoundFrontendClassOutput& ClassOutput) { return ClassOutput.NodeID == InNode.ID; };
 
-				int32 NumOutputsRemoved = GraphClassPtr->Interface.Outputs.RemoveAll(IsOutputWithSameNodeID);
+				int32 NumOutputsRemoved = GraphClass->Interface.Outputs.RemoveAll(IsOutputWithSameNodeID);
 
 				bool bDidRemoveNode = RemoveNode(InNode);
 
@@ -2850,9 +2884,9 @@ namespace Metasound
 
 		bool FGraphController::ContainsNodesAndClassesByPredicate(TFunctionRef<bool (const FMetasoundFrontendClass&, const FMetasoundFrontendNode&)> InPredicate) const
 		{
-			if (GraphClassPtr.IsValid())
+			if (FMetasoundFrontendGraphClass* GraphClass = GraphClassPtr.Get())
 			{
-				for (FMetasoundFrontendNode& Node : GraphClassPtr->Graph.Nodes)
+				for (FMetasoundFrontendNode& Node : GraphClass->Graph.Nodes)
 				{
 					const FMetasoundFrontendClass* NodeClass = OwningDocument->FindClassWithID(Node.ClassID).Get();
 					if (nullptr != NodeClass)
@@ -2876,14 +2910,17 @@ namespace Metasound
 		{
 			TArray<FNodeAndClass> NodesAndClasses;
 
-			if (GraphClassPtr.IsValid())
+			if (FMetasoundFrontendGraphClass* GraphClass = GraphClassPtr.Get())
 			{
-				for (FMetasoundFrontendNode& Node : GraphClassPtr->Graph.Nodes)
+				for (FMetasoundFrontendNode& Node : GraphClass->Graph.Nodes)
 				{
 					FNodeAccessPtr NodePtr = GraphClassPtr.GetNodeWithNodeID(Node.ID);
 					FConstClassAccessPtr NodeClassPtr = OwningDocument->FindClassWithID(Node.ClassID);
 
-					if (NodeClassPtr && NodePtr)
+					const bool bIsValidNodePtr = (nullptr != NodePtr.Get());
+					const bool bIsValidNodeClassPtr = (nullptr != NodeClassPtr.Get());
+
+					if (bIsValidNodePtr && bIsValidNodeClassPtr)
 					{
 						NodesAndClasses.Add(FGraphController::FNodeAndClass{ NodePtr, NodeClassPtr });
 					}
@@ -2901,14 +2938,17 @@ namespace Metasound
 		{
 			TArray<FConstNodeAndClass> NodesAndClasses;
 
-			if (GraphClassPtr.IsValid())
+			if (const FMetasoundFrontendGraphClass* GraphClass = GraphClassPtr.Get())
 			{
-				for (const FMetasoundFrontendNode& Node : GraphClassPtr->Graph.Nodes)
+				for (const FMetasoundFrontendNode& Node : GraphClass->Graph.Nodes)
 				{
 					FConstNodeAccessPtr NodePtr = GraphClassPtr.GetNodeWithNodeID(Node.ID);
 					FConstClassAccessPtr NodeClassPtr = OwningDocument->FindClassWithID(Node.ClassID);
 
-					if (NodeClassPtr && NodePtr)
+					const bool bIsValidNodePtr = (nullptr != NodePtr.Get());
+					const bool bIsValidNodeClassPtr = (nullptr != NodeClassPtr.Get());
+
+					if (bIsValidNodePtr && bIsValidNodeClassPtr)
 					{
 						NodesAndClasses.Add(FGraphController::FConstNodeAndClass{ NodePtr, NodeClassPtr });
 					}
@@ -2924,12 +2964,11 @@ namespace Metasound
 
 		TArray<FGraphController::FNodeAndClass> FGraphController::GetNodesAndClassesByPredicate(TFunctionRef<bool (const FMetasoundFrontendClass&, const FMetasoundFrontendNode&)> InPredicate)
 		{
-
 			TArray<FNodeAndClass> NodesAndClasses;
 
-			if (GraphClassPtr.IsValid())
+			if (FMetasoundFrontendGraphClass* GraphClass = GraphClassPtr.Get())
 			{
-				for (FMetasoundFrontendNode& Node : GraphClassPtr->Graph.Nodes)
+				for (FMetasoundFrontendNode& Node : GraphClass->Graph.Nodes)
 				{
 					FConstClassAccessPtr NodeClassPtr = OwningDocument->FindClassWithID(Node.ClassID);
 					if (const FMetasoundFrontendClass* NodeClass = NodeClassPtr.Get())
@@ -2954,9 +2993,9 @@ namespace Metasound
 		{
 			TArray<FConstNodeAndClass> NodesAndClasses;
 
-			if (GraphClassPtr.IsValid())
+			if (const FMetasoundFrontendGraphClass* GraphClass = GraphClassPtr.Get())
 			{
-				for (const FMetasoundFrontendNode& Node : GraphClassPtr->Graph.Nodes)
+				for (const FMetasoundFrontendNode& Node : GraphClass->Graph.Nodes)
 				{
 					FConstClassAccessPtr NodeClassPtr = OwningDocument->FindClassWithID(Node.ClassID);
 					if (const FMetasoundFrontendClass* NodeClass = NodeClassPtr.Get())
@@ -2979,14 +3018,10 @@ namespace Metasound
 
 		FNodeHandle FGraphController::GetNodeByPredicate(TFunctionRef<bool (const FMetasoundFrontendClass&, const FMetasoundFrontendNode&)> InPredicate)
 		{
-			if (GraphClassPtr.IsValid())
+			TArray<FNodeAndClass> NodeAndClass = GetNodesAndClassesByPredicate(InPredicate);
+			if (NodeAndClass.Num() > 0)
 			{
-				TArray<FNodeAndClass> NodeAndClass = GetNodesAndClassesByPredicate(InPredicate);
-
-				if (NodeAndClass.Num() > 0)
-				{
-					return GetNodeHandle(NodeAndClass[0]);
-				}
+				return GetNodeHandle(NodeAndClass[0]);
 			}
 
 			return FInvalidNodeController::GetInvalid();
@@ -2994,14 +3029,10 @@ namespace Metasound
 
 		FConstNodeHandle FGraphController::GetNodeByPredicate(TFunctionRef<bool (const FMetasoundFrontendClass&, const FMetasoundFrontendNode&)> InPredicate) const
 		{
-			if (GraphClassPtr.IsValid())
+			TArray<FConstNodeAndClass> NodeAndClass = GetNodesAndClassesByPredicate(InPredicate);
+			if (NodeAndClass.Num() > 0)
 			{
-				TArray<FConstNodeAndClass> NodeAndClass = GetNodesAndClassesByPredicate(InPredicate);
-
-				if (NodeAndClass.Num() > 0)
-				{
-					return GetNodeHandle(NodeAndClass[0]);
-				}
+				return GetNodeHandle(NodeAndClass[0]);
 			}
 
 			return FInvalidNodeController::GetInvalid();
@@ -3009,20 +3040,12 @@ namespace Metasound
 
 		TArray<FNodeHandle> FGraphController::GetNodesByPredicate(TFunctionRef<bool (const FMetasoundFrontendClass&, const FMetasoundFrontendNode&)> InFilterFunc)
 		{
-			if (GraphClassPtr.IsValid())
-			{
-				return GetNodeHandles(GetNodesAndClassesByPredicate(InFilterFunc));
-			}
-			return TArray<FNodeHandle>();
+			return GetNodeHandles(GetNodesAndClassesByPredicate(InFilterFunc));
 		}
 
 		TArray<FConstNodeHandle> FGraphController::GetNodesByPredicate(TFunctionRef<bool (const FMetasoundFrontendClass&, const FMetasoundFrontendNode&)> InFilterFunc) const
 		{
-			if (GraphClassPtr.IsValid())
-			{
-				return GetNodeHandles(GetNodesAndClassesByPredicate(InFilterFunc));
-			}
-			return TArray<FConstNodeHandle>();
+			return GetNodeHandles(GetNodesAndClassesByPredicate(InFilterFunc));
 		}
 
 		TArray<FNodeHandle> FGraphController::GetNodeHandles(TArrayView<const FGraphController::FNodeAndClass> InNodesAndClasses)
@@ -3059,66 +3082,75 @@ namespace Metasound
 
 		FNodeHandle FGraphController::GetNodeHandle(const FGraphController::FNodeAndClass& InNodeAndClass)
 		{
-			if (InNodeAndClass.IsValid() && GraphClassPtr.IsValid())
+			FMetasoundFrontendNode* Node = InNodeAndClass.Node.Get();
+			const FMetasoundFrontendClass* NodeClass = InNodeAndClass.Class.Get();
+			
+			if ((nullptr != Node) && (nullptr != NodeClass))
 			{
 				FGraphHandle OwningGraph = this->AsShared();
 				FGraphAccessPtr GraphPtr = GraphClassPtr.GetGraph();
 
-				switch (InNodeAndClass.Class->Metadata.Type)
+				switch (NodeClass->Metadata.Type)
 				{
 					case EMetasoundFrontendClassType::Input:
-						if (FClassInputAccessPtr OwningGraphClassInputPtr = FindInputDescriptionWithNodeID(InNodeAndClass.Node->ID))
 						{
-							FInputNodeController::FInitParams InitParams
+							FClassInputAccessPtr OwningGraphClassInputPtr = FindInputDescriptionWithNodeID(Node->ID);
+							if (nullptr != OwningGraphClassInputPtr.Get())
 							{
-								InNodeAndClass.Node,
-								InNodeAndClass.Class,
-								OwningGraphClassInputPtr,
-								GraphPtr,
-								OwningGraph
-							};
-							return FInputNodeController::CreateInputNodeHandle(InitParams);
-						}
-						else
-						{
-							// TODO: This supports input nodes introduced during subgraph inflation. Input nodes
-							// should be replaced with value nodes once they are implemented. 
-							FNodeController::FInitParams InitParams
+								FInputNodeController::FInitParams InitParams
+								{
+									InNodeAndClass.Node,
+									InNodeAndClass.Class,
+									OwningGraphClassInputPtr,
+									GraphPtr,
+									OwningGraph
+								};
+								return FInputNodeController::CreateInputNodeHandle(InitParams);
+							}
+							else
 							{
-								InNodeAndClass.Node,
-								InNodeAndClass.Class,
-								GraphPtr,
-								OwningGraph
-							};
-							return FNodeController::CreateNodeHandle(InitParams);
+								// TODO: This supports input nodes introduced during subgraph inflation. Input nodes
+								// should be replaced with value nodes once they are implemented. 
+								FNodeController::FInitParams InitParams
+								{
+									InNodeAndClass.Node,
+									InNodeAndClass.Class,
+									GraphPtr,
+									OwningGraph
+								};
+								return FNodeController::CreateNodeHandle(InitParams);
+							}
 						}
 						break;
 
 					case EMetasoundFrontendClassType::Output:
-						if (FClassOutputAccessPtr OwningGraphClassOutputPtr = FindOutputDescriptionWithNodeID(InNodeAndClass.Node->ID))
 						{
-							FOutputNodeController::FInitParams InitParams
+							FClassOutputAccessPtr OwningGraphClassOutputPtr = FindOutputDescriptionWithNodeID(Node->ID);
+							if (nullptr != OwningGraphClassOutputPtr.Get())
 							{
-								InNodeAndClass.Node,
-								InNodeAndClass.Class,
-								OwningGraphClassOutputPtr,
-								GraphPtr,
-								OwningGraph
-							};
-							return FOutputNodeController::CreateOutputNodeHandle(InitParams);
-						}
-						else
-						{
-							// TODO: This supports output nodes introduced during subgraph inflation. Output nodes
-							// should be replaced with value nodes once they are implemented. 
-							FNodeController::FInitParams InitParams
+								FOutputNodeController::FInitParams InitParams
+								{
+									InNodeAndClass.Node,
+									InNodeAndClass.Class,
+									OwningGraphClassOutputPtr,
+									GraphPtr,
+									OwningGraph
+								};
+								return FOutputNodeController::CreateOutputNodeHandle(InitParams);
+							}
+							else
 							{
-								InNodeAndClass.Node,
-								InNodeAndClass.Class,
-								GraphPtr,
-								OwningGraph
-							};
-							return FNodeController::CreateNodeHandle(InitParams);
+								// TODO: This supports output nodes introduced during subgraph inflation. Output nodes
+								// should be replaced with value nodes once they are implemented. 
+								FNodeController::FInitParams InitParams
+								{
+									InNodeAndClass.Node,
+									InNodeAndClass.Class,
+									GraphPtr,
+									OwningGraph
+								};
+								return FNodeController::CreateNodeHandle(InitParams);
+							}
 						}
 						break;
 
@@ -3158,40 +3190,49 @@ namespace Metasound
 
 		FConstNodeHandle FGraphController::GetNodeHandle(const FGraphController::FConstNodeAndClass& InNodeAndClass) const
 		{
-			if (InNodeAndClass.IsValid() && GraphClassPtr.IsValid())
+			const FMetasoundFrontendNode* Node = InNodeAndClass.Node.Get();
+			const FMetasoundFrontendClass* NodeClass = InNodeAndClass.Class.Get();
+			
+			if ((nullptr != Node) && (nullptr != NodeClass))
 			{
 				FConstGraphHandle OwningGraph = this->AsShared();
 				FConstGraphAccessPtr GraphPtr = GraphClassPtr.GetGraph();
 
-				switch (InNodeAndClass.Class->Metadata.Type)
+				switch (NodeClass->Metadata.Type)
 				{
 					case EMetasoundFrontendClassType::Input:
-						if (FConstClassInputAccessPtr OwningGraphClassInputPtr = FindInputDescriptionWithNodeID(InNodeAndClass.Node->ID))
 						{
-							FInputNodeController::FInitParams InitParams
+							FConstClassInputAccessPtr OwningGraphClassInputPtr = FindInputDescriptionWithNodeID(Node->ID);
+							if (nullptr != OwningGraphClassInputPtr.Get())
 							{
-								ConstCastAccessPtr<FNodeAccessPtr>(InNodeAndClass.Node),
-								InNodeAndClass.Class,
-								ConstCastAccessPtr<FClassInputAccessPtr>(OwningGraphClassInputPtr),
-								ConstCastAccessPtr<FGraphAccessPtr>(GraphPtr),
-								ConstCastSharedRef<IGraphController>(OwningGraph)
-							};
-							return FInputNodeController::CreateConstInputNodeHandle(InitParams);
+								FInputNodeController::FInitParams InitParams
+								{
+									ConstCastAccessPtr<FNodeAccessPtr>(InNodeAndClass.Node),
+									InNodeAndClass.Class,
+									ConstCastAccessPtr<FClassInputAccessPtr>(OwningGraphClassInputPtr),
+									ConstCastAccessPtr<FGraphAccessPtr>(GraphPtr),
+									ConstCastSharedRef<IGraphController>(OwningGraph)
+								};
+								return FInputNodeController::CreateConstInputNodeHandle(InitParams);
+							}
 						}
 						break;
 
 					case EMetasoundFrontendClassType::Output:
-						if (FConstClassOutputAccessPtr OwningGraphClassOutputPtr = FindOutputDescriptionWithNodeID(InNodeAndClass.Node->ID))
 						{
-							FOutputNodeController::FInitParams InitParams
+							FConstClassOutputAccessPtr OwningGraphClassOutputPtr = FindOutputDescriptionWithNodeID(Node->ID);
+							if (nullptr != OwningGraphClassOutputPtr.Get())
 							{
-								ConstCastAccessPtr<FNodeAccessPtr>(InNodeAndClass.Node),
-								InNodeAndClass.Class,
-								ConstCastAccessPtr<FClassOutputAccessPtr>(OwningGraphClassOutputPtr),
-								ConstCastAccessPtr<FGraphAccessPtr>(GraphPtr),
-								ConstCastSharedRef<IGraphController>(OwningGraph)
-							};
-							return FOutputNodeController::CreateConstOutputNodeHandle(InitParams);
+								FOutputNodeController::FInitParams InitParams
+								{
+									ConstCastAccessPtr<FNodeAccessPtr>(InNodeAndClass.Node),
+									InNodeAndClass.Class,
+									ConstCastAccessPtr<FClassOutputAccessPtr>(OwningGraphClassOutputPtr),
+									ConstCastAccessPtr<FGraphAccessPtr>(GraphPtr),
+									ConstCastSharedRef<IGraphController>(OwningGraph)
+								};
+								return FOutputNodeController::CreateConstOutputNodeHandle(InitParams);
+							}
 						}
 						break;
 
@@ -3219,72 +3260,72 @@ namespace Metasound
 
 		FMetasoundFrontendClassInput* FGraphController::FindInputDescriptionWithName(const FString& InName)
 		{
-			if (GraphClassPtr.IsValid())
+			if (FMetasoundFrontendGraphClass* GraphClass = GraphClassPtr.Get())
 			{
-				return GraphClassPtr->Interface.Inputs.FindByPredicate([&](const FMetasoundFrontendClassInput& Desc) { return Desc.Name == InName; });
+				return GraphClass->Interface.Inputs.FindByPredicate([&](const FMetasoundFrontendClassInput& Desc) { return Desc.Name == InName; });
 			}
 			return nullptr;
 		}
 
 		const FMetasoundFrontendClassInput* FGraphController::FindInputDescriptionWithName(const FString& InName) const
 		{
-			if (GraphClassPtr.IsValid())
+			if (const FMetasoundFrontendGraphClass* GraphClass = GraphClassPtr.Get())
 			{
-				return GraphClassPtr->Interface.Inputs.FindByPredicate([&](const FMetasoundFrontendClassInput& Desc) { return Desc.Name == InName; });
+				return GraphClass->Interface.Inputs.FindByPredicate([&](const FMetasoundFrontendClassInput& Desc) { return Desc.Name == InName; });
 			}
 			return nullptr;
 		}
 
 		FMetasoundFrontendClassOutput* FGraphController::FindOutputDescriptionWithName(const FString& InName)
 		{
-			if (GraphClassPtr.IsValid())
+			if (FMetasoundFrontendGraphClass* GraphClass = GraphClassPtr.Get())
 			{
-				return GraphClassPtr->Interface.Outputs.FindByPredicate([&](const FMetasoundFrontendClassOutput& Desc) { return Desc.Name == InName; });
+				return GraphClass->Interface.Outputs.FindByPredicate([&](const FMetasoundFrontendClassOutput& Desc) { return Desc.Name == InName; });
 			}
 			return nullptr;
 		}
 
 		const FMetasoundFrontendClassOutput* FGraphController::FindOutputDescriptionWithName(const FString& InName) const
 		{
-			if (GraphClassPtr.IsValid())
+			if (const FMetasoundFrontendGraphClass* GraphClass = GraphClassPtr.Get())
 			{
-				return GraphClassPtr->Interface.Outputs.FindByPredicate([&](const FMetasoundFrontendClassOutput& Desc) { return Desc.Name == InName; });
+				return GraphClass->Interface.Outputs.FindByPredicate([&](const FMetasoundFrontendClassOutput& Desc) { return Desc.Name == InName; });
 			}
 			return nullptr;
 		}
 
 		FMetasoundFrontendClassInput* FGraphController::FindInputDescriptionWithVertexID(const FGuid& InVertexID)
 		{
-			if (GraphClassPtr.IsValid())
+			if (FMetasoundFrontendGraphClass* GraphClass = GraphClassPtr.Get())
 			{
-				return GraphClassPtr->Interface.Inputs.FindByPredicate([&](const FMetasoundFrontendClassInput& Desc) { return Desc.VertexID == InVertexID; });
+				return GraphClass->Interface.Inputs.FindByPredicate([&](const FMetasoundFrontendClassInput& Desc) { return Desc.VertexID == InVertexID; });
 			}
 			return nullptr;
 		}
 
 		const FMetasoundFrontendClassInput* FGraphController::FindInputDescriptionWithVertexID(const FGuid& InVertexID) const
 		{
-			if (GraphClassPtr.IsValid())
+			if (const FMetasoundFrontendGraphClass* GraphClass = GraphClassPtr.Get())
 			{
-				return GraphClassPtr->Interface.Inputs.FindByPredicate([&](const FMetasoundFrontendClassInput& Desc) { return Desc.VertexID == InVertexID; });
+				return GraphClass->Interface.Inputs.FindByPredicate([&](const FMetasoundFrontendClassInput& Desc) { return Desc.VertexID == InVertexID; });
 			}
 			return nullptr;
 		}
 
 		FMetasoundFrontendClassOutput* FGraphController::FindOutputDescriptionWithVertexID(const FGuid& InVertexID)
 		{
-			if (GraphClassPtr.IsValid())
+			if (FMetasoundFrontendGraphClass* GraphClass = GraphClassPtr.Get())
 			{
-				return GraphClassPtr->Interface.Outputs.FindByPredicate([&](const FMetasoundFrontendClassOutput& Desc) { return Desc.VertexID == InVertexID; });
+				return GraphClass->Interface.Outputs.FindByPredicate([&](const FMetasoundFrontendClassOutput& Desc) { return Desc.VertexID == InVertexID; });
 			}
 			return nullptr;
 		}
 
 		const FMetasoundFrontendClassOutput* FGraphController::FindOutputDescriptionWithVertexID(const FGuid& InVertexID) const
 		{
-			if (GraphClassPtr.IsValid())
+			if (const FMetasoundFrontendGraphClass* GraphClass = GraphClassPtr.Get())
 			{
-				return GraphClassPtr->Interface.Outputs.FindByPredicate([&](const FMetasoundFrontendClassOutput& Desc) { return Desc.VertexID == InVertexID; });
+				return GraphClass->Interface.Outputs.FindByPredicate([&](const FMetasoundFrontendClassOutput& Desc) { return Desc.VertexID == InVertexID; });
 			}
 			return nullptr;
 		}
@@ -3339,23 +3380,23 @@ namespace Metasound
 
 		bool FDocumentController::IsValid() const
 		{
-			return DocumentPtr.IsValid();
+			return (nullptr != DocumentPtr.Get());
 		}
 
 		TArray<FMetasoundFrontendClass> FDocumentController::GetDependencies() const
 		{
-			if (DocumentPtr.IsValid())
+			if (const FMetasoundFrontendDocument* Document = DocumentPtr.Get())
 			{
-				return DocumentPtr->Dependencies;
+				return Document->Dependencies;
 			}
 			return TArray<FMetasoundFrontendClass>();
 		}
 
 		TArray<FMetasoundFrontendGraphClass> FDocumentController::GetSubgraphs() const
 		{
-			if (DocumentPtr.IsValid())
+			if (const FMetasoundFrontendDocument* Document = DocumentPtr.Get())
 			{
-				return DocumentPtr->Subgraphs;
+				return Document->Subgraphs;
 			}
 			return TArray<FMetasoundFrontendGraphClass>();
 		}
@@ -3378,65 +3419,63 @@ namespace Metasound
 
 		bool FDocumentController::AddDuplicateSubgraph(const FMetasoundFrontendGraphClass& InGraphToCopy, const FMetasoundFrontendDocument& InOtherDocument)
 		{
-			FMetasoundFrontendDocument* ThisDocument = DocumentPtr.Get();
-			if (nullptr == ThisDocument)
+			if (FMetasoundFrontendDocument* Document = DocumentPtr.Get())
 			{
-				UE_LOG(LogMetaSound, Error, TEXT("Cannot add subgraph to invalid document"));
-				return false;
-			}
-
-			// Direct copy of subgraph
-			bool bSuccess = true;
-			FMetasoundFrontendGraphClass SubgraphCopy(InGraphToCopy);
+				// Direct copy of subgraph
+				bool bSuccess = true;
+				FMetasoundFrontendGraphClass SubgraphCopy(InGraphToCopy);
 
-			for (FMetasoundFrontendNode& Node : SubgraphCopy.Graph.Nodes)
-			{
-				const FGuid OriginalClassID = Node.ClassID;
-
-				auto IsClassWithClassID = [&](const FMetasoundFrontendClass& InClass) -> bool
+				for (FMetasoundFrontendNode& Node : SubgraphCopy.Graph.Nodes)
 				{
-					return InClass.ID == OriginalClassID;
-				};
+					const FGuid OriginalClassID = Node.ClassID;
 
-				if (const FMetasoundFrontendClass* OriginalNodeClass = InOtherDocument.Dependencies.FindByPredicate(IsClassWithClassID))
-				{
-					// Should not be a graph class since it's in the dependencies list
-					check(EMetasoundFrontendClassType::Graph != OriginalNodeClass->Metadata.Type);
+					auto IsClassWithClassID = [&](const FMetasoundFrontendClass& InClass) -> bool
+					{
+						return InClass.ID == OriginalClassID;
+					};
 
-					if (const FMetasoundFrontendClass* NodeClass = FindOrAddClass(OriginalNodeClass->Metadata).Get())
+					if (const FMetasoundFrontendClass* OriginalNodeClass = InOtherDocument.Dependencies.FindByPredicate(IsClassWithClassID))
 					{
-						// All this just to update this ID. Maybe having globally 
-						// consistent class IDs would help. Or using the classname & version as
-						// a class ID. 
-						Node.ClassID = NodeClass->ID;
+						// Should not be a graph class since it's in the dependencies list
+						check(EMetasoundFrontendClassType::Graph != OriginalNodeClass->Metadata.Type);
+
+						if (const FMetasoundFrontendClass* NodeClass = FindOrAddClass(OriginalNodeClass->Metadata).Get())
+						{
+							// All this just to update this ID. Maybe having globally 
+							// consistent class IDs would help. Or using the classname & version as
+							// a class ID. 
+							Node.ClassID = NodeClass->ID;
+						}
+						else
+						{
+							UE_LOG(LogMetaSound, Error, TEXT("Failed to add subgraph dependency [Class:%s]"), *OriginalNodeClass->Metadata.ClassName.ToString());
+							bSuccess = false;
+						}
 					}
-					else
+					else if (const FMetasoundFrontendGraphClass* OriginalNodeGraphClass = InOtherDocument.Subgraphs.FindByPredicate(IsClassWithClassID))
 					{
-						UE_LOG(LogMetaSound, Error, TEXT("Failed to add subgraph dependency [Class:%s]"), *OriginalNodeClass->Metadata.ClassName.ToString());
-						bSuccess = false;
+						bSuccess = bSuccess && AddDuplicateSubgraph(*OriginalNodeGraphClass, InOtherDocument);
+						if (!bSuccess)
+						{
+							break;
+						}
 					}
-				}
-				else if (const FMetasoundFrontendGraphClass* OriginalNodeGraphClass = InOtherDocument.Subgraphs.FindByPredicate(IsClassWithClassID))
-				{
-					bSuccess = bSuccess && AddDuplicateSubgraph(*OriginalNodeGraphClass, InOtherDocument);
-					if (!bSuccess)
+					else
 					{
-						break;
+						bSuccess = false;
+						UE_LOG(LogMetaSound, Error, TEXT("Failed to copy subgraph. Subgraph document is missing dependency info for node [Node:%s, NodeID:%s]"), *Node.Name, *Node.ID.ToString());
 					}
 				}
-				else
+
+				if (bSuccess)
 				{
-					bSuccess = false;
-					UE_LOG(LogMetaSound, Error, TEXT("Failed to copy subgraph. Subgraph document is missing dependency info for node [Node:%s, NodeID:%s]"), *Node.Name, *Node.ID.ToString());
+					Document->Subgraphs.Add(SubgraphCopy);
 				}
-			}
 
-			if (bSuccess)
-			{
-				ThisDocument->Subgraphs.Add(SubgraphCopy);
+				return bSuccess;
 			}
 
-			return bSuccess;
+			return false;
 		}
 
 		const FMetasoundFrontendVersion& FDocumentController::GetArchetypeVersion() const
@@ -3505,7 +3544,7 @@ namespace Metasound
 		{
 			FConstClassAccessPtr MetasoundClass = FindDependencyWithID(InClassID);
 
-			if (!MetasoundClass.IsValid())
+			if (nullptr == MetasoundClass.Get())
 			{
 				MetasoundClass = FindSubgraphWithID(InClassID);
 			}
@@ -3515,17 +3554,17 @@ namespace Metasound
 
 		void FDocumentController::SetMetadata(const FMetasoundFrontendDocumentMetadata& InMetadata)
 		{
-			if (DocumentPtr.IsValid())
+			if (FMetasoundFrontendDocument* Document = DocumentPtr.Get())
 			{
-				DocumentPtr->Metadata = InMetadata;
+				Document->Metadata = InMetadata;
 			}
 		}
 
 		const FMetasoundFrontendDocumentMetadata& FDocumentController::GetMetadata() const
 		{
-			if (DocumentPtr.IsValid())
+			if (const FMetasoundFrontendDocument* Document = DocumentPtr.Get())
 			{
-				return DocumentPtr->Metadata;
+				return Document->Metadata;
 			}
 
 			return FrontendControllerIntrinsics::GetInvalidValueConstRef<FMetasoundFrontendDocumentMetadata>();
@@ -3538,44 +3577,48 @@ namespace Metasound
 
 		FConstClassAccessPtr FDocumentController::FindOrAddClass(const FNodeRegistryKey& InKey)
 		{
-			FConstClassAccessPtr ClassPtr = FindClass(InKey);
-
-			auto AddClass = [=](FMetasoundFrontendClass&& NewClassDescription)
+			if (FMetasoundFrontendDocument* Document = DocumentPtr.Get())
 			{
-				FConstClassAccessPtr NewClassPtr;
-				if (DocumentPtr.IsValid())
+				FConstClassAccessPtr ClassPtr = FindClass(InKey);
+
+				auto AddClass = [=](FMetasoundFrontendClass&& NewClassDescription)
 				{
+					FConstClassAccessPtr NewClassPtr;
+
 					// Cannot add a subgraph using this method because dependencies
 					// of external graph are not added in this method.
 					check(EMetasoundFrontendClassType::Graph != NewClassDescription.Metadata.Type);
 					NewClassDescription.ID = FGuid::NewGuid();
 
-					DocumentPtr->Dependencies.Add(MoveTemp(NewClassDescription));
+					Document->Dependencies.Add(MoveTemp(NewClassDescription));
 					NewClassPtr = FindClass(InKey);
-				}
-				return NewClassPtr;
-			};
 
-			if (const FMetasoundFrontendClass* MetasoundClass = ClassPtr.Get())
-			{
-				// External node classes must match version to return shared definition.
-				if (MetasoundClass->Metadata.Type == EMetasoundFrontendClassType::External)
+					return NewClassPtr;
+				};
+
+				if (const FMetasoundFrontendClass* MetasoundClass = ClassPtr.Get())
 				{
-					// TODO: Assuming we want to recheck classes when they add another
-					// node, this should be replace with a call to synchronize a 
-					// single class.
-					FMetasoundFrontendClass NewClass = GenerateClassDescription(InKey);
-					if (NewClass.Metadata.Version.Major != ClassPtr->Metadata.Version.Major)
+					// External node classes must match version to return shared definition.
+					if (MetasoundClass->Metadata.Type == EMetasoundFrontendClassType::External)
 					{
-						return AddClass(MoveTemp(NewClass));
+						// TODO: Assuming we want to recheck classes when they add another
+						// node, this should be replace with a call to synchronize a 
+						// single class.
+						FMetasoundFrontendClass NewClass = GenerateClassDescription(InKey);
+						if (NewClass.Metadata.Version.Major != MetasoundClass->Metadata.Version.Major)
+						{
+							return AddClass(MoveTemp(NewClass));
+						}
 					}
+
+					return ClassPtr;
 				}
 
-				return ClassPtr;
+				FMetasoundFrontendClass NewClass = GenerateClassDescription(InKey);
+				return AddClass(MoveTemp(NewClass));
 			}
 
-			FMetasoundFrontendClass NewClass = GenerateClassDescription(InKey);
-			return AddClass(MoveTemp(NewClass));
+			return FConstClassAccessPtr();
 		}
 
 		FConstClassAccessPtr FDocumentController::FindClass(const FMetasoundFrontendClassMetadata& InMetadata) const
@@ -3586,65 +3629,76 @@ namespace Metasound
 		FConstClassAccessPtr FDocumentController::FindOrAddClass(const FMetasoundFrontendClassMetadata& InMetadata)
 		{
 			FConstClassAccessPtr ClassPtr = FindClass(InMetadata);
-
-			if (!DocumentPtr.IsValid())
-			{
-				return ClassPtr;
-			}
-
-			// External node classes must match major version to return shared definition.
-			if (EMetasoundFrontendClassType::External == InMetadata.Type)
+			
+			if (const FMetasoundFrontendClass* Class = ClassPtr.Get())
 			{
-				if (ClassPtr.IsValid())
+				// External node classes must match major version to return shared definition.
+				if (EMetasoundFrontendClassType::External == InMetadata.Type)
 				{
-					if (InMetadata.Version.Major != ClassPtr->Metadata.Version.Major)
+					if (InMetadata.Version.Major != Class->Metadata.Version.Major)
 					{
+						// Mismatched major version. Reset class pointer to null.
 						ClassPtr = FConstClassAccessPtr();
 					}
 				}
 			}
 
-			if (!ClassPtr.IsValid())
+
+			const bool bNoMatchingClassFoundInDocument = (nullptr == ClassPtr.Get());
+			if (bNoMatchingClassFoundInDocument)
 			{
-				if ((EMetasoundFrontendClassType::External == InMetadata.Type) || (EMetasoundFrontendClassType::Input == InMetadata.Type) || (EMetasoundFrontendClassType::Output == InMetadata.Type))
+				// If no matching class found, attempt to add a class matching the metadata.
+				if (FMetasoundFrontendDocument* Document = DocumentPtr.Get())
 				{
-					FMetasoundFrontendClass NewClass;
-					FNodeRegistryKey Key = FRegistry::GetRegistryKey(InMetadata);
-
-					if (FRegistry::GetFrontendClassFromRegistered(Key, NewClass))
-					{
-						NewClass.ID = FGuid::NewGuid();
-						DocumentPtr->Dependencies.Add(NewClass);
-					}
-					else
+					switch (InMetadata.Type)
 					{
-						UE_LOG(LogMetaSound, Error,
-							TEXT("Cannot add external dependency. No Metasound class found with matching registry key [Key:%s, Name:%s, Version:%s]. Suggested solution \"%s\" by %s."),
-							*Key,
-							*InMetadata.ClassName.GetFullName().ToString(),
-							*InMetadata.Version.ToString(),
-							*InMetadata.PromptIfMissing.ToString(),
-							*InMetadata.Author.ToString());
+						case EMetasoundFrontendClassType::External:
+						case EMetasoundFrontendClassType::Input:
+						case EMetasoundFrontendClassType::Output:
+						{
+							FMetasoundFrontendClass NewClass;
+							FNodeRegistryKey Key = NodeRegistryKey::CreateKey(InMetadata);
+
+							if (FRegistry::GetFrontendClassFromRegistered(Key, NewClass))
+							{
+								NewClass.ID = FGuid::NewGuid();
+								Document->Dependencies.Add(NewClass);
+							}
+							else
+							{
+								UE_LOG(LogMetaSound, Error,
+									TEXT("Cannot add external dependency. No Metasound class found with matching registry key [Key:%s, Name:%s, Version:%s]. Suggested solution \"%s\" by %s."),
+									*Key,
+									*InMetadata.ClassName.GetFullName().ToString(),
+									*InMetadata.Version.ToString(),
+									*InMetadata.PromptIfMissing.ToString(),
+									*InMetadata.Author.ToString());
+							}
+						}
+						break;
+
+						case EMetasoundFrontendClassType::Graph:
+						{
+							FMetasoundFrontendGraphClass NewClass;
+							NewClass.ID = FGuid::NewGuid();
+							NewClass.Metadata = InMetadata;
+
+							Document->Subgraphs.Add(NewClass);
+						}
+						break;
+
+						default:
+						{
+							UE_LOG(LogMetaSound, Error, TEXT(
+								"Unsupported metasound class type for node: \"%s\" (%s)."),
+								*InMetadata.ClassName.GetFullName().ToString(),
+								*InMetadata.Version.ToString());
+							checkNoEntry();
+						}
 					}
-				}
-				else if (EMetasoundFrontendClassType::Graph == InMetadata.Type)
-				{
-					FMetasoundFrontendGraphClass NewClass;
-					NewClass.ID = FGuid::NewGuid();
-					NewClass.Metadata = InMetadata;
 
-					DocumentPtr->Subgraphs.Add(NewClass);
-				}
-				else
-				{
-					UE_LOG(LogMetaSound, Error, TEXT(
-						"Unsupported metasound class type for node: \"%s\" (%s)."),
-						*InMetadata.ClassName.GetFullName().ToString(),
-						*InMetadata.Version.ToString());
-					checkNoEntry();
+					ClassPtr = FindClass(InMetadata);
 				}
-
-				ClassPtr = FindClass(InMetadata);
 			}
 
 			return ClassPtr;
@@ -3652,66 +3706,64 @@ namespace Metasound
 
 		void FDocumentController::SynchronizeDependencies()
 		{
-			if (!DocumentPtr.IsValid())
+			if (FMetasoundFrontendDocument* Document = DocumentPtr.Get())
 			{
-				return;
-			}
-
-			int32 NumDependenciesRemovedThisItr = 0;
+				int32 NumDependenciesRemovedThisItr = 0;
 
-			// Repeatedly remove unreferenced dependencies until there are
-			// no unreferenced dependencies left.
-			do
-			{
-				TSet<FGuid> ReferencedDependencyIDs;
-				auto AddNodeClassIDToSet = [&](const FMetasoundFrontendNode& Node)
+				// Repeatedly remove unreferenced dependencies until there are
+				// no unreferenced dependencies left.
+				do
 				{
-					ReferencedDependencyIDs.Add(Node.ClassID);
-				};
+					TSet<FGuid> ReferencedDependencyIDs;
+					auto AddNodeClassIDToSet = [&](const FMetasoundFrontendNode& Node)
+					{
+						ReferencedDependencyIDs.Add(Node.ClassID);
+					};
 
-				auto AddGraphNodeClassIDsToSet = [&](const FMetasoundFrontendGraphClass& GraphClass)
-				{
-					Algo::ForEach(GraphClass.Graph.Nodes, AddNodeClassIDToSet);
-				};
+					auto AddGraphNodeClassIDsToSet = [&](const FMetasoundFrontendGraphClass& GraphClass)
+					{
+						Algo::ForEach(GraphClass.Graph.Nodes, AddNodeClassIDToSet);
+					};
 
-				// Referenced dependencies in root class
-				Algo::ForEach(DocumentPtr->RootGraph.Graph.Nodes, AddNodeClassIDToSet);
+					// Referenced dependencies in root class
+					Algo::ForEach(Document->RootGraph.Graph.Nodes, AddNodeClassIDToSet);
 
-				// Referenced dependencies in subgraphs
-				Algo::ForEach(DocumentPtr->Subgraphs, AddGraphNodeClassIDsToSet);
+					// Referenced dependencies in subgraphs
+					Algo::ForEach(Document->Subgraphs, AddGraphNodeClassIDsToSet);
 
-				auto IsDependencyUnreferenced = [&](const FMetasoundFrontendClass& ClassDependency)
-				{
-					return !ReferencedDependencyIDs.Contains(ClassDependency.ID);
-				};
+					auto IsDependencyUnreferenced = [&](const FMetasoundFrontendClass& ClassDependency)
+					{
+						return !ReferencedDependencyIDs.Contains(ClassDependency.ID);
+					};
 
-				NumDependenciesRemovedThisItr = DocumentPtr->Dependencies.RemoveAllSwap(IsDependencyUnreferenced);
-			}
-			while (NumDependenciesRemovedThisItr > 0);
+					NumDependenciesRemovedThisItr = Document->Dependencies.RemoveAllSwap(IsDependencyUnreferenced);
+				}
+				while (NumDependenciesRemovedThisItr > 0);
 
-			/*
-			FRegistry* Registry = FRegistry::Get();
-			if (ensure(Registry))
-			{
-				// Make sure classes are up-to-date with registered versions of class
-				for (FMetasoundFrontendClass& Class : DocumentPtr->Dependencies)
+				/*
+				FRegistry* Registry = FRegistry::Get();
+				if (ensure(Registry))
 				{
-					FMetasoundFrontendClass RegisteredClass;
-					FNodeRegistryKey Key = FRegistry::GetRegistryKey(Class.Metadata);
-					if (Registry->FindFrontendClassFromRegistered(Key, RegisteredClass))
+					// Make sure classes are up-to-date with registered versions of class
+					for (FMetasoundFrontendClass& Class : DocumentPtr->Dependencies)
 					{
-						Class = RegisteredClass;
-						// TODO: update nodes for class.
-						// TODO: perform minor version updating.
+						FMetasoundFrontendClass RegisteredClass;
+						FNodeRegistryKey Key = FRegistry::GetRegistryKey(Class.Metadata);
+						if (Registry->FindFrontendClassFromRegistered(Key, RegisteredClass))
+						{
+							Class = RegisteredClass;
+							// TODO: update nodes for class.
+							// TODO: perform minor version updating.
+						}
 					}
 				}
+				*/
 			}
-			*/
 		}
 
 		FGraphHandle FDocumentController::GetRootGraph()
 		{
-			if (DocumentPtr.IsValid())
+			if (IsValid())
 			{
 				FGraphClassAccessPtr GraphClass = DocumentPtr.GetRootGraph();
 				return FGraphController::CreateGraphHandle(FGraphController::FInitParams{GraphClass, this->AsShared()});
@@ -3722,7 +3774,7 @@ namespace Metasound
 
 		FConstGraphHandle FDocumentController::GetRootGraph() const
 		{
-			if (DocumentPtr.IsValid())
+			if (IsValid())
 			{
 				FConstGraphClassAccessPtr GraphClass = DocumentPtr.GetRootGraph();
 				return FGraphController::CreateConstGraphHandle(FGraphController::FInitParams
@@ -3738,9 +3790,9 @@ namespace Metasound
 		{
 			TArray<FGraphHandle> Subgraphs;
 
-			if (DocumentPtr.IsValid())
+			if (FMetasoundFrontendDocument* Document = DocumentPtr.Get())
 			{
-				for (FMetasoundFrontendGraphClass& GraphClass : DocumentPtr->Subgraphs)
+				for (FMetasoundFrontendGraphClass& GraphClass : Document->Subgraphs)
 				{
 					Subgraphs.Add(GetSubgraphWithClassID(GraphClass.ID));
 				}
@@ -3753,9 +3805,9 @@ namespace Metasound
 		{
 			TArray<FConstGraphHandle> Subgraphs;
 
-			if (DocumentPtr.IsValid())
+			if (const FMetasoundFrontendDocument* Document = DocumentPtr.Get())
 			{
-				for (const FMetasoundFrontendGraphClass& GraphClass : DocumentPtr->Subgraphs)
+				for (const FMetasoundFrontendGraphClass& GraphClass : Document->Subgraphs)
 				{
 					Subgraphs.Add(GetSubgraphWithClassID(GraphClass.ID));
 				}
@@ -3780,12 +3832,12 @@ namespace Metasound
 
 		bool FDocumentController::ExportToJSONAsset(const FString& InAbsolutePath) const
 		{
-			if (DocumentPtr.IsValid())
+			if (const FMetasoundFrontendDocument* Document = DocumentPtr.Get())
 			{
 				if (TUniquePtr<FArchive> FileWriter = TUniquePtr<FArchive>(IFileManager::Get().CreateFileWriter(*InAbsolutePath)))
 				{
 					TJsonStructSerializerBackend<DefaultCharType> Backend(*FileWriter, EStructSerializerBackendFlags::Default);
-					FStructSerializer::Serialize<FMetasoundFrontendDocument>(*DocumentPtr, Backend);
+					FStructSerializer::Serialize<FMetasoundFrontendDocument>(*Document, Backend);
 			
 					FileWriter->Close();
 
@@ -3802,45 +3854,27 @@ namespace Metasound
 
 		FString FDocumentController::ExportToJSON() const
 		{
-			TArray<uint8> WriterBuffer;
-			FMemoryWriter MemWriter(WriterBuffer);
+			FString Output;
 
-			Metasound::TJsonStructSerializerBackend<Metasound::DefaultCharType> Backend(MemWriter, EStructSerializerBackendFlags::Default);
-			FStructSerializer::Serialize<FMetasoundFrontendDocument>(*DocumentPtr, Backend);
+			if (const FMetasoundFrontendDocument* Document = DocumentPtr.Get())
+			{
+				TArray<uint8> WriterBuffer;
+				FMemoryWriter MemWriter(WriterBuffer);
 
-			MemWriter.Close();
+				Metasound::TJsonStructSerializerBackend<Metasound::DefaultCharType> Backend(MemWriter, EStructSerializerBackendFlags::Default);
+				FStructSerializer::Serialize<FMetasoundFrontendDocument>(*Document, Backend);
 
-			// null terminator
-			WriterBuffer.AddZeroed(sizeof(ANSICHAR));
+				MemWriter.Close();
 
-			FString Output;
-			Output.AppendChars(reinterpret_cast<ANSICHAR*>(WriterBuffer.GetData()), WriterBuffer.Num() / sizeof(ANSICHAR));
+				// null terminator
+				WriterBuffer.AddZeroed(sizeof(ANSICHAR));
 
-			return Output;
-		}
-		
-		bool FDocumentController::IsMatchingMetasoundClass(const FMetasoundFrontendClassMetadata& InMetadataA, const FMetasoundFrontendClassMetadata& InMetadataB) 
-		{
-			if (InMetadataA.Type == InMetadataB.Type)
-			{
-				if (InMetadataA.ClassName == InMetadataB.ClassName)
-				{
-					return FRegistry::GetRegistryKey(InMetadataA) == FRegistry::GetRegistryKey(InMetadataB);
-				}
+				Output.AppendChars(reinterpret_cast<ANSICHAR*>(WriterBuffer.GetData()), WriterBuffer.Num() / sizeof(ANSICHAR));
 			}
-			return false;
-		}
-
-		bool FDocumentController::IsMatchingMetasoundClass(const FNodeClassInfo& InNodeClass, const FMetasoundFrontendClassMetadata& InMetadata) 
-		{
-			return FRegistry::GetRegistryKey(InNodeClass) == FRegistry::GetRegistryKey(InMetadata);
-		}
 
-		bool FDocumentController::IsMatchingMetasoundClass(const FNodeRegistryKey& InKey, const FMetasoundFrontendClassMetadata& InMetadata) 
-		{
-			return InKey == FRegistry::GetRegistryKey(InMetadata);
+			return Output;
 		}
-
+		
 		FDocumentAccess FDocumentController::ShareAccess() 
 		{
 			FDocumentAccess Access;
diff --git a/Engine/Plugins/Runtime/Metasound/Source/MetasoundFrontend/Private/MetasoundFrontendStandardController.h b/Engine/Plugins/Runtime/Metasound/Source/MetasoundFrontend/Private/MetasoundFrontendStandardController.h
--- a/Engine/Plugins/Runtime/Metasound/Source/MetasoundFrontend/Private/MetasoundFrontendStandardController.h
+++ b/Engine/Plugins/Runtime/Metasound/Source/MetasoundFrontend/Private/MetasoundFrontendStandardController.h
@@ -835,16 +835,12 @@ namespace Metasound
 			{
 				FNodeAccessPtr Node;
 				FConstClassAccessPtr Class;
-
-				bool IsValid() const { return Node.IsValid() && Class.IsValid(); }
 			};
 
 			struct FConstNodeAndClass
 			{
 				FConstNodeAccessPtr Node;
 				FConstClassAccessPtr Class;
-
-				bool IsValid() const { return Node.IsValid() && Class.IsValid(); }
 			};
 
 			bool ContainsNodesAndClassesByPredicate(TFunctionRef<bool (const FMetasoundFrontendClass&, const FMetasoundFrontendNode&)> InPredicate) const;
@@ -953,10 +949,6 @@ namespace Metasound
 			bool ExportToJSONAsset(const FString& InAbsolutePath) const override;
 			FString ExportToJSON() const override;
 
-			static bool IsMatchingMetasoundClass(const FMetasoundFrontendClassMetadata& InMetadataA, const FMetasoundFrontendClassMetadata& InMetadataB);
-			static bool IsMatchingMetasoundClass(const FNodeClassInfo& InNodeClass, const FMetasoundFrontendClassMetadata& InMetadata);
-			static bool IsMatchingMetasoundClass(const FNodeRegistryKey& InKey, const FMetasoundFrontendClassMetadata& InMetadata);
-
 		protected:
 			FDocumentAccess ShareAccess() override;
 			FConstDocumentAccess ShareAccess() const override;
diff --git a/Engine/Plugins/Runtime/Metasound/Source/MetasoundFrontend/Private/MetasoundFrontendSubgraphNodeController.cpp b/Engine/Plugins/Runtime/Metasound/Source/MetasoundFrontend/Private/MetasoundFrontendSubgraphNodeController.cpp
--- a/Engine/Plugins/Runtime/Metasound/Source/MetasoundFrontend/Private/MetasoundFrontendSubgraphNodeController.cpp
+++ b/Engine/Plugins/Runtime/Metasound/Source/MetasoundFrontend/Private/MetasoundFrontendSubgraphNodeController.cpp
@@ -73,16 +73,19 @@ namespace Metasound
 
 		FNodeHandle FSubgraphNodeController::CreateNodeHandle(const FSubgraphNodeController::FInitParams& InParams)
 		{
-			if (InParams.NodePtr.IsValid() && InParams.ClassPtr.IsValid())
+			if (FMetasoundFrontendNode* Node = InParams.NodePtr.Get())
 			{
-				// Cannot make a valid node handle if the node description and class description differ
-				if (InParams.NodePtr->ClassID == InParams.ClassPtr->ID)
+				if (const FMetasoundFrontendClass* Class = InParams.ClassPtr.Get())
 				{
-					return MakeShared<FSubgraphNodeController>(EPrivateToken::Token, InParams);
-				}
-				else
-				{
-					UE_LOG(LogMetaSound, Warning, TEXT("Frontend Node [NodeID:%s, ClassID:%s] is not of expected class class [ClassID:%s]"), *InParams.NodePtr->ID.ToString(), *InParams.NodePtr->ClassID.ToString(), *InParams.ClassPtr->ID.ToString());
+					// Cannot make a valid node handle if the node description and class description differ
+					if (Node->ClassID == Class->ID)
+					{
+						return MakeShared<FSubgraphNodeController>(EPrivateToken::Token, InParams);
+					}
+					else
+					{
+						UE_LOG(LogMetaSound, Warning, TEXT("Frontend Node [NodeID:%s, ClassID:%s] is not of expected class class [ClassID:%s]"), *Node->ID.ToString(), *Node->ClassID.ToString(), *Class->ID.ToString());
+					}
 				}
 			}
 			return FInvalidNodeController::GetInvalid();
@@ -90,16 +93,19 @@ namespace Metasound
 
 		FConstNodeHandle FSubgraphNodeController::CreateConstNodeHandle(const FSubgraphNodeController::FInitParams& InParams)
 		{
-			if (InParams.NodePtr.IsValid() && InParams.ClassPtr.IsValid())
+			if (FMetasoundFrontendNode* Node = InParams.NodePtr.Get())
 			{
-				// Cannot make a valid node handle if the node description and class description differ
-				if (InParams.NodePtr->ClassID == InParams.ClassPtr->ID)
+				if (const FMetasoundFrontendClass* Class = InParams.ClassPtr.Get())
 				{
-					return MakeShared<const FSubgraphNodeController>(EPrivateToken::Token, InParams);
-				}
-				else
-				{
-					UE_LOG(LogMetaSound, Warning, TEXT("Frontend Node [NodeID:%s, ClassID:%s] is not of expected class class [ClassID:%s]"), *InParams.NodePtr->ID.ToString(), *InParams.NodePtr->ClassID.ToString(), *InParams.ClassPtr->ID.ToString());
+					// Cannot make a valid node handle if the node description and class description differ
+					if (Node->ClassID == Class->ID)
+					{
+						return MakeShared<const FSubgraphNodeController>(EPrivateToken::Token, InParams);
+					}
+					else
+					{
+						UE_LOG(LogMetaSound, Warning, TEXT("Frontend Node [NodeID:%s, ClassID:%s] is not of expected class class [ClassID:%s]"), *Node->ID.ToString(), *Node->ClassID.ToString(), *Class->ID.ToString());
+					}
 				}
 			}
 			return FInvalidNodeController::GetInvalid();
@@ -107,7 +113,7 @@ namespace Metasound
 
 		bool FSubgraphNodeController::IsValid() const
 		{
-			return FBaseNodeController::IsValid() && GraphPtr.IsValid();
+			return FBaseNodeController::IsValid() && (nullptr != GraphPtr.Get());
 		}
 
 		int32 FSubgraphNodeController::GetNumInputs() const
diff --git a/Engine/Plugins/Runtime/Metasound/Source/MetasoundFrontend/Public/MetasoundAccessPtr.h b/Engine/Plugins/Runtime/Metasound/Source/MetasoundFrontend/Public/MetasoundAccessPtr.h
--- a/Engine/Plugins/Runtime/Metasound/Source/MetasoundFrontend/Public/MetasoundAccessPtr.h
+++ b/Engine/Plugins/Runtime/Metasound/Source/MetasoundFrontend/Public/MetasoundAccessPtr.h
@@ -54,15 +54,6 @@ namespace Metasound
 
 			static Type FallbackObject;
 
-			/** Returns true if object is valid. Should only be called in same
-			 * thread which created the token. */
-			bool IsValid() const { return (Get() != nullptr); }
-
-			FORCEINLINE explicit operator bool() const
-			{
-				return IsValid();
-			}
-
 			/** Returns a pointer to the accessible object. If the object is 
 			 * inaccessible, a nullptr is returned. 
 			 */
@@ -70,36 +61,10 @@ namespace Metasound
 			{
 #if METASOUND_FRONTEND_ACCESSPTR_DEBUG_INFO 
 				CachedObjectPtr = GetObject();
-#endif
+				return CachedObjectPtr;
+#else
 				return GetObject();
-			}
-
-			/** Returns a pointer to the accessible object. If the object is 
-			 * inaccessible, an assert is triggered and the static fallback object
-			 * is returned. 
-			 */
-			Type& operator*() const
-			{
-				if (Type* Object = Get())
-				{
-					return *Object;
-				}
-
-				checkNoEntry();
-				return FallbackObject;
-			}
-
-			/** Returns a pointer to the accessible object. If the object is 
-			 * inaccessible, a nullptr is returned. 
-			 */
-			Type* operator->() const
-			{
-				if (Type* Object = Get())
-				{
-					return Object;
-				}
-
-				return &FallbackObject;
+#endif
 			}
 
 			/** Returns an access pointer to a member of the wrapped object. 
@@ -114,11 +79,11 @@ namespace Metasound
 			{
 				using MemberType = typename AccessPtrType::ObjectType;
 
-				TFunction<MemberType*()> GetMemberFromObject = [=]() -> MemberType*
+				TFunction<MemberType*()> GetMemberFromObject = [GetObject=this->GetObject, GetMember=MoveTemp(InGetMember)]() -> MemberType*
 				{
-					if (Type* Object = Get())
+					if (Type* Object = GetObject())
 					{
-						return InGetMember(*Object);
+						return GetMember(*Object);
 					}
 					return static_cast<MemberType*>(nullptr);
 				};
@@ -139,7 +104,7 @@ namespace Metasound
 			{
 				Type* RefPtr = &InRef;
 
-				GetObject = [=]() -> Type*
+				GetObject = [AccessToken=MoveTemp(AccessToken), RefPtr]() -> Type*
 				{
 					Type* Object = nullptr;
 					if (AccessToken.IsValid())
@@ -158,10 +123,9 @@ namespace Metasound
 			template<typename OtherType>
 			TAccessPtr(const TAccessPtr<OtherType>& InOther, EConstCast InTag)
 			{
-				TFunction<OtherType*()> OtherGetObject = InOther.GetObject;
-				GetObject = [=]() -> Type*
+				GetObject = [GetOtherObject=InOther.GetObject]() -> Type*
 				{
-					return const_cast<Type*>(OtherGetObject());
+					return const_cast<Type*>(GetOtherObject());
 				};
 #if METASOUND_FRONTEND_ACCESSPTR_DEBUG_INFO 
 				Get();
@@ -175,11 +139,9 @@ namespace Metasound
 			>
 			TAccessPtr(const TAccessPtr<OtherType>& InOther, EDerivedCopy InTag=EDerivedCopy::Tag)
 			{
-				TFunction<OtherType*()> OtherGetObject = InOther.GetObject;
-
-				GetObject = [=]() -> Type*
+				GetObject = [GetOtherObject=InOther.GetObject]() -> Type*
 				{
-					return static_cast<Type*>(OtherGetObject());
+					return static_cast<Type*>(GetOtherObject());
 				};
 #if METASOUND_FRONTEND_ACCESSPTR_DEBUG_INFO 
 				Get();
diff --git a/Engine/Plugins/Runtime/Metasound/Source/MetasoundFrontend/Public/MetasoundFrontendRegistries.h b/Engine/Plugins/Runtime/Metasound/Source/MetasoundFrontend/Public/MetasoundFrontendRegistries.h
--- a/Engine/Plugins/Runtime/Metasound/Source/MetasoundFrontend/Public/MetasoundFrontendRegistries.h
+++ b/Engine/Plugins/Runtime/Metasound/Source/MetasoundFrontend/Public/MetasoundFrontendRegistries.h
@@ -55,6 +55,7 @@ namespace Metasound
 	{
 		using FNodeRegistryKey = FString;
 
+
 		// Returns true if the registry key is a valid key. 
 		//
 		// This does *not* connote that the registry key exists in the registry.
@@ -354,6 +355,33 @@ namespace Metasound
 			FNodeRegistryKey Key;
 			FNodeClassInfo NodeClassInfo;
 		};
+
+		namespace NodeRegistryKey
+		{
+			// Returns true if the registry key is a valid key. 
+			//
+			// This does *not* connote that the registry key exists in the registry.
+			METASOUNDFRONTEND_API bool IsValid(const FNodeRegistryKey& InKey);
+
+			// Returns true if both keys represent the same entry in the node registry.
+			METASOUNDFRONTEND_API bool IsEqual(const FNodeRegistryKey& InLHS, const FNodeRegistryKey& InRHS);
+
+			// Returns true if the class metadata and key represent the same entry in the node registry.
+			METASOUNDFRONTEND_API bool IsEqual(const FMetasoundFrontendClassMetadata& InLHS, const FNodeRegistryKey& InRHS);
+
+			// Returns true if the class info and key represent the same entry in the node registry.
+			METASOUNDFRONTEND_API bool IsEqual(const FNodeClassInfo& InLHS, const FNodeRegistryKey& InRHS);
+
+			// Returns true if the class metadatas represent the same entry in the node registry.
+			METASOUNDFRONTEND_API bool IsEqual(const FMetasoundFrontendClassMetadata& InLHS, const FMetasoundFrontendClassMetadata& InRHS);
+
+			// Returns true if the class info and class metadata represent the same entry in the node registry.
+			METASOUNDFRONTEND_API bool IsEqual(const FNodeClassInfo& InLHS, const FMetasoundFrontendClassMetadata& InRHS);
+
+			METASOUNDFRONTEND_API FNodeRegistryKey CreateKey(const FNodeClassMetadata& InNodeMetadata);
+			METASOUNDFRONTEND_API FNodeRegistryKey CreateKey(const FMetasoundFrontendClassMetadata& InNodeMetadata);
+			METASOUNDFRONTEND_API FNodeRegistryKey CreateKey(const FNodeClassInfo& ClassInfo);
+		}
 	} // namespace Frontend
 } // namespace Metasound
 
