Added 3 Trigger nodes to Metasounds
TriggerControl, TriggerOnce, and TriggerOnValueChange

diff --git a/Engine/Plugins/Runtime/Metasound/Source/MetasoundStandardNodes/Private/MetasoundTriggerControlNode.cpp b/Engine/Plugins/Runtime/Metasound/Source/MetasoundStandardNodes/Private/MetasoundTriggerControlNode.cpp
new file mode 100644
--- /dev/null
+++ b/Engine/Plugins/Runtime/Metasound/Source/MetasoundStandardNodes/Private/MetasoundTriggerControlNode.cpp
@@ -0,0 +1,235 @@
+// Copyright Epic Games, Inc. All Rights Reserved.
+
+#include "CoreMinimal.h"
+#include "MetasoundBuilderInterface.h"
+#include "MetasoundDataReferenceCollection.h"
+#include "MetasoundExecutableOperator.h"
+#include "MetasoundFacade.h"
+#include "MetasoundNode.h"
+#include "MetasoundNodeInterface.h"
+#include "MetasoundNodeRegistrationMacro.h"
+#include "MetasoundOperatorInterface.h"
+#include "MetasoundPrimitives.h"
+#include "MetasoundSampleCounter.h"
+#include "MetasoundStandardNodesNames.h"
+#include "MetasoundStandardNodesCategories.h"
+#include "MetasoundParamHelper.h"
+
+#define LOCTEXT_NAMESPACE "MetasoundStandardNodes_TriggerControlNode"
+
+namespace Metasound
+{
+
+	namespace TriggerControlVertexNames
+	{
+		METASOUND_PARAM(InputEnter, "Trigger In", "The input trigger. This will pass through if the node is open.");
+		METASOUND_PARAM(InputOpen, "Open", "Once Triggered, allows triggers to pass through.");
+		METASOUND_PARAM(InputClose, "Close", "Once triggered, prevents all triggers from passing through.");
+		METASOUND_PARAM(InputToggle, "Toggle", "Once triggered, opens the node if closed, and closes the node if opened.");
+		METASOUND_PARAM(InputStartClosed, "Start Closed", "Whether the node should be closed when the Metasound begins.");
+		
+		METASOUND_PARAM(OutputExit, "Trigger Out", "The output trigger.");
+	}
+
+	class FTriggerControlOperator : public TExecutableOperator<FTriggerControlOperator>
+	{
+	public:
+		static const FNodeClassMetadata& GetNodeInfo();
+		static const FVertexInterface& GetVertexInterface();
+		static TUniquePtr<IOperator> CreateOperator(const FCreateOperatorParams& InParams, FBuildErrorArray& OutErrors);
+
+		FTriggerControlOperator(const FOperatorSettings& InSettings,
+			const FTriggerReadRef& InTriggerEnter,
+			const FTriggerReadRef& InTriggerOpen,
+			const FTriggerReadRef& InTriggerClose,
+			const FTriggerReadRef& InTriggerToggle,
+			const FBoolReadRef& InStartClosed);
+
+		virtual FDataReferenceCollection GetInputs() const override;
+
+		virtual FDataReferenceCollection GetOutputs() const override;
+
+		void Execute();
+
+	private:
+		// Try to go through gate
+		FTriggerReadRef TriggerEnterInput;
+		// Open the gate
+		FTriggerReadRef TriggerOpenInput;
+		// Close the gate
+		FTriggerReadRef TriggerCloseInput;
+		// Switch between open and closed
+		FTriggerReadRef TriggerToggleInput;
+		// Whether the gate starts closed
+		FBoolReadRef bStartClosedInput;
+
+		// If gate is open, sends trigger
+		FTriggerWriteRef TriggerExitOutput;
+
+		// Status of the gate
+		bool bIsGateOpen;
+	};
+
+	FTriggerControlOperator::FTriggerControlOperator(const FOperatorSettings& InSettings,
+		const FTriggerReadRef& InTriggerEnter,
+		const FTriggerReadRef& InTriggerOpen,
+		const FTriggerReadRef& InTriggerClose,
+		const FTriggerReadRef& InTriggerToggle,
+		const FBoolReadRef& InStartClosed)
+		: TriggerEnterInput(InTriggerEnter)
+		, TriggerOpenInput(InTriggerOpen)
+		, TriggerCloseInput(InTriggerClose)
+		, TriggerToggleInput(InTriggerToggle)
+		, bStartClosedInput(InStartClosed)
+		, TriggerExitOutput(FTriggerWriteRef::CreateNew(InSettings))
+		, bIsGateOpen(!*bStartClosedInput)
+	{
+	}
+
+	FDataReferenceCollection FTriggerControlOperator::GetInputs() const
+	{
+		using namespace TriggerControlVertexNames;
+
+		FDataReferenceCollection InputDataReferences;
+		InputDataReferences.AddDataReadReference(METASOUND_GET_PARAM_NAME(InputEnter), TriggerEnterInput);
+		InputDataReferences.AddDataReadReference(METASOUND_GET_PARAM_NAME(InputOpen), TriggerOpenInput);
+		InputDataReferences.AddDataReadReference(METASOUND_GET_PARAM_NAME(InputClose), TriggerCloseInput);
+		InputDataReferences.AddDataReadReference(METASOUND_GET_PARAM_NAME(InputToggle), TriggerToggleInput);
+		InputDataReferences.AddDataReadReference(METASOUND_GET_PARAM_NAME(InputStartClosed), bStartClosedInput);
+		return InputDataReferences;
+	}
+
+	FDataReferenceCollection FTriggerControlOperator::GetOutputs() const
+	{
+		using namespace TriggerControlVertexNames;
+
+		FDataReferenceCollection OutputDataReferences;
+		OutputDataReferences.AddDataReadReference(METASOUND_GET_PARAM_NAME(OutputExit), TriggerExitOutput);
+
+		return OutputDataReferences;
+	}
+
+	void FTriggerControlOperator::Execute()
+	{
+		TriggerExitOutput->AdvanceBlock();
+		
+		// Open Gate
+		TriggerOpenInput->ExecuteBlock(
+			[](int32, int32)
+			{
+			},
+			[this](int32 StartFrame, int32 EndFrame)
+			{
+				bIsGateOpen = true;
+			}
+		);
+		// Close gate
+		TriggerCloseInput->ExecuteBlock(
+			[](int32, int32)
+			{
+			},
+			[this](int32 StartFrame, int32 EndFrame)
+			{
+				bIsGateOpen = false;
+			}
+		);
+		// Toggle gate status
+		TriggerToggleInput->ExecuteBlock(
+			[](int32, int32)
+			{
+			},
+			[this](int32 StartFrame, int32 EndFrame)
+			{
+				bIsGateOpen = !bIsGateOpen;
+			}
+		);
+
+		// Pass through trigger if gate is open
+		TriggerEnterInput->ExecuteBlock(
+			[](int32, int32) 
+			{
+			},
+			[this](int32 StartFrame, int32 EndFrame)
+			{
+				if (bIsGateOpen)
+				{
+					TriggerExitOutput->TriggerFrame(StartFrame);
+				}
+			}
+		);		
+
+	}
+
+
+	TUniquePtr<IOperator> FTriggerControlOperator::CreateOperator(const FCreateOperatorParams& InParams, FBuildErrorArray& OutErrors)
+	{
+		using namespace TriggerControlVertexNames;
+
+		const FInputVertexInterface& InputInterface = GetVertexInterface().GetInputInterface();
+
+		FTriggerReadRef TriggerEnterIn = InParams.InputDataReferences.GetDataReadReferenceOrConstruct<FTrigger>(METASOUND_GET_PARAM_NAME(InputEnter), InParams.OperatorSettings);
+		FTriggerReadRef TriggerOpenIn = InParams.InputDataReferences.GetDataReadReferenceOrConstruct<FTrigger>(METASOUND_GET_PARAM_NAME(InputOpen), InParams.OperatorSettings);
+		FTriggerReadRef TriggerCloseIn = InParams.InputDataReferences.GetDataReadReferenceOrConstruct<FTrigger>(METASOUND_GET_PARAM_NAME(InputClose), InParams.OperatorSettings);
+		FTriggerReadRef TriggerToggleIn = InParams.InputDataReferences.GetDataReadReferenceOrConstruct<FTrigger>(METASOUND_GET_PARAM_NAME(InputToggle), InParams.OperatorSettings);
+		FBoolReadRef bStartClosedIn = InParams.InputDataReferences.GetDataReadReferenceOrConstructWithVertexDefault<bool>(InputInterface, METASOUND_GET_PARAM_NAME(InputStartClosed), InParams.OperatorSettings);
+
+		return MakeUnique<FTriggerControlOperator>(InParams.OperatorSettings, TriggerEnterIn, TriggerOpenIn, TriggerCloseIn, TriggerToggleIn, bStartClosedIn);
+	}
+
+	const FVertexInterface& FTriggerControlOperator::GetVertexInterface()
+	{
+		using namespace TriggerControlVertexNames;
+
+		static const FVertexInterface Interface(
+			FInputVertexInterface(
+				TInputDataVertexModel<FTrigger>(METASOUND_GET_PARAM_NAME_AND_TT(InputEnter)),
+				TInputDataVertexModel<FTrigger>(METASOUND_GET_PARAM_NAME_AND_TT(InputOpen)),
+				TInputDataVertexModel<FTrigger>(METASOUND_GET_PARAM_NAME_AND_TT(InputClose)),
+				TInputDataVertexModel<FTrigger>(METASOUND_GET_PARAM_NAME_AND_TT(InputToggle)),
+				TInputDataVertexModel<bool>(METASOUND_GET_PARAM_NAME_AND_TT(InputStartClosed), true)
+			),
+			FOutputVertexInterface(
+				TOutputDataVertexModel<FTrigger>(METASOUND_GET_PARAM_NAME_AND_TT(OutputExit))
+			)
+		);
+
+		return Interface;
+	}
+
+	const FNodeClassMetadata& FTriggerControlOperator::GetNodeInfo()
+	{
+		auto InitNodeInfo = []() -> FNodeClassMetadata
+		{
+			FNodeClassMetadata Info;
+			Info.ClassName = { Metasound::StandardNodes::Namespace, TEXT("Trigger Control"), TEXT("") };
+			Info.MajorVersion = 1;
+			Info.MinorVersion = 0;
+			Info.DisplayName = LOCTEXT("Metasound_TriggerControlNodeDisplayName", "Trigger Control");
+			Info.Description = LOCTEXT("Metasound_TriggerControlNodeDescription", "Control whether input triggers are passed through.");
+			Info.Author = PluginAuthor;
+			Info.PromptIfMissing = PluginNodeMissingPrompt;
+			Info.DefaultInterface = GetVertexInterface();
+			Info.CategoryHierarchy.Emplace(StandardNodes::TriggerUtils);
+
+			return Info;
+		};
+
+		static const FNodeClassMetadata Info = InitNodeInfo();
+
+		return Info;
+	}
+
+	// Node Class
+	class FTriggerControlNode : public FNodeFacade
+	{
+	public:
+		FTriggerControlNode(const FNodeInitData& InitData)
+			: FNodeFacade(InitData.InstanceName, InitData.InstanceID, TFacadeOperatorClass<FTriggerControlOperator>())
+		{
+		}
+	};
+
+	METASOUND_REGISTER_NODE(FTriggerControlNode)
+}
+
+#undef LOCTEXT_NAMESPACE
diff --git a/Engine/Plugins/Runtime/Metasound/Source/MetasoundStandardNodes/Private/MetasoundTriggerOnValueChangeNode.cpp b/Engine/Plugins/Runtime/Metasound/Source/MetasoundStandardNodes/Private/MetasoundTriggerOnValueChangeNode.cpp
new file mode 100644
--- /dev/null
+++ b/Engine/Plugins/Runtime/Metasound/Source/MetasoundStandardNodes/Private/MetasoundTriggerOnValueChangeNode.cpp
@@ -0,0 +1,195 @@
+// Copyright Epic Games, Inc. All Rights Reserved.
+
+#include "CoreMinimal.h"
+#include "MetasoundBuilderInterface.h"
+#include "MetasoundDataReferenceCollection.h"
+#include "MetasoundExecutableOperator.h"
+#include "MetasoundFacade.h"
+#include "MetasoundNode.h"
+#include "MetasoundNodeInterface.h"
+#include "MetasoundNodeRegistrationMacro.h"
+#include "MetasoundOperatorInterface.h"
+#include "MetasoundPrimitives.h"
+#include "MetasoundStandardNodesNames.h"
+#include "MetasoundStandardNodesCategories.h"
+#include "MetasoundParamHelper.h"
+
+#define LOCTEXT_NAMESPACE "MetasoundStandardNodes_TriggerOnValueChangeNode"
+
+namespace Metasound
+{
+
+	namespace TriggerOnValueChangeVertexNames
+	{
+		METASOUND_PARAM(InputValue, "Value", "The input value to watch. Whenever this changes, the output trigger is sent.");
+		
+		METASOUND_PARAM(OutputOnChange, "Trigger", "The output trigger.");
+	}
+
+	namespace TriggerOnValueChangePrivate
+	{
+		template <typename ValueType>
+		static bool IsValueEqual(ValueType InValueA, ValueType InValueB)
+		{
+			return InValueA == InValueB;
+		}
+
+		static bool IsValueEqual(float InValueA, float InValueB)
+		{
+			return FMath::IsNearlyEqual(InValueA, InValueB);
+		}		
+	}
+
+
+	template <typename ValueType>
+	class TTriggerOnValueChangeOperator : public TExecutableOperator<TTriggerOnValueChangeOperator<ValueType>>
+	{
+	public:
+		static const FNodeClassMetadata& GetNodeInfo();
+		static const FVertexInterface& GetVertexInterface();
+		static TUniquePtr<IOperator> CreateOperator(const FCreateOperatorParams& InParams, FBuildErrorArray& OutErrors);
+
+		TTriggerOnValueChangeOperator(const FOperatorSettings& InSettings, const TDataReadReference<ValueType>& InValue);
+
+		virtual FDataReferenceCollection GetInputs() const override;
+
+		virtual FDataReferenceCollection GetOutputs() const override;
+
+		void Execute();
+
+	private:
+		// Parameter to watch for changes
+		TDataReadReference<ValueType> ValueInput;
+
+		// If gate is open, sends trigger
+		FTriggerWriteRef TriggerOnChangeOutput;
+
+		// Status of the gate
+		ValueType PrevValue;
+	};
+
+	template <typename ValueType>
+	TTriggerOnValueChangeOperator<ValueType>::TTriggerOnValueChangeOperator(const FOperatorSettings& InSettings, const TDataReadReference<ValueType>& InValue)
+		: ValueInput(InValue)
+		, TriggerOnChangeOutput(FTriggerWriteRef::CreateNew(InSettings))
+		, PrevValue(*InValue)
+	{
+	}
+
+	template <typename ValueType>
+	FDataReferenceCollection TTriggerOnValueChangeOperator<ValueType>::GetInputs() const
+	{
+		using namespace TriggerOnValueChangeVertexNames;
+
+		FDataReferenceCollection InputDataReferences;
+		InputDataReferences.AddDataReadReference(METASOUND_GET_PARAM_NAME(InputValue), ValueInput);
+		return InputDataReferences;
+	}
+
+	template <typename ValueType>
+	FDataReferenceCollection TTriggerOnValueChangeOperator<ValueType>::GetOutputs() const
+	{
+		using namespace TriggerOnValueChangeVertexNames;
+
+		FDataReferenceCollection OutputDataReferences;
+		OutputDataReferences.AddDataReadReference(METASOUND_GET_PARAM_NAME(OutputOnChange), TriggerOnChangeOutput);
+
+		return OutputDataReferences;
+	}
+
+	template <typename ValueType>
+	void TTriggerOnValueChangeOperator<ValueType>::Execute()
+	{
+		using namespace TriggerOnValueChangePrivate;
+
+		TriggerOnChangeOutput->AdvanceBlock();
+
+		// If value changes, call the trigger at the start of the audio block
+		if (!IsValueEqual(*ValueInput, PrevValue))
+		{
+			PrevValue = *ValueInput;
+			TriggerOnChangeOutput->TriggerFrame(0);
+		}
+
+	}
+
+	template <typename ValueType>
+	TUniquePtr<IOperator> TTriggerOnValueChangeOperator<ValueType>::CreateOperator(const FCreateOperatorParams& InParams, FBuildErrorArray& OutErrors)
+	{
+		using namespace TriggerOnValueChangeVertexNames;
+
+		const FInputVertexInterface& InputInterface = GetVertexInterface().GetInputInterface();
+
+		TDataReadReference<ValueType> ValueIn = InParams.InputDataReferences.GetDataReadReferenceOrConstructWithVertexDefault<ValueType>(InputInterface, METASOUND_GET_PARAM_NAME(InputValue), InParams.OperatorSettings);
+
+		return MakeUnique<TTriggerOnValueChangeOperator>(InParams.OperatorSettings, ValueIn);
+	}
+
+	template <typename ValueType>
+	const FVertexInterface& TTriggerOnValueChangeOperator<ValueType>::GetVertexInterface()
+	{
+		using namespace TriggerOnValueChangeVertexNames;
+
+		static const FVertexInterface Interface(
+			FInputVertexInterface(
+				TInputDataVertexModel<ValueType>(METASOUND_GET_PARAM_NAME_AND_TT(InputValue), 1)
+			),
+			FOutputVertexInterface(
+				TOutputDataVertexModel<FTrigger>(METASOUND_GET_PARAM_NAME_AND_TT(OutputOnChange))
+			)
+		);
+
+		return Interface;
+	}
+
+	template <typename ValueType>
+	const FNodeClassMetadata& TTriggerOnValueChangeOperator<ValueType>::GetNodeInfo()
+	{
+		auto InitNodeInfo = []() -> FNodeClassMetadata
+		{
+			const FName DataTypeName = GetMetasoundDataTypeName<ValueType>();
+			const FName OperatorName = TEXT("Trigger On Value Change");
+			const FText NodeDisplayName = FText::Format(LOCTEXT("TriggerOnValueChangeName", "Trigger On Value Change ({0})"), GetMetasoundDataTypeDisplayText<ValueType>());
+			const FText NodeDescription = LOCTEXT("TriggerOnValueChangeNameDesc", "Triggers when a given value changes.");
+
+			FNodeClassMetadata Info;
+			Info.ClassName = { Metasound::StandardNodes::Namespace, OperatorName, DataTypeName };
+			Info.MajorVersion = 1;
+			Info.MinorVersion = 0;
+			Info.DisplayName = NodeDisplayName;
+			Info.Description = NodeDescription;
+			Info.Author = PluginAuthor;
+			Info.PromptIfMissing = PluginNodeMissingPrompt;
+			Info.DefaultInterface = GetVertexInterface();
+			Info.CategoryHierarchy.Emplace(StandardNodes::TriggerUtils);
+
+			return Info;
+		};
+
+		static const FNodeClassMetadata Info = InitNodeInfo();
+
+		return Info;
+	}
+
+	// Node Class
+	template <typename ValueType>
+	class TTriggerOnValueChangeNode : public FNodeFacade
+	{
+	public:
+		TTriggerOnValueChangeNode(const FNodeInitData& InitData)
+			: FNodeFacade(InitData.InstanceName, InitData.InstanceID, TFacadeOperatorClass<TTriggerOnValueChangeOperator<ValueType>>())
+		{
+		}
+	};
+
+	using FTriggerOnInt32Change = TTriggerOnValueChangeNode<int32>;
+	METASOUND_REGISTER_NODE(FTriggerOnInt32Change)
+
+	using FTriggerOnFloatChange = TTriggerOnValueChangeNode<float>;
+	METASOUND_REGISTER_NODE(FTriggerOnFloatChange)
+
+	using FTriggerOnBoolChange = TTriggerOnValueChangeNode<bool>;
+	METASOUND_REGISTER_NODE(FTriggerOnBoolChange)
+}
+
+#undef LOCTEXT_NAMESPACE
diff --git a/Engine/Plugins/Runtime/Metasound/Source/MetasoundStandardNodes/Private/MetasoundTriggerOnceNode.cpp b/Engine/Plugins/Runtime/Metasound/Source/MetasoundStandardNodes/Private/MetasoundTriggerOnceNode.cpp
new file mode 100644
--- /dev/null
+++ b/Engine/Plugins/Runtime/Metasound/Source/MetasoundStandardNodes/Private/MetasoundTriggerOnceNode.cpp
@@ -0,0 +1,197 @@
+// Copyright Epic Games, Inc. All Rights Reserved.
+
+#include "CoreMinimal.h"
+#include "MetasoundBuilderInterface.h"
+#include "MetasoundDataReferenceCollection.h"
+#include "MetasoundExecutableOperator.h"
+#include "MetasoundFacade.h"
+#include "MetasoundNode.h"
+#include "MetasoundNodeInterface.h"
+#include "MetasoundNodeRegistrationMacro.h"
+#include "MetasoundOperatorInterface.h"
+#include "MetasoundPrimitives.h"
+#include "MetasoundSampleCounter.h"
+#include "MetasoundStandardNodesNames.h"
+#include "MetasoundStandardNodesCategories.h"
+#include "MetasoundParamHelper.h"
+
+#define LOCTEXT_NAMESPACE "MetasoundStandardNodes_TriggerOnceNode"
+
+namespace Metasound
+{
+
+	namespace TriggerOnceVertexNames
+	{
+		METASOUND_PARAM(InputEnter, "Trigger In", "The input trigger. After the first time this is triggered, no subsequent triggers will pass through.");
+		METASOUND_PARAM(InputReset, "Reset", "When triggered, opens the node to allow another trigger through.");
+		METASOUND_PARAM(InputStartClosed, "Start Closed", "Whether the node should be closed when the Metasound begins.");
+		
+		METASOUND_PARAM(OutputExit, "Trigger Out", "The output trigger.");
+	}
+
+	class FTriggerOnceOperator : public TExecutableOperator<FTriggerOnceOperator>
+	{
+	public:
+		static const FNodeClassMetadata& GetNodeInfo();
+		static const FVertexInterface& GetVertexInterface();
+		static TUniquePtr<IOperator> CreateOperator(const FCreateOperatorParams& InParams, FBuildErrorArray& OutErrors);
+
+		FTriggerOnceOperator(const FOperatorSettings& InSettings,
+			const FTriggerReadRef& InTriggerEnter,
+			const FTriggerReadRef& InTriggerReset,
+			const FBoolReadRef& InStartClosed);
+
+		virtual FDataReferenceCollection GetInputs() const override;
+
+		virtual FDataReferenceCollection GetOutputs() const override;
+
+		void Execute();
+
+	private:
+		// Try to go through gate
+		FTriggerReadRef TriggerEnterInput;
+		// Open the gate
+		FTriggerReadRef TriggerResetInput;
+		// Whether the gate starts closed
+		FBoolReadRef bStartClosedInput;
+
+		// If gate is open, sends trigger
+		FTriggerWriteRef TriggerExitOutput;
+
+		// Status of the gate
+		bool bIsGateOpen;
+	};
+
+	FTriggerOnceOperator::FTriggerOnceOperator(const FOperatorSettings& InSettings,
+		const FTriggerReadRef& InTriggerEnter,
+		const FTriggerReadRef& InTriggerReset,
+		const FBoolReadRef& InStartClosed)
+		: TriggerEnterInput(InTriggerEnter)
+		, TriggerResetInput(InTriggerReset)
+		, bStartClosedInput(InStartClosed)
+		, TriggerExitOutput(FTriggerWriteRef::CreateNew(InSettings))
+		, bIsGateOpen(!*bStartClosedInput)
+	{
+	}
+
+	FDataReferenceCollection FTriggerOnceOperator::GetInputs() const
+	{
+		using namespace TriggerOnceVertexNames;
+
+		FDataReferenceCollection InputDataReferences;
+		InputDataReferences.AddDataReadReference(METASOUND_GET_PARAM_NAME(InputEnter), TriggerEnterInput);
+		InputDataReferences.AddDataReadReference(METASOUND_GET_PARAM_NAME(InputReset), TriggerResetInput);
+		InputDataReferences.AddDataReadReference(METASOUND_GET_PARAM_NAME(InputStartClosed), bStartClosedInput);
+		return InputDataReferences;
+	}
+
+	FDataReferenceCollection FTriggerOnceOperator::GetOutputs() const
+	{
+		using namespace TriggerOnceVertexNames;
+
+		FDataReferenceCollection OutputDataReferences;
+		OutputDataReferences.AddDataReadReference(METASOUND_GET_PARAM_NAME(OutputExit), TriggerExitOutput);
+
+		return OutputDataReferences;
+	}
+
+	void FTriggerOnceOperator::Execute()
+	{
+		TriggerExitOutput->AdvanceBlock();
+
+		// Pass through trigger if gate is open
+		TriggerEnterInput->ExecuteBlock(
+			[](int32, int32) 
+			{
+			},
+			[this](int32 StartFrame, int32 EndFrame)
+			{
+				if (bIsGateOpen)
+				{
+					bIsGateOpen = false;
+					TriggerExitOutput->TriggerFrame(StartFrame);
+				}
+			}
+		);		
+		// Open Gate
+		TriggerResetInput->ExecuteBlock(
+			[](int32, int32)
+			{
+			},
+			[this](int32 StartFrame, int32 EndFrame)
+			{
+				bIsGateOpen = true;
+			}
+			);
+
+	}
+
+
+	TUniquePtr<IOperator> FTriggerOnceOperator::CreateOperator(const FCreateOperatorParams& InParams, FBuildErrorArray& OutErrors)
+	{
+		using namespace TriggerOnceVertexNames;
+
+		const FInputVertexInterface& InputInterface = GetVertexInterface().GetInputInterface();
+
+		FTriggerReadRef TriggerEnterIn = InParams.InputDataReferences.GetDataReadReferenceOrConstruct<FTrigger>(METASOUND_GET_PARAM_NAME(InputEnter), InParams.OperatorSettings);
+		FTriggerReadRef TriggerResetIn = InParams.InputDataReferences.GetDataReadReferenceOrConstruct<FTrigger>(METASOUND_GET_PARAM_NAME(InputReset), InParams.OperatorSettings);
+		FBoolReadRef bStartClosedIn = InParams.InputDataReferences.GetDataReadReferenceOrConstructWithVertexDefault<bool>(InputInterface, METASOUND_GET_PARAM_NAME(InputStartClosed), InParams.OperatorSettings);
+
+		return MakeUnique<FTriggerOnceOperator>(InParams.OperatorSettings, TriggerEnterIn, TriggerResetIn, bStartClosedIn);
+	}
+
+	const FVertexInterface& FTriggerOnceOperator::GetVertexInterface()
+	{
+		using namespace TriggerOnceVertexNames;
+
+		static const FVertexInterface Interface(
+			FInputVertexInterface(
+				TInputDataVertexModel<FTrigger>(METASOUND_GET_PARAM_NAME_AND_TT(InputEnter)),
+				TInputDataVertexModel<FTrigger>(METASOUND_GET_PARAM_NAME_AND_TT(InputReset)),
+				TInputDataVertexModel<bool>(METASOUND_GET_PARAM_NAME_AND_TT(InputStartClosed), false)
+			),
+			FOutputVertexInterface(
+				TOutputDataVertexModel<FTrigger>(METASOUND_GET_PARAM_NAME_AND_TT(OutputExit))
+			)
+		);
+
+		return Interface;
+	}
+
+	const FNodeClassMetadata& FTriggerOnceOperator::GetNodeInfo()
+	{
+		auto InitNodeInfo = []() -> FNodeClassMetadata
+		{
+			FNodeClassMetadata Info;
+			Info.ClassName = { Metasound::StandardNodes::Namespace, TEXT("Trigger Once"), TEXT("") };
+			Info.MajorVersion = 1;
+			Info.MinorVersion = 0;
+			Info.DisplayName = LOCTEXT("Metasound_TriggerOnceNodeDisplayName", "Trigger Once");
+			Info.Description = LOCTEXT("Metasound_TriggerOnceNodeDescription", "Sends an output trigger the first time the node is triggered, and ignores all others (can be re-opened).");
+			Info.Author = PluginAuthor;
+			Info.PromptIfMissing = PluginNodeMissingPrompt;
+			Info.DefaultInterface = GetVertexInterface();
+			Info.CategoryHierarchy.Emplace(StandardNodes::TriggerUtils);
+
+			return Info;
+		};
+
+		static const FNodeClassMetadata Info = InitNodeInfo();
+
+		return Info;
+	}
+
+	// Node Class
+	class FTriggerOnceNode : public FNodeFacade
+	{
+	public:
+		FTriggerOnceNode(const FNodeInitData& InitData)
+			: FNodeFacade(InitData.InstanceName, InitData.InstanceID, TFacadeOperatorClass<FTriggerOnceOperator>())
+		{
+		}
+	};
+
+	METASOUND_REGISTER_NODE(FTriggerOnceNode)
+}
+
+#undef LOCTEXT_NAMESPACE
