Fixed a Crash in the Crossfade (Float) node caused by an unclamped input value

diff --git a/Engine/Plugins/Runtime/Metasound/Source/MetasoundStandardNodes/Private/MetasoundCrossfadeNode.cpp b/Engine/Plugins/Runtime/Metasound/Source/MetasoundStandardNodes/Private/MetasoundCrossfadeNode.cpp
--- a/Engine/Plugins/Runtime/Metasound/Source/MetasoundStandardNodes/Private/MetasoundCrossfadeNode.cpp
+++ b/Engine/Plugins/Runtime/Metasound/Source/MetasoundStandardNodes/Private/MetasoundCrossfadeNode.cpp
@@ -270,7 +270,7 @@ namespace Metasound
 			{
 				PrevCrossfadeValue = CurrentCrossfadeValue;
 				IndexA = (int32)FMath::Floor(CurrentCrossfadeValue);
-				IndexB = IndexA + 1;
+				IndexB = FMath::Clamp(IndexA + 1, 0.0f, (float)(NumInputs - 1));
 				Alpha = CurrentCrossfadeValue - (float)IndexA;
 			}
 
