Fix for looping sounds w/ metasound wave player

diff --git a/Engine/Plugins/Runtime/Metasound/Source/MetasoundEngine/Private/MetasoundWavePlayerNode.cpp b/Engine/Plugins/Runtime/Metasound/Source/MetasoundEngine/Private/MetasoundWavePlayerNode.cpp
--- a/Engine/Plugins/Runtime/Metasound/Source/MetasoundEngine/Private/MetasoundWavePlayerNode.cpp
+++ b/Engine/Plugins/Runtime/Metasound/Source/MetasoundEngine/Private/MetasoundWavePlayerNode.cpp
@@ -547,8 +547,11 @@ namespace Metasound
 
 						*LoopPercent = 0.0f;
 
-						// Now seek to the loop start frame
-						SeekDecoder(LoopSeekTime);
+						// Now seek to the loop start frame if we are looping a custom amount. Otherwise, we let the decoder do the looping.
+						if (LoopDurationSeconds->GetSeconds() > 0.0f)
+						{
+							SeekDecoder(LoopSeekTime);
+						}
 					}
 					else
 					{
