Fixing dropped metasound triggers.
Having payload IDs determined by atomic counters on data channels.
Checked in on behalf of [REDACTED_HANDLE]

#pf https://horde.devtools.epicgames.com/job/60f5d15635476b0001780bb9

diff --git a/Engine/Plugins/Runtime/Metasound/Source/MetasoundFrontend/Public/MetasoundRouter.h b/Engine/Plugins/Runtime/Metasound/Source/MetasoundFrontend/Public/MetasoundRouter.h
--- a/Engine/Plugins/Runtime/Metasound/Source/MetasoundFrontend/Public/MetasoundRouter.h
+++ b/Engine/Plugins/Runtime/Metasound/Source/MetasoundFrontend/Public/MetasoundRouter.h
@@ -16,6 +16,7 @@
 #include "Misc/Guid.h"
 #include "UObject/NameTypes.h"
 
+#include <atomic>
 #include <type_traits>
 
 /**
@@ -243,7 +244,8 @@ namespace Metasound
 			return NumAliveSenders.GetValue();
 		}
 
-		virtual int64 GetPayloadID() const { return 0; }
+		virtual int64 GetPayloadID() const { return INDEX_NONE; }
+		virtual int64 CreateNewPayloadID() = 0;
 		virtual bool PushOpaque(ISender& InSender) { return true; };
 		virtual void PopOpaque(IReceiver& InReceiver) {};
 		virtual bool IsEmpty() const { return true; };
@@ -299,6 +301,7 @@ namespace Metasound
 
 		TNonOperationalDataChannel(const FOperatorSettings& InOperatorSettings)
 			: IDataChannel(FName())
+			, PayloadIDCounter(0)
 		{}
 
 		virtual bool PushOpaque(ISender& InSender) override 
@@ -311,6 +314,12 @@ namespace Metasound
 
 		}
 
+		virtual int64 CreateNewPayloadID() override
+		{
+			// Atomic return and post increment.
+			return PayloadIDCounter++;
+		}
+
 		virtual bool IsEmpty() const override
 		{
 			return true;
@@ -336,6 +345,8 @@ namespace Metasound
 		{
 			return TUniquePtr<ISender>();
 		}
+	private:
+		std::atomic<int64> PayloadIDCounter;
 	};
 
 	// Specializations for trivially copyable types:
@@ -423,7 +434,7 @@ namespace Metasound
 		// Push from the SenderBuffer to the data channel.
 		bool PushToDataChannel()
 		{
-			++PayloadID;
+			PayloadID = DataChannel->CreateNewPayloadID();
 			if (!Payload.IsValid())
 			{
 				Payload = MakeUnique<TDataType>(SenderBuffer.Pop());
@@ -537,6 +548,7 @@ namespace Metasound
 		TCopyableDataChannel(const FOperatorSettings& InOperatorSettings)
 			: IDataChannel(GetDataTypeName())
 			, OperatorSettings(InOperatorSettings)
+			, PayloadIDCounter(0)
 		{}
 
 		virtual int64 GetPayloadID() const override
@@ -544,6 +556,12 @@ namespace Metasound
 			return PayloadID;
 		}
 
+		virtual int64 CreateNewPayloadID() override
+		{
+			// Atomic return and post increment.
+			return PayloadIDCounter++;
+		}
+
 		virtual bool PushOpaque(ISender& InSender) override
 		{
 			TSender<TDataType>& CastSender = InSender.GetAs<TSender<TDataType>>();
@@ -632,9 +650,11 @@ namespace Metasound
 		TUniquePtr<TDataType> AtomicData;
 		FCriticalSection AtomicDataLock;
 
-		int64 PayloadID = 0;
+		int64 PayloadID = INDEX_NONE;
 
 		FOperatorSettings OperatorSettings;
+
+		std::atomic<int64> PayloadIDCounter;
 	};
 
 	// AUDIO SPECIALIZATIONS:
@@ -670,7 +690,7 @@ namespace Metasound
 
 		int64 GetPayloadID() const override
 		{
-			return 0;
+			return INDEX_NONE;
 		}
 
 		bool PushLiteral(const FLiteral& InLiteral) override
@@ -789,6 +809,7 @@ namespace Metasound
 	public:
 		TAudioDataChannel(const FOperatorSettings& InOperatorSettings)
 			: IDataChannel(GetDataTypeName())
+			, PayloadIDCounter(0)
 		{
 			AudioBuses.AddDefaulted(MaxChannels);
 		}
@@ -806,6 +827,12 @@ namespace Metasound
 			}
 		}
 
+		virtual int64 CreateNewPayloadID() override
+		{
+			// Atomic return and post increment.
+			return PayloadIDCounter++;
+		}
+
 		virtual bool IsEmpty() const override
 		{
 			return false;
@@ -846,6 +873,7 @@ namespace Metasound
 
 	private:
 		TArray<Audio::FPatchMixerSplitter> AudioBuses;
+		std::atomic<int64> PayloadIDCounter;
 	};
 
 
