Fix for MetaSound registry incremental unregister throwing a check.

Check hit was in FTransformRegistrationEventsToClasses::Transform and was due to prior reduce step not being applied on the merged result from the previous step.

#rnx

diff --git a/Engine/Plugins/Runtime/Metasound/Source/MetasoundFrontend/Private/MetasoundFrontendQuery.cpp b/Engine/Plugins/Runtime/Metasound/Source/MetasoundFrontend/Private/MetasoundFrontendQuery.cpp
--- a/Engine/Plugins/Runtime/Metasound/Source/MetasoundFrontend/Private/MetasoundFrontendQuery.cpp
+++ b/Engine/Plugins/Runtime/Metasound/Source/MetasoundFrontend/Private/MetasoundFrontendQuery.cpp
@@ -220,8 +220,8 @@ namespace Metasound
 				return Increment(InOutResult);
 			}
 
-			virtual void Reset() override 
-			{ 
+			virtual void Reset() override
+			{
 				if (Source.IsValid())
 				{
 					Source->Reset();
@@ -446,7 +446,7 @@ namespace Metasound
 				EResultModificationState ResultState = EResultModificationState::Unmodified;
 
 				if (Step.IsValid())
-				{ 	
+				{
 					if (InIncremental.GetSelection().Num() > 0)
 					{
 						TSet<FKey> NewKeys;
@@ -483,7 +483,7 @@ namespace Metasound
 				EResultModificationState ResultState = EResultModificationState::Unmodified;
 
 				if (Step.IsValid())
-				{ 	
+				{
 					if (InOutResult.GetSelection().Num() > 0)
 					{
 						TArray<FFrontendQueryEntry*> SortedSelection = GetSortedSelection(InOutResult.GetSelection());
@@ -1247,24 +1247,24 @@ namespace Metasound
 			{
 				FFrontendQueryStep& Step = *Steps[StepIndex];
 
-				IncrementalResultState = Step.Increment(IncrementalResult);
-
-				const bool bMergeIncrementalResults = Step.IsMergeRequiredForIncremental() && (EResultModificationState::Modified == IncrementalResultState);
-
+				const bool bMergeIncrementalResults = Step.IsMergeRequiredForIncremental();
 				if (bMergeIncrementalResults)
 				{
 					const bool bIsResultCacheEmpty = StepResultCache[StepIndex].GetSelection().Num() == 0;
-
 					if (bIsResultCacheEmpty)
 					{
-						// If the prior result for this step is empty, we can continue
-						// using the incremental path.
-						StepResultCache[StepIndex] = IncrementalResult;
+						// If the prior result for this step is empty, continue using the incremental path.
+						IncrementalResultState = Step.Increment(IncrementalResult);
+						if (EResultModificationState::Modified == IncrementalResultState)
+						{
+							StepResultCache[StepIndex] = IncrementalResult;
+						}
 					}
 					else
 					{
 						Step.Merge(IncrementalResult, StepResultCache[StepIndex]);
 
+						IncrementalResultState = Step.Execute(StepResultCache[StepIndex]);
 						*Result = StepResultCache[StepIndex];
 
 						// After results are merged, downstream steps can no longer 
@@ -1274,6 +1274,10 @@ namespace Metasound
 						return;
 					}
 				}
+				else
+				{
+					IncrementalResultState = Step.Increment(IncrementalResult);
+				}
 			}
 		}
 
