First pass on MetaSound Composition/Preset Support
- Stubbed in logic & assets for MetaSound Editor UX node type differentiation

diff --git a/Engine/Plugins/Runtime/Metasound/Source/MetasoundEditor/Private/MetasoundAssetTypeActions.cpp b/Engine/Plugins/Runtime/Metasound/Source/MetasoundEditor/Private/MetasoundAssetTypeActions.cpp
--- a/Engine/Plugins/Runtime/Metasound/Source/MetasoundEditor/Private/MetasoundAssetTypeActions.cpp
+++ b/Engine/Plugins/Runtime/Metasound/Source/MetasoundEditor/Private/MetasoundAssetTypeActions.cpp
@@ -153,9 +153,8 @@ namespace Metasound
 		void FAssetTypeActions_MetaSoundSource::RegisterMenuActions()
 		{
 			using namespace AssetTypeActionsPrivate;
-			// TODO: Enable Composition
-			// RegisterPresetAction<UMetaSoundSource, UMetaSoundSourceFactory>(SourcePresetLabel, SourcePresetToolTip);
-			// RegisterPresetAction<UMetaSound, UMetaSoundFactory, UMetaSoundSource>(PresetLabel, PresetToolTip);
+			RegisterPresetAction<UMetaSoundSource, UMetaSoundSourceFactory>(SourcePresetLabel, SourcePresetToolTip);
+			RegisterPresetAction<UMetaSound, UMetaSoundFactory, UMetaSoundSource>(PresetLabel, PresetToolTip);
 		}
 	} // namespace Editor
 } // namespace Metasound
diff --git a/Engine/Plugins/Runtime/Metasound/Source/MetasoundEditor/Private/MetasoundEditor.cpp b/Engine/Plugins/Runtime/Metasound/Source/MetasoundEditor/Private/MetasoundEditor.cpp
--- a/Engine/Plugins/Runtime/Metasound/Source/MetasoundEditor/Private/MetasoundEditor.cpp
+++ b/Engine/Plugins/Runtime/Metasound/Source/MetasoundEditor/Private/MetasoundEditor.cpp
@@ -64,27 +64,6 @@ namespace Metasound
 {
 	namespace Editor
 	{
-		namespace Private
-		{
-			void ValidateAndUpdate(UMetasoundEditorGraph& Graph)
-			{
-				using namespace Frontend;
-
-				FGraphValidationResults Results;
-				Graph.Validate(Results);
-
-				// Only auto-update minor upgrades
-				for (const FGraphNodeValidationResult& Result : Results.GetResults())
-				{
-					FMetasoundFrontendVersionNumber MinorUpdateVersion = Result.Node->GetMinorUpdateAvailable();
-					if (MinorUpdateVersion.IsValid())
-					{
-						Result.Node->UpdateToVersion(MinorUpdateVersion, true /* bInPropagateErrorMessages */);
-					}
-				}
-			}
-		}
-
 		static const TArray<FText> NodeSectionNames
 		{
 			LOCTEXT("NodeSectionName_Invalid", "INVALID"),
@@ -334,6 +313,8 @@ namespace Metasound
 
 		void FEditor::InitMetasoundEditor(const EToolkitMode::Type Mode, const TSharedPtr<IToolkitHost>& InitToolkitHost, UObject* ObjectToEdit)
 		{
+			using namespace Metasound::Frontend;
+
 			check(ObjectToEdit);
 			checkf(IMetasoundUObjectRegistry::Get().IsRegisteredClass(ObjectToEdit), TEXT("Object passed in was not registered as a valid metasound archetype!"));
 			
@@ -342,15 +323,8 @@ namespace Metasound
 
 			Metasound = ObjectToEdit;
 
-			FMetasoundAssetBase* MetasoundAsset = IMetasoundUObjectRegistry::Get().GetObjectAsAssetBase(Metasound);
-			check(MetasoundAsset);
-
-			const bool bDidUpdateAsset = MetasoundAsset->VersionAsset();
-			if (bDidUpdateAsset)
-			{
-				MetasoundAsset->RegisterGraphWithFrontend();
-			}
-			FGraphBuilder::SynchronizeGraph(*Metasound);
+			const bool bAutoUpdate = true;
+			FGraphBuilder::ValidateGraph(*Metasound, bAutoUpdate);
 
 			GEditor->RegisterForUndo(this);
 
@@ -413,8 +387,6 @@ namespace Metasound
 
 			ExtendToolbar();
 			RegenerateMenusAndToolbars();
-
-			Private::ValidateAndUpdate(GetMetaSoundGraphChecked());
 		}
 
 		UObject* FEditor::GetMetasoundObject() const
@@ -474,10 +446,7 @@ namespace Metasound
 				MetasoundGraphEditor->NotifyGraphChanged();
 			}
 
-			if (MetasoundInterfaceMenu.IsValid())
-			{
-				MetasoundInterfaceMenu->RefreshAllActions(true /* bPreserveExpansion */);
-			}
+			RefreshInterface();
 
 			FSlateApplication::Get().DismissAllMenus();
 		}
@@ -509,7 +478,7 @@ namespace Metasound
 							Metasound->MarkPackageDirty();
 						}
 						MetasoundGraphEditor->NotifyGraphChanged();
-						MetasoundInterfaceMenu->RefreshAllActions(true /* bPreserveExpansion */);
+						RefreshInterface();
 					}
 				}
 			}
@@ -828,6 +797,11 @@ namespace Metasound
 		{
 			if (USoundBase* MetasoundToPlay = Cast<USoundBase>(Metasound))
 			{
+				if (!FGraphBuilder::ValidateGraph(*Metasound, false /* bAutoUpdate */))
+				{
+					return;
+				}
+
 				// Set the send to the audio bus that is used for analyzing the metasound output
 				check(GEditor);
 				if (UAudioComponent* PreviewComp = GEditor->PlayPreviewSound(MetasoundToPlay))
@@ -914,6 +888,11 @@ namespace Metasound
 			return PlayTime;
 		}
 
+		TSharedPtr<SGraphEditor> FEditor::GetGraphEditor() const
+		{
+			return MetasoundGraphEditor;
+		}
+
 		void FEditor::Stop()
 		{
 			check(GEditor);
@@ -1187,8 +1166,9 @@ namespace Metasound
 						{
 							if (UMetasoundEditorGraphExternalNode* ExternalNode = Cast<UMetasoundEditorGraphExternalNode>(Object))
 							{
-								FMetasoundFrontendVersionNumber MajorUpdateVersion = ExternalNode->GetMajorUpdateAvailable();
-								if (MajorUpdateVersion.IsValid())
+								FMetasoundFrontendVersionNumber MajorUpdateVersion = ExternalNode->FindHighestVersionInRegistry();
+								Metasound::Frontend::FConstNodeHandle NodeHandle = ExternalNode->GetConstNodeHandle();
+								if (MajorUpdateVersion.IsValid() && MajorUpdateVersion > NodeHandle->GetClassMetadata().GetVersion())
 								{
 									ExternalNode->UpdateToVersion(MajorUpdateVersion, false /* bInPropagateErrorMessages */);
 								}
@@ -1328,7 +1308,7 @@ namespace Metasound
 
 			FGraphBuilder::DeleteVariableNodeHandle(*Variable);
 			Graph->RemoveVariable(*Variable);
-			MetasoundInterfaceMenu->RefreshAllActions(true /* bPreserveExpansion */);
+			RefreshInterface();
 
 			if (!NameIDPair.Name.IsNone())
 			{
@@ -1341,6 +1321,8 @@ namespace Metasound
 					}
 				}
 			}
+
+			Metasound::Editor::FGraphBuilder::RegisterGraphWithFrontend(Graph->GetMetasoundChecked());
 		}
 
 		void FEditor::DeleteSelected()
@@ -1733,6 +1715,14 @@ namespace Metasound
 			MetasoundGraphEditor->NotifyGraphChanged();
 		}
 
+		void FEditor::RefreshInterface()
+		{
+			if (MetasoundInterfaceMenu.IsValid())
+			{
+				MetasoundInterfaceMenu->RefreshAllActions(true /* bPreserveExpansion */);
+			}
+		}
+
 		bool FEditor::CanPasteNodes()
 		{
 			if (!IsGraphEditable())
@@ -1768,7 +1758,7 @@ namespace Metasound
 
 		void FEditor::OnInputNameChanged(FGuid InNodeID)
 		{
-			if (!MetasoundInterfaceMenu.IsValid())
+			if (!MetasoundInterfaceMenu.IsValid() || !Metasound)
 			{
 				return;
 			}
@@ -1795,6 +1785,8 @@ namespace Metasound
 					}
 				}
 			}
+
+			Metasound::Editor::FGraphBuilder::RegisterGraphWithFrontend(*Metasound);
 		}
 
 		void FEditor::OnOutputNameChanged(FGuid InNodeID)
@@ -1826,6 +1818,8 @@ namespace Metasound
 					}
 				}
 			}
+
+			Metasound::Editor::FGraphBuilder::RegisterGraphWithFrontend(*Metasound);
 		}
 
 		void FEditor::CollectAllActions(FGraphActionListBuilderBase& OutAllActions)
@@ -2035,6 +2029,8 @@ namespace Metasound
 					Frontend::FNodeHandle NodeHandle = FGraphBuilder::AddInputNodeHandle(*Metasound, DataTypeName, FText::GetEmpty());
 					if (ensure(NodeHandle->IsValid()))
 					{
+						NameToSelect = *NodeHandle->GetDisplayName().ToString();
+
 						TObjectPtr<UMetasoundEditorGraphInput> Input = Graph.FindOrAddInput(NodeHandle);
 						if (ensure(Input))
 						{
@@ -2058,9 +2054,7 @@ namespace Metasound
 					Frontend::FNodeHandle NodeHandle = FGraphBuilder::AddOutputNodeHandle(*Metasound, DataTypeName, FText::GetEmpty());
 					if (ensure(NodeHandle->IsValid()))
 					{
-						const FText DisplayName = FGraphBuilder::GenerateUniqueOutputDisplayName(*Metasound);
-						NodeHandle->SetDisplayName(DisplayName);
-						NameToSelect = *DisplayName.ToString();
+						NameToSelect = *NodeHandle->GetDisplayName().ToString();
 
 						TObjectPtr<UMetasoundEditorGraphOutput> Output = Graph.FindOrAddOutput(NodeHandle);
 						if (ensure(Output))
@@ -2080,10 +2074,12 @@ namespace Metasound
 				return FReply::Unhandled();
 			}
 
+			FGraphBuilder::RegisterGraphWithFrontend(*Metasound);
+
 			if (MetasoundInterfaceMenu.IsValid())
 			{
 				MetasoundInterfaceMenu->RefreshAllActions(/* bPreserveExpansion */ true);
-				if (NameToSelect.IsNone())
+				if (!NameToSelect.IsNone())
 				{
 					MetasoundInterfaceMenu->SelectItemByName(NameToSelect);
 				}
diff --git a/Engine/Plugins/Runtime/Metasound/Source/MetasoundEditor/Private/MetasoundEditor.h b/Engine/Plugins/Runtime/Metasound/Source/MetasoundEditor/Private/MetasoundEditor.h
--- a/Engine/Plugins/Runtime/Metasound/Source/MetasoundEditor/Private/MetasoundEditor.h
+++ b/Engine/Plugins/Runtime/Metasound/Source/MetasoundEditor/Private/MetasoundEditor.h
@@ -71,6 +71,7 @@ namespace Metasound
 			virtual void UnregisterTabSpawners(const TSharedRef<FTabManager>& TabManager) override;
 
 			double GetPlayTime() const;
+			TSharedPtr<SGraphEditor> GetGraphEditor() const;
 
 			/** Edits the specified Metasound object */
 			void InitMetasoundEditor(const EToolkitMode::Type Mode, const TSharedPtr<IToolkitHost>& InitToolkitHost, UObject* ObjectToEdit);
@@ -103,13 +104,16 @@ namespace Metasound
 			/** Whether pasting the currently selected nodes is permissible */
 			bool CanPasteNodes();
 
-			/* Duplicates the selected node(s) in the graph */
+			/** Duplicates the selected node(s) in the graph */
 			void DuplicateNodes();
 
-			/* Pastes node(s) from the clipboard to the graph */
+			/** Pastes node(s) from the clipboard to the graph */
 			void PasteNodes(const FVector2D* InLocation = nullptr);
 			void PasteNodes(const FVector2D* InLocation, const FText& InTransactionText);
 
+			/** Forces all UX pertaining to the root graph's interface to be refreshed. */
+			void RefreshInterface();
+
 			/* Whether the displayed graph is marked as editable */
 			bool IsGraphEditable() const;
 
diff --git a/Engine/Plugins/Runtime/Metasound/Source/MetasoundEditor/Private/MetasoundEditorGraph.cpp b/Engine/Plugins/Runtime/Metasound/Source/MetasoundEditor/Private/MetasoundEditorGraph.cpp
--- a/Engine/Plugins/Runtime/Metasound/Source/MetasoundEditor/Private/MetasoundEditorGraph.cpp
+++ b/Engine/Plugins/Runtime/Metasound/Source/MetasoundEditor/Private/MetasoundEditorGraph.cpp
@@ -37,11 +37,32 @@ TArray<UMetasoundEditorGraphNode*> UMetasoundEditorGraphVariable::GetNodes() con
 	return Nodes;
 }
 
+void UMetasoundEditorGraphVariable::SetDescription(const FText& InDescription)
+{
+	using namespace Metasound::Editor;
+	using namespace Metasound::Frontend;
+
+	const FText TransactionLabel = FText::Format(LOCTEXT("SetVariableTooltipFormat", "Set MetaSound {0}'s ToolTip"), GetVariableLabel());
+	const FScopedTransaction Transaction(TransactionLabel);
+
+	Modify();
+	if (UMetasoundEditorGraph* Graph = Cast<UMetasoundEditorGraph>(GetOuter()))
+	{
+		UObject& MetaSound = Graph->GetMetasoundChecked();
+		MetaSound.Modify();
+
+		FNodeHandle NodeHandle = GetNodeHandle();
+		NodeHandle->SetDescription(InDescription);
+
+		FGraphBuilder::RegisterGraphWithFrontend(MetaSound);
+	}
+}
+
 void UMetasoundEditorGraphVariable::SetDisplayName(const FText& InNewName)
 {
 	using namespace Metasound::Frontend;
 
-	const FText TransactionLabel = FText::Format(LOCTEXT("Rename Variable", "Rename Metasound {0}"), GetVariableLabel());
+	const FText TransactionLabel = FText::Format(LOCTEXT("SetVariableDisplayNameFormat", "Rename Metasound {0}"), GetVariableLabel());
 	const FScopedTransaction Transaction(TransactionLabel);
 
 	Modify();
@@ -98,7 +119,7 @@ void UMetasoundEditorGraphVariable::SetDataType(FName InNewType)
 		return;
 	}
 
-	ClassName = NewNodeHandle->GetClassMetadata().ClassName;
+	ClassName = NewNodeHandle->GetClassMetadata().GetClassName();
 	NodeID = NewNodeHandle->GetID();
 	TypeName = InNewType;
 
@@ -116,6 +137,8 @@ void UMetasoundEditorGraphVariable::SetDataType(FName InNewType)
 	// Notify now that the node has a new ID (doing so before creating & syncing Frontend Node &
 	// EdGraph variable can result in refreshing editors while in a desync'ed state)
 	NameChanged.Broadcast(NodeID);
+
+	FGraphBuilder::RegisterGraphWithFrontend(Metasound);
 }
 
 Metasound::Frontend::FNodeHandle UMetasoundEditorGraphVariable::GetNodeHandle() const
@@ -181,7 +204,6 @@ bool UMetasoundEditorGraphVariable::CanRename(const FText& InNewName, FText& Out
 	return bIsNameValid;
 }
 
-#if WITH_EDITORONLY_DATA
 void UMetasoundEditorGraphInputLiteral::PostEditUndo()
 {
 	Super::PostEditUndo();
@@ -191,7 +213,6 @@ void UMetasoundEditorGraphInputLiteral::PostEditUndo()
 		Input->OnLiteralChanged(false /* bPostTransaction */);
 	}
 }
-#endif // WITH_EDITORONLY_DATA
 
 void UMetasoundEditorGraphInput::UpdateDocumentInput(bool bPostTransaction)
 {
@@ -222,6 +243,10 @@ void UMetasoundEditorGraphInput::UpdateDocumentInput(bool bPostTransaction)
 
 	const FGuid VertexID = GraphHandle->GetVertexIDForInputVertex(NodeName);
 	GraphHandle->SetDefaultInput(VertexID, Literal->GetDefault());
+
+	// Disabled as internal call to validation to all other open graphs
+	// is expensive and can be spammed by dragging values 
+// 	Metasound::Editor::FGraphBuilder::RegisterGraphWithFrontend(*Metasound);
 }
 
 void UMetasoundEditorGraphInput::UpdatePreviewInstance(const Metasound::FVertexKey& InParameterName, TScriptInterface<IAudioParameterInterface>& InParamInterface) const
@@ -255,7 +280,6 @@ const FText& UMetasoundEditorGraphInput::GetVariableLabel() const
 	return Label;
 }
 
-#if WITH_EDITORONLY_DATA
 void UMetasoundEditorGraphInput::PostEditUndo()
 {
 	Super::PostEditUndo();
@@ -265,13 +289,17 @@ void UMetasoundEditorGraphInput::PostEditUndo()
 		if (UMetasoundEditorGraph* MetasoundGraph = Cast<UMetasoundEditorGraph>(GetOuter()))
 		{
 			MetasoundGraph->RemoveVariable(*CastChecked<UMetasoundEditorGraphVariable>(this));
+
+			if (UObject* Object = MetasoundGraph->GetMetasound())
+			{
+				Metasound::Editor::FGraphBuilder::RegisterGraphWithFrontend(*Object);
+			}
 		}
 		return;
 	}
 
 	OnLiteralChanged(false /* bPostTransaction */);
 }
-#endif // WITH_EDITORONLY_DATA
 
 void UMetasoundEditorGraphInput::OnLiteralChanged(bool bPostTransaction)
 {
@@ -346,6 +374,11 @@ UMetasoundEditorGraphInputNode* UMetasoundEditorGraph::CreateInputNode(Metasound
 		return nullptr;
 	}
 
+	if (!ensure(InNodeHandle->GetClassMetadata().GetType() == EMetasoundFrontendClassType::Input))
+	{
+		return nullptr;
+	}
+
 	UEdGraphNode* NewEdGraphNode = CreateNode(UMetasoundEditorGraphInputNode::StaticClass(), bInSelectNewNode);
 	UMetasoundEditorGraphInputNode* NewInputNode = CastChecked<UMetasoundEditorGraphInputNode>(NewEdGraphNode);
 	if (ensure(NewInputNode))
@@ -391,14 +424,40 @@ UObject& UMetasoundEditorGraph::GetMetasoundChecked() const
 	return *ParentMetasound;
 }
 
-void UMetasoundEditorGraph::Synchronize()
+void UMetasoundEditorGraph::RegisterGraphWithFrontend()
 {
+	using namespace Metasound::Editor;
+
 	if (UObject* ParentMetasound = GetOuter())
 	{
-		Metasound::Editor::FGraphBuilder::SynchronizeGraph(*ParentMetasound);
+		FGraphBuilder::RegisterGraphWithFrontend(*ParentMetasound);
 	}
 }
 
+bool UMetasoundEditorGraph::Synchronize()
+{
+	using namespace Metasound::Editor;
+
+	if (UObject* ParentMetasound = GetOuter())
+	{
+		return FGraphBuilder::SynchronizeGraph(*ParentMetasound);
+	}
+
+	return false;
+}
+
+bool UMetasoundEditorGraph::Validate(bool bInAutoUpdate)
+{
+	using namespace Metasound::Editor;
+
+	if (UObject* ParentMetasound = GetOuter())
+	{
+		return FGraphBuilder::ValidateGraph(*ParentMetasound, bInAutoUpdate);
+	}
+
+	return false;
+}
+
 UMetasoundEditorGraphInput* UMetasoundEditorGraph::FindInput(FGuid InNodeID) const
 {
 	const TObjectPtr<UMetasoundEditorGraphInput>* Input = Inputs.FindByPredicate([InNodeID](const TObjectPtr<UMetasoundEditorGraphInput>& Input)
@@ -420,10 +479,10 @@ UMetasoundEditorGraphInput* UMetasoundEditorGraph::FindOrAddInput(Metasound::Fro
 
 	ensure(InNodeHandle->GetNumInputs() == 1);
 	InNodeHandle->IterateConstInputs([InGraph = &Graph, InTypeName = &TypeName, InVertexID = &VertexID](FConstInputHandle InputHandle)
-		{
-			*InTypeName = InputHandle->GetDataType();
-			*InVertexID = (*InGraph)->GetVertexIDForInputVertex(InputHandle->GetName());
-		});
+	{
+		*InTypeName = InputHandle->GetDataType();
+		*InVertexID = (*InGraph)->GetVertexIDForInputVertex(InputHandle->GetName());
+	});
 
 	const FGuid NodeID = InNodeHandle->GetID();
 	if (TObjectPtr<UMetasoundEditorGraphInput> Input = FindInput(NodeID))
@@ -435,7 +494,7 @@ UMetasoundEditorGraphInput* UMetasoundEditorGraph::FindOrAddInput(Metasound::Fro
 	if (UMetasoundEditorGraphInput* NewInput = NewObject<UMetasoundEditorGraphInput>(this, FName(), RF_Transactional))
 	{
 		NewInput->NodeID = NodeID;
-		NewInput->ClassName = InNodeHandle->GetClassMetadata().ClassName;
+		NewInput->ClassName = InNodeHandle->GetClassMetadata().GetClassName();
 		NewInput->TypeName = TypeName;
 
 		FMetasoundFrontendLiteral DefaultLiteral = Graph->GetDefaultInput(VertexID);
@@ -490,9 +549,10 @@ UMetasoundEditorGraphOutput* UMetasoundEditorGraph::FindOrAddOutput(Metasound::F
 	if (UMetasoundEditorGraphOutput* NewOutput = NewObject<UMetasoundEditorGraphOutput>(this, FName(), RF_Transactional))
 	{
 		NewOutput->NodeID = NodeID;
-		NewOutput->ClassName = InNodeHandle->GetClassMetadata().ClassName;
+		NewOutput->ClassName = InNodeHandle->GetClassMetadata().GetClassName();
 		NewOutput->TypeName = TypeName;
 		Outputs.Add(NewOutput);
+
 		return NewOutput;
 	}
 
@@ -573,24 +633,29 @@ bool UMetasoundEditorGraph::RemoveVariable(UMetasoundEditorGraphVariable& InVari
 {
 	if (UMetasoundEditorGraphInput* Input = Cast<UMetasoundEditorGraphInput>(&InVariable))
 	{
-		return Inputs.Remove(Input) > 0;
+		if (!Inputs.Remove(Input))
+		{
+			return false;
+		}
 	}
-
-	if (UMetasoundEditorGraphOutput* Output = Cast<UMetasoundEditorGraphOutput>(&InVariable))
+	else if (UMetasoundEditorGraphOutput* Output = Cast<UMetasoundEditorGraphOutput>(&InVariable))
 	{
-		return Outputs.Remove(Output) > 0;
+		if (!Outputs.Remove(Output))
+		{
+			return false;
+		}
 	}
 
-	return false;
+	return true;
 }
 
-bool UMetasoundEditorGraph::Validate(Metasound::Editor::FGraphValidationResults& OutResults)
+bool UMetasoundEditorGraph::ValidateInternal(Metasound::Editor::FGraphValidationResults& OutResults, bool bClearUpgradeMessaging)
 {
 	using namespace Metasound::Editor;
 	using namespace Metasound::Frontend;
 
 	bool bMarkDirty = false;
-	bool bIsValid = false;
+	bool bIsValid = true;
 
 	OutResults = FGraphValidationResults();
 
@@ -599,7 +664,7 @@ bool UMetasoundEditorGraph::Validate(Metasound::Editor::FGraphValidationResults&
 	for (UMetasoundEditorGraphExternalNode* ExternalNode : ExternalNodes)
 	{
 		FGraphNodeValidationResult NodeResult(*ExternalNode);
-		bIsValid |= ExternalNode->Validate(NodeResult);
+		bIsValid &= ExternalNode->Validate(NodeResult, bClearUpgradeMessaging);
 		bMarkDirty |= NodeResult.bIsDirty;
 		OutResults.NodeResults.Add(NodeResult);
 	}
diff --git a/Engine/Plugins/Runtime/Metasound/Source/MetasoundEditor/Private/MetasoundEditorGraphBuilder.cpp b/Engine/Plugins/Runtime/Metasound/Source/MetasoundEditor/Private/MetasoundEditorGraphBuilder.cpp
--- a/Engine/Plugins/Runtime/Metasound/Source/MetasoundEditor/Private/MetasoundEditorGraphBuilder.cpp
+++ b/Engine/Plugins/Runtime/Metasound/Source/MetasoundEditor/Private/MetasoundEditorGraphBuilder.cpp
@@ -7,6 +7,7 @@
 #include "EdGraph/EdGraphNode.h"
 #include "EdGraph/EdGraphPin.h"
 #include "Framework/Notifications/NotificationManager.h"
+#include "GraphEditor.h"
 #include "Metasound.h"
 #include "MetasoundAssetBase.h"
 #include "MetasoundEditor.h"
@@ -14,6 +15,7 @@
 #include "MetasoundEditorGraphNode.h"
 #include "MetasoundEditorGraphInputNodes.h"
 #include "MetasoundEditorGraphSchema.h"
+#include "MetasoundEditorGraphValidation.h"
 #include "MetasoundEditorModule.h"
 #include "MetasoundEditorSettings.h"
 #include "MetasoundFrontendQuery.h"
@@ -53,13 +55,7 @@ namespace Metasound
 
 		namespace GraphBuilderPrivate
 		{
-			namespace NodeLayout
-			{
-				static const FVector2D BufferX { 250.0f, 0.0f };
-				static const FVector2D BufferY { 0.0f, 100.0f };
-			}
-
-			void DeleteNode(UObject& InMetasound, Frontend::FNodeHandle InNodeHandle)
+			void DeleteNode(UObject& InMetaSound, Frontend::FNodeHandle InNodeHandle)
 			{
 				if (InNodeHandle->IsValid())
 				{
@@ -71,7 +67,7 @@ namespace Metasound
 				}
 			}
 
-			void InitializeGraph(UObject& InMetaSound)
+			bool InitializeGraph(UObject& InMetaSound)
 			{
 				FMetasoundAssetBase* MetaSoundAsset = IMetasoundUObjectRegistry::Get().GetObjectAsAssetBase(&InMetaSound);
 				check(MetaSoundAsset);
@@ -83,56 +79,62 @@ namespace Metasound
 					UMetasoundEditorGraph* Graph = NewObject<UMetasoundEditorGraph>(&InMetaSound, FName(), RF_Transactional);
 					Graph->Schema = UMetasoundEditorGraphSchema::StaticClass();
 					MetaSoundAsset->SetGraph(Graph);
+
+					return true;
 				}
+
+				return false;
 			}
 		}
 
-		UMetasoundEditorGraphExternalNode* FGraphBuilder::AddExternalNode(UObject& InMetasound, Frontend::FNodeHandle& InNodeHandle, FVector2D InLocation, bool bInSelectNewNode)
+		UMetasoundEditorGraphExternalNode* FGraphBuilder::AddExternalNode(UObject& InMetaSound, Frontend::FNodeHandle& InNodeHandle, FVector2D InLocation, bool bInSelectNewNode)
 		{
 			UMetasoundEditorGraphExternalNode* NewGraphNode = nullptr;
-			if (!ensure(InNodeHandle->GetClassMetadata().Type == EMetasoundFrontendClassType::External))
+			if (!ensure(InNodeHandle->GetClassMetadata().GetType() == EMetasoundFrontendClassType::External))
 			{
 				return nullptr;
 			}
 
-			FMetasoundAssetBase* MetasoundAsset = IMetasoundUObjectRegistry::Get().GetObjectAsAssetBase(&InMetasound);
-			check(MetasoundAsset);
-			UEdGraph& Graph = MetasoundAsset->GetGraphChecked();
+			FMetasoundAssetBase* MetaSoundAsset = IMetasoundUObjectRegistry::Get().GetObjectAsAssetBase(&InMetaSound);
+			check(MetaSoundAsset);
+			UEdGraph& Graph = MetaSoundAsset->GetGraphChecked();
 			FGraphNodeCreator<UMetasoundEditorGraphExternalNode> NodeCreator(Graph);
 
 			NewGraphNode = NodeCreator.CreateNode(bInSelectNewNode);
-			NewGraphNode->ClassName = InNodeHandle->GetClassMetadata().ClassName;
+			NewGraphNode->ClassName = InNodeHandle->GetClassMetadata().GetClassName();
 
 			NodeCreator.Finalize();
-			InitGraphNode(InNodeHandle, NewGraphNode, InMetasound);
+			InitGraphNode(InNodeHandle, NewGraphNode, InMetaSound);
 
 			SynchronizeNodeLocation(InLocation, InNodeHandle, *NewGraphNode);
 
 			return NewGraphNode;
 		}
 
-		UMetasoundEditorGraphExternalNode* FGraphBuilder::AddExternalNode(UObject& InMetasound, const FMetasoundFrontendClassMetadata& InMetadata, FVector2D InLocation, bool bInSelectNewNode)
+		UMetasoundEditorGraphExternalNode* FGraphBuilder::AddExternalNode(UObject& InMetaSound, const FMetasoundFrontendClassMetadata& InMetadata, FVector2D InLocation, bool bInSelectNewNode)
 		{
-			FMetasoundAssetBase* MetasoundAsset = IMetasoundUObjectRegistry::Get().GetObjectAsAssetBase(&InMetasound);
-			check(MetasoundAsset);
+			FMetasoundAssetBase* MetaSoundAsset = IMetasoundUObjectRegistry::Get().GetObjectAsAssetBase(&InMetaSound);
+			check(MetaSoundAsset);
+
+			Frontend::FNodeHandle NodeHandle = MetaSoundAsset->GetRootGraphHandle()->AddNode(InMetadata);
+			MetaSoundAsset->RebuildReferencedAssets(UMetaSoundAssetSubsystem::Get());
 
-			Frontend::FNodeHandle NodeHandle = MetasoundAsset->GetRootGraphHandle()->AddNode(InMetadata);
-			return AddExternalNode(InMetasound, NodeHandle, InLocation, bInSelectNewNode);
+			return AddExternalNode(InMetaSound, NodeHandle, InLocation, bInSelectNewNode);
 		}
 
-		UMetasoundEditorGraphOutputNode* FGraphBuilder::AddOutputNode(UObject& InMetasound, Frontend::FNodeHandle& InNodeHandle, FVector2D InLocation, bool bInSelectNewNode)
+		UMetasoundEditorGraphOutputNode* FGraphBuilder::AddOutputNode(UObject& InMetaSound, Frontend::FNodeHandle& InNodeHandle, FVector2D InLocation, bool bInSelectNewNode)
 		{
 			using namespace Frontend;
 
 			UMetasoundEditorGraphOutputNode* NewGraphNode = nullptr;
-			if (!ensure(InNodeHandle->GetClassMetadata().Type == EMetasoundFrontendClassType::Output))
+			if (!ensure(InNodeHandle->GetClassMetadata().GetType() == EMetasoundFrontendClassType::Output))
 			{
 				return nullptr;
 			}
 
-			FMetasoundAssetBase* MetasoundAsset = IMetasoundUObjectRegistry::Get().GetObjectAsAssetBase(&InMetasound);
-			check(MetasoundAsset);
-			UEdGraph& Graph = MetasoundAsset->GetGraphChecked();
+			FMetasoundAssetBase* MetaSoundAsset = IMetasoundUObjectRegistry::Get().GetObjectAsAssetBase(&InMetaSound);
+			check(MetaSoundAsset);
+			UEdGraph& Graph = MetaSoundAsset->GetGraphChecked();
 			FGraphNodeCreator<UMetasoundEditorGraphOutputNode> NodeCreator(Graph);
 
 			NewGraphNode = NodeCreator.CreateNode(bInSelectNewNode);
@@ -140,14 +142,62 @@ namespace Metasound
 			NewGraphNode->Output = MetasoundGraph->FindOrAddOutput(InNodeHandle);
 
 			NodeCreator.Finalize();
-			InitGraphNode(InNodeHandle, NewGraphNode, InMetasound);
+			InitGraphNode(InNodeHandle, NewGraphNode, InMetaSound);
 
 			SynchronizeNodeLocation(InLocation, InNodeHandle, *NewGraphNode);
 
 			return NewGraphNode;
 		}
 
-		void FGraphBuilder::InitGraphNode(Frontend::FNodeHandle& InNodeHandle, UMetasoundEditorGraphNode* NewGraphNode, UObject& InMetasound)
+		bool FGraphBuilder::ValidateGraph(UObject& InMetaSoundToUpdate, bool bInAutoUpdate)
+		{
+			using namespace Frontend;
+			using namespace GraphBuilderPrivate;
+
+			SynchronizeGraph(InMetaSoundToUpdate);
+
+			TSharedPtr<SGraphEditor> GraphEditor;
+			TSharedPtr<FEditor> MetaSoundEditor = GetEditorForMetasound(InMetaSoundToUpdate);
+			if (MetaSoundEditor.IsValid())
+			{
+				GraphEditor = MetaSoundEditor->GetGraphEditor();
+			}
+
+			FMetasoundAssetBase* MetaSoundAsset = IMetasoundUObjectRegistry::Get().GetObjectAsAssetBase(&InMetaSoundToUpdate);
+			check(MetaSoundAsset);
+
+			// Run initial validation pass before updating to capture upgrade message state prior to auto-update
+			UMetasoundEditorGraph& Graph = *CastChecked<UMetasoundEditorGraph>(&MetaSoundAsset->GetGraphChecked());
+			FGraphValidationResults Results;
+			Graph.ValidateInternal(Results);
+
+			bool bAutoUpdated = false;
+			if (Results.IsValid() && bInAutoUpdate)
+			{
+				check(MetaSoundAsset);
+				if (MetaSoundAsset->AutoUpdate(UMetaSoundAssetSubsystem::Get(), true /* bMarkDirty*/, true /* bUpdateReferencedAssets*/))
+				{
+					SynchronizeGraph(InMetaSoundToUpdate);
+
+					Graph.ValidateInternal(Results, false /* bClearUpgradeMessaging */);
+					for (const FGraphNodeValidationResult& Result : Results.GetResults())
+					{
+						if (GraphEditor.IsValid())
+						{
+							check(Result.Node);
+							GraphEditor->RefreshNode(*Result.Node);
+							MetaSoundEditor->RefreshInterface();
+						}
+					}
+
+					InMetaSoundToUpdate.MarkPackageDirty();
+				}
+			}
+
+			return Results.IsValid();
+		}
+
+		void FGraphBuilder::InitGraphNode(Frontend::FNodeHandle& InNodeHandle, UMetasoundEditorGraphNode* NewGraphNode, UObject& InMetaSound)
 		{
 			NewGraphNode->CreateNewGuid();
 			NewGraphNode->SetNodeID(InNodeHandle->GetID());
@@ -209,10 +259,10 @@ namespace Metasound
 
 		FText FGraphBuilder::GenerateUniqueNameByFilter(const UObject& InMetaSound, const FText& InBaseText, TUniqueFunction<bool(const Frontend::FConstGraphHandle&, const FText&)> InIsValidNameFilter)
 		{
-			const FMetasoundAssetBase* MetasoundAsset = IMetasoundUObjectRegistry::Get().GetObjectAsAssetBase(&InMetaSound);
-			check(MetasoundAsset);
+			const FMetasoundAssetBase* MetaSoundAsset = IMetasoundUObjectRegistry::Get().GetObjectAsAssetBase(&InMetaSound);
+			check(MetaSoundAsset);
 
-			const Frontend::FConstGraphHandle GraphHandle = MetasoundAsset->GetRootGraphHandle();
+			const Frontend::FConstGraphHandle GraphHandle = MetaSoundAsset->GetRootGraphHandle();
 
 			static const FText VariableUniqueNameFormat = LOCTEXT("VariableUniqueNameFormat", "{0}{1}");
 
@@ -227,16 +277,16 @@ namespace Metasound
 			return NewName;
 		}
 
-		TSharedPtr<FEditor> FGraphBuilder::GetEditorForGraph(const UObject& Metasound)
+		TSharedPtr<FEditor> FGraphBuilder::GetEditorForMetasound(const UObject& Metasound)
 		{
-			TSharedPtr<IToolkit> FoundAssetEditor = FToolkitManager::Get().FindEditorForAsset(CastChecked<const UObject>(&Metasound));
+			TSharedPtr<IToolkit> FoundAssetEditor = FToolkitManager::Get().FindEditorForAsset(&Metasound);
 			return StaticCastSharedPtr<FEditor, IToolkit>(FoundAssetEditor);
 		}
 
 		TSharedPtr<FEditor> FGraphBuilder::GetEditorForGraph(const UEdGraph& EdGraph)
 		{
 			const UMetasoundEditorGraph* MetasoundGraph = CastChecked<const UMetasoundEditorGraph>(&EdGraph);
-			return GetEditorForGraph(MetasoundGraph->GetMetasoundChecked());
+			return GetEditorForMetasound(MetasoundGraph->GetMetasoundChecked());
 		}
 
 		FLinearColor FGraphBuilder::GetPinCategoryColor(const FEdGraphPinType& PinType)
@@ -362,14 +412,14 @@ namespace Metasound
 			InNodeHandle->SetNodeStyle(Style);
 		}
 
-		UMetasoundEditorGraphInputNode* FGraphBuilder::AddInputNode(UObject& InMetasound, Frontend::FNodeHandle InNodeHandle, FVector2D InLocation, bool bInSelectNewNode)
+		UMetasoundEditorGraphInputNode* FGraphBuilder::AddInputNode(UObject& InMetaSound, Frontend::FNodeHandle InNodeHandle, FVector2D InLocation, bool bInSelectNewNode)
 		{
 			using namespace Frontend;
 
-			FMetasoundAssetBase* MetasoundAsset = IMetasoundUObjectRegistry::Get().GetObjectAsAssetBase(&InMetasound);
-			check(MetasoundAsset);
+			FMetasoundAssetBase* MetaSoundAsset = IMetasoundUObjectRegistry::Get().GetObjectAsAssetBase(&InMetaSound);
+			check(MetaSoundAsset);
 
-			UMetasoundEditorGraph* MetasoundGraph = Cast<UMetasoundEditorGraph>(MetasoundAsset->GetGraph());
+			UMetasoundEditorGraph* MetasoundGraph = Cast<UMetasoundEditorGraph>(MetaSoundAsset->GetGraph());
 			if (!ensure(MetasoundGraph))
 			{
 				return nullptr;
@@ -523,7 +573,7 @@ namespace Metasound
 			return true;
 		}
 
-		Frontend::FNodeHandle FGraphBuilder::AddNodeHandle(UObject& InMetasound, UMetasoundEditorGraphNode& InGraphNode)
+		Frontend::FNodeHandle FGraphBuilder::AddNodeHandle(UObject& InMetaSound, UMetasoundEditorGraphNode& InGraphNode)
 		{
 			using namespace Frontend;
 
@@ -537,7 +587,7 @@ namespace Metasound
 					UMetasoundEditorGraphInput* Input = InputNode->Input;
 					if (ensure(Input))
 					{
-						NodeHandle = FGraphBuilder::AddInputNodeHandle(InMetasound, Input->TypeName, InGraphNode.GetTooltipText());
+						NodeHandle = AddInputNodeHandle(InMetaSound, Input->TypeName, InGraphNode.GetTooltipText());
 						NodeHandle->SetDisplayName(FText::FromString(Pin->GetName()));
 					}
 				}
@@ -552,7 +602,7 @@ namespace Metasound
 					UMetasoundEditorGraphOutput* Output = OutputNode->Output;
 					if (ensure(Output))
 					{
-						NodeHandle = FGraphBuilder::AddOutputNodeHandle(InMetasound, Output->TypeName, InGraphNode.GetTooltipText());
+						NodeHandle = FGraphBuilder::AddOutputNodeHandle(InMetaSound, Output->TypeName, InGraphNode.GetTooltipText());
 						NodeHandle->SetDisplayName(FText::FromString(Pin->GetName()));
 					}
 				}
@@ -564,13 +614,15 @@ namespace Metasound
 				bool bDidFindClassWithName = ISearchEngine::Get().FindClassWithHighestVersion(ExternalNode->ClassName.ToNodeClassName(), FrontendClass);
 				if (ensure(bDidFindClassWithName))
 				{
-					FMetasoundAssetBase* MetasoundAsset = IMetasoundUObjectRegistry::Get().GetObjectAsAssetBase(&InMetasound);
-					check(MetasoundAsset);
+					FMetasoundAssetBase* MetaSoundAsset = IMetasoundUObjectRegistry::Get().GetObjectAsAssetBase(&InMetaSound);
+					check(MetaSoundAsset);
 
-					Frontend::FNodeHandle NewNode = MetasoundAsset->GetRootGraphHandle()->AddNode(FrontendClass.Metadata);
+					Frontend::FNodeHandle NewNode = MetaSoundAsset->GetRootGraphHandle()->AddNode(FrontendClass.Metadata);
 					ExternalNode->SetNodeID(NewNode->GetID());
 
 					NodeHandle = NewNode;
+
+					MetaSoundAsset->RebuildReferencedAssets(UMetaSoundAssetSubsystem::Get());
 				}
 			}
 
@@ -584,90 +636,43 @@ namespace Metasound
 			return NodeHandle;
 		}
 
-		Frontend::FNodeHandle FGraphBuilder::AddInputNodeHandle(UObject& InMetasound, const FName InTypeName, const FText& InToolTip, const FMetasoundFrontendLiteral* InDefaultValue)
+		Frontend::FNodeHandle FGraphBuilder::AddInputNodeHandle(UObject& InMetaSound, const FName InTypeName, const FText& InToolTip, const FMetasoundFrontendLiteral* InDefaultValue)
 		{
-			FMetasoundAssetBase* MetasoundAsset = IMetasoundUObjectRegistry::Get().GetObjectAsAssetBase(&InMetasound);
-			check(MetasoundAsset);
-
-			Frontend::FGraphHandle GraphHandle = MetasoundAsset->GetRootGraphHandle();
-			const FGuid VertexID = FGuid::NewGuid();
-
-			FMetasoundFrontendClassInput Description;
-
-			// The user currently never interfaces with the name, only the display name,
-			// so just set to VertexID to make unique, avoid confusion, and eventually
-			// this field will most likely go away (or become the display name).
-			Description.Name = VertexID.ToString();
-
-			Description.TypeName = InTypeName;
-			Description.Metadata.Description = InToolTip;
-			Description.VertexID = VertexID;
-
-			Frontend::FNodeHandle NodeHandle = GraphHandle->AddInputVertex(Description);
-			if (ensure(NodeHandle->IsValid()))
-			{
-				const FText DisplayName = FGraphBuilder::GenerateUniqueInputDisplayName(InMetasound);
-				GraphHandle->SetInputDisplayName(Description.Name, DisplayName);
-
-				if (InDefaultValue)
-				{
-					GraphHandle->SetDefaultInput(VertexID, *InDefaultValue);
-				}
-				else
-				{
-					GraphHandle->SetDefaultInputToDefaultLiteralOfType(VertexID);
-				}
-			}
+			FMetasoundAssetBase* MetaSoundAsset = IMetasoundUObjectRegistry::Get().GetObjectAsAssetBase(&InMetaSound);
+			check(MetaSoundAsset);
 
-			return NodeHandle;
+			const FText DisplayName = GenerateUniqueInputDisplayName(InMetaSound);
+			return MetaSoundAsset->GetRootGraphHandle()->AddInputNode(InTypeName, InToolTip, InDefaultValue, &DisplayName);
 		}
 
-		Frontend::FNodeHandle FGraphBuilder::AddOutputNodeHandle(UObject& InMetasound, const FName InTypeName, const FText& InToolTip)
+		Frontend::FNodeHandle FGraphBuilder::AddOutputNodeHandle(UObject& InMetaSound, const FName InTypeName, const FText& InToolTip)
 		{
-			FMetasoundAssetBase* MetasoundAsset = IMetasoundUObjectRegistry::Get().GetObjectAsAssetBase(&InMetasound);
-			check(MetasoundAsset);
-
-			Frontend::FGraphHandle GraphHandle = MetasoundAsset->GetRootGraphHandle();
-			const FGuid VertexID = FGuid::NewGuid();
-
-			FMetasoundFrontendClassOutput Description;
-
-			// The user currently never interfaces with the name, only the display name,
-			// so just set to VertexID to make unique, avoid confusion, and eventually
-			// this field will most likely go away (or become the display name).
-			Description.Name = VertexID.ToString();
-
-			Description.TypeName = InTypeName;
-			Description.Metadata.Description = InToolTip;
-			Description.VertexID = VertexID;
-
-			Frontend::FNodeHandle NodeHandle = GraphHandle->AddOutputVertex(Description);
-
-			const FText DisplayName = FGraphBuilder::GenerateUniqueOutputDisplayName(InMetasound);
-			GraphHandle->SetOutputDisplayName(Description.Name, DisplayName);
+			FMetasoundAssetBase* MetaSoundAsset = IMetasoundUObjectRegistry::Get().GetObjectAsAssetBase(&InMetaSound);
+			check(MetaSoundAsset);
 
-			return NodeHandle;
+			const FText DisplayName = GenerateUniqueOutputDisplayName(InMetaSound);
+			return MetaSoundAsset->GetRootGraphHandle()->AddOutputNode(InTypeName, InToolTip, &DisplayName);
 		}
 
-		UMetasoundEditorGraphNode* FGraphBuilder::AddNode(UObject& InMetasound, Frontend::FNodeHandle InNodeHandle, FVector2D InLocation, bool bInSelectNewNode)
+		UMetasoundEditorGraphNode* FGraphBuilder::AddNode(UObject& InMetaSound, Frontend::FNodeHandle InNodeHandle, FVector2D InLocation, bool bInSelectNewNode)
 		{
-			switch (InNodeHandle->GetClassMetadata().Type)
+			switch (InNodeHandle->GetClassMetadata().GetType())
 			{
 				case EMetasoundFrontendClassType::Input:
 				{
-					return CastChecked<UMetasoundEditorGraphNode>(AddInputNode(InMetasound, InNodeHandle, InLocation, bInSelectNewNode));
+					return CastChecked<UMetasoundEditorGraphNode>(AddInputNode(InMetaSound, InNodeHandle, InLocation, bInSelectNewNode));
 				}
 				break;
 
 				case EMetasoundFrontendClassType::External:
 				{
-					return CastChecked<UMetasoundEditorGraphNode>(AddExternalNode(InMetasound, InNodeHandle, InLocation, bInSelectNewNode));
+					return CastChecked<UMetasoundEditorGraphNode>(AddExternalNode(InMetaSound, InNodeHandle, InLocation, bInSelectNewNode));
 				}
 				break;
 
 				case EMetasoundFrontendClassType::Output:
 				{
-					return CastChecked<UMetasoundEditorGraphNode>(AddOutputNode(InMetasound, InNodeHandle, InLocation, bInSelectNewNode));
+					return CastChecked<UMetasoundEditorGraphNode>(AddOutputNode(InMetaSound, InNodeHandle, InLocation, bInSelectNewNode));
 				}
 				break;
 
@@ -708,7 +713,6 @@ namespace Metasound
 			FOutputHandle OutputHandle = GetOutputHandleFromPin(&InOutputPin);
 			if (!InputHandle->IsValid() || !OutputHandle->IsValid())
 			{
-				InInputPin.BreakLinkTo(&InOutputPin);
 				return false;
 			}
 
@@ -765,7 +769,7 @@ namespace Metasound
 			}
 		}
 
-		void FGraphBuilder::InitMetaSound(UObject& InMetaSound, const FString& InAuthor, bool bConformToArchetype)
+		void FGraphBuilder::InitMetaSound(UObject& InMetaSound, const FString& InAuthor)
 		{
 			using namespace Frontend;
 			using namespace GraphBuilderPrivate;
@@ -773,12 +777,11 @@ namespace Metasound
 			FMetasoundFrontendClassMetadata Metadata;
 
 			// 1. Set default class Metadata
-			Metadata.ClassName = FMetasoundFrontendClassName(FName(), *FGuid::NewGuid().ToString(), FName());
-			Metadata.Version.Major = 1;
-			Metadata.Version.Minor = 0;
-			Metadata.DisplayName = FText::FromString(InMetaSound.GetName());
-			Metadata.Type = EMetasoundFrontendClassType::Graph;
-			Metadata.Author = FText::FromString(InAuthor);
+			Metadata.SetClassName(FMetasoundFrontendClassName(FName(), *FGuid::NewGuid().ToString(), FName()));
+			Metadata.SetVersion({ 1, 0 });
+			Metadata.SetDisplayName(FText::FromString(InMetaSound.GetName()));
+			Metadata.SetType(EMetasoundFrontendClassType::Graph);
+			Metadata.SetAuthor(FText::FromString(InAuthor));
 
 			FMetasoundAssetBase* MetaSoundAsset = IMetasoundUObjectRegistry::Get().GetObjectAsAssetBase(&InMetaSound);
 			check(MetaSoundAsset);
@@ -790,35 +793,32 @@ namespace Metasound
 			DocMetadata.Version.Number = FVersionDocument::GetMaxVersion();
 			DocumentHandle->SetMetadata(DocMetadata);
 
-			if (bConformToArchetype)
-			{
-				MetaSoundAsset->ConformDocumentToArchetype();
-			}
+			MetaSoundAsset->ConformDocumentToArchetype();
 
 			FGraphHandle GraphHandle = MetaSoundAsset->GetRootGraphHandle();
 			FVector2D InputNodeLocation = FVector2D::ZeroVector;
-			FVector2D ExternalNodeLocation = InputNodeLocation + NodeLayout::BufferX;
-			FVector2D OutputNodeLocation = ExternalNodeLocation + NodeLayout::BufferX;
+			FVector2D ExternalNodeLocation = InputNodeLocation + DisplayStyle::NodeLayout::DefaultOffsetX;
+			FVector2D OutputNodeLocation = ExternalNodeLocation + DisplayStyle::NodeLayout::DefaultOffsetX;
 
 			TArray<FNodeHandle> NodeHandles = GraphHandle->GetNodes();
 			for (FNodeHandle& NodeHandle : NodeHandles)
 			{
-				const EMetasoundFrontendClassType NodeType = NodeHandle->GetClassMetadata().Type;
+				const EMetasoundFrontendClassType NodeType = NodeHandle->GetClassMetadata().GetType();
 				FVector2D NewLocation;
 				if (NodeType == EMetasoundFrontendClassType::Input)
 				{
 					NewLocation = InputNodeLocation;
-					InputNodeLocation += NodeLayout::BufferY;
+					InputNodeLocation += DisplayStyle::NodeLayout::DefaultOffsetY;
 				}
 				else if (NodeType == EMetasoundFrontendClassType::Output)
 				{
 					NewLocation = OutputNodeLocation;
-					OutputNodeLocation += NodeLayout::BufferY;
+					OutputNodeLocation += DisplayStyle::NodeLayout::DefaultOffsetY;
 				}
 				else
 				{
 					NewLocation = ExternalNodeLocation;
-					ExternalNodeLocation += NodeLayout::BufferY;
+					ExternalNodeLocation += DisplayStyle::NodeLayout::DefaultOffsetY;
 				}
 				FMetasoundFrontendNodeStyle Style = NodeHandle->GetNodeStyle();
 				// TODO: Find consistent location for controlling node locations.
@@ -826,8 +826,6 @@ namespace Metasound
 				Style.Display.Locations = {{FGuid::NewGuid(), NewLocation}};
 				NodeHandle->SetNodeStyle(Style);
 			}
-
-			MetaSoundAsset->RegisterGraphWithFrontend();
 		}
 
 		void FGraphBuilder::InitMetaSoundPreset(UObject& InMetaSoundReferenced, UObject& InMetaSoundPreset)
@@ -835,180 +833,22 @@ namespace Metasound
 			using namespace Frontend;
 			using namespace GraphBuilderPrivate;
 
-			FMetasoundAssetBase* MetaSoundReferencedAsset = IMetasoundUObjectRegistry::Get().GetObjectAsAssetBase(&InMetaSoundReferenced);
-			check(MetaSoundReferencedAsset);
-
-			FConstGraphHandle ReferencedGraphHandle = MetaSoundReferencedAsset->GetRootGraphHandle();
-			if (!ensure(ReferencedGraphHandle->IsValid()))
-			{
-				return;
-			}
-
-			FMetasoundAssetBase* MetaSoundPresetAsset = IMetasoundUObjectRegistry::Get().GetObjectAsAssetBase(&InMetaSoundPreset);
-			check(MetaSoundPresetAsset);
-			FGraphHandle PresetGraphHandle = MetaSoundPresetAsset->GetRootGraphHandle();
-			if (!ensure(PresetGraphHandle->IsValid()))
-			{
-				return;
-			}
-
-			// Ensure referenced MetaSound is up-to-date
-			check(MetaSoundReferencedAsset);
-			MetaSoundReferencedAsset->VersionAsset();
-
-			// Ensure asset is registered with frontend so the class
-			// can be queried when reference is added. This may modify
-			// registry properties, so must be non-const.
-			MetaSoundReferencedAsset->RegisterGraphWithFrontend();
+			FMetasoundAssetBase* PresetAsset = IMetasoundUObjectRegistry::Get().GetObjectAsAssetBase(&InMetaSoundPreset);
+			check(PresetAsset);
 
-			// Initialize EdGraph, but don't conform to archetype as this would provide the
-			// default interface (inputs/outputs), which will be added below when duplicating
-			// node's interface to the preset.
-			GraphBuilderPrivate::InitializeGraph(InMetaSoundPreset);
-
-			// 1a. Preset graph topology is considered "read-only", so mark as such
+			// Mark preset as auto-update and non-editable
+			FGraphHandle PresetGraphHandle = PresetAsset->GetRootGraphHandle();
 			FMetasoundFrontendGraphStyle Style = PresetGraphHandle->GetGraphStyle();
 			Style.bIsGraphEditable = false;
 			PresetGraphHandle->SetGraphStyle(Style);
 
-			// 1b. Set archetype accordingly
-			const FMetasoundFrontendVersion& RefArchetype = MetaSoundReferencedAsset->GetDocumentHandle()->GetArchetypeVersion();
-			if (ensure(MetaSoundPresetAsset->IsArchetypeSupported(RefArchetype)))
-			{
-				MetaSoundPresetAsset->ConformDocumentToArchetype();
-			}
-
-			// 2. Generate referenced inputs & outputs
-			FVector2D InputNodeLocation = FVector2D::ZeroVector;
-			FVector2D ExternalNodeLocation = InputNodeLocation + NodeLayout::BufferX;
-			FVector2D OutputNodeLocation = ExternalNodeLocation + NodeLayout::BufferX;
-
-			TMap<FString, FString> PresetInputToReferenceInputMap;
-			TMap<FString, FString> PresetOutputToReferenceOutputMap;
-
-			ReferencedGraphHandle->IterateConstNodes([&](FConstNodeHandle RefGraphInputNode)
-			{
-				// Required inputs are added via the prior step (archetype conformation).
-				// For these, find the corrisponding input & create an editor node only.
-				const bool bIsRequired = RefGraphInputNode->IsRequired();
+			FMetasoundFrontendClassMetadata Metadata = PresetGraphHandle->GetGraphMetadata();
+			Metadata.SetAutoUpdateManagesInterface(true);
+			PresetGraphHandle->SetGraphMetadata(Metadata);
 
-				const FText& DisplayName = RefGraphInputNode->GetDisplayName();
-				const FText& Description = RefGraphInputNode->GetDescription();
-
-				// Should only ever be one
-				ensure(RefGraphInputNode->GetNumInputs() == 1);
-				RefGraphInputNode->IterateConstInputs([&](FConstInputHandle RefGraphInputHandle)
-				{
-					FNodeHandle NodeHandle = INodeController::GetInvalidHandle();
-					if (bIsRequired)
-					{
-						NodeHandle = PresetGraphHandle->GetInputNodeWithName(RefGraphInputHandle->GetName());
-					}
-					else
-					{
-						FGuid InputID = ReferencedGraphHandle->GetVertexIDForInputVertex(RefGraphInputHandle->GetName());
-						const FMetasoundFrontendLiteral DefaultLiteral = ReferencedGraphHandle->GetDefaultInput(InputID);
-
-						const FName DataTypeName = RefGraphInputHandle->GetDataType();
-						NodeHandle = AddInputNodeHandle(InMetaSoundPreset, DataTypeName, Description, &DefaultLiteral);
-						NodeHandle->SetDisplayName(DisplayName);
-					}
-
-					if (ensure(NodeHandle->IsValid()))
-					{
-						UMetasoundEditorGraphInputNode* InputNode = AddInputNode(InMetaSoundPreset, NodeHandle, InputNodeLocation, false /* bSelectNewNode */);
-						if (ensure(InputNode))
-						{
-							PresetInputToReferenceInputMap.Add(NodeHandle->GetNodeName(), RefGraphInputHandle->GetName());
-						}
-						InputNodeLocation += NodeLayout::BufferY;
-					}
-				});
-			},
-			EMetasoundFrontendClassType::Input);
-
-			ReferencedGraphHandle->IterateConstNodes([&](FConstNodeHandle RefGraphOutputNode)
-			{
-				// Required inputs are added via the prior step (archetype conformation).
-				// For these, find the corrisponding input & create an editor node only.
-				const bool bIsRequired = RefGraphOutputNode->IsRequired();
-
-				const FText& DisplayName = RefGraphOutputNode->GetDisplayName();
-				const FText& Description = RefGraphOutputNode->GetDescription();
-
-				// Should only ever be one
-				ensure(RefGraphOutputNode->GetNumOutputs() == 1);
-				RefGraphOutputNode->IterateConstOutputs([&](FConstOutputHandle RefGraphOutputHandle)
-				{
-					Frontend::FNodeHandle NodeHandle = Frontend::INodeController::GetInvalidHandle();
-					if (bIsRequired)
-					{
-						NodeHandle = PresetGraphHandle->GetOutputNodeWithName(RefGraphOutputHandle->GetName());
-					}
-					else
-					{
-						const FName DataTypeName = RefGraphOutputHandle->GetDataType();
-						NodeHandle = AddOutputNodeHandle(InMetaSoundPreset, DataTypeName, Description);
-						NodeHandle->SetDisplayName(DisplayName);
-					}
-
-					if (ensure(NodeHandle->IsValid()))
-					{
-						UMetasoundEditorGraphOutputNode* OutputNode = AddOutputNode(InMetaSoundPreset, NodeHandle, OutputNodeLocation, false /* bSelectNewNode */);
-						if (ensure(OutputNode))
-						{
-							PresetOutputToReferenceOutputMap.Add(OutputNode->GetNodeHandle()->GetNodeName(), RefGraphOutputHandle->GetName());
-						}
-						OutputNodeLocation += NodeLayout::BufferY;
-					}
-				});
-			},
-			EMetasoundFrontendClassType::Output);
-
-			// 3. Generate a referencing node to the given referenced MetaSound.
-			FMetasoundFrontendClassMetadata ReferencedClassMetadata = ReferencedGraphHandle->GetGraphMetadata();
-
-			// Swap type on look-up as it will be referenced as an externally defined class relative to the new Preset asset
-			ReferencedClassMetadata.Type = EMetasoundFrontendClassType::External;
-			UMetasoundEditorGraphExternalNode* ReferencedNode = AddExternalNode(InMetaSoundPreset, ReferencedClassMetadata, ExternalNodeLocation, false /* bSelectNewNode */);
-			FNodeHandle ReferencedNodeHandle = ReferencedNode->GetNodeHandle();
-
-			// 4. Connect preset's respective inputs & outputs to generated referencing node.
-			PresetGraphHandle->IterateNodes([&](FNodeHandle InputNode)
-			{
-				const FString& Name = InputNode->GetNodeName();
-
-				// Should only ever be one
-				ensure(InputNode->GetNumOutputs() == 1);
-				InputNode->IterateOutputs([&](FOutputHandle OutputHandle)
-				{
-					const FString& InputName = PresetInputToReferenceInputMap.FindChecked(Name);
-					const TArray<FInputHandle> ReferenceInputs = ReferencedNodeHandle->GetInputsWithVertexName(InputName);
-					if (ensure(ReferenceInputs.Num() == 1))
-					{
-						ensure(ReferenceInputs[0]->Connect(*OutputHandle));
-					}
-				});
-			},
-			EMetasoundFrontendClassType::Input);
-
-			PresetGraphHandle->IterateNodes([&](FNodeHandle OutputNode)
-			{
-				const FString& Name = OutputNode->GetNodeName();
-
-				// Should only ever be one
-				ensure(OutputNode->GetNumInputs() == 1);
-				OutputNode->IterateInputs([&](FInputHandle InputHandle)
-				{
-					const FString& OutputName = PresetOutputToReferenceOutputMap.FindChecked(Name);
-					const TArray<FOutputHandle> ReferenceInputs = ReferencedNodeHandle->GetOutputsWithVertexName(OutputName);
-					if (ensure(ReferenceInputs.Num() == 1))
-					{
-						ensure(InputHandle->Connect(*ReferenceInputs[0]));
-					}
-				});
-			},
-			EMetasoundFrontendClassType::Output);
+			const FMetasoundAssetBase* ReferencedAsset = IMetasoundUObjectRegistry::Get().GetObjectAsAssetBase(&InMetaSoundReferenced);
+			check(ReferencedAsset);
+			FRebuildPresetRootGraph(ReferencedAsset->GetDocumentHandle()).Transform(PresetAsset->GetDocumentHandle());
 		}
 
 		void FGraphBuilder::DeleteVariableNodeHandle(UMetasoundEditorGraphVariable& InVariable)
@@ -1076,7 +916,7 @@ namespace Metasound
 			Frontend::FGraphHandle GraphHandle = NodeHandle->GetOwningGraph();
 			if (GraphHandle->IsValid())
 			{
-				switch (NodeHandle->GetClassMetadata().Type)
+				switch (NodeHandle->GetClassMetadata().GetType())
 				{
 					case EMetasoundFrontendClassType::Output:
 					case EMetasoundFrontendClassType::Input:
@@ -1114,6 +954,15 @@ namespace Metasound
 			{
 				SynchronizeGraph(Graph->GetMetasoundChecked());
 			}
+			// SynchronizeGraph calls rebuild references, but it needs to be
+			// called to remove the given reference irrespective of error state.
+			else
+			{
+				FMetasoundAssetBase* MetaSoundAsset = IMetasoundUObjectRegistry::Get().GetObjectAsAssetBase(&Graph->GetMetasoundChecked());
+				check(MetaSoundAsset);
+
+				MetaSoundAsset->RebuildReferencedAssets(UMetaSoundAssetSubsystem::Get());
+			}
 
 			return bSuccess;
 		}
@@ -1138,7 +987,7 @@ namespace Metasound
 			// Only add input pins of the node is not an input node. Input nodes
 			// have their input pins hidden because they cannot be connected internal
 			// to the graph.
-			if (EMetasoundFrontendClassType::Input != NodeHandle->GetClassMetadata().Type)
+			if (EMetasoundFrontendClassType::Input != NodeHandle->GetClassMetadata().GetType())
 			{
 				TArray<FInputHandle> InputHandles = NodeHandle->GetInputs();
 				InputHandles = NodeHandle->GetInputStyle().SortDefaults(InputHandles);
@@ -1151,7 +1000,7 @@ namespace Metasound
 			// Only add output pins of the node is not an output node. Output nodes
 			// have their output pins hidden because they cannot be connected internal
 			// to the graph.
-			if (EMetasoundFrontendClassType::Output != NodeHandle->GetClassMetadata().Type)
+			if (EMetasoundFrontendClassType::Output != NodeHandle->GetClassMetadata().GetType())
 			{
 				TArray<FOutputHandle> OutputHandles = NodeHandle->GetOutputs();
 				OutputHandles = NodeHandle->GetOutputStyle().SortDefaults(OutputHandles);
@@ -1177,8 +1026,47 @@ namespace Metasound
 			}
 		}
 
+		void FGraphBuilder::RegisterGraphWithFrontend(UObject& InMetaSound)
+		{
+			FMetasoundAssetBase* MetaSoundAsset = IMetasoundUObjectRegistry::Get().GetObjectAsAssetBase(&InMetaSound);
+			check(MetaSoundAsset);
+
+			MetaSoundAsset->RegisterGraphWithFrontend();
+
+			TArray<UObject*> EditedAssets = GEditor->GetEditorSubsystem<UAssetEditorSubsystem>()->GetAllEditedAssets();
+			for (UObject* Asset : EditedAssets)
+			{
+				if (Asset != &InMetaSound && IMetasoundUObjectRegistry::Get().IsRegisteredClass(Asset))
+				{
+					ValidateGraph(*Asset, true /* bAutoUpdate */);
+				}
+			}
+		}
+
+		void FGraphBuilder::UnregisterGraphWithFrontend(UObject& InMetaSound)
+		{
+			FMetasoundAssetBase* MetaSoundAsset = IMetasoundUObjectRegistry::Get().GetObjectAsAssetBase(&InMetaSound);
+			check(MetaSoundAsset);
+
+			MetaSoundAsset->UnregisterGraphWithFrontend();
+
+			TArray<UObject*> EditedAssets = GEditor->GetEditorSubsystem<UAssetEditorSubsystem>()->GetAllEditedAssets();
+			for (UObject* Asset : EditedAssets)
+			{
+				if (Asset != &InMetaSound && IMetasoundUObjectRegistry::Get().IsRegisteredClass(Asset))
+				{
+					ValidateGraph(*Asset, true /* bAutoUpdate */);
+				}
+			}
+		}
+
 		bool FGraphBuilder::IsMatchingInputHandleAndPin(const Frontend::FInputHandle& InInputHandle, const UEdGraphPin& InEditorPin)
 		{
+			if (InEditorPin.Direction != EGPD_Input)
+			{
+				return false;
+			}
+
 			Frontend::FInputHandle PinInputHandle = GetInputHandleFromPin(&InEditorPin);
 			if (PinInputHandle->GetID() == InInputHandle->GetID())
 			{
@@ -1190,6 +1078,11 @@ namespace Metasound
 
 		bool FGraphBuilder::IsMatchingOutputHandleAndPin(const Frontend::FOutputHandle& InOutputHandle, const UEdGraphPin& InEditorPin)
 		{
+			if (InEditorPin.Direction != EGPD_Output)
+			{
+				return false;
+			}
+
 			Frontend::FOutputHandle PinOutputHandle = GetOutputHandleFromPin(&InEditorPin);
 			if (PinOutputHandle->GetID() == InOutputHandle->GetID())
 			{
@@ -1262,18 +1155,18 @@ namespace Metasound
 			return NewPin;
 		}
 
-		bool FGraphBuilder::SynchronizeConnections(UObject& InMetasound)
+		bool FGraphBuilder::SynchronizeConnections(UObject& InMetaSound)
 		{
 			using namespace Frontend;
 
 			bool bIsGraphDirty = false;
 
-			FMetasoundAssetBase* MetasoundAsset = IMetasoundUObjectRegistry::Get().GetObjectAsAssetBase(&InMetasound);
-			check(MetasoundAsset);
+			FMetasoundAssetBase* MetaSoundAsset = IMetasoundUObjectRegistry::Get().GetObjectAsAssetBase(&InMetaSound);
+			check(MetaSoundAsset);
 
-			FGraphHandle GraphHandle = MetasoundAsset->GetRootGraphHandle();
+			FGraphHandle GraphHandle = MetaSoundAsset->GetRootGraphHandle();
 
-			UMetasoundEditorGraph* EditorGraph = CastChecked<UMetasoundEditorGraph>(MetasoundAsset->GetGraph());
+			UMetasoundEditorGraph* EditorGraph = CastChecked<UMetasoundEditorGraph>(MetaSoundAsset->GetGraph());
 
 
 			TArray<UMetasoundEditorGraphNode*> EditorNodes;
@@ -1292,7 +1185,7 @@ namespace Metasound
 
 				FNodeHandle Node = EditorNode->GetNodeHandle();
 
-				if (EMetasoundFrontendClassType::Input == Node->GetClassMetadata().Type)
+				if (EMetasoundFrontendClassType::Input == Node->GetClassMetadata().GetType())
 				{
 					// Skip this node if it is an input node. Input pins on input 
 					// nodes are not connected internal to the graph.
@@ -1378,25 +1271,25 @@ namespace Metasound
 			return bIsGraphDirty;
 		}
 
-		bool FGraphBuilder::SynchronizeGraph(UObject& InMetasound)
+		bool FGraphBuilder::SynchronizeGraph(UObject& InMetaSound)
 		{
 			using namespace Frontend;
 
-			GraphBuilderPrivate::InitializeGraph(InMetasound);
+			GraphBuilderPrivate::InitializeGraph(InMetaSound);
 
-			bool bIsEditorGraphDirty = SynchronizeVariables(InMetasound);
+			bool bIsEditorGraphDirty = SynchronizeVariables(InMetaSound);
 
-			FMetasoundAssetBase* MetasoundAsset = IMetasoundUObjectRegistry::Get().GetObjectAsAssetBase(&InMetasound);
-			check(MetasoundAsset);
+			FMetasoundAssetBase* MetaSoundAsset = IMetasoundUObjectRegistry::Get().GetObjectAsAssetBase(&InMetaSound);
+			check(MetaSoundAsset);
 
 			// Get all external nodes from Frontend graph.  Input and output references will only be added/synchronized
 			// if required when synchronizing connections (as they are not required to inhabit editor graph).
-			FGraphHandle GraphHandle = MetasoundAsset->GetRootGraphHandle();
+			FGraphHandle GraphHandle = MetaSoundAsset->GetRootGraphHandle();
 			TArray<FNodeHandle> FrontendNodes = GraphHandle->GetNodes();
 
 			// Get all editor nodes from editor graph (some nodes on graph may *NOT* be metasound ed nodes,
 			// such as comment boxes, etc, so just get nodes of class UMetasoundEditorGraph).
-			UMetasoundEditorGraph* EditorGraph = CastChecked<UMetasoundEditorGraph>(MetasoundAsset->GetGraph());
+			UMetasoundEditorGraph* EditorGraph = CastChecked<UMetasoundEditorGraph>(MetaSoundAsset->GetGraph());
 			TArray<UMetasoundEditorGraphNode*> EditorNodes;
 			EditorGraph->GetNodesOfClass(EditorNodes);
 
@@ -1486,7 +1379,7 @@ namespace Metasound
 				TArray<UMetasoundEditorGraphNode*> AddedNodes;
 				for (const TPair<FGuid, FVector2D>& Location : NewStyle.Display.Locations)
 				{
-					UMetasoundEditorGraphNode* NewNode = AddNode(InMetasound, Node, Location.Value, false /* bInSelectNewNode */);
+					UMetasoundEditorGraphNode* NewNode = AddNode(InMetaSound, Node, Location.Value, false /* bInSelectNewNode */);
 					if (ensure(NewNode))
 					{
 						FAssociatedNodes& AssociatedNodeData = AssociatedNodes.FindOrAdd(Node->GetID());
@@ -1522,7 +1415,9 @@ namespace Metasound
 			}
 
 			// Synchronize connections.
-			bIsEditorGraphDirty |= SynchronizeConnections(InMetasound);
+			bIsEditorGraphDirty |= SynchronizeConnections(InMetaSound);
+
+			MetaSoundAsset->RebuildReferencedAssets(UMetaSoundAssetSubsystem::Get());
 			return bIsEditorGraphDirty;
 		}
 
@@ -1589,7 +1484,7 @@ namespace Metasound
 				}
 			}
 
-			const EMetasoundFrontendClassType ClassType = InNode->GetClassMetadata().Type;
+			const EMetasoundFrontendClassType ClassType = InNode->GetClassMetadata().GetType();
 
 			// Only add input pins of the node is not an input node. Input nodes
 			// have their input pins hidden because they cannot be connected internal
@@ -1664,22 +1559,22 @@ namespace Metasound
 			return OldValue != InPin.DefaultValue;
 		}
 
-		bool FGraphBuilder::SynchronizeVariables(UObject& InMetasound)
+		bool FGraphBuilder::SynchronizeVariables(UObject& InMetaSound)
 		{
 			using namespace Frontend;
 
 			bool bIsEditorGraphDirty = false;
 
-			FMetasoundAssetBase* MetasoundAsset = IMetasoundUObjectRegistry::Get().GetObjectAsAssetBase(&InMetasound);
-			check(MetasoundAsset);
-			UMetasoundEditorGraph* Graph = CastChecked<UMetasoundEditorGraph>(MetasoundAsset->GetGraph());
+			FMetasoundAssetBase* MetaSoundAsset = IMetasoundUObjectRegistry::Get().GetObjectAsAssetBase(&InMetaSound);
+			check(MetaSoundAsset);
+			UMetasoundEditorGraph* Graph = CastChecked<UMetasoundEditorGraph>(MetaSoundAsset->GetGraph());
 
 			IMetasoundEditorModule& EditorModule = FModuleManager::GetModuleChecked<IMetasoundEditorModule>("MetaSoundEditor");
 
 			TSet<UMetasoundEditorGraphInput*> Inputs;
 			TSet<UMetasoundEditorGraphOutput*> Outputs;
 
-			FGraphHandle GraphHandle = MetasoundAsset->GetRootGraphHandle();
+			FGraphHandle GraphHandle = MetaSoundAsset->GetRootGraphHandle();
 			GraphHandle->IterateNodes([&](FNodeHandle NodeHandle)
 			{
 				if (UMetasoundEditorGraphInput* Input = Graph->FindInput(NodeHandle->GetID()))
diff --git a/Engine/Plugins/Runtime/Metasound/Source/MetasoundEditor/Private/MetasoundEditorGraphBuilder.h b/Engine/Plugins/Runtime/Metasound/Source/MetasoundEditor/Private/MetasoundEditorGraphBuilder.h
--- a/Engine/Plugins/Runtime/Metasound/Source/MetasoundEditor/Private/MetasoundEditorGraphBuilder.h
+++ b/Engine/Plugins/Runtime/Metasound/Source/MetasoundEditor/Private/MetasoundEditorGraphBuilder.h
@@ -12,6 +12,7 @@
 
 
 // Forward Declarations
+class SGraphEditor;
 class UEdGraphNode;
 class UEdGraphPin;
 class UMetaSound;
@@ -33,7 +34,7 @@ namespace Metasound
 
 		class FGraphBuilder
 		{
-			static void InitGraphNode(Frontend::FNodeHandle& InNodeHandle, UMetasoundEditorGraphNode* NewGraphNode, UObject& InMetasound);
+			static void InitGraphNode(Frontend::FNodeHandle& InNodeHandle, UMetasoundEditorGraphNode* NewGraphNode, UObject& InMetaSound);
 
 		public:
 			static const FName PinCategoryAudio;
@@ -53,41 +54,50 @@ namespace Metasound
 			static const FText FunctionMenuName;
 
 			// Adds an EdGraph node to mirror the provided FNodeHandle.
-			static UMetasoundEditorGraphNode* AddNode(UObject& InMetasound, Frontend::FNodeHandle InNodeHandle, FVector2D InLocation, bool bInSelectNewNode = true);
+			static UMetasoundEditorGraphNode* AddNode(UObject& InMetaSound, Frontend::FNodeHandle InNodeHandle, FVector2D InLocation, bool bInSelectNewNode = true);
 
-			// Convenience functions for retrieving the editor for the given Metasound/EdGraph
-			static TSharedPtr<FEditor> GetEditorForGraph(const UObject& Metasound);
+			// Convenience functions for retrieving the editor for the given MetaSound/EdGraph
+			static TSharedPtr<FEditor> GetEditorForMetasound(const UObject& MetaSound);
 			static TSharedPtr<FEditor> GetEditorForGraph(const UEdGraph& EdGraph);
 
+			// Validates MetaSound and optionally performs AutoUpdate logic in Frontend with any additional logic required to refresh editor & respective editor object state.
+			static bool ValidateGraph(UObject& InMetaSoundToUpdate, bool bAutoUpdate = true);
+
+			// Wraps RegisterGraphWithFrontend logic in Frontend with any additional logic required to refresh editor & respective editor object state.
+			static void RegisterGraphWithFrontend(UObject& InMetaSound);
+
+			// Wraps RegisterGraphWithFrontend logic in Frontend with any additional logic required to refresh editor & respective editor object state.
+			static void UnregisterGraphWithFrontend(UObject& InMetaSound);
+
 			// Adds a node handle to mirror the provided graph node and binds to it.  Does *NOT* mirror existing EdGraph connections
 			// nor does it remove existing bound Frontend Node (if set) from associated Frontend Graph.
-			static Frontend::FNodeHandle AddNodeHandle(UObject& InMetasound, UMetasoundEditorGraphNode& InGraphNode);
+			static Frontend::FNodeHandle AddNodeHandle(UObject& InMetaSound, UMetasoundEditorGraphNode& InGraphNode);
 
 			// Adds a corresponding UMetasoundEditorGraphInputNode for the provided node handle.
-			static UMetasoundEditorGraphInputNode* AddInputNode(UObject& InMetasound, Frontend::FNodeHandle InNodeHandle, FVector2D InLocation, bool bInSelectNewNode = true);
+			static UMetasoundEditorGraphInputNode* AddInputNode(UObject& InMetaSound, Frontend::FNodeHandle InNodeHandle, FVector2D InLocation, bool bInSelectNewNode = true);
 
 			// Generates FNodeHandle for the given external node data. Does not bind or create EdGraph representation of given node.
 			static Frontend::FNodeHandle AddInputNodeHandle(
-				UObject& InMetasound,
+				UObject& InMetaSound,
 				const FName InTypeName,
 				const FText& InToolTip,
 				const FMetasoundFrontendLiteral* InDefaultValue = nullptr);
 
 			// Adds a corresponding UMetasoundEditorGraphExternalNode for the provided node handle.
-			static UMetasoundEditorGraphExternalNode* AddExternalNode(UObject& InMetasound, Frontend::FNodeHandle& InNodeHandle, FVector2D InLocation, bool bInSelectNewNode = true);
+			static UMetasoundEditorGraphExternalNode* AddExternalNode(UObject& InMetaSound, Frontend::FNodeHandle& InNodeHandle, FVector2D InLocation, bool bInSelectNewNode = true);
 
 			// Adds an externally-defined node with the given class info to both the editor and document graphs.
 			// Generates analogous FNodeHandle.
-			static UMetasoundEditorGraphExternalNode* AddExternalNode(UObject& InMetasound, const FMetasoundFrontendClassMetadata& InMetadata, FVector2D InLocation, bool bInSelectNewNode = true);
+			static UMetasoundEditorGraphExternalNode* AddExternalNode(UObject& InMetaSound, const FMetasoundFrontendClassMetadata& InMetadata, FVector2D InLocation, bool bInSelectNewNode = true);
 
 			// Synchronizes node location data
 			static void SynchronizeNodeLocation(FVector2D InLocation, Frontend::FNodeHandle InNodeHandle, UMetasoundEditorGraphNode& InNode);
 
 			// Adds an output node to the editor graph that corresponds to the provided node handle.
-			static UMetasoundEditorGraphOutputNode* AddOutputNode(UObject& InMetasound, Frontend::FNodeHandle& InNodeHandle, FVector2D InLocation, bool bInSelectNewNode = true);
+			static UMetasoundEditorGraphOutputNode* AddOutputNode(UObject& InMetaSound, Frontend::FNodeHandle& InNodeHandle, FVector2D InLocation, bool bInSelectNewNode = true);
 
 			// Generates analogous FNodeHandle for the given internal node data. Does not bind nor create EdGraph representation of given node.
-			static Frontend::FNodeHandle AddOutputNodeHandle(UObject& InMetasound, const FName InTypeName, const FText& InToolTip);
+			static Frontend::FNodeHandle AddOutputNodeHandle(UObject& InMetaSound, const FName InTypeName, const FText& InToolTip);
 
 			// Attempts to connect Frontend node counterparts together for provided pins.  Returns true if succeeded,
 			// and breaks pin link and returns false if failed.  If bConnectEdPins is set, will attempt to connect
@@ -132,7 +142,7 @@ namespace Metasound
 			static FLinearColor GetPinCategoryColor(const FEdGraphPinType& PinType);
 
 			// Initializes MetaSound with default inputs & outputs.
-			static void InitMetaSound(UObject& InMetaSound, const FString& InAuthor, bool bConformToArchetype = true);
+			static void InitMetaSound(UObject& InMetaSound, const FString& InAuthor);
 
 			// Initializes a MetaSound Preset using the provided ReferencedMetaSound asset's
 			// root graph as the sole, encapsulated topology.
@@ -153,29 +163,29 @@ namespace Metasound
 			// Refreshes pin state from class FrontendClassVertexMetadata
 			static void RefreshPinMetadata(UEdGraphPin& InPin, const FMetasoundFrontendVertexMetadata& InMetadata);
 
-			// Adds and removes nodes, pins and connections so that the UEdGraph of the Metasound matches the FMetasoundFrontendDocumentModel
+			// Adds and removes nodes, pins and connections so that the UEdGraph of the MetaSound matches the FMetasoundFrontendDocumentModel
 			//
 			// @return True if the UEdGraph was altered, false otherwise. 
-			static bool SynchronizeGraph(UObject& InMetasound);
+			static bool SynchronizeGraph(UObject& InMetaSound);
 
 			// Adds and removes pins so that the UMetasoundEditorGraphNode matches the InNode.
 			//
 			// @return True if the UMetasoundEditorGraphNode was altered. False otherwise.
 			static bool SynchronizeNodePins(UMetasoundEditorGraphNode& InEditorNode, Frontend::FNodeHandle InNode, bool bRemoveUnusedPins = true, bool bLogChanges = true);
 
-			// Adds and removes connections so that the UEdGraph of the metasound has the same
+			// Adds and removes connections so that the UEdGraph of the MetaSound has the same
 			// connections as the FMetasoundFrontendDocument graph.
 			//
 			// @return True if the UEdGraph was altered. False otherwise.
-			static bool SynchronizeConnections(UObject& InMetasound);
+			static bool SynchronizeConnections(UObject& InMetaSound);
 
 			// Synchronizes literal for a given input with the EdGraph's pin value.
 			static bool SynchronizePinLiteral(UEdGraphPin& InPin);
 
-			// Synchronizes inputs and outputs for the given Metasound.
+			// Synchronizes inputs and outputs for the given MetaSound.
 			//
 			// @return True if the UEdGraph was altered. False otherwise.
-			static bool SynchronizeVariables(UObject& InMetasound);
+			static bool SynchronizeVariables(UObject& InMetaSound);
 
 			// Returns true if the FInputHandle and UEdGraphPin match each other.
 			static bool IsMatchingInputHandleAndPin(const Frontend::FInputHandle& InInputHandle, const UEdGraphPin& InEditorPin);
diff --git a/Engine/Plugins/Runtime/Metasound/Source/MetasoundEditor/Private/MetasoundEditorGraphInputNodes.cpp b/Engine/Plugins/Runtime/Metasound/Source/MetasoundEditor/Private/MetasoundEditorGraphInputNodes.cpp
new file mode 100644
--- /dev/null
+++ b/Engine/Plugins/Runtime/Metasound/Source/MetasoundEditor/Private/MetasoundEditorGraphInputNodes.cpp
@@ -0,0 +1,347 @@
+// Copyright Epic Games, Inc. All Rights Reserved.
+#include "MetasoundEditorGraphInputNodes.h"
+
+#include "Algo/Count.h"
+#include "Algo/Transform.h"
+#include "Audio/AudioParameterInterface.h"
+#include "CoreMinimal.h"
+#include "EdGraph/EdGraphNode.h"
+#include "GraphEditorSettings.h"
+#include "MetasoundDataReference.h"
+#include "MetasoundEditorGraph.h"
+#include "MetasoundEditorGraphNode.h"
+#include "MetasoundEditorSettings.h"
+#include "MetasoundFrontend.h"
+#include "MetasoundFrontendController.h"
+#include "MetasoundFrontendDocument.h"
+#include "MetasoundFrontendLiteral.h"
+#include "MetasoundFrontendRegistries.h"
+#include "MetasoundPrimitives.h"
+#include "Misc/Guid.h"
+#include "Sound/SoundWave.h"
+#include "UObject/ObjectMacros.h"
+#include "UObject/SoftObjectPath.h"
+
+namespace Metasound
+{
+	namespace Editor
+	{
+		namespace InputPrivate
+		{
+			template <typename TPODType>
+			void ConvertLiteral(const FMetasoundFrontendLiteral& InLiteral, TPODType& OutValue)
+			{
+				const FName& TypeName = GetMetasoundDataTypeName<TPODType>();
+				OutValue = InLiteral.ToLiteral(TypeName).Value.Get<TPODType>();
+			}
+
+			template <typename TPODType, typename TLiteralType = TPODType>
+			void ConvertLiteralToArray(const FMetasoundFrontendLiteral& InLiteral, TArray<TLiteralType>& OutArray)
+			{
+				const FName& TypeName = *FString(GetMetasoundDataTypeString<TPODType>() + TEXT(":Array"));
+				TArray<TPODType> NewValue = InLiteral.ToLiteral(TypeName).Value.Get<TArray<TPODType>>();
+				Algo::Transform(NewValue, OutArray, [](const TPODType& InValue) { return TLiteralType { InValue }; });
+			}
+		}
+	}
+}
+
+FMetasoundFrontendClassName UMetasoundEditorGraphInputNode::GetClassName() const
+{
+	if (ensure(Input))
+	{
+		return Input->ClassName;
+	}
+
+	return Super::GetClassName();
+}
+
+void UMetasoundEditorGraphInputNode::UpdatePreviewInstance(const Metasound::FVertexKey& InParameterName, TScriptInterface<IAudioParameterInterface>& InParameterInterface) const
+{
+	if (Input)
+	{
+		Input->UpdatePreviewInstance(InParameterName, InParameterInterface);
+	}
+}
+
+FGuid UMetasoundEditorGraphInputNode::GetNodeID() const
+{
+	if (Input)
+	{
+		return Input->NodeID;
+	}
+
+	return Super::GetNodeID();
+}
+
+FLinearColor UMetasoundEditorGraphInputNode::GetNodeTitleColor() const
+{
+	if (const UMetasoundEditorSettings* EditorSettings = GetDefault<UMetasoundEditorSettings>())
+	{
+		return EditorSettings->InputNodeTitleColor;
+	}
+
+	return Super::GetNodeTitleColor();
+}
+
+FSlateIcon UMetasoundEditorGraphInputNode::GetNodeTitleIcon() const
+{
+	static const FName NativeIconName = "MetasoundEditor.Graph.Node.Class.Input";
+	return FSlateIcon("MetaSoundStyle", NativeIconName);
+}
+
+#if WITH_EDITORONLY_DATA
+void UMetasoundEditorGraphInputNode::PostEditUndo()
+{
+	Super::PostEditUndo();
+
+	if (Input)
+	{
+		Input->OnLiteralChanged(false /* bPostTransaction */);
+	}
+}
+#endif // WITH_EDITORONLY_DATA
+
+void UMetasoundEditorGraphInputNode::SetNodeID(FGuid InNodeID)
+{
+	Input->NodeID = InNodeID;
+}
+
+FMetasoundFrontendLiteral UMetasoundEditorGraphInputBool::GetDefault() const 
+{
+	FMetasoundFrontendLiteral Literal;
+	Literal.Set(Default.Value);
+	return Literal;
+}
+
+EMetasoundFrontendLiteralType UMetasoundEditorGraphInputBool::GetLiteralType() const 
+{
+	return EMetasoundFrontendLiteralType::Boolean;
+}
+
+void UMetasoundEditorGraphInputBool::SetFromLiteral(const FMetasoundFrontendLiteral& InLiteral) 
+{
+	Metasound::Editor::InputPrivate::ConvertLiteral<bool>(InLiteral, Default.Value);
+}
+
+void UMetasoundEditorGraphInputBool::UpdatePreviewInstance(const Metasound::FVertexKey& InParameterName, TScriptInterface<IAudioParameterInterface>& InParameterInterface) const
+	{
+		InParameterInterface->SetBool(*InParameterName, Default.Value);
+	}
+
+FMetasoundFrontendLiteral UMetasoundEditorGraphInputBoolArray::GetDefault() const
+	{
+		TArray<bool> BoolArray;
+		Algo::Transform(Default, BoolArray, [](const FMetasoundEditorGraphInputBoolRef& InValue) { return InValue.Value; });
+
+		FMetasoundFrontendLiteral Literal;
+		Literal.Set(BoolArray);
+		return Literal;
+	}
+
+EMetasoundFrontendLiteralType UMetasoundEditorGraphInputBoolArray::GetLiteralType() const
+	{
+		return EMetasoundFrontendLiteralType::BooleanArray;
+	}
+
+void UMetasoundEditorGraphInputBoolArray::SetFromLiteral(const FMetasoundFrontendLiteral& InLiteral)
+	{
+		Metasound::Editor::InputPrivate::ConvertLiteralToArray<bool, FMetasoundEditorGraphInputBoolRef>(InLiteral, Default);
+	}
+
+void UMetasoundEditorGraphInputBoolArray::UpdatePreviewInstance(const Metasound::FVertexKey& InParameterName, TScriptInterface<IAudioParameterInterface>& InParameterInterface) const
+{
+	TArray<bool> BoolArray;
+	Algo::Transform(Default, BoolArray, [](const FMetasoundEditorGraphInputBoolRef& InValue) { return InValue.Value; });
+	InParameterInterface->SetBoolArray(*InParameterName, BoolArray);
+}
+
+FMetasoundFrontendLiteral UMetasoundEditorGraphInputInt::GetDefault() const
+{
+	FMetasoundFrontendLiteral Literal;
+	Literal.Set(Default.Value);
+	return Literal;
+}
+
+EMetasoundFrontendLiteralType UMetasoundEditorGraphInputInt::GetLiteralType() const
+{
+	return EMetasoundFrontendLiteralType::Integer;
+}
+
+void UMetasoundEditorGraphInputInt::SetFromLiteral(const FMetasoundFrontendLiteral& InLiteral)
+{
+	Metasound::Editor::InputPrivate::ConvertLiteral<int32>(InLiteral, Default.Value);
+}
+
+void UMetasoundEditorGraphInputInt::UpdatePreviewInstance(const Metasound::FVertexKey& InParameterName, TScriptInterface<IAudioParameterInterface>& InParameterInterface) const
+{
+	InParameterInterface->SetInt(*InParameterName, Default.Value);
+}
+
+FMetasoundFrontendLiteral UMetasoundEditorGraphInputIntArray::GetDefault() const
+{
+	TArray<int32> IntArray;
+	Algo::Transform(Default, IntArray, [](const FMetasoundEditorGraphInputIntRef& InValue) { return InValue.Value; });
+
+	FMetasoundFrontendLiteral Literal;
+	Literal.Set(IntArray);
+
+	return Literal;
+}
+
+EMetasoundFrontendLiteralType UMetasoundEditorGraphInputIntArray::GetLiteralType() const
+{
+	return EMetasoundFrontendLiteralType::IntegerArray;
+}
+
+void UMetasoundEditorGraphInputIntArray::SetFromLiteral(const FMetasoundFrontendLiteral& InLiteral)
+{
+	Metasound::Editor::InputPrivate::ConvertLiteralToArray<int32, FMetasoundEditorGraphInputIntRef>(InLiteral, Default);
+}
+
+void UMetasoundEditorGraphInputIntArray::UpdatePreviewInstance(const Metasound::FVertexKey& InParameterName, TScriptInterface<IAudioParameterInterface>& InParameterInterface) const
+{
+	TArray<int32> IntArray;
+	Algo::Transform(Default, IntArray, [](const FMetasoundEditorGraphInputIntRef& InValue) { return InValue.Value; });
+	InParameterInterface->SetIntArray(*InParameterName, IntArray);
+}
+
+FMetasoundFrontendLiteral UMetasoundEditorGraphInputFloat::GetDefault() const
+{
+	FMetasoundFrontendLiteral Literal;
+	Literal.Set(Default);
+	return Literal;
+}
+
+EMetasoundFrontendLiteralType UMetasoundEditorGraphInputFloat::GetLiteralType() const
+{
+	return EMetasoundFrontendLiteralType::Float;
+}
+
+void UMetasoundEditorGraphInputFloat::SetFromLiteral(const FMetasoundFrontendLiteral& InLiteral)
+{
+	Metasound::Editor::InputPrivate::ConvertLiteral<float>(InLiteral, Default);
+}
+
+void UMetasoundEditorGraphInputFloat::UpdatePreviewInstance(const Metasound::FVertexKey& InParameterName, TScriptInterface<IAudioParameterInterface>& InParameterInterface) const
+{
+	InParameterInterface->SetFloat(*InParameterName, Default);
+}
+
+FMetasoundFrontendLiteral UMetasoundEditorGraphInputFloatArray::GetDefault() const
+{
+	FMetasoundFrontendLiteral Literal;
+	Literal.Set(Default);
+	return Literal;
+}
+
+EMetasoundFrontendLiteralType UMetasoundEditorGraphInputFloatArray::GetLiteralType() const
+{
+	return EMetasoundFrontendLiteralType::FloatArray;
+}
+
+void UMetasoundEditorGraphInputFloatArray::SetFromLiteral(const FMetasoundFrontendLiteral& InLiteral)
+{
+	Metasound::Editor::InputPrivate::ConvertLiteralToArray<float>(InLiteral, Default);
+}
+
+void UMetasoundEditorGraphInputFloatArray::UpdatePreviewInstance(const Metasound::FVertexKey& InParameterName, TScriptInterface<IAudioParameterInterface>& InParameterInterface) const
+{
+	InParameterInterface->SetFloatArray(*InParameterName, Default);
+}
+
+FMetasoundFrontendLiteral UMetasoundEditorGraphInputString::GetDefault() const
+{
+	FMetasoundFrontendLiteral Literal;
+	Literal.Set(Default);
+	return Literal;
+}
+
+EMetasoundFrontendLiteralType UMetasoundEditorGraphInputString::GetLiteralType() const
+{
+	return EMetasoundFrontendLiteralType::String;
+}
+
+void UMetasoundEditorGraphInputString::SetFromLiteral(const FMetasoundFrontendLiteral& InLiteral)
+{
+	Metasound::Editor::InputPrivate::ConvertLiteral<FString>(InLiteral, Default);
+}
+
+void UMetasoundEditorGraphInputString::UpdatePreviewInstance(const Metasound::FVertexKey& InParameterName, TScriptInterface<IAudioParameterInterface>& InParameterInterface) const
+{
+	InParameterInterface->SetString(*InParameterName, Default);
+}
+
+FMetasoundFrontendLiteral UMetasoundEditorGraphInputStringArray::GetDefault() const
+{
+	FMetasoundFrontendLiteral Literal;
+	Literal.Set(Default);
+	return Literal;
+}
+
+EMetasoundFrontendLiteralType UMetasoundEditorGraphInputStringArray::GetLiteralType() const
+{
+	return EMetasoundFrontendLiteralType::StringArray;
+}
+
+void UMetasoundEditorGraphInputStringArray::SetFromLiteral(const FMetasoundFrontendLiteral& InLiteral)
+{
+	Metasound::Editor::InputPrivate::ConvertLiteralToArray<FString>(InLiteral, Default);
+}
+
+void UMetasoundEditorGraphInputStringArray::UpdatePreviewInstance(const Metasound::FVertexKey& InParameterName, TScriptInterface<IAudioParameterInterface>& InParameterInterface) const
+{
+	InParameterInterface->SetStringArray(*InParameterName, Default);
+}
+
+FMetasoundFrontendLiteral UMetasoundEditorGraphInputObject::GetDefault() const
+{
+	FMetasoundFrontendLiteral Literal;
+	Literal.Set(Default.Object);
+	return Literal;
+}
+
+EMetasoundFrontendLiteralType UMetasoundEditorGraphInputObject::GetLiteralType() const
+{
+	return EMetasoundFrontendLiteralType::UObject;
+}
+
+void UMetasoundEditorGraphInputObject::SetFromLiteral(const FMetasoundFrontendLiteral& InLiteral)
+{
+	ensure(InLiteral.TryGet(Default.Object));
+}
+
+void UMetasoundEditorGraphInputObject::UpdatePreviewInstance(const Metasound::FVertexKey& InParameterName, TScriptInterface<IAudioParameterInterface>& InParameterInterface) const
+{
+	// TODO. We need proxy object here safely.
+}
+
+FMetasoundFrontendLiteral UMetasoundEditorGraphInputObjectArray::GetDefault() const
+{
+	TArray<UObject*> ObjectArray;
+	Algo::Transform(Default, ObjectArray, [](const FMetasoundEditorGraphInputObjectRef& InValue) { return InValue.Object; });
+
+	FMetasoundFrontendLiteral Literal;
+	Literal.Set(ObjectArray);
+	return Literal;
+}
+
+EMetasoundFrontendLiteralType UMetasoundEditorGraphInputObjectArray::GetLiteralType() const
+{
+	return EMetasoundFrontendLiteralType::UObjectArray;
+}
+
+void UMetasoundEditorGraphInputObjectArray::SetFromLiteral(const FMetasoundFrontendLiteral& InLiteral)
+{
+	TArray<UObject*> ObjectArray;
+	ensure(InLiteral.TryGet(ObjectArray));
+
+	Algo::Transform(ObjectArray, Default, [](UObject* InValue) { return FMetasoundEditorGraphInputObjectRef { InValue }; });
+}
+
+void UMetasoundEditorGraphInputObjectArray::UpdatePreviewInstance(const Metasound::FVertexKey& InParameterName, TScriptInterface<IAudioParameterInterface>& InParameterInterface) const
+{
+	TArray<UObject*> ObjectArray;
+	Algo::Transform(Default, ObjectArray, [](const FMetasoundEditorGraphInputObjectRef& InValue) { return InValue.Object; });
+	// TODO. We need proxy object here safely.
+}
diff --git a/Engine/Plugins/Runtime/Metasound/Source/MetasoundEditor/Private/MetasoundEditorGraphNode.cpp b/Engine/Plugins/Runtime/Metasound/Source/MetasoundEditor/Private/MetasoundEditorGraphNode.cpp
--- a/Engine/Plugins/Runtime/Metasound/Source/MetasoundEditor/Private/MetasoundEditorGraphNode.cpp
+++ b/Engine/Plugins/Runtime/Metasound/Source/MetasoundEditor/Private/MetasoundEditorGraphNode.cpp
@@ -15,6 +15,7 @@
 #include "MetasoundEditorGraphSchema.h"
 #include "MetasoundEditorGraphValidation.h"
 #include "MetasoundEditorModule.h"
+#include "MetasoundEditorSettings.h"
 #include "MetasoundFrontend.h"
 #include "MetasoundFrontendRegistries.h"
 #include "MetasoundFrontendSearchEngine.h"
@@ -334,25 +335,24 @@ FString UMetasoundEditorGraphNode::GetPinMetaData(FName InPinName, FName InKey)
 	return Super::GetPinMetaData(InPinName, InKey);
 }
 
-void UMetasoundEditorGraphNode::PostEditImport()
+void UMetasoundEditorGraphNode::PreSave(FObjectPreSaveContext InSaveContext)
 {
-}
+	using namespace Metasound::Editor;
 
-void UMetasoundEditorGraphNode::PostEditChangeProperty(FPropertyChangedEvent& InEvent)
-{
-	const FName PropertyName = InEvent.GetPropertyName();
-	if (PropertyName == GET_MEMBER_NAME_CHECKED(UEdGraphNode, NodePosX)
-		|| PropertyName == GET_MEMBER_NAME_CHECKED(UEdGraphNode, NodePosY))
+	Super::PreSave(InSaveContext);
+
+	// Required to refresh upgrade nodes that are stale when saving.
+	if (TSharedPtr<FEditor> MetaSoundEditor = FGraphBuilder::GetEditorForMetasound(GetMetasoundChecked()))
 	{
-		UpdatePosition();
+		if (TSharedPtr<SGraphEditor> GraphEditor = MetaSoundEditor->GetGraphEditor())
+		{
+			GraphEditor->RefreshNode(*this);
+		}
 	}
-
-	Super::PostEditChangeProperty(InEvent);
 }
 
-void UMetasoundEditorGraphNode::PostEditChangeChainProperty(FPropertyChangedChainEvent& InEvent)
+void UMetasoundEditorGraphNode::PostEditImport()
 {
-	Super::PostEditChangeChainProperty(InEvent);
 }
 
 void UMetasoundEditorGraphNode::PostEditUndo()
@@ -483,6 +483,22 @@ void UMetasoundEditorGraphOutputNode::SetNodeID(FGuid InNodeID)
 	}
 }
 
+FLinearColor UMetasoundEditorGraphOutputNode::GetNodeTitleColor() const
+{
+	if (const UMetasoundEditorSettings* EditorSettings = GetDefault<UMetasoundEditorSettings>())
+	{
+		return EditorSettings->OutputNodeTitleColor;
+	}
+
+	return Super::GetNodeTitleColor();
+}
+
+FSlateIcon UMetasoundEditorGraphOutputNode::GetNodeTitleIcon() const
+{
+	static const FName NativeIconName = "MetasoundEditor.Graph.Node.Class.Output";
+	return FSlateIcon("MetaSoundStyle", NativeIconName);
+}
+
 bool UMetasoundEditorGraphExternalNode::RefreshPinMetadata()
 {
 	using namespace Metasound::Frontend;
@@ -538,44 +554,21 @@ bool UMetasoundEditorGraphExternalNode::RefreshPinMetadata()
 	return bModified;
 }
 
-FMetasoundFrontendVersionNumber UMetasoundEditorGraphExternalNode::GetMajorUpdateAvailable() const
+FMetasoundFrontendVersionNumber UMetasoundEditorGraphExternalNode::FindHighestVersionInRegistry() const
 {
-	using namespace Metasound::Frontend;
-
-	const FMetasoundFrontendClassMetadata& Metadata = GetConstNodeHandle()->GetClassMetadata();
-	const FMetasoundFrontendVersionNumber& CurrentVersion = Metadata.Version;
-	Metasound::FNodeClassName NodeClassName = Metadata.ClassName.ToNodeClassName();
-
-	FMetasoundFrontendClass ClassWithHighestVersion;
-	if (ISearchEngine::Get().FindClassWithHighestVersion(NodeClassName, ClassWithHighestVersion))
-	{
-		if (ClassWithHighestVersion.Metadata.Version.Major > CurrentVersion.Major)
-		{
-			return ClassWithHighestVersion.Metadata.Version;
-		}
-	}
+	return GetConstNodeHandle()->FindHighestVersionInRegistry();
+}
 
-	return FMetasoundFrontendVersionNumber::GetInvalid();
+FMetasoundFrontendVersionNumber UMetasoundEditorGraphExternalNode::FindHighestMinorVersionInRegistry() const
+{
+	return GetConstNodeHandle()->FindHighestMinorVersionInRegistry();
 }
 
-FMetasoundFrontendVersionNumber UMetasoundEditorGraphExternalNode::GetMinorUpdateAvailable() const
+bool UMetasoundEditorGraphExternalNode::CanAutoUpdate() const
 {
 	using namespace Metasound::Frontend;
 
-	const FMetasoundFrontendClassMetadata& Metadata = GetConstNodeHandle()->GetClassMetadata();
-	const FMetasoundFrontendVersionNumber& CurrentVersion = Metadata.Version;
-	Metasound::FNodeClassName NodeClassName = Metadata.ClassName.ToNodeClassName();
-
-	FMetasoundFrontendClass ClassWithHighestVersion;
-	if (ISearchEngine::Get().FindClassWithMajorVersion(NodeClassName, CurrentVersion.Major, ClassWithHighestVersion))
-	{
-		if (ClassWithHighestVersion.Metadata.Version.Minor > CurrentVersion.Minor)
-		{
-			return ClassWithHighestVersion.Metadata.Version;
-		}
-	}
-
-	return FMetasoundFrontendVersionNumber::GetInvalid();
+	return GetConstNodeHandle()->CanAutoUpdate();
 }
 
 UMetasoundEditorGraphExternalNode* UMetasoundEditorGraphExternalNode::UpdateToVersion(const FMetasoundFrontendVersionNumber& InNewVersion, bool bInPropagateErrorMessages)
@@ -583,131 +576,67 @@ UMetasoundEditorGraphExternalNode* UMetasoundEditorGraphExternalNode::UpdateToVe
 	using namespace Metasound::Editor;
 	using namespace Metasound::Frontend;
 
-	const FMetasoundFrontendClassMetadata& Metadata = GetConstNodeHandle()->GetClassMetadata();
-	const TArray<FMetasoundFrontendClass> SortedVersions = ISearchEngine::Get().FindClassesWithName(Metadata.ClassName.ToNodeClassName(), true /* bInSortByVersion */);
+	FNodeHandle InitNodeHandle = GetNodeHandle();
+	const FGuid InitNodeID = GetNodeID();
+	const FMetasoundFrontendVersionNumber InitVersion = InitNodeHandle->GetClassMetadata().GetVersion();
 
-	auto IsClassOfNewVersion = [InNewVersion](const FMetasoundFrontendClass& SortedClass)
-	{
-		return SortedClass.Metadata.Version == InNewVersion;
-	};
+	FNodeHandle ReplacementNode = InitNodeHandle->ReplaceWithVersion(InNewVersion);
+	const FGuid ReplacementNodeID = ReplacementNode->GetID();
 
-	if (!ensure(SortedVersions.ContainsByPredicate(IsClassOfNewVersion)))
+	if (!ensure(InitNodeID != ReplacementNodeID))
 	{
 		return this;
 	}
 
-	struct FConnectionInfo
+	const FMetasoundFrontendNodeStyle& Style = ReplacementNode->GetNodeStyle();
+	FVector2D NodeLocation = FVector2D::ZeroVector;
+	const FVector2D* StyleLocation = Style.Display.Locations.Find(NodeGuid);
+	if (ensure(StyleLocation))
 	{
-		TArray<UEdGraphPin*> PinsLinkedTo;
-		FName DataType;
-		FString DefaultValue;
-	};
+		NodeLocation = *StyleLocation;
+	}
 
-	// Cache pin connections by name to try so they can be
-	// hooked back up after swapping to the new class version.
-	TMap<FName, FConnectionInfo> InputConnections;
-	IteratePins([Connections = &InputConnections](UEdGraphPin& Pin, int32 /* Index */)
-	{
-		FConstInputHandle InputHandle = FGraphBuilder::GetConstInputHandleFromPin(&Pin);
-		Connections->Add(Pin.GetFName(), FConnectionInfo
-		{
-			Pin.LinkedTo,
-			InputHandle->GetDataType(),
-			Pin.DefaultValue
-		});
-	}, EGPD_Input);
+	UObject& MetaSound = GetMetasoundChecked();
+	UMetasoundEditorGraphExternalNode* ReplacementEdNode = FGraphBuilder::AddExternalNode(MetaSound, ReplacementNode, NodeLocation, false /* bInSelectNewNode */);
 
-	TMap<FName, FConnectionInfo> OutputConnections;
-	IteratePins([Connections = &OutputConnections](UEdGraphPin& Pin, int32 /* Index */)
+	if (ensure(ReplacementEdNode))
 	{
-		FConstOutputHandle OutputHandle = FGraphBuilder::GetConstOutputHandleFromPin(&Pin);
-		Connections->Add(Pin.GetFName(), FConnectionInfo
-		{
-			Pin.LinkedTo,
-			OutputHandle->GetDataType()
-		});
-	}, EGPD_Output);
+		GetGraph()->RemoveNode(this);
+	}
 
-	FNodeHandle NodeHandle = GetNodeHandle();
-	FMetasoundFrontendVersionNumber InitVersion = NodeHandle->GetClassMetadata().Version;
+	MetaSound.Modify();
 
-	// Only update the version on the Metadata. Name and version are the only fields that pertain
-	// to building valid class data (via the GetRegistryKey look-up).  All other Metadata is
-	// superfluous for this use case.
-	FMetasoundFrontendClassMetadata NewMetadata = NodeHandle->GetClassMetadata();
-	NewMetadata.Version = InNewVersion;
+	// TODO: Pare down to SynchronizeNode/NodeConnections call to avoid having to synchronize the entire graph for all version calls
+	FGraphBuilder::SynchronizeGraph(MetaSound);
 
-	const FVector2D NodePos(NodePosX, NodePosY);
-	if (!ensureAlways(FGraphBuilder::DeleteNode(*this)))
+	if (InitVersion != InNewVersion)
 	{
-		return this;
+		NodeUpgradeMessage = FText::Format(LOCTEXT("NodeVersionUpgradeMessageFormat", "Class interface updated: {0} to {1}"),
+			FText::FromString(*InitVersion.ToString()),
+			FText::FromString(*InNewVersion.ToString())
+		);
 	}
-
-	Modify();
-	UObject& Metasound = GetMetasoundChecked();
-	UMetasoundEditorGraphExternalNode* ReplacementNode = FGraphBuilder::AddExternalNode(Metasound, NewMetadata, NodePos);
-	if (!ensureAlways(ReplacementNode))
+	else
 	{
-		return this;
+		NodeUpgradeMessage = FText::Format(LOCTEXT("NodeVersionRefreshMessageFormat", "Class interface updated: detected interface change for {0}"),
+			FText::FromString(*InitVersion.ToString())
+		);
 	}
 
-	ReplacementNode->Modify();
-	ReplacementNode->IteratePins([&](UEdGraphPin& Pin, int32 /* Index */)
-	{
-		if (FConnectionInfo* ConnectionInfo = InputConnections.Find(Pin.GetFName()))
-		{
-			FInputHandle InputHandle = FGraphBuilder::GetInputHandleFromPin(&Pin);
-			if (ConnectionInfo->DataType == InputHandle->GetDataType())
-			{
-				Pin.DefaultValue = ConnectionInfo->DefaultValue;
-				FMetasoundFrontendLiteral LiteralValue;
-				if (FGraphBuilder::GetPinLiteral(Pin, LiteralValue))
-				{
-					InputHandle->SetLiteral(LiteralValue);
-				}
-
-				for (UEdGraphPin* OutputPin : ConnectionInfo->PinsLinkedTo)
-				{
-					FGraphBuilder::ConnectNodes(Pin, *OutputPin, true /* bInConnectEdPins */);
-				}
-			}
-		}
-	}, EGPD_Input);
-
-	ReplacementNode->IteratePins([&](UEdGraphPin& Pin, int32 /* Index */)
-	{
-		if (FConnectionInfo* ConnectionInfo = OutputConnections.Find(Pin.GetFName()))
-		{
-			FConstOutputHandle OutputHandle = FGraphBuilder::GetConstOutputHandleFromPin(&Pin);
-			if (ConnectionInfo->DataType == OutputHandle->GetDataType())
-			{
-				for (UEdGraphPin* InputPin : ConnectionInfo->PinsLinkedTo)
-				{
-					FGraphBuilder::ConnectNodes(*InputPin, Pin, true /* bInConnectEdPins */);
-				}
-			}
-		}
-	}, EGPD_Output);
-
-	ReplacementNode->NodeUpgradeMessage = FText::Format(LOCTEXT("NodeVersionUpgradeMessageFormat", "Class upgraded from {0} to {1}"),
-		FText::FromString(*InitVersion.ToString()),
-		FText::FromString(*InNewVersion.ToString())
-	);
-
-	if (ReplacementNode != this)
+	if (InitNodeID != ReplacementNodeID)
 	{
 		if (bInPropagateErrorMessages)
 		{
-			ReplacementNode->bHasCompilerMessage = bHasCompilerMessage;
-			ReplacementNode->ErrorMsg = ErrorMsg;
-			ReplacementNode->ErrorType = ErrorType;
+			ReplacementEdNode->bHasCompilerMessage = bHasCompilerMessage;
+			ReplacementEdNode->ErrorMsg = ErrorMsg;
+			ReplacementEdNode->ErrorType = ErrorType;
 		}
 	}
 
-	return ReplacementNode;
+	return ReplacementEdNode;
 }
 
-bool UMetasoundEditorGraphExternalNode::Validate(Metasound::Editor::FGraphNodeValidationResult& OutResult)
+bool UMetasoundEditorGraphExternalNode::Validate(Metasound::Editor::FGraphNodeValidationResult& OutResult, bool bClearUpgradeMessage)
 {
 	using namespace Metasound::Editor;
 	using namespace Metasound::Frontend;
@@ -716,6 +645,12 @@ bool UMetasoundEditorGraphExternalNode::Validate(Metasound::Editor::FGraphNodeVa
 
 	FNodeHandle NodeHandle = GetNodeHandle();
 
+	if (ErrorType != EMessageSeverity::Info)
+	{
+		ErrorType = EMessageSeverity::Info;
+		OutResult.bIsDirty = true;
+	}
+
 	if (bHasCompilerMessage)
 	{
 		bHasCompilerMessage = false;
@@ -723,33 +658,31 @@ bool UMetasoundEditorGraphExternalNode::Validate(Metasound::Editor::FGraphNodeVa
 		OutResult.bIsDirty = true;
 	}
 
-	if (!ensure(NodeHandle->IsValid()))
+	if (bClearUpgradeMessage && !NodeUpgradeMessage.IsEmpty())
 	{
-		GraphNodePrivate::SetGraphNodeMessage(*this, EMessageSeverity::Error, TEXT("Node is invalid (Frontend node not found)"));
-		OutResult.bIsInvalid = true;
+		NodeUpgradeMessage = FText::GetEmpty();
 		OutResult.bIsDirty = true;
-		return true;
 	}
 
-	const FMetasoundFrontendClassMetadata& Metadata = NodeHandle->GetClassMetadata();
-	if (!ensure(Metadata.Type == EMetasoundFrontendClassType::External))
+	if (!NodeHandle->IsValid())
 	{
-		GraphNodePrivate::SetGraphNodeMessage(*this, EMessageSeverity::Error, TEXT("Node is of type 'External' but Frontend counterpart is not natively-defined."));
-		OutResult.bIsDirty = true;
+		GraphNodePrivate::SetGraphNodeMessage(*this, EMessageSeverity::Error, TEXT("Class definition missing (If defined in asset, corrisponding graph definition may be invalid)"));
 		OutResult.bIsInvalid = true;
-		return true;
+		OutResult.bIsDirty = true;
+		return !OutResult.bIsInvalid;
 	}
 
-	const Metasound::FNodeClassName NodeClassName = Metadata.ClassName.ToNodeClassName();
+	const FMetasoundFrontendClassMetadata& Metadata = NodeHandle->GetClassMetadata();
+	const Metasound::FNodeClassName NodeClassName = Metadata.GetClassName().ToNodeClassName();
 	const TArray<FMetasoundFrontendClass> SortedClasses = ISearchEngine::Get().FindClassesWithName(NodeClassName, true /* bSortByVersion */);
 
 	if (SortedClasses.IsEmpty())
 	{
 		OutResult.MissingClass = NodeClassName.GetFullName();
-		const FString NewErrorMsg = FString::Format(TEXT("Missing Class '{0}': {1}"),
+		const FString NewErrorMsg = FString::Format(TEXT("Class definition '{0}' missing: {1}"),
 			{
-				*Metadata.ClassName.ToString(),
-				*Metadata.PromptIfMissing.ToString()
+				*Metadata.GetClassName().ToString(),
+				*Metadata.GetPromptIfMissing().ToString()
 			});
 
 		for (UEdGraphPin* Pin : Pins)
@@ -760,7 +693,7 @@ bool UMetasoundEditorGraphExternalNode::Validate(Metasound::Editor::FGraphNodeVa
 		GraphNodePrivate::SetGraphNodeMessage(*this, EMessageSeverity::Error, NewErrorMsg);
 		OutResult.bIsDirty = true;
 		OutResult.bIsInvalid = true;
-		return true;
+		return !OutResult.bIsInvalid;
 	}
 
 	for (UEdGraphPin* Pin : Pins)
@@ -772,50 +705,82 @@ bool UMetasoundEditorGraphExternalNode::Validate(Metasound::Editor::FGraphNodeVa
 		}
 	}
 
-	if (SortedClasses[0].Metadata.Version.Major > Metadata.Version.Major)
+	const FMetasoundFrontendVersionNumber& CurrentVersion = Metadata.GetVersion();
+	if (SortedClasses[0].Metadata.GetVersion() > CurrentVersion)
 	{
-		const FMetasoundFrontendVersionNumber& CurrentVersion = NodeHandle->GetClassMetadata().Version;
-
 		FMetasoundFrontendClass HighestMinorVersionClass;
 		FString NodeMsg;
 		EMessageSeverity::Type Severity;
 
-		// Only compare major version as minor can be incorrect and graph
-		// will still build and divert to using the new upgraded version.
-		if (ISearchEngine::Get().FindClassWithMajorVersion(NodeClassName, CurrentVersion.Major, HighestMinorVersionClass))
+		auto ClassVersionExists = [InCurrentVersion = &CurrentVersion](const FMetasoundFrontendClass& AvailableClass)
+		{
+			return AvailableClass.Metadata.GetVersion() == *InCurrentVersion;
+		};
+		if (SortedClasses.ContainsByPredicate(ClassVersionExists))
 		{
 			NodeMsg = FString::Format(TEXT("Node class '{0}' out-of-date: Eligible for upgrade to {1}"),
 			{
-				*Metadata.ClassName.ToString(),
-				*SortedClasses[0].Metadata.Version.ToString()
+				*Metadata.GetClassName().ToString(),
+				*SortedClasses[0].Metadata.GetVersion().ToString()
 			});
 			Severity = EMessageSeverity::Warning;
 		}
 		else
 		{
-			NodeMsg = FString::Format(TEXT("Node class '{0} ({1})' is missing.  Eligible for upgrade to {2}"),
+			NodeMsg = FString::Format(TEXT("Node class '{0} ({1})' is missing.  Highest version ({2}) found."),
 			{
-				*Metadata.ClassName.ToString(),
-				*Metadata.Version.ToString(),
-				*SortedClasses[0].Metadata.Version.ToString()
+				*Metadata.GetClassName().ToString(),
+				*Metadata.GetVersion().ToString(),
+				*SortedClasses[0].Metadata.GetVersion().ToString()
 			});
 			Severity = EMessageSeverity::Error;
+			OutResult.bIsInvalid = true;
 		}
 
 		GraphNodePrivate::SetGraphNodeMessage(*this, Severity, NodeMsg);
 		OutResult.bIsDirty = true;
+		return !OutResult.bIsInvalid;
 	}
 
-	auto IsMinorUpgradeVersion = [&](const FMetasoundFrontendVersionNumber& Version)
+	OutResult.bIsAutoUpdateEligible = CanAutoUpdate();
+
+	const FNodeRegistryKey RegistryKey = NodeRegistryKey::CreateKey(Metadata);
+	bIsClassNative = FMetasoundFrontendRegistryContainer::Get()->IsNodeNative(RegistryKey);
+
+	return !OutResult.bIsInvalid;
+}
+
+void UMetasoundEditorGraphExternalNode::PostLoad()
+{
+	Super::PostLoad();
+}
+
+FLinearColor UMetasoundEditorGraphExternalNode::GetNodeTitleColor() const
+{
+	if (const UMetasoundEditorSettings* EditorSettings = GetDefault<UMetasoundEditorSettings>())
 	{
-		if (Version.Major != Metadata.Version.Major)
+		if (bIsClassNative)
 		{
-			return false;
+			return EditorSettings->NativeNodeTitleColor;
 		}
 
-		return Version.Minor > Metadata.Version.Minor;
-	};
+		return EditorSettings->AssetReferenceNodeTitleColor;
+	}
+
+	return Super::GetNodeTitleColor();
+}
 
-	return false;
+FSlateIcon UMetasoundEditorGraphExternalNode::GetNodeTitleIcon() const
+{
+	if (bIsClassNative)
+	{
+		static const FName NativeIconName = "MetasoundEditor.Graph.Node.Class.Native";
+		return FSlateIcon("MetaSoundStyle", NativeIconName);
+	}
+	else
+	{
+		static const FName GraphRefIconName = "MetasoundEditor.Graph.Node.Class.Graph";
+		return FSlateIcon("MetasoundStyle", GraphRefIconName);
+	}
 }
 #undef LOCTEXT_NAMESPACE
