First pass on MetaSound Composition/Preset Support

diff --git a/Engine/Plugins/Runtime/Metasound/Source/MetasoundEditor/Private/MetasoundEditorGraphSchema.cpp b/Engine/Plugins/Runtime/Metasound/Source/MetasoundEditor/Private/MetasoundEditorGraphSchema.cpp
--- a/Engine/Plugins/Runtime/Metasound/Source/MetasoundEditor/Private/MetasoundEditorGraphSchema.cpp
+++ b/Engine/Plugins/Runtime/Metasound/Source/MetasoundEditor/Private/MetasoundEditorGraphSchema.cpp
@@ -230,6 +230,8 @@ UEdGraphNode* FMetasoundGraphSchemaAction_NewInput::PerformAction(UEdGraph* Pare
 		NewGraphNode->Modify();
 		UEdGraphNode* EdGraphNode = CastChecked<UEdGraphNode>(NewGraphNode);
 		SchemaPrivate::TryConnectNewNodeToPin(*EdGraphNode, FromPin);
+
+		FGraphBuilder::RegisterGraphWithFrontend(ParentMetasound);
 		return EdGraphNode;
 	}
 
@@ -284,6 +286,7 @@ UEdGraphNode* FMetasoundGraphSchemaAction_PromoteToInput::PerformAction(UEdGraph
 						MetasoundEditor->OnInputNameChanged(NewGraphNode->GetNodeID());
 					}
 
+					FGraphBuilder::RegisterGraphWithFrontend(ParentMetasound);
 					return EdGraphNode;
 				}
 			}
@@ -326,6 +329,7 @@ UEdGraphNode* FMetasoundGraphSchemaAction_NewOutput::PerformAction(UEdGraph* Par
 	if (UMetasoundEditorGraphOutputNode* NewGraphNode = FGraphBuilder::AddOutputNode(ParentMetasound, NodeHandle, Location, bSelectNewNode))
 	{
 		SchemaPrivate::TryConnectNewNodeToPin(*NewGraphNode, FromPin);
+		FGraphBuilder::RegisterGraphWithFrontend(ParentMetasound);
 		return NewGraphNode;
 	}
 
@@ -437,6 +441,7 @@ void UMetasoundEditorGraphSchema::GetGraphContextActions(FGraphContextMenuBuilde
 	using namespace Metasound::Frontend;
 
 	FActionClassFilters ClassFilters;
+	FConstGraphHandle GraphHandle = IGraphController::GetInvalidHandle();
 	if (const UEdGraphPin* FromPin = ContextMenuBuilder.FromPin)
 	{
 		if (FromPin->Direction == EGPD_Input)
@@ -448,7 +453,7 @@ void UMetasoundEditorGraphSchema::GetGraphContextActions(FGraphContextMenuBuilde
 			};
 
 			// Show only input nodes as output nodes can only connected if FromPin is input
-			FConstGraphHandle GraphHandle = InputHandle->GetOwningNode()->GetOwningGraph();
+			GraphHandle = InputHandle->GetOwningNode()->GetOwningGraph();
 			GetDataTypeInputNodeActions(ContextMenuBuilder, GraphHandle, [InputHandle](FConstNodeHandle NodeHandle)
 			{
 				bool bHasOutputOfType = false;
@@ -477,7 +482,7 @@ void UMetasoundEditorGraphSchema::GetGraphContextActions(FGraphContextMenuBuilde
 			};
 
 			// Show only output nodes as input nodes can only connected if FromPin is output
-			FConstGraphHandle GraphHandle = OutputHandle->GetOwningNode()->GetOwningGraph();
+			GraphHandle = OutputHandle->GetOwningNode()->GetOwningGraph();
 			GetDataTypeOutputNodeActions(ContextMenuBuilder, GraphHandle, [OutputHandle](FConstNodeHandle NodeHandle)
 			{
 				bool bHasInputOfType = false;
@@ -503,14 +508,14 @@ void UMetasoundEditorGraphSchema::GetGraphContextActions(FGraphContextMenuBuilde
 		{
 			FMetasoundAssetBase* MetasoundAsset = IMetasoundUObjectRegistry::Get().GetObjectAsAssetBase(Metasound);
 			check(MetasoundAsset);
-			FGraphHandle GraphHandle = MetasoundAsset->GetRootGraphHandle();
+			GraphHandle = MetasoundAsset->GetRootGraphHandle();
 
 			GetDataTypeInputNodeActions(ContextMenuBuilder, GraphHandle);
 			GetDataTypeOutputNodeActions(ContextMenuBuilder, GraphHandle);
 		}
 	}
 
-	GetFunctionActions(ContextMenuBuilder, ClassFilters);
+	GetFunctionActions(ContextMenuBuilder, ClassFilters, true /* bShowSelectedActions */, GraphHandle);
 	GetConversionActions(ContextMenuBuilder, ClassFilters);
 }
 
@@ -541,8 +546,9 @@ void UMetasoundEditorGraphSchema::GetContextMenuActions(class UToolMenu* Menu, c
 		// and node registry is reporting a major update is available.
 		if (const UMetasoundEditorGraphExternalNode* ExternalNode = Cast<UMetasoundEditorGraphExternalNode>(Context->Node))
 		{
-			FMetasoundFrontendVersionNumber MajorUpdateVersion = ExternalNode->GetMajorUpdateAvailable();
-			if (MajorUpdateVersion.IsValid())
+			FMetasoundFrontendVersionNumber MajorUpdateVersion = ExternalNode->FindHighestVersionInRegistry();
+			Metasound::Frontend::FConstNodeHandle NodeHandle = ExternalNode->GetConstNodeHandle();
+			if (MajorUpdateVersion.IsValid() && MajorUpdateVersion > NodeHandle->GetClassMetadata().GetVersion())
 			{
 				Section.AddMenuEntry(FEditorCommands::Get().UpdateNodes);
 			}
@@ -713,7 +719,7 @@ FText UMetasoundEditorGraphSchema::GetPinDisplayName(const UEdGraphPin* Pin) con
 
 	UMetasoundEditorGraphNode* Node = CastChecked<UMetasoundEditorGraphNode>(Pin->GetOwningNode());
 	FNodeHandle NodeHandle = Node->GetNodeHandle();
-	const EMetasoundFrontendClassType ClassType = NodeHandle->GetClassMetadata().Type;
+	const EMetasoundFrontendClassType ClassType = NodeHandle->GetClassMetadata().GetType();
 
 	switch (ClassType)
 	{
@@ -745,9 +751,10 @@ FText UMetasoundEditorGraphSchema::GetPinDisplayName(const UEdGraphPin* Pin) con
 
 		case EMetasoundFrontendClassType::External:
 		{
+			const FString PinName = Pin->GetName();
 			if (Pin->Direction == EGPD_Input)
 			{
-				TArray<FInputHandle> InputHandles = NodeHandle->GetInputsWithVertexName(Pin->GetName());
+				TArray<FInputHandle> InputHandles = NodeHandle->GetInputsWithVertexName(PinName);
 				if (ensure(InputHandles.Num() > 0))
 				{
 					return InputHandles[0]->GetDisplayName();
@@ -755,7 +762,7 @@ FText UMetasoundEditorGraphSchema::GetPinDisplayName(const UEdGraphPin* Pin) con
 			}
 			else
 			{
-				TArray<FOutputHandle> OutputHandles = NodeHandle->GetOutputsWithVertexName(Pin->GetName());
+				TArray<FOutputHandle> OutputHandles = NodeHandle->GetOutputsWithVertexName(PinName);
 				if (ensure(OutputHandles.Num() > 0))
 				{
 					return OutputHandles[0]->GetDisplayName();
@@ -821,70 +828,72 @@ void UMetasoundEditorGraphSchema::BreakPinLinks(UEdGraphPin& TargetPin, bool bSe
 
 void UMetasoundEditorGraphSchema::GetAssetsGraphHoverMessage(const TArray<FAssetData>& Assets, const UEdGraph* HoverGraph, FString& OutTooltipText, bool& OutOkIcon) const
 {
-	// Enable Composition
-	//using namespace Metasound::Editor;
-
-	//OutOkIcon = true;
-
-	//if (!HoverGraph)
-	//{
-	//	OutOkIcon = false;
-	//	return;
-	//}
-
-	//OutTooltipText = TEXT("Add MetaSound reference to Graph.");
-
-	//IMetasoundEditorModule& EditorModule = FModuleManager::GetModuleChecked<IMetasoundEditorModule>("MetaSoundEditor");
-	//for (const FAssetData& Data : Assets)
-	//{
-	//	if (!EditorModule.IsMetaSoundAssetClass(Data.GetClass()->GetFName()))
-	//	{
-	//		OutOkIcon = false;
-	//		OutTooltipText = TEXT("Asset(s) must all be MetaSounds.");
-	//		break;
-	//	}
-	//}
+	using namespace Metasound::Editor;
+
+	OutOkIcon = true;
+
+	if (!HoverGraph)
+	{
+		OutOkIcon = false;
+		return;
+	}
+
+	OutTooltipText = TEXT("Add MetaSound reference to Graph.");
+
+	IMetasoundEditorModule& EditorModule = FModuleManager::GetModuleChecked<IMetasoundEditorModule>("MetaSoundEditor");
+	for (const FAssetData& Data : Assets)
+	{
+		if (!EditorModule.IsMetaSoundAssetClass(Data.GetClass()->GetFName()))
+		{
+			OutOkIcon = false;
+			OutTooltipText = TEXT("Asset(s) must all be MetaSounds.");
+			break;
+		}
+	}
 }
 
 void UMetasoundEditorGraphSchema::DroppedAssetsOnGraph(const TArray<FAssetData>& Assets, const FVector2D& GraphPosition, UEdGraph* Graph) const
 {
-	// Enable Composition
-	//using namespace Metasound;
-	//using namespace Metasound::Frontend;
-
-	//const FScopedTransaction Transaction(LOCTEXT("DropMetaSoundOnGraph", "Drop MetaSound On Graph"));
-
-	//UMetasoundEditorGraph* MetaSoundGraph = CastChecked<UMetasoundEditorGraph>(Graph);
-	//UObject& MetaSound = MetaSoundGraph->GetMetasoundChecked();
-	//MetaSound.Modify();
-	//Graph->Modify();
-
-	//for (const FAssetData& DroppedAsset : Assets)
-	//{
-	//	if (UObject* DroppedMetaSound = DroppedAsset.GetAsset())
-	//	{
-	//		FMetasoundAssetBase* DroppedMetaSoundAsset = IMetasoundUObjectRegistry::Get().GetObjectAsAssetBase(&DroppedMetaSound);
-	//		check(DroppedMetaSoundAsset);
-
-	//		const FNodeClassInfo ClassInfo = DroppedMetaSoundAsset->GetAssetClassInfo();
-	//		const FNodeRegistryKey RegistryKey = FMetasoundFrontendRegistryContainer::Get()->GetRegistryKey(ClassInfo);
-	//		if (ensure(IsValidNodeRegistryKey(RegistryKey)))
-	//		{
-	//			if (!FMetasoundFrontendRegistryContainer::Get()->IsNodeRegistered(RegistryKey))
-	//			{
-	//				DroppedMetaSoundAsset->RegisterGraphWithFrontend();
-	//			}
-	//		}
-
-	//		const FMetasoundFrontendClassName ClassName = MetaSoundAsset->GetAssetClassInfo().ClassName;
-
-	//		FMetasoundFrontendClass Class;
-	//		if (ensure(ISearchEngine::Get().FindClassWithHighestVersion(ClassName.ToNodeClassName(), Class)))
-	//		{
-	//			Metasound::Editor::FGraphBuilder::AddExternalNode(MetaSound, Class.Metadata, GraphPosition);
-	//		}
-	//	}
-	//}
+	using namespace Metasound;
+	using namespace Metasound::Frontend;
+
+	const FScopedTransaction Transaction(LOCTEXT("DropMetaSoundOnGraph", "Drop MetaSound On Graph"));
+
+	UMetasoundEditorGraph* MetaSoundGraph = CastChecked<UMetasoundEditorGraph>(Graph);
+	UObject& MetaSound = MetaSoundGraph->GetMetasoundChecked();
+	MetaSound.Modify();
+	Graph->Modify();
+
+	for (const FAssetData& DroppedAsset : Assets)
+	{
+		if (UObject* DroppedObject = DroppedAsset.GetAsset())
+		{
+			FMetasoundAssetBase* DroppedMetaSoundAsset = IMetasoundUObjectRegistry::Get().GetObjectAsAssetBase(DroppedObject);
+			if (!DroppedMetaSoundAsset)
+			{
+				continue;
+			}
+
+			// This may not be necessary as dropping an asset on the graph may load it, thus triggering the registration from the MetaSoundAssetSubsystem.
+			const FNodeClassInfo ClassInfo = DroppedMetaSoundAsset->GetAssetClassInfo();
+			const FNodeRegistryKey RegistryKey = NodeRegistryKey::CreateKey(ClassInfo);
+			if (ensure(IsValidNodeRegistryKey(RegistryKey)))
+			{
+				if (!FMetasoundFrontendRegistryContainer::Get()->IsNodeRegistered(RegistryKey))
+				{
+					DroppedMetaSoundAsset->RegisterGraphWithFrontend();
+				}
+			}
+
+			const FMetasoundFrontendClassName ClassName = DroppedMetaSoundAsset->GetAssetClassInfo().ClassName;
+
+			FMetasoundFrontendClass Class;
+			if (ensure(ISearchEngine::Get().FindClassWithHighestVersion(ClassName.ToNodeClassName(), Class)))
+			{
+				Metasound::Editor::FGraphBuilder::AddExternalNode(MetaSound, Class.Metadata, GraphPosition);
+			}
+		}
+	}
 }
 
 void UMetasoundEditorGraphSchema::DroppedAssetsOnNode(const TArray<FAssetData>& Assets, const FVector2D& GraphPosition, UEdGraphNode* Node) const
@@ -912,21 +921,21 @@ void UMetasoundEditorGraphSchema::GetConversionActions(FGraphActionMenuBuilder&
 		}
 
 		FMetasoundFrontendClassMetadata Metadata = FrontendClass.Metadata;
-		const FText Tooltip = Metadata.Author.IsEmpty()
-			? Metadata.Description
-			: FText::Format(LOCTEXT("MetasoundTooltipAuthorFormat", "{0}\nAuthor: {1}"), Metadata.Description, Metadata.Author);
+		const FText Tooltip = Metadata.GetAuthor().IsEmpty()
+			? Metadata.GetDescription()
+			: FText::Format(LOCTEXT("MetasoundTooltipAuthorFormat", "{0}\nAuthor: {1}"), Metadata.GetDescription(), Metadata.GetAuthor());
 
-		if (!Metadata.CategoryHierarchy.IsEmpty() && !Metadata.CategoryHierarchy[0].CompareTo(Editor::FGraphBuilder::ConvertMenuName))
+		if (!Metadata.GetCategoryHierarchy().IsEmpty() && !Metadata.GetCategoryHierarchy()[0].CompareTo(Editor::FGraphBuilder::ConvertMenuName))
 		{
 			TSharedPtr<FMetasoundGraphSchemaAction_NewNode> NewNodeAction = MakeShared<FMetasoundGraphSchemaAction_NewNode>
 			(
 				Editor::FGraphBuilder::ConvertMenuName,
-				Metadata.DisplayName,
+				Metadata.GetDisplayName(),
 				Tooltip,
 				0
 			);
 
-			Metadata.Type = EMetasoundFrontendClassType::External;
+			Metadata.SetType(EMetasoundFrontendClassType::External);
 			NewNodeAction->ClassMetadata = Metadata;
 			ActionMenuBuilder.AddAction(NewNodeAction);
 
@@ -990,13 +999,23 @@ void UMetasoundEditorGraphSchema::GetDataTypeOutputNodeActions(FGraphContextMenu
 	SchemaPrivate::GetDataTypeActions<FMetasoundGraphSchemaAction_NewOutput>(ActionQuery);
 }
 
-void UMetasoundEditorGraphSchema::GetFunctionActions(FGraphActionMenuBuilder& ActionMenuBuilder, Metasound::Editor::FActionClassFilters InFilters, bool bShowSelectedActions) const
+void UMetasoundEditorGraphSchema::GetFunctionActions(FGraphActionMenuBuilder& ActionMenuBuilder, Metasound::Editor::FActionClassFilters InFilters, bool bShowSelectedActions, Metasound::Frontend::FConstGraphHandle InGraphHandle) const
 {
 	using namespace Metasound::Editor;
 	using namespace Metasound::Frontend;
 
 	const FText MenuJoinFormat = LOCTEXT("MetasoundActionsFormatSubCategory", "{0}|{1}");
 
+	const UMetaSoundAssetSubsystem& AssetSubsystem = UMetaSoundAssetSubsystem::Get();
+	FMetasoundAssetBase* ParentAsset = nullptr;
+	if (InGraphHandle->IsValid())
+	{
+		FMetasoundFrontendClassMetadata AssetMetadata = InGraphHandle->GetGraphMetadata();
+		AssetMetadata.SetType(EMetasoundFrontendClassType::External);
+		const FNodeRegistryKey RegistryKey = NodeRegistryKey::CreateKey(AssetMetadata);
+		ParentAsset = AssetSubsystem.FindAssetFromKey(RegistryKey);
+	}
+
 	const bool bIncludeDeprecated = static_cast<bool>(EnableDeprecatedMetaSoundNodeClassCreationCVar);
 	const TArray<FMetasoundFrontendClass> FrontendClasses = ISearchEngine::Get().FindAllClasses(bIncludeDeprecated);
 	for (const FMetasoundFrontendClass& FrontendClass : FrontendClasses)
@@ -1011,23 +1030,35 @@ void UMetasoundEditorGraphSchema::GetFunctionActions(FGraphActionMenuBuilder& Ac
 			continue;
 		}
 
+		if (ParentAsset && FrontendClass.Metadata.GetType() == EMetasoundFrontendClassType::External)
+		{
+			const FNodeRegistryKey RegistryKey = NodeRegistryKey::CreateKey(FrontendClass.Metadata);
+			if (const FSoftObjectPath* Path = AssetSubsystem.FindObjectPathFromKey(RegistryKey))
+			{
+				if (ParentAsset->AddingReferenceCausesLoop(*Path, AssetSubsystem))
+				{
+					continue;
+				}
+			}
+		}
+
 		FMetasoundFrontendClassMetadata Metadata = FrontendClass.Metadata;
-		const FText Tooltip = Metadata.Author.IsEmpty()
-			? Metadata.Description
-			: FText::Format(LOCTEXT("MetasoundTooltipAuthorFormat", "{0}\nAuthor: {1}"), Metadata.Description, Metadata.Author);
+		const FText Tooltip = Metadata.GetAuthor().IsEmpty()
+			? Metadata.GetDescription()
+			: FText::Format(LOCTEXT("MetasoundTooltipAuthorFormat", "{0}\nAuthor: {1}"), Metadata.GetDescription(), Metadata.GetAuthor());
 
-		if (Metadata.CategoryHierarchy.IsEmpty() || Metadata.CategoryHierarchy[0].CompareTo(FGraphBuilder::ConvertMenuName))
+		if (Metadata.GetCategoryHierarchy().IsEmpty() || Metadata.GetCategoryHierarchy()[0].CompareTo(FGraphBuilder::ConvertMenuName))
 		{
-			const FText CategoriesText = FText::Join(LOCTEXT("MetasoundActionsCategoryDelim", "|"), Metadata.CategoryHierarchy);
+			const FText CategoriesText = FText::Join(LOCTEXT("MetasoundActionsCategoryDelim", "|"), Metadata.GetCategoryHierarchy());
 			TSharedPtr<FMetasoundGraphSchemaAction_NewNode> NewNodeAction = MakeShared<FMetasoundGraphSchemaAction_NewNode>
 			(
 				FText::Format(MenuJoinFormat, FGraphBuilder::FunctionMenuName, CategoriesText),
-				Metadata.DisplayName,
+				Metadata.GetDisplayName(),
 				Tooltip,
 				0
 			);
 
-			Metadata.Type = EMetasoundFrontendClassType::External;
+			Metadata.SetType(EMetasoundFrontendClassType::External);
 			NewNodeAction->ClassMetadata = Metadata;
 			ActionMenuBuilder.AddAction(NewNodeAction);
 		}
diff --git a/Engine/Plugins/Runtime/Metasound/Source/MetasoundEditor/Private/MetasoundEditorModule.cpp b/Engine/Plugins/Runtime/Metasound/Source/MetasoundEditor/Private/MetasoundEditorModule.cpp
--- a/Engine/Plugins/Runtime/Metasound/Source/MetasoundEditor/Private/MetasoundEditorModule.cpp
+++ b/Engine/Plugins/Runtime/Metasound/Source/MetasoundEditor/Private/MetasoundEditorModule.cpp
@@ -42,6 +42,8 @@
 #include "Styling/SlateStyleRegistry.h"
 #include "Styling/SlateTypes.h"
 #include "Templates/SharedPointer.h"
+#include "UObject/UObjectGlobals.h"
+#include "Animation/Skeleton.h"
 
 DEFINE_LOG_CATEGORY(LogMetasoundEditor);
 
@@ -102,6 +104,11 @@ namespace Metasound
 					Set("MetasoundEditor.Graph.TriggerPin.Connected", new IMAGE_BRUSH(TEXT("Graph/pin_trigger_connected"), Icon15x11));
 					Set("MetasoundEditor.Graph.TriggerPin.Disconnected", new IMAGE_BRUSH(TEXT("Graph/pin_trigger_disconnected"), Icon15x11));
 
+					Set("MetasoundEditor.Graph.Node.Class.Native", new IMAGE_BRUSH_SVG(TEXT("Icons/native_node"), FVector2D(8.0f, 16.0f)));
+					Set("MetasoundEditor.Graph.Node.Class.Graph", new IMAGE_BRUSH_SVG(TEXT("Icons/graph_node"), Icon16));
+					Set("MetasoundEditor.Graph.Node.Class.Input", new IMAGE_BRUSH_SVG(TEXT("Icons/input_node"), FVector2D(16.0f, 13.0f)));
+					Set("MetasoundEditor.Graph.Node.Class.Output", new IMAGE_BRUSH_SVG(TEXT("Icons/output_node"), FVector2D(16.0f, 13.0f)));
+
 					Set("MetasoundEditor.Graph.Node.Math.Add", new FSlateImageBrush(RootToContentDir(TEXT("/Graph/node_math_add_40x.png")), Icon40x40));
 					Set("MetasoundEditor.Graph.Node.Math.Divide", new FSlateImageBrush(RootToContentDir(TEXT("/Graph/node_math_divide_40x.png")), Icon40x40));
 					Set("MetasoundEditor.Graph.Node.Math.Modulo", new FSlateImageBrush(RootToContentDir(TEXT("/Graph/node_math_modulo_40x.png")), Icon40x40));
@@ -146,7 +153,52 @@ namespace Metasound
 
 				if (IsMetaSoundAssetClass(InAssetData.AssetClass))
 				{
-					UMetaSoundAssetSubsystem::Get().AddOrUpdateAsset(InAssetData);
+					// Use the editor version of UnregisterWithFrontend so it refreshes any open MetaSound editors
+					UMetaSoundAssetSubsystem::Get().AddOrUpdateAsset(InAssetData, false /* bUnregisterWithFrontend */);
+					if (UObject* AssetObject = InAssetData.GetAsset())
+					{
+						FGraphBuilder::RegisterGraphWithFrontend(*AssetObject);
+					}
+				}
+			}
+
+			void OnPackageReloaded(const EPackageReloadPhase InPackageReloadPhase, FPackageReloadedEvent* InPackageReloadedEvent)
+			{
+				using namespace Metasound;
+				using namespace Metasound::Editor;
+				using namespace Metasound::Frontend;
+
+				if (!InPackageReloadedEvent)
+				{
+					return;
+				}
+
+				if (InPackageReloadPhase != EPackageReloadPhase::OnPackageFixup)
+				{
+					return;
+				}
+
+				for (const TPair<UObject*, UObject*>& Pair : InPackageReloadedEvent->GetRepointedObjects())
+				{
+					if (UObject* Obj = Pair.Key)
+					{
+						if (IsMetaSoundAssetClass(Obj->GetClass()->GetFName()))
+						{
+							// Use the editor version of UnregisterWithFrontend so it refreshes any open MetaSound editors
+							UMetaSoundAssetSubsystem::Get().RemoveAsset(*Pair.Key, false /* bUnregisterWithFrontend */);
+							FGraphBuilder::UnregisterGraphWithFrontend(*Pair.Key);
+						}
+					}
+
+					if (UObject* Obj = Pair.Value)
+					{
+						if (IsMetaSoundAssetClass(Obj->GetClass()->GetFName()))
+						{
+							// Use the editor version of RegisterWithFrontend so it refreshes any open MetaSound editors
+							UMetaSoundAssetSubsystem::Get().AddOrUpdateAsset(*Pair.Value, false /* bRegisterWithFrontend */);
+							FGraphBuilder::RegisterGraphWithFrontend(*Pair.Value);
+						}
+					}
 				}
 			}
 
@@ -161,27 +213,32 @@ namespace Metasound
 				AssetRegistryModule.Get().OnAssetUpdated().AddRaw(this, &FModule::AddOrUpdateClassRegistryAsset);
 				AssetRegistryModule.Get().OnAssetRemoved().AddRaw(this, &FModule::RemoveAssetFromClassRegistry);
 				AssetRegistryModule.Get().OnAssetRenamed().AddRaw(this, &FModule::RenameAssetInClassRegistry);
+
+				FCoreUObjectDelegates::OnPackageReloaded.AddRaw(this, &FModule::OnPackageReloaded);
 			}
 
 			void RemoveAssetFromClassRegistry(const FAssetData& InAssetData)
 			{
-				using namespace Metasound;
-				using namespace Metasound::Frontend;
-
 				if (IsMetaSoundAssetClass(InAssetData.AssetClass))
 				{
-					UMetaSoundAssetSubsystem::Get().RemoveAsset(InAssetData);
+					// Use the editor version of UnregisterWithFrontend so it refreshes any open MetaSound editors
+					UMetaSoundAssetSubsystem::Get().RemoveAsset(InAssetData, false /* bUnregisterWithFrontend */);
+					if (UObject* AssetObject = InAssetData.GetAsset())
+					{
+						FGraphBuilder::UnregisterGraphWithFrontend(*AssetObject);
+					}
 				}
 			}
 
 			void RenameAssetInClassRegistry(const FAssetData& InAssetData, const FString& InOldObjectPath)
 			{
-				using namespace Metasound;
-				using namespace Metasound::Frontend;
-
 				if (IsMetaSoundAssetClass(InAssetData.AssetClass))
 				{
-					UMetaSoundAssetSubsystem::Get().AddOrUpdateAsset(InAssetData);
+					UMetaSoundAssetSubsystem::Get().RenameAsset(InAssetData, false /* bReregisterWithFrontend */);
+					if (UObject* AssetObject = InAssetData.GetAsset())
+					{
+						FGraphBuilder::RegisterGraphWithFrontend(*AssetObject);
+					}
 				}
 			}
 
@@ -325,6 +382,8 @@ namespace Metasound
 					AssetRegistryModule->Get().OnAssetUpdated().RemoveAll(this);
 					AssetRegistryModule->Get().OnAssetRemoved().RemoveAll(this);
 					AssetRegistryModule->Get().OnAssetRenamed().RemoveAll(this);
+
+					FCoreUObjectDelegates::OnPackageReloaded.RemoveAll(this);
 				}
 			}
 
@@ -361,6 +420,7 @@ namespace Metasound
 
 			virtual bool IsMetaSoundAssetClass(const FName InClassName) const override
 			{
+				// TODO: Move to IMetasoundUObjectRegistry (overload IsRegisteredClass to take in class name?)
 				return MetaSoundClassNames.Contains(InClassName);
 			}
 
@@ -369,8 +429,7 @@ namespace Metasound
 				// Register Metasound asset type actions
 				IAssetTools& AssetTools = FModuleManager::LoadModuleChecked<FAssetToolsModule>(AssetToolName).Get();
 
-				// TODO: Enable Composition
-				//AddAssetAction<FAssetTypeActions_MetaSound>(AssetTools, AssetActions);
+				AddAssetAction<FAssetTypeActions_MetaSound>(AssetTools, AssetActions);
 				AddAssetAction<FAssetTypeActions_MetaSoundSource>(AssetTools, AssetActions);
 
 				FPropertyEditorModule& PropertyModule = FModuleManager::LoadModuleChecked<FPropertyEditorModule>("PropertyEditor");
@@ -428,10 +487,9 @@ namespace Metasound
 				MetaSoundClassNames.Add(UMetaSound::StaticClass()->GetFName());
 				MetaSoundClassNames.Add(UMetaSoundSource::StaticClass()->GetFName());
 
-				//  TODO: Enable Composition
-				//FAssetTypeActions_MetaSound::RegisterMenuActions();
-				//FAssetTypeActions_MetaSoundSource::RegisterMenuActions();
-				// InitializeAssetClassRegistry();
+				FAssetTypeActions_MetaSound::RegisterMenuActions();
+				FAssetTypeActions_MetaSoundSource::RegisterMenuActions();
+				InitializeAssetClassRegistry();
 			}
 
 			virtual void ShutdownModule() override
@@ -467,8 +525,7 @@ namespace Metasound
 					GraphPanelPinFactory.Reset();
 				}
 
-				//  TODO: Enable Composition
-				// ShutdownAssetClassRegistry();
+				ShutdownAssetClassRegistry();
 
 				AssetActions.Reset();
 				DataTypeInfo.Reset();
diff --git a/Engine/Plugins/Runtime/Metasound/Source/MetasoundEditor/Private/MetasoundEditorSettings.cpp b/Engine/Plugins/Runtime/Metasound/Source/MetasoundEditor/Private/MetasoundEditorSettings.cpp
--- a/Engine/Plugins/Runtime/Metasound/Source/MetasoundEditor/Private/MetasoundEditorSettings.cpp
+++ b/Engine/Plugins/Runtime/Metasound/Source/MetasoundEditor/Private/MetasoundEditorSettings.cpp
@@ -21,4 +21,9 @@ UMetasoundEditorSettings::UMetasoundEditorSettings(const FObjectInitializer& Obj
 	StringPinTypeColor = FLinearColor(1.0f, 0.0f, 0.660537f, 1.0f);				// bright pink
 	TimePinTypeColor = FLinearColor(0.3f, 1.0f, 1.0f, 1.0f);					// cyan
 	TriggerPinTypeColor = FLinearColor(1.0f, 1.0f, 1.0f, 1.0f);					// white
+
+	NativeNodeTitleColor = FLinearColor(0.4f, 0.85f, 0.35f, 1.0f);				// i forgot
+	AssetReferenceNodeTitleColor = FLinearColor(0.047f, 0.686f, 0.988f);		// sky blue
+	InputNodeTitleColor = FLinearColor(0.168f, 1.0f, 0.7294f);					// sea foam
+	OutputNodeTitleColor = FLinearColor(1.0f, 0.878f, 0.1686f);					// yellow
 }
diff --git a/Engine/Plugins/Runtime/Metasound/Source/MetasoundEditor/Private/MetasoundEditorSettings.h b/Engine/Plugins/Runtime/Metasound/Source/MetasoundEditor/Private/MetasoundEditorSettings.h
--- a/Engine/Plugins/Runtime/Metasound/Source/MetasoundEditor/Private/MetasoundEditorSettings.h
+++ b/Engine/Plugins/Runtime/Metasound/Source/MetasoundEditor/Private/MetasoundEditorSettings.h
@@ -66,6 +66,22 @@ public:
 	UPROPERTY(EditAnywhere, config, Category=PinColors)
 	FLinearColor TriggerPinTypeColor;
 
+	/** Native node class title color */
+	UPROPERTY(EditAnywhere, config, Category = NodeTitleColors)
+	FLinearColor NativeNodeTitleColor;
+
+	/** Title color for references to MetaSound assets */
+	UPROPERTY(EditAnywhere, config, Category = NodeTitleColors)
+	FLinearColor AssetReferenceNodeTitleColor;
+
+	/** Input node title color */
+	UPROPERTY(EditAnywhere, config, Category = NodeTitleColors)
+	FLinearColor InputNodeTitleColor;
+
+	/** Output node title color */
+	UPROPERTY(EditAnywhere, config, Category = NodeTitleColors)
+	FLinearColor OutputNodeTitleColor;
+
 	/** Determines which details view to show in Metasounds Editor */
 	UPROPERTY(Transient)
 	EMetasoundActiveDetailView DetailView = EMetasoundActiveDetailView::Metasound;
diff --git a/Engine/Plugins/Runtime/Metasound/Source/MetasoundEditor/Private/MetasoundFactory.cpp b/Engine/Plugins/Runtime/Metasound/Source/MetasoundEditor/Private/MetasoundFactory.cpp
--- a/Engine/Plugins/Runtime/Metasound/Source/MetasoundEditor/Private/MetasoundFactory.cpp
+++ b/Engine/Plugins/Runtime/Metasound/Source/MetasoundEditor/Private/MetasoundFactory.cpp
@@ -20,15 +20,14 @@ namespace Metasound
 			T* MetaSoundObject = NewObject<T>(InParent, InName, InFlags);
 			check(MetaSoundObject);
 
+			FGraphBuilder::InitMetaSound(*MetaSoundObject, UKismetSystemLibrary::GetPlatformUserName());
+
 			if (InReferencedMetaSoundObject)
 			{
-				FGraphBuilder::InitMetaSound(*MetaSoundObject, UKismetSystemLibrary::GetPlatformUserName(), false /* bConformToArchetype */);
 				FGraphBuilder::InitMetaSoundPreset(*InReferencedMetaSoundObject, *MetaSoundObject);
 			}
-			else
-			{
-				FGraphBuilder::InitMetaSound(*MetaSoundObject, UKismetSystemLibrary::GetPlatformUserName());
-			}
+
+			FGraphBuilder::SynchronizeGraph(*MetaSoundObject);
 
 			return MetaSoundObject;
 		}
diff --git a/Engine/Plugins/Runtime/Metasound/Source/MetasoundEditor/Private/MetasoundNodeDetailCustomization.cpp b/Engine/Plugins/Runtime/Metasound/Source/MetasoundEditor/Private/MetasoundNodeDetailCustomization.cpp
--- a/Engine/Plugins/Runtime/Metasound/Source/MetasoundEditor/Private/MetasoundNodeDetailCustomization.cpp
+++ b/Engine/Plugins/Runtime/Metasound/Source/MetasoundEditor/Private/MetasoundNodeDetailCustomization.cpp
@@ -462,7 +462,7 @@ namespace Metasound
 		{
 		}
 
-		void FMetasoundVariableDataTypeSelector::AddDataTypeSelector(IDetailLayoutBuilder& InDetailLayout, const FText& InRowName, TWeakObjectPtr<UMetasoundEditorGraphVariable> InGraphVariable, bool bIsRequired)
+		void FMetasoundVariableDataTypeSelector::AddDataTypeSelector(IDetailLayoutBuilder& InDetailLayout, const FText& InRowName, TWeakObjectPtr<UMetasoundEditorGraphVariable> InGraphVariable, bool bIsEnabled)
 		{
 			DetailLayoutBuilder = &InDetailLayout;
 
@@ -524,7 +524,7 @@ namespace Metasound
 			});
 
 			CategoryBuilder.AddCustomRow(InRowName)
-			.IsEnabled(!bIsRequired)
+			.IsEnabled(bIsEnabled)
 			.NameContent()
 			[
 				SNew(STextBlock)
@@ -546,7 +546,7 @@ namespace Metasound
 					{
 						OnBaseDataTypeChanged(InGraphVariable, ItemSelected, SelectInfo);
 					})
-					.IsEnabled(!bIsRequired)
+					.IsEnabled(bIsEnabled)
 				]
 				+ SHorizontalBox::Slot()
 				.FillWidth(0.40f)
@@ -598,15 +598,16 @@ namespace Metasound
 			IDetailCategoryBuilder& CategoryBuilder = DetailLayout.EditCategory("General");
 
 			const bool bIsRequired = IsRequired();
+			const bool bIsGraphEditable = IsGraphEditable();
 			DisplayNameEditableTextBox = SNew(SEditableTextBox)
 				.Text(this, &FMetasoundInputDetailCustomization::GetDisplayName)
 				.OnTextChanged(this, &FMetasoundInputDetailCustomization::OnDisplayNameChanged)
 				.OnTextCommitted(this, &FMetasoundInputDetailCustomization::OnDisplayNameCommitted)
-				.IsReadOnly(bIsRequired)
+				.IsReadOnly(bIsRequired || !bIsGraphEditable)
 				.Font(IDetailLayoutBuilder::GetDetailFont());
 
 			CategoryBuilder.AddCustomRow(VariableCustomizationPrivate::InputNameText)
-			.EditCondition(!bIsRequired, nullptr)
+			.EditCondition(!bIsRequired && bIsGraphEditable, nullptr)
 			.NameContent()
 			[
 				SNew(STextBlock)
@@ -630,7 +631,7 @@ namespace Metasound
 			];
 
 			CategoryBuilder.AddCustomRow(VariableCustomizationPrivate::NodeTooltipText)
-			.EditCondition(!bIsRequired, nullptr)
+			.EditCondition(!bIsRequired && bIsGraphEditable, nullptr)
 			.NameContent()
 			[
 				SNew(STextBlock)
@@ -642,29 +643,29 @@ namespace Metasound
 				SNew(SMultiLineEditableTextBox)
 				.Text(this, &FMetasoundInputDetailCustomization::GetTooltip)
 				.OnTextCommitted(this, &FMetasoundInputDetailCustomization::OnTooltipCommitted)
-				.IsReadOnly(bIsRequired)
+				.IsReadOnly(bIsRequired || !bIsGraphEditable)
 				.ModiferKeyForNewLine(EModifierKey::Shift)
 				.RevertTextOnEscape(true)
 				.WrapTextAt(VariableCustomizationPrivate::DetailsTitleMaxWidth - VariableCustomizationPrivate::DetailsTitleWrapPadding)
 				.Font(IDetailLayoutBuilder::GetDetailFont())
 			];
 
-			AddDataTypeSelector(DetailLayout, VariableCustomizationPrivate::DataTypeNameText, GraphVariable, bIsRequired);
-
-			CategoryBuilder.AddCustomRow(LOCTEXT("InputPrivate", "Private"))
-			.Visibility(TAttribute<EVisibility>(EVisibility::Hidden))
-			.NameContent()
-			[
-				SNew(STextBlock)
-				.Text(LOCTEXT("InputPrivate", "Private"))
-				.Font(IDetailLayoutBuilder::GetDetailFont())
-			]
-			.ValueContent()
-			[
-				SNew(SCheckBox)
-				.IsChecked(this, &FMetasoundInputDetailCustomization::OnGetPrivateCheckboxState)
-				.OnCheckStateChanged(this, &FMetasoundInputDetailCustomization::OnPrivateChanged)
-			];
+			AddDataTypeSelector(DetailLayout, VariableCustomizationPrivate::DataTypeNameText, GraphVariable, !bIsRequired && bIsGraphEditable);
+
+// 			CategoryBuilder.AddCustomRow(LOCTEXT("InputPrivate", "Private"))
+// 			.Visibility(TAttribute<EVisibility>(EVisibility::Hidden))
+// 			.NameContent()
+// 			[
+// 				SNew(STextBlock)
+// 				.Text(LOCTEXT("InputPrivate", "Private"))
+// 				.Font(IDetailLayoutBuilder::GetDetailFont())
+// 			]
+// 			.ValueContent()
+// 			[
+// 				SNew(SCheckBox)
+// 				.IsChecked(this, &FMetasoundInputDetailCustomization::OnGetPrivateCheckboxState)
+// 				.OnCheckStateChanged(this, &FMetasoundInputDetailCustomization::OnPrivateChanged)
+// 			];
 
 			FNodeHandle NodeHandle = GraphVariable->GetNodeHandle();
 			const TArray<FOutputHandle>& Outputs = NodeHandle->GetOutputs();
@@ -879,15 +880,16 @@ namespace Metasound
 			IDetailCategoryBuilder& CategoryBuilder = DetailLayout.EditCategory("General");
 
 			const bool bIsRequired = IsRequired();
+			const bool bIsGraphEditable = IsGraphEditable();
 			DisplayNameEditableTextBox = SNew(SEditableTextBox)
 				.Text(this, &FMetasoundOutputDetailCustomization::GetDisplayName)
 				.OnTextChanged(this, &FMetasoundOutputDetailCustomization::OnDisplayNameChanged)
 				.OnTextCommitted(this, &FMetasoundOutputDetailCustomization::OnDisplayNameCommitted)
-				.IsReadOnly(bIsRequired)
+				.IsReadOnly(bIsRequired || !bIsGraphEditable)
 				.Font(IDetailLayoutBuilder::GetDetailFont());
 
 			CategoryBuilder.AddCustomRow(VariableCustomizationPrivate::OutputNameText)
-			.EditCondition(!bIsRequired, nullptr)
+			.EditCondition(!bIsRequired && bIsGraphEditable, nullptr)
 			.NameContent()
 			[
 				SNew(STextBlock)
@@ -910,7 +912,7 @@ namespace Metasound
 			];
 
 			CategoryBuilder.AddCustomRow(VariableCustomizationPrivate::NodeTooltipText)
-			.EditCondition(!bIsRequired, nullptr)
+			.EditCondition(!bIsRequired && bIsGraphEditable, nullptr)
 			.NameContent()
 			[
 				SNew(STextBlock)
@@ -922,29 +924,29 @@ namespace Metasound
 				SNew(SMultiLineEditableTextBox)
 				.Text(this, &FMetasoundOutputDetailCustomization::GetTooltip)
 				.OnTextCommitted(this, &FMetasoundOutputDetailCustomization::OnTooltipCommitted)
-				.IsReadOnly(bIsRequired)
+				.IsReadOnly(bIsRequired || !bIsGraphEditable)
 				.ModiferKeyForNewLine(EModifierKey::Shift)
 				.RevertTextOnEscape(true)
 				.WrapTextAt(VariableCustomizationPrivate::DetailsTitleMaxWidth - VariableCustomizationPrivate::DetailsTitleWrapPadding)
 				.Font(IDetailLayoutBuilder::GetDetailFont())
 			];
 
-			AddDataTypeSelector(DetailLayout, VariableCustomizationPrivate::DataTypeNameText, GraphVariable, bIsRequired);
-
-			CategoryBuilder.AddCustomRow(LOCTEXT("OutputPrivate", "Private"))
-			.Visibility(TAttribute<EVisibility>(EVisibility::Hidden))
-			.NameContent()
-			[
-				SNew(STextBlock)
-				.Text(LOCTEXT("OutputPrivate", "Private"))
-				.Font(IDetailLayoutBuilder::GetDetailFont())
-			]
-			.ValueContent()
-			[
-				SNew(SCheckBox)
-				.IsChecked(this, &FMetasoundOutputDetailCustomization::OnGetPrivateCheckboxState)
-				.OnCheckStateChanged(this, &FMetasoundOutputDetailCustomization::OnPrivateChanged)
-			];
+			AddDataTypeSelector(DetailLayout, VariableCustomizationPrivate::DataTypeNameText, GraphVariable, !bIsRequired && bIsGraphEditable);
+
+// 			CategoryBuilder.AddCustomRow(LOCTEXT("OutputPrivate", "Private"))
+// 			.Visibility(TAttribute<EVisibility>(EVisibility::Hidden))
+// 			.NameContent()
+// 			[
+// 				SNew(STextBlock)
+// 				.Text(LOCTEXT("OutputPrivate", "Private"))
+// 				.Font(IDetailLayoutBuilder::GetDetailFont())
+// 			]
+// 			.ValueContent()
+// 			[
+// 				SNew(SCheckBox)
+// 				.IsChecked(this, &FMetasoundOutputDetailCustomization::OnGetPrivateCheckboxState)
+// 				.OnCheckStateChanged(this, &FMetasoundOutputDetailCustomization::OnPrivateChanged)
+// 			];
 		}
 
 		void FMetasoundOutputDetailCustomization::SetDefaultPropertyMetaData(TSharedRef<IPropertyHandle> InDefaultPropertyHandle) const
diff --git a/Engine/Plugins/Runtime/Metasound/Source/MetasoundEditor/Private/MetasoundNodeDetailCustomization.h b/Engine/Plugins/Runtime/Metasound/Source/MetasoundEditor/Private/MetasoundNodeDetailCustomization.h
--- a/Engine/Plugins/Runtime/Metasound/Source/MetasoundEditor/Private/MetasoundNodeDetailCustomization.h
+++ b/Engine/Plugins/Runtime/Metasound/Source/MetasoundEditor/Private/MetasoundNodeDetailCustomization.h
@@ -157,12 +157,23 @@ namespace Metasound
 
 				if (GraphVariable.IsValid())
 				{
-					return GraphVariable->GetNodeHandle()->GetDisplayName();
+					return GraphVariable->GetConstNodeHandle()->GetDisplayName();
 				}
 
 				return FText::GetEmpty();
 			}
 
+			bool IsGraphEditable() const
+			{
+				if (GraphVariable.IsValid())
+				{
+					Metasound::Frontend::FConstNodeHandle NodeHandle = GraphVariable->GetConstNodeHandle();
+					return NodeHandle->GetOwningGraph()->GetGraphStyle().bIsGraphEditable;
+				}
+
+				return false;
+			}
+
 			bool IsRequired() const
 			{
 				if (GraphVariable.IsValid())
@@ -179,17 +190,7 @@ namespace Metasound
 
 				if (GraphVariable.IsValid())
 				{
-					const FText TransactionLabel = FText::Format(LOCTEXT("SetTooltip", "Set the MetaSound {0}'s tooltip"), VariableLabel);
-					const FScopedTransaction Transaction(TransactionLabel);
-
-					GraphVariable->Modify();
-					if (UMetasoundEditorGraph* Graph = Cast<UMetasoundEditorGraph>(GraphVariable->GetOuter()))
-					{
-						Graph->GetMetasoundChecked().Modify();
-					}
-
-					FNodeHandle NodeHandle = GraphVariable->GetNodeHandle();
-					NodeHandle->SetDescription(InNewText);
+					GraphVariable->SetDescription(InNewText);
 				}
 			}
 
diff --git a/Engine/Plugins/Runtime/Metasound/Source/MetasoundEditor/Private/SMetasoundGraphNode.cpp b/Engine/Plugins/Runtime/Metasound/Source/MetasoundEditor/Private/SMetasoundGraphNode.cpp
--- a/Engine/Plugins/Runtime/Metasound/Source/MetasoundEditor/Private/SMetasoundGraphNode.cpp
+++ b/Engine/Plugins/Runtime/Metasound/Source/MetasoundEditor/Private/SMetasoundGraphNode.cpp
@@ -37,6 +37,7 @@
 #include "Widgets/SBoxPanel.h"
 #include "Widgets/SOverlay.h"
 #include "Widgets/SWidget.h"
+#include "Widgets/Text/SInlineEditableTextBlock.h"
 #include "SGraphPinComboBox.h"
 #include "SMetasoundEnumPin.h"
 #include "UObject/ScriptInterface.h"
@@ -146,73 +147,94 @@ TSharedPtr<SGraphPin> SMetasoundGraphNode::CreatePinWidget(UEdGraphPin* InPin) c
 	using namespace Metasound::Editor;
 	using namespace Metasound::Frontend;
 
+	TSharedPtr<SGraphPin> PinWidget;
+
 	if (const UMetasoundEditorGraphSchema* GraphSchema = Cast<const UMetasoundEditorGraphSchema>(InPin->GetSchema()))
 	{
 		// Don't show default value field for container types
 		if (InPin->PinType.ContainerType != EPinContainerType::None)
 		{
-			return SNew(SGraphPin, InPin);
+			PinWidget = SNew(SGraphPin, InPin)
+				.ToolTipText(this, &SMetasoundGraphNode::GetPinTooltip, InPin);
 		}
 
-		if (InPin->PinType.PinCategory == FGraphBuilder::PinCategoryAudio)
+		else if (InPin->PinType.PinCategory == FGraphBuilder::PinCategoryAudio)
 		{
-			return SNew(SGraphPin, InPin);
+			PinWidget = SNew(SGraphPin, InPin)
+				.ToolTipText(this, &SMetasoundGraphNode::GetPinTooltip, InPin);
 		}
 
-		if (InPin->PinType.PinCategory == FGraphBuilder::PinCategoryBoolean)
+		else if (InPin->PinType.PinCategory == FGraphBuilder::PinCategoryBoolean)
 		{
-			return SNew(SGraphPinBool, InPin);
+			PinWidget = SNew(SGraphPinBool, InPin)
+				.ToolTipText(this, &SMetasoundGraphNode::GetPinTooltip, InPin);
 		}
 
-		//if (InPin->PinType.PinCategory == FGraphBuilder::PinCategoryDouble)
+		//else if (InPin->PinType.PinCategory == FGraphBuilder::PinCategoryDouble)
 		//{
-		//	return SNew(SGraphPinNum<double>, InPin);
+		//	PinWidget = SNew(SGraphPinNum<double>, InPin)
+		//		.ToolTipText(this, &SMetasoundGraphNode::GetPinTooltip, InPin);
 		//}
 
-		if (InPin->PinType.PinCategory == FGraphBuilder::PinCategoryFloat)
+		else if (InPin->PinType.PinCategory == FGraphBuilder::PinCategoryFloat)
 		{
-			return SNew(SGraphPinNum<float>, InPin);
+			PinWidget = SNew(SGraphPinNum<float>, InPin)
+				.ToolTipText(this, &SMetasoundGraphNode::GetPinTooltip, InPin);
 		}
 
-		if (InPin->PinType.PinCategory == FGraphBuilder::PinCategoryInt32)
+		else if (InPin->PinType.PinCategory == FGraphBuilder::PinCategoryInt32)
 		{
 			if (SMetasoundEnumPin::FindEnumInterfaceFromPin(InPin))
 			{
-				return SNew(SMetasoundEnumPin, InPin);
+				PinWidget = SNew(SMetasoundEnumPin, InPin)
+					.ToolTipText(this, &SMetasoundGraphNode::GetPinTooltip, InPin);
+			}
+			else
+			{
+				PinWidget = SNew(SGraphPinInteger, InPin)
+					.ToolTipText(this, &SMetasoundGraphNode::GetPinTooltip, InPin);
 			}
-			return SNew(SGraphPinInteger, InPin);
 		}
 
 		//if (InPin->PinType.PinCategory == FGraphBuilder::PinCategoryInt64)
 		//{
-		//	return SNew(SGraphPinNum<int64>, InPin);
+		//	PinWidget = SNew(SGraphPinNum<int64>, InPin)
+		//		.ToolTipText(this, &SMetasoundGraphNode::GetPinTooltip, InPin);
 		//}
 
-		if (InPin->PinType.PinCategory == FGraphBuilder::PinCategoryObject)
+		else if (InPin->PinType.PinCategory == FGraphBuilder::PinCategoryObject)
 		{
-			return SNew(SGraphPinObject, InPin);
+			PinWidget = SNew(SGraphPinObject, InPin)
+				.ToolTipText(this, &SMetasoundGraphNode::GetPinTooltip, InPin);
 		}
 
-		if (InPin->PinType.PinCategory == FGraphBuilder::PinCategoryString)
+		else if (InPin->PinType.PinCategory == FGraphBuilder::PinCategoryString)
 		{
-			return SNew(SGraphPinString, InPin);
+			PinWidget = SNew(SGraphPinString, InPin)
+				.ToolTipText(this, &SMetasoundGraphNode::GetPinTooltip, InPin);
 		}
 
-		if (InPin->PinType.PinCategory == FGraphBuilder::PinCategoryTrigger)
+		else if (InPin->PinType.PinCategory == FGraphBuilder::PinCategoryTrigger)
 		{
-			TSharedPtr<SGraphPin> TriggerPin = SNew(SGraphPin, InPin);
+			PinWidget = SNew(SGraphPin, InPin)
+				.ToolTipText(this, &SMetasoundGraphNode::GetPinTooltip, InPin);
 
 			if (const ISlateStyle* MetasoundStyle = FSlateStyleRegistry::FindSlateStyle("MetaSoundStyle"))
 			{
 				const FSlateBrush* PinConnectedBrush = MetasoundStyle->GetBrush(TEXT("MetasoundEditor.Graph.TriggerPin.Connected"));
 				const FSlateBrush* PinDisconnectedBrush = MetasoundStyle->GetBrush(TEXT("MetasoundEditor.Graph.TriggerPin.Disconnected"));
-				TriggerPin->SetCustomPinIcon(PinConnectedBrush, PinDisconnectedBrush);
+				PinWidget->SetCustomPinIcon(PinConnectedBrush, PinDisconnectedBrush);
 			}
-			return TriggerPin;
 		}
 	}
 
-	return SNew(SGraphPin, InPin);
+	if (!PinWidget.IsValid())
+	{
+		PinWidget = SNew(SGraphPin, InPin)
+			.ToolTipText(this, &SMetasoundGraphNode::GetPinTooltip, InPin);
+	}
+
+	return PinWidget;
 }
 
 void SMetasoundGraphNode::CreateStandardPinWidget(UEdGraphPin* InPin)
@@ -243,6 +265,23 @@ void SMetasoundGraphNode::CreateStandardPinWidget(UEdGraphPin* InPin)
 	}
 }
 
+FText SMetasoundGraphNode::GetPinTooltip(UEdGraphPin* InPin) const
+{
+	using namespace Metasound::Editor;
+	using namespace Metasound::Frontend;
+
+	if (InPin->Direction == EGPD_Input)
+	{
+		FConstInputHandle InputHandle = FGraphBuilder::GetConstInputHandleFromPin(InPin);
+		return InputHandle->GetTooltip();
+	}
+	else
+	{
+		FConstOutputHandle OutputHandle = FGraphBuilder::GetConstOutputHandleFromPin(InPin);
+		return OutputHandle->GetTooltip();
+	}
+}
+
 TSharedRef<SWidget> SMetasoundGraphNode::CreateTitleWidget(TSharedPtr<SNodeTitle> NodeTitle)
 {
 	Metasound::Frontend::FNodeHandle NodeHandle = GetMetasoundNode().GetNodeHandle();
@@ -251,7 +290,58 @@ TSharedRef<SWidget> SMetasoundGraphNode::CreateTitleWidget(TSharedPtr<SNodeTitle
 		return SNullWidget::NullWidget;
 	}
 
-	return SGraphNode::CreateTitleWidget(NodeTitle);
+	TSharedPtr<SHorizontalBox> TitleBoxWidget = SNew(SHorizontalBox);
+
+	FSlateIcon NodeIcon = GetMetasoundNode().GetNodeTitleIcon();
+	if (const FSlateBrush* IconBrush = NodeIcon.GetIcon())
+	{
+		if (IconBrush != FStyleDefaults::GetNoBrush())
+		{
+			TSharedPtr<SImage> Image;
+			TitleBoxWidget->AddSlot()
+			.AutoWidth()
+			.VAlign(VAlign_Center)
+			.HAlign(HAlign_Right)
+			[
+				SNew(SHorizontalBox)
+				+ SHorizontalBox::Slot()
+				.Padding(0.0f, 0.0f, 4.0f, 0.0f)
+				.VAlign(VAlign_Center)
+				.HAlign(HAlign_Right)
+				[
+					SAssignNew(Image, SImage)
+				]
+			];
+			Image->SetColorAndOpacity(TAttribute<FSlateColor>::CreateLambda([this]() { return FSlateColor(GetNodeTitleColorOverride()); }));
+			Image->SetImage(IconBrush);
+		}
+	}
+
+	TitleBoxWidget->AddSlot()
+	.AutoWidth()
+	[
+		SGraphNode::CreateTitleWidget(NodeTitle)
+	];
+
+	InlineEditableText->SetColorAndOpacity(TAttribute<FLinearColor>::Create(TAttribute<FLinearColor>::FGetter::CreateSP(this, &SMetasoundGraphNode::GetNodeTitleColorOverride)));
+
+	return TitleBoxWidget.ToSharedRef();
+}
+
+FLinearColor SMetasoundGraphNode::GetNodeTitleColorOverride() const
+{
+	FLinearColor ReturnTitleColor = GraphNode->IsDeprecated() ? FLinearColor::Red : GetNodeObj()->GetNodeTitleColor();
+
+	if (!GraphNode->IsNodeEnabled() || GraphNode->IsDisplayAsDisabledForced() || GraphNode->IsNodeUnrelated())
+	{
+		ReturnTitleColor *= FLinearColor(0.5f, 0.5f, 0.5f, 0.4f);
+	}
+	else
+	{
+		ReturnTitleColor.A = FadeCurve.GetLerp();
+	}
+
+	return ReturnTitleColor;
 }
 
 void SMetasoundGraphNode::SetDefaultTitleAreaWidget(TSharedRef<SOverlay> DefaultTitleAreaWidget)
@@ -315,6 +405,7 @@ void SMetasoundGraphNode::SetDefaultTitleAreaWidget(TSharedRef<SOverlay> Default
 				.Size(FVector2D(20,20))
 			]
 		];
+
 	}
 	else
 	{
diff --git a/Engine/Plugins/Runtime/Metasound/Source/MetasoundEditor/Public/MetasoundEditorGraph.h b/Engine/Plugins/Runtime/Metasound/Source/MetasoundEditor/Public/MetasoundEditorGraph.h
--- a/Engine/Plugins/Runtime/Metasound/Source/MetasoundEditor/Public/MetasoundEditorGraph.h
+++ b/Engine/Plugins/Runtime/Metasound/Source/MetasoundEditor/Public/MetasoundEditorGraph.h
@@ -11,6 +11,7 @@
 #include "MetasoundSource.h"
 #include "MetasoundUObjectRegistry.h"
 #include "UObject/ObjectMacros.h"
+#include "UObject/ObjectSaveContext.h"
 #include "UObject/ScriptInterface.h"
 
 #include "MetasoundEditorGraph.generated.h"
@@ -64,6 +65,7 @@ public:
 	bool IsRequired() const;
 	TArray<UMetasoundEditorGraphNode*> GetNodes() const;
 	void SetDataType(FName InNewType);
+	void SetDescription(const FText& InDescription);
 	void SetDisplayName(const FText& InNewName);
 
 	virtual EMetasoundFrontendClassType GetClassType() const { return EMetasoundFrontendClassType::Invalid; }
@@ -95,9 +97,7 @@ public:
 	{
 	}
 
-#if WITH_EDITORONLY_DATA
 	virtual void PostEditUndo() override;
-#endif // WITH_EDITORONLY_DATA
 };
 
 UCLASS()
@@ -119,9 +119,9 @@ public:
 	void OnDataTypeChanged() override;
 	void OnLiteralChanged(bool bPostTransaction = true);
 
-#if WITH_EDITORONLY_DATA
+#if WITH_EDITOR
 	virtual void PostEditUndo() override;
-#endif // WITH_EDITORONLY_DATA
+#endif // WITH_EDITOR
 };
 
 UCLASS()
@@ -158,11 +158,14 @@ public:
 	void SetPreviewID(uint32 InPreviewID);
 	bool IsPreviewing() const;
 
-	virtual void Synchronize();
+	virtual bool Synchronize() override;
+	virtual bool Validate(bool bInAutoUpdate) override;
 
-	bool Validate(Metasound::Editor::FGraphValidationResults& OutResults);
+	virtual void RegisterGraphWithFrontend() override;
 
 private:
+	bool ValidateInternal(Metasound::Editor::FGraphValidationResults& OutResults, bool bClearUpgradeMessaging = true);
+
 	// Preview ID is the Unique ID provided by the UObject that implements
 	// a sound's ParameterInterface when a sound begins playing.
 	uint32 PreviewID = INDEX_NONE;
@@ -186,4 +189,5 @@ public:
 
 	friend class UMetaSoundFactory;
 	friend class UMetaSoundSourceFactory;
+	friend class Metasound::Editor::FGraphBuilder;
 };
diff --git a/Engine/Plugins/Runtime/Metasound/Source/MetasoundEditor/Public/MetasoundEditorGraphInputNodes.h b/Engine/Plugins/Runtime/Metasound/Source/MetasoundEditor/Public/MetasoundEditorGraphInputNodes.h
--- a/Engine/Plugins/Runtime/Metasound/Source/MetasoundEditor/Public/MetasoundEditorGraphInputNodes.h
+++ b/Engine/Plugins/Runtime/Metasound/Source/MetasoundEditor/Public/MetasoundEditorGraphInputNodes.h
@@ -6,6 +6,7 @@
 #include "Audio/AudioParameterInterface.h"
 #include "CoreMinimal.h"
 #include "EdGraph/EdGraphNode.h"
+#include "GraphEditorSettings.h"
 #include "MetasoundDataReference.h"
 #include "MetasoundEditorGraph.h"
 #include "MetasoundEditorGraphNode.h"
@@ -26,32 +27,6 @@
 class UEdGraphPin;
 class UMetaSound;
 
-namespace Metasound
-{
-	namespace Editor
-	{
-		// Forward Declarations
-		class FGraphBuilder;
-
-		namespace InputPrivate
-		{
-			template <typename TPODType>
-			void ConvertLiteral(const FMetasoundFrontendLiteral& InLiteral, TPODType& OutValue)
-			{
-				const FName& TypeName = GetMetasoundDataTypeName<TPODType>();
-				OutValue = InLiteral.ToLiteral(TypeName).Value.Get<TPODType>();
-			}
-
-			template <typename TPODType, typename TLiteralType = TPODType>
-			void ConvertLiteralToArray(const FMetasoundFrontendLiteral& InLiteral, TArray<TLiteralType>& OutArray)
-			{
-				const FName& TypeName = *FString(GetMetasoundDataTypeString<TPODType>() + TEXT(":Array"));
-				TArray<TPODType> NewValue = InLiteral.ToLiteral(TypeName).Value.Get<TArray<TPODType>>();
-				Algo::Transform(NewValue, OutArray, [](const TPODType& InValue) { return TLiteralType { InValue }; });
-			}
-		}
-	}
-}
 
 UCLASS()
 class METASOUNDEDITOR_API UMetasoundEditorGraphInputNode : public UMetasoundEditorGraphNode
@@ -65,51 +40,21 @@ public:
 	TObjectPtr<UMetasoundEditorGraphInput> Input;
 
 public:
-	virtual FMetasoundFrontendClassName GetClassName() const override
-	{
-		if (ensure(Input))
-		{
-			return Input->ClassName;
-		}
-
-		return Super::GetClassName();
-	}
-
-	virtual void UpdatePreviewInstance(const Metasound::FVertexKey& InParameterName, TScriptInterface<IAudioParameterInterface>& InParameterInterface) const
-	{
-		if (Input)
-		{
-			Input->UpdatePreviewInstance(InParameterName, InParameterInterface);
-		}
-	}
-	
-	virtual FGuid GetNodeID() const override
-	{
-		if (Input)
-		{
-			return Input->NodeID;
-		}
-
-		return Super::GetNodeID();
-	}
+	virtual void UpdatePreviewInstance(const Metasound::FVertexKey& InParameterName, TScriptInterface<IAudioParameterInterface>& InParameterInterface) const;
+
+	virtual FMetasoundFrontendClassName GetClassName() const;
+
+	virtual FGuid GetNodeID() const override;
+	virtual FLinearColor GetNodeTitleColor() const override;
+	virtual FSlateIcon GetNodeTitleIcon() const override;
 
 #if WITH_EDITORONLY_DATA
-	virtual void PostEditUndo() override
-	{
-		Super::PostEditUndo();
-
-		if (Input)
-		{
-			Input->OnLiteralChanged(false /* bPostTransaction */);
-		}
-	}
+	virtual void PostEditUndo() override;
 #endif // WITH_EDITORONLY_DATA
 
 protected:
-	virtual void SetNodeID(FGuid InNodeID) override
-	{
-		Input->NodeID = InNodeID;
-	}
+	virtual void SetNodeID(FGuid InNodeID) override;
+
 
 	friend class Metasound::Editor::FGraphBuilder;
 };
@@ -135,27 +80,10 @@ public:
 	UPROPERTY(EditAnywhere, Category = DefaultValue)
 	FMetasoundEditorGraphInputBoolRef Default;
 
-	virtual FMetasoundFrontendLiteral GetDefault() const override
-	{
-		FMetasoundFrontendLiteral Literal;
-		Literal.Set(Default.Value);
-		return Literal;
-	}
-
-	virtual EMetasoundFrontendLiteralType GetLiteralType() const override
-	{
-		return EMetasoundFrontendLiteralType::Boolean;
-	}
-
-	virtual void SetFromLiteral(const FMetasoundFrontendLiteral& InLiteral) override
-	{
-		Metasound::Editor::InputPrivate::ConvertLiteral<bool>(InLiteral, Default.Value);
-	}
-
-	virtual void UpdatePreviewInstance(const Metasound::FVertexKey& InParameterName, TScriptInterface<IAudioParameterInterface>& InParameterInterface) const override
-	{
-		InParameterInterface->SetBool(*InParameterName, Default.Value);
-	}
+	virtual FMetasoundFrontendLiteral GetDefault() const override;
+	virtual EMetasoundFrontendLiteralType GetLiteralType() const override;
+	virtual void SetFromLiteral(const FMetasoundFrontendLiteral& InLiteral) override;
+	virtual void UpdatePreviewInstance(const Metasound::FVertexKey& InParameterName, TScriptInterface<IAudioParameterInterface>& InParameterInterface) const override;
 };
 
 UCLASS(MinimalAPI)
@@ -169,32 +97,10 @@ public:
 	UPROPERTY(EditAnywhere, Category = DefaultValue)
 	TArray<FMetasoundEditorGraphInputBoolRef> Default;
 
-	virtual FMetasoundFrontendLiteral GetDefault() const override
-	{
-		TArray<bool> BoolArray;
-		Algo::Transform(Default, BoolArray, [](const FMetasoundEditorGraphInputBoolRef& InValue) { return InValue.Value; });
-
-		FMetasoundFrontendLiteral Literal;
-		Literal.Set(BoolArray);
-		return Literal;
-	}
-
-	virtual EMetasoundFrontendLiteralType GetLiteralType() const override
-	{
-		return EMetasoundFrontendLiteralType::BooleanArray;
-	}
-
-	virtual void SetFromLiteral(const FMetasoundFrontendLiteral& InLiteral) override
-	{
-		Metasound::Editor::InputPrivate::ConvertLiteralToArray<bool, FMetasoundEditorGraphInputBoolRef>(InLiteral, Default);
-	}
-
-	virtual void UpdatePreviewInstance(const Metasound::FVertexKey& InParameterName, TScriptInterface<IAudioParameterInterface>& InParameterInterface) const override
-	{
-		TArray<bool> BoolArray;
-		Algo::Transform(Default, BoolArray, [](const FMetasoundEditorGraphInputBoolRef& InValue) { return InValue.Value; });
-		InParameterInterface->SetBoolArray(*InParameterName, BoolArray);
-	}
+	virtual FMetasoundFrontendLiteral GetDefault() const override;
+	virtual EMetasoundFrontendLiteralType GetLiteralType() const override;
+	virtual void SetFromLiteral(const FMetasoundFrontendLiteral& InLiteral) override;
+	virtual void UpdatePreviewInstance(const Metasound::FVertexKey& InParameterName, TScriptInterface<IAudioParameterInterface>& InParameterInterface) const override;
 };
 
 // Broken out to be able to customize and swap enum behavior for basic integer literal behavior
@@ -218,27 +124,10 @@ public:
 	UPROPERTY(EditAnywhere, Category = DefaultValue)
 	FMetasoundEditorGraphInputIntRef Default;
 
-	virtual FMetasoundFrontendLiteral GetDefault() const override
-	{
-		FMetasoundFrontendLiteral Literal;
-		Literal.Set(Default.Value);
-		return Literal;
-	}
-
-	virtual EMetasoundFrontendLiteralType GetLiteralType() const override
-	{
-		return EMetasoundFrontendLiteralType::Integer;
-	}
-
-	virtual void SetFromLiteral(const FMetasoundFrontendLiteral& InLiteral) override
-	{
-		Metasound::Editor::InputPrivate::ConvertLiteral<int32>(InLiteral, Default.Value);
-	}
-
-	virtual void UpdatePreviewInstance(const Metasound::FVertexKey& InParameterName, TScriptInterface<IAudioParameterInterface>& InParameterInterface) const override
-	{
-		InParameterInterface->SetInt(*InParameterName, Default.Value);
-	}
+	virtual FMetasoundFrontendLiteral GetDefault() const override;
+	virtual EMetasoundFrontendLiteralType GetLiteralType() const override;
+	virtual void SetFromLiteral(const FMetasoundFrontendLiteral& InLiteral) override;
+	virtual void UpdatePreviewInstance(const Metasound::FVertexKey& InParameterName, TScriptInterface<IAudioParameterInterface>& InParameterInterface) const override;
 };
 
 UCLASS(MinimalAPI)
@@ -252,33 +141,10 @@ public:
 	UPROPERTY(EditAnywhere, Category = DefaultValue)
 	TArray<FMetasoundEditorGraphInputIntRef> Default;
 
-	virtual FMetasoundFrontendLiteral GetDefault() const override
-	{
-		TArray<int32> IntArray;
-		Algo::Transform(Default, IntArray, [](const FMetasoundEditorGraphInputIntRef& InValue) { return InValue.Value; });
-
-		FMetasoundFrontendLiteral Literal;
-		Literal.Set(IntArray);
-
-		return Literal;
-	}
-
-	virtual EMetasoundFrontendLiteralType GetLiteralType() const override
-	{
-		return EMetasoundFrontendLiteralType::IntegerArray;
-	}
-
-	virtual void SetFromLiteral(const FMetasoundFrontendLiteral& InLiteral) override
-	{
-		Metasound::Editor::InputPrivate::ConvertLiteralToArray<int32, FMetasoundEditorGraphInputIntRef>(InLiteral, Default);
-	}
-
-	void UpdatePreviewInstance(const Metasound::FVertexKey& InParameterName, TScriptInterface<IAudioParameterInterface>& InParameterInterface) const override
-	{
-		TArray<int32> IntArray;
-		Algo::Transform(Default, IntArray, [](const FMetasoundEditorGraphInputIntRef& InValue) { return InValue.Value; });
-		InParameterInterface->SetIntArray(*InParameterName, IntArray);
-	}
+	virtual FMetasoundFrontendLiteral GetDefault() const override;
+	virtual EMetasoundFrontendLiteralType GetLiteralType() const override;
+	virtual void SetFromLiteral(const FMetasoundFrontendLiteral& InLiteral) override;
+	virtual void UpdatePreviewInstance(const Metasound::FVertexKey& InParameterName, TScriptInterface<IAudioParameterInterface>& InParameterInterface) const override;
 };
 
 UCLASS(MinimalAPI)
@@ -292,27 +158,10 @@ public:
 	UPROPERTY(EditAnywhere, Category = DefaultValue)
 	float Default = 0.f;
 
-	virtual FMetasoundFrontendLiteral GetDefault() const override
-	{
-		FMetasoundFrontendLiteral Literal;
-		Literal.Set(Default);
-		return Literal;
-	}
-
-	virtual EMetasoundFrontendLiteralType GetLiteralType() const override
-	{
-		return EMetasoundFrontendLiteralType::Float;
-	}
-
-	virtual void SetFromLiteral(const FMetasoundFrontendLiteral& InLiteral) override
-	{
-		Metasound::Editor::InputPrivate::ConvertLiteral<float>(InLiteral, Default);
-	}
-
-	virtual void UpdatePreviewInstance(const Metasound::FVertexKey& InParameterName, TScriptInterface<IAudioParameterInterface>& InParameterInterface) const override
-	{
-		InParameterInterface->SetFloat(*InParameterName, Default);
-	}
+	virtual FMetasoundFrontendLiteral GetDefault() const override;
+	virtual EMetasoundFrontendLiteralType GetLiteralType() const override;
+	virtual void SetFromLiteral(const FMetasoundFrontendLiteral& InLiteral) override;
+	virtual void UpdatePreviewInstance(const Metasound::FVertexKey& InParameterName, TScriptInterface<IAudioParameterInterface>& InParameterInterface) const override;
 };
 
 UCLASS(MinimalAPI)
@@ -326,27 +175,10 @@ public:
 	UPROPERTY(EditAnywhere, Category = DefaultValue)
 	TArray<float> Default;
 
-	virtual FMetasoundFrontendLiteral GetDefault() const override
-	{
-		FMetasoundFrontendLiteral Literal;
-		Literal.Set(Default);
-		return Literal;
-	}
-
-	virtual EMetasoundFrontendLiteralType GetLiteralType() const override
-	{
-		return EMetasoundFrontendLiteralType::FloatArray;
-	}
-
-	virtual void SetFromLiteral(const FMetasoundFrontendLiteral& InLiteral) override
-	{
-		Metasound::Editor::InputPrivate::ConvertLiteralToArray<float>(InLiteral, Default);
-	}
-
-	virtual void UpdatePreviewInstance(const Metasound::FVertexKey& InParameterName, TScriptInterface<IAudioParameterInterface>& InParameterInterface) const override
-	{
-		InParameterInterface->SetFloatArray(*InParameterName, Default);
-	}
+	virtual FMetasoundFrontendLiteral GetDefault() const override;
+	virtual EMetasoundFrontendLiteralType GetLiteralType() const override;
+	virtual void SetFromLiteral(const FMetasoundFrontendLiteral& InLiteral) override;
+	virtual void UpdatePreviewInstance(const Metasound::FVertexKey& InParameterName, TScriptInterface<IAudioParameterInterface>& InParameterInterface) const override;
 };
 
 UCLASS(MinimalAPI)
@@ -360,27 +192,10 @@ public:
 	UPROPERTY(EditAnywhere, Category = DefaultValue)
 	FString Default;
 
-	virtual FMetasoundFrontendLiteral GetDefault() const override
-	{
-		FMetasoundFrontendLiteral Literal;
-		Literal.Set(Default);
-		return Literal;
-	}
-
-	virtual EMetasoundFrontendLiteralType GetLiteralType() const override
-	{
-		return EMetasoundFrontendLiteralType::String;
-	}
-
-	virtual void SetFromLiteral(const FMetasoundFrontendLiteral& InLiteral) override
-	{
-		Metasound::Editor::InputPrivate::ConvertLiteral<FString>(InLiteral, Default);
-	}
-
-	virtual void UpdatePreviewInstance(const Metasound::FVertexKey& InParameterName, TScriptInterface<IAudioParameterInterface>& InParameterInterface) const override
-	{
-		InParameterInterface->SetString(*InParameterName, Default);
-	}
+	virtual FMetasoundFrontendLiteral GetDefault() const override;
+	virtual EMetasoundFrontendLiteralType GetLiteralType() const override;
+	virtual void SetFromLiteral(const FMetasoundFrontendLiteral& InLiteral) override;
+	virtual void UpdatePreviewInstance(const Metasound::FVertexKey& InParameterName, TScriptInterface<IAudioParameterInterface>& InParameterInterface) const override;
 };
 
 UCLASS(MinimalAPI)
@@ -394,27 +209,10 @@ public:
 	UPROPERTY(EditAnywhere, Category = DefaultValue)
 	TArray<FString> Default;
 
-	virtual FMetasoundFrontendLiteral GetDefault() const override
-	{
-		FMetasoundFrontendLiteral Literal;
-		Literal.Set(Default);
-		return Literal;
-	}
-
-	virtual EMetasoundFrontendLiteralType GetLiteralType() const override
-	{
-		return EMetasoundFrontendLiteralType::StringArray;
-	}
-
-	virtual void SetFromLiteral(const FMetasoundFrontendLiteral& InLiteral) override
-	{
-		Metasound::Editor::InputPrivate::ConvertLiteralToArray<FString>(InLiteral, Default);
-	}
-
-	virtual void UpdatePreviewInstance(const Metasound::FVertexKey& InParameterName, TScriptInterface<IAudioParameterInterface>& InParameterInterface) const override
-	{
-		InParameterInterface->SetStringArray(*InParameterName, Default);
-	}
+	virtual FMetasoundFrontendLiteral GetDefault() const override;
+	virtual EMetasoundFrontendLiteralType GetLiteralType() const override;
+	virtual void SetFromLiteral(const FMetasoundFrontendLiteral& InLiteral) override;
+	virtual void UpdatePreviewInstance(const Metasound::FVertexKey& InParameterName, TScriptInterface<IAudioParameterInterface>& InParameterInterface) const override;
 };
 
 // Broken out to be able to customize and swap AllowedClass based on provided object proxy
@@ -438,27 +236,10 @@ public:
 	UPROPERTY(EditAnywhere, Category = DefaultValue)
 	FMetasoundEditorGraphInputObjectRef Default;
 
-	virtual FMetasoundFrontendLiteral GetDefault() const override
-	{
-		FMetasoundFrontendLiteral Literal;
-		Literal.Set(Default.Object);
-		return Literal;
-	}
-
-	virtual EMetasoundFrontendLiteralType GetLiteralType() const override
-	{
-		return EMetasoundFrontendLiteralType::UObject;
-	}
-
-	virtual void SetFromLiteral(const FMetasoundFrontendLiteral& InLiteral) override
-	{
-		ensure(InLiteral.TryGet(Default.Object));
-	}
-
-	virtual void UpdatePreviewInstance(const Metasound::FVertexKey& InParameterName, TScriptInterface<IAudioParameterInterface>& InParameterInterface) const override
-	{
-		// TODO. We need proxy object here safely.
-	}
+	virtual FMetasoundFrontendLiteral GetDefault() const override;
+	virtual EMetasoundFrontendLiteralType GetLiteralType() const override;
+	virtual void SetFromLiteral(const FMetasoundFrontendLiteral& InLiteral) override;
+	virtual void UpdatePreviewInstance(const Metasound::FVertexKey& InParameterName, TScriptInterface<IAudioParameterInterface>& InParameterInterface) const override;
 };
 
 UCLASS(MinimalAPI)
@@ -472,33 +253,8 @@ public:
 	UPROPERTY(EditAnywhere, Category = DefaultValue)
 	TArray<FMetasoundEditorGraphInputObjectRef> Default;
 
-	virtual FMetasoundFrontendLiteral GetDefault() const override
-	{
-		TArray<UObject*> ObjectArray;
-		Algo::Transform(Default, ObjectArray, [](const FMetasoundEditorGraphInputObjectRef& InValue) { return InValue.Object; });
-
-		FMetasoundFrontendLiteral Literal;
-		Literal.Set(ObjectArray);
-		return Literal;
-	}
-
-	virtual EMetasoundFrontendLiteralType GetLiteralType() const override
-	{
-		return EMetasoundFrontendLiteralType::UObjectArray;
-	}
-
-	virtual void SetFromLiteral(const FMetasoundFrontendLiteral& InLiteral) override
-	{
-		TArray<UObject*> ObjectArray;
-		ensure(InLiteral.TryGet(ObjectArray));
-
-		Algo::Transform(ObjectArray, Default, [](UObject* InValue) { return FMetasoundEditorGraphInputObjectRef { InValue }; });
-	}
-
-	virtual void UpdatePreviewInstance(const Metasound::FVertexKey& InParameterName, TScriptInterface<IAudioParameterInterface>& InParameterInterface) const override
-	{
-		TArray<UObject*> ObjectArray;
-		Algo::Transform(Default, ObjectArray, [](const FMetasoundEditorGraphInputObjectRef& InValue) { return InValue.Object; });
-		// TODO. We need proxy object here safely.
-	}
+	virtual FMetasoundFrontendLiteral GetDefault() const override;
+	virtual EMetasoundFrontendLiteralType GetLiteralType() const override;
+	virtual void SetFromLiteral(const FMetasoundFrontendLiteral& InLiteral) override;
+	virtual void UpdatePreviewInstance(const Metasound::FVertexKey& InParameterName, TScriptInterface<IAudioParameterInterface>& InParameterInterface) const override;
 };
\ No newline at end of file
diff --git a/Engine/Plugins/Runtime/Metasound/Source/MetasoundEditor/Public/MetasoundEditorGraphNode.h b/Engine/Plugins/Runtime/Metasound/Source/MetasoundEditor/Public/MetasoundEditorGraphNode.h
--- a/Engine/Plugins/Runtime/Metasound/Source/MetasoundEditor/Public/MetasoundEditorGraphNode.h
+++ b/Engine/Plugins/Runtime/Metasound/Source/MetasoundEditor/Public/MetasoundEditorGraphNode.h
@@ -11,6 +11,7 @@
 #include "MetasoundFrontendRegistries.h"
 #include "Misc/Guid.h"
 #include "Sound/SoundWave.h"
+#include "Textures/SlateIcon.h"
 #include "UObject/ObjectMacros.h"
 
 #include "MetasoundEditorGraphNode.generated.h"
@@ -62,18 +63,17 @@ public:
 	// End of UEdGraphNode interface
 
 	// UObject interface
-	virtual void PostLoad() override;
+	virtual void PreSave(FObjectPreSaveContext InSaveContext) override;
 
-#if WITH_EDITOR
-	virtual void PostEditChangeProperty(FPropertyChangedEvent& InEvent) override;
-	virtual void PostEditChangeChainProperty(FPropertyChangedChainEvent& InEvent) override;
+	virtual void PostLoad() override;
 	virtual void PostEditUndo() override;
-#endif // WITH_EDITOR
 
 	virtual void PostEditImport() override;
 	virtual void PostDuplicate(bool bDuplicateForPIE) override;
 	// End of UObject interface
 
+	virtual FSlateIcon GetNodeTitleIcon() const { return FSlateIcon(); }
+
 	virtual bool CanAddInputPin() const
 	{
 		return false;
@@ -117,6 +117,8 @@ public:
 	virtual bool CanUserDeleteNode() const override;
 
 protected:
+	virtual FLinearColor GetNodeTitleColor() const override;
+	virtual FSlateIcon GetNodeTitleIcon() const override;
 	virtual void SetNodeID(FGuid InNodeID) override;
 
 	friend class Metasound::Editor::FGraphBuilder;
@@ -134,18 +136,31 @@ protected:
 	UPROPERTY()
 	FGuid NodeID;
 
+	// Whether or not the referenced class is natively defined
+	// (false if defined in another asset). Cached from node
+	// implementation for fast access when validated.
+	UPROPERTY(Transient)
+	bool bIsClassNative = false;
+
 public:
+	virtual void PostLoad() override;
+
 	virtual FMetasoundFrontendClassName GetClassName() const override { return ClassName; }
 	virtual FGuid GetNodeID() const override { return NodeID; }
+	virtual FLinearColor GetNodeTitleColor() const override;
+	virtual FSlateIcon GetNodeTitleIcon() const override;
 
-	FMetasoundFrontendVersionNumber GetMajorUpdateAvailable() const;
-	FMetasoundFrontendVersionNumber GetMinorUpdateAvailable() const;
+	FMetasoundFrontendVersionNumber FindHighestVersionInRegistry() const;
+	FMetasoundFrontendVersionNumber FindHighestMinorVersionInRegistry() const;
+	bool CanAutoUpdate() const;
 
 	// Attempts to replace this node with a new one of the same class and given version number.
 	// If this node is already of the given version, returns itself. If update fails, returns this node.
 	UMetasoundEditorGraphExternalNode* UpdateToVersion(const FMetasoundFrontendVersionNumber& InNewVersion, bool bInPropagateErrorMessages);
 
-	bool Validate(Metasound::Editor::FGraphNodeValidationResult& OutResult);
+	// Validates node and returns whether or not the node is valid.
+	bool Validate(Metasound::Editor::FGraphNodeValidationResult& OutResult, bool bClearUpgradeMessage = true);
+
 
 protected:
 	// Refreshes all Pin Metadata from the associated Frontend node's default class interface.
diff --git a/Engine/Plugins/Runtime/Metasound/Source/MetasoundEditor/Public/MetasoundEditorGraphSchema.h b/Engine/Plugins/Runtime/Metasound/Source/MetasoundEditor/Public/MetasoundEditorGraphSchema.h
--- a/Engine/Plugins/Runtime/Metasound/Source/MetasoundEditor/Public/MetasoundEditorGraphSchema.h
+++ b/Engine/Plugins/Runtime/Metasound/Source/MetasoundEditor/Public/MetasoundEditorGraphSchema.h
@@ -208,7 +208,7 @@ class UMetasoundEditorGraphSchema : public UEdGraphSchema
 private:
 	/** Adds actions for creating actions associated with graph DataTypes */
 	void GetConversionActions(FGraphActionMenuBuilder& ActionMenuBuilder, Metasound::Editor::FActionClassFilters InFilters = Metasound::Editor::FActionClassFilters(), bool bShowSelectedActions = true) const;
-	void GetFunctionActions(FGraphActionMenuBuilder& ActionMenuBuilder, Metasound::Editor::FActionClassFilters InFilters = Metasound::Editor::FActionClassFilters(), bool bShowSelectedActions = true) const;
+	void GetFunctionActions(FGraphActionMenuBuilder& ActionMenuBuilder, Metasound::Editor::FActionClassFilters InFilters = Metasound::Editor::FActionClassFilters(), bool bShowSelectedActions = true, Metasound::Frontend::FConstGraphHandle InGraphHandle = Metasound::Frontend::IGraphController::GetInvalidHandle()) const;
 
 	void GetDataTypeInputNodeActions(FGraphContextMenuBuilder& InMenuBuilder, Metasound::Frontend::FConstGraphHandle InGraphHandle, Metasound::Editor::FInterfaceNodeFilterFunction InFilter = Metasound::Editor::FInterfaceNodeFilterFunction(), bool bShowSelectedActions = true) const;
 	void GetDataTypeOutputNodeActions(FGraphContextMenuBuilder& InMenuBuilder, Metasound::Frontend::FConstGraphHandle InGraphHandle, Metasound::Editor::FInterfaceNodeFilterFunction InFilter = Metasound::Editor::FInterfaceNodeFilterFunction(), bool bShowSelectedActions = true) const;
diff --git a/Engine/Plugins/Runtime/Metasound/Source/MetasoundEditor/Public/MetasoundEditorGraphValidation.h b/Engine/Plugins/Runtime/Metasound/Source/MetasoundEditor/Public/MetasoundEditorGraphValidation.h
--- a/Engine/Plugins/Runtime/Metasound/Source/MetasoundEditor/Public/MetasoundEditorGraphValidation.h
+++ b/Engine/Plugins/Runtime/Metasound/Source/MetasoundEditor/Public/MetasoundEditorGraphValidation.h
@@ -16,6 +16,9 @@ namespace Metasound
 			// Node associated with validation result
 			UMetasoundEditorGraphExternalNode* Node = nullptr;
 
+			// Whether or not the node can be auto-updated
+			bool bIsAutoUpdateEligible = false;
+
 			// Whether associated node is in invalid state (document is corrupt and no Frontend representation could be
 			// found for the node)
 			bool bIsInvalid = false;
@@ -43,6 +46,16 @@ namespace Metasound
 				return NodeResults;
 			}
 
+			// Returns whether or not the graph is in a valid state
+			bool IsValid() const
+			{
+				auto NodeIsInvalid = [](const FGraphNodeValidationResult& Result)
+				{
+					return Result.bIsInvalid || !Result.MissingClass.IsNone();
+				};
+				return !NodeResults.ContainsByPredicate(NodeIsInvalid);
+			}
+
 			// Returns set of missing class names from the validation results
 			TSet<FName> FindMissingClasses() const
 			{
diff --git a/Engine/Plugins/Runtime/Metasound/Source/MetasoundEngine/Private/Metasound.cpp b/Engine/Plugins/Runtime/Metasound/Source/MetasoundEngine/Private/Metasound.cpp
--- a/Engine/Plugins/Runtime/Metasound/Source/MetasoundEngine/Private/Metasound.cpp
+++ b/Engine/Plugins/Runtime/Metasound/Source/MetasoundEngine/Private/Metasound.cpp
@@ -8,6 +8,7 @@
 #include "MetasoundAudioFormats.h"
 #include "MetasoundEngineArchetypes.h"
 #include "MetasoundEngineEnvironment.h"
+#include "MetasoundEnvironment.h"
 #include "MetasoundFrontendController.h"
 #include "MetasoundFrontendQuery.h"
 #include "MetasoundFrontendQuerySteps.h"
@@ -19,7 +20,7 @@
 #include "MetasoundPrimitives.h"
 #include "MetasoundReceiveNode.h"
 #include "MetasoundTrigger.h"
-#include "MetasoundEnvironment.h"
+#include "MetasoundUObjectRegistry.h"
 #include "UObject/ObjectSaveContext.h"
 
 #if WITH_EDITORONLY_DATA
@@ -36,22 +37,10 @@ UMetaSound::UMetaSound(const FObjectInitializer& ObjectInitializer)
 }
 
 #if WITH_EDITOR
-void UMetaSound::PreSave(FObjectPreSaveContext SaveContext)
-{
-	Super::PreSave(SaveContext);
-
-	// TODO: Enable Composition
-	// RegisterGraphWithFrontend();
-}
-
 void UMetaSound::PostEditUndo()
 {
 	Super::PostEditUndo();
-
-	if (Graph)
-	{
-		Graph->Synchronize();
-	}
+	Metasound::PostAssetUndo(*this);
 }
 
 void UMetaSound::PostEditChangeProperty(FPropertyChangedEvent& InEvent)
@@ -59,9 +48,26 @@ void UMetaSound::PostEditChangeProperty(FPropertyChangedEvent& InEvent)
 	Super::PostEditChangeProperty(InEvent);
 	Metasound::PostEditChangeProperty(*this, InEvent);
 }
-
 #endif // WITHEDITOR
 
+void UMetaSound::BeginDestroy()
+{
+	UnregisterGraphWithFrontend();
+	Super::BeginDestroy();
+}
+
+void UMetaSound::PreSave(FObjectPreSaveContext InSaveContext)
+{
+	Super::PreSave(InSaveContext);
+	Metasound::PreSaveAsset(*this);
+}
+
+void UMetaSound::Serialize(FArchive& InArchive)
+{
+	Super::Serialize(InArchive);
+	Metasound::SerializeToArchive(*this, InArchive);
+}
+
 #if WITH_EDITORONLY_DATA
 UEdGraph* UMetaSound::GetGraph()
 {
diff --git a/Engine/Plugins/Runtime/Metasound/Source/MetasoundEngine/Private/MetasoundSource.cpp b/Engine/Plugins/Runtime/Metasound/Source/MetasoundEngine/Private/MetasoundSource.cpp
--- a/Engine/Plugins/Runtime/Metasound/Source/MetasoundEngine/Private/MetasoundSource.cpp
+++ b/Engine/Plugins/Runtime/Metasound/Source/MetasoundEngine/Private/MetasoundSource.cpp
@@ -53,26 +53,10 @@ UMetaSoundSource::UMetaSoundSource(const FObjectInitializer& ObjectInitializer)
 }
 
 #if WITH_EDITOR
-
-void UMetaSoundSource::PreSave(FObjectPreSaveContext SaveContext)
-{
-	Super::PreSave(SaveContext);
-
-	// TODO: Move these to be run anytime the interface changes or
-	// the node is versioned
-
-	//  TODO: Enable Composition
-	// RegisterGraphWithFrontend();
-}
-
 void UMetaSoundSource::PostEditUndo()
 {
 	Super::PostEditUndo();
-
-	if (Graph)
-	{
-		Graph->Synchronize();
-	}
+	Metasound::PostAssetUndo(*this);
 }
 
 void UMetaSoundSource::PostEditChangeProperty(FPropertyChangedEvent& InEvent)
@@ -120,6 +104,24 @@ void UMetaSoundSource::PostEditChangeProperty(FPropertyChangedEvent& InEvent)
 }
 #endif // WITH_EDITOR
 
+void UMetaSoundSource::BeginDestroy()
+{
+	UnregisterGraphWithFrontend();
+	Super::BeginDestroy();
+}
+
+void UMetaSoundSource::PreSave(FObjectPreSaveContext InSaveContext)
+{
+	Super::PreSave(InSaveContext);
+	Metasound::PreSaveAsset(*this);
+}
+
+void UMetaSoundSource::Serialize(FArchive& InArchive)
+{
+	Super::Serialize(InArchive);
+	Metasound::SerializeToArchive(*this, InArchive);
+}
+
 #if WITH_EDITORONLY_DATA
 UEdGraph* UMetaSoundSource::GetGraph()
 {
@@ -222,16 +224,6 @@ ISoundGeneratorPtr UMetaSoundSource::CreateSoundGenerator(const FSoundGeneratorI
 		return ISoundGeneratorPtr(nullptr);
 	}
 
-	// Check version and attempt to version up if out-of-date
-	if (VersionAsset(false /* bMarkDirty */))
-	{
-		const FMetasoundFrontendDocumentMetadata& Metadata = GetDocumentHandle()->GetMetadata();
-		UE_LOG(LogMetaSound, Display,
-			TEXT("Dynamically versioned UMetaSoundSource asset [Name:%s] to %s (Resave to avoid update at runtime)."),
-			*GetName(),
-			*Metadata.Version.Number.ToString());
-	}
-
 	// Inject receive nodes for unused transmittable inputs. Perform edits on a copy
 	// of the document to avoid altering document.
 	// TODO: Use a light wrapper of the document instead of a copy.
diff --git a/Engine/Plugins/Runtime/Metasound/Source/MetasoundEngine/Private/MetasoundUObjectRegistry.cpp b/Engine/Plugins/Runtime/Metasound/Source/MetasoundEngine/Private/MetasoundUObjectRegistry.cpp
--- a/Engine/Plugins/Runtime/Metasound/Source/MetasoundEngine/Private/MetasoundUObjectRegistry.cpp
+++ b/Engine/Plugins/Runtime/Metasound/Source/MetasoundEngine/Private/MetasoundUObjectRegistry.cpp
@@ -22,7 +22,7 @@ namespace Metasound
 {
 	namespace AssetSubsystemPrivate
 	{
-		static bool GetAssetClassInfo(const FAssetData& InAssetData, Frontend::FNodeClassInfo& OutInfo)
+		bool GetAssetClassInfo(const FAssetData& InAssetData, Frontend::FNodeClassInfo& OutInfo)
 		{
 			using namespace Metasound;
 			using namespace Metasound::Frontend;
@@ -74,8 +74,7 @@ namespace Metasound
 
 void UMetaSoundAssetSubsystem::Initialize(FSubsystemCollectionBase& InCollection)
 {
-//  TODO: Enable Composition
-//	FCoreDelegates::OnPostEngineInit.AddUObject(this, &UMetaSoundAssetSubsystem::PostEngineInit);
+	FCoreDelegates::OnPostEngineInit.AddUObject(this, &UMetaSoundAssetSubsystem::PostEngineInit);
 }
 
 void UMetaSoundAssetSubsystem::PostEngineInit()
@@ -118,7 +117,26 @@ void UMetaSoundAssetSubsystem::Deinitialize()
 {
 }
 
-void UMetaSoundAssetSubsystem::AddOrUpdateAsset(const FAssetData& InAssetData)
+void UMetaSoundAssetSubsystem::AddOrUpdateAsset(UObject& InObject, bool bInRegisterWithFrontend)
+{
+	using namespace Metasound;
+	using namespace Metasound::AssetSubsystemPrivate;
+	using namespace Metasound::Frontend;
+
+	FMetasoundAssetBase* MetaSoundAsset = Metasound::IMetasoundUObjectRegistry::Get().GetObjectAsAssetBase(&InObject);
+	check(MetaSoundAsset);
+
+	if (bInRegisterWithFrontend)
+	{
+		MetaSoundAsset->RegisterGraphWithFrontend();
+	}
+
+	FNodeClassInfo ClassInfo = MetaSoundAsset->GetAssetClassInfo();
+	const FNodeRegistryKey RegistryKey = NodeRegistryKey::CreateKey(ClassInfo);
+	PathMap.FindOrAdd(RegistryKey) = InObject.GetPathName();
+}
+
+void UMetaSoundAssetSubsystem::AddOrUpdateAsset(const FAssetData& InAssetData, bool bInRegisterWithFrontend)
 {
 	using namespace Metasound;
 	using namespace Metasound::AssetSubsystemPrivate;
@@ -129,7 +147,7 @@ void UMetaSoundAssetSubsystem::AddOrUpdateAsset(const FAssetData& InAssetData)
 
 	FNodeClassInfo ClassInfo;
 	bool bClassInfoFound = GetAssetClassInfo(InAssetData, ClassInfo);
-	if (!bClassInfoFound || bLoadRequiredToRegisterAssetClasses)
+	if (!bClassInfoFound || bLoadRequiredToRegisterAssetClasses || bInRegisterWithFrontend)
 	{
 		UObject* Object = nullptr;
 
@@ -152,61 +170,98 @@ void UMetaSoundAssetSubsystem::AddOrUpdateAsset(const FAssetData& InAssetData)
 
 		if (ensure(Object))
 		{
-			
-			FMetasoundAssetBase* MetaSoundAsset = Metasound::IMetasoundUObjectRegistry::Get().GetObjectAsAssetBase(Object);
-			check(MetaSoundAsset);
+			AddOrUpdateAsset(*Object, bInRegisterWithFrontend);
+			return;
+		}
+	}
 
-			// Must version to ensure registration uses the correct key,
-			// which must be based off of most up-to-date document model
-			// for safety.
-			MetaSoundAsset->VersionAsset(false /* bMarkDirty */);
+	const FNodeRegistryKey RegistryKey = NodeRegistryKey::CreateKey(ClassInfo);
+	PathMap.FindOrAdd(RegistryKey) = InAssetData.ObjectPath;
+}
 
-			MetaSoundAsset->RegisterGraphWithFrontend();
-		}
+FMetasoundAssetBase* UMetaSoundAssetSubsystem::FindAssetFromKey(const Metasound::Frontend::FNodeRegistryKey& RegistryKey) const
+{
+	if (const FSoftObjectPath* ObjectPath = FindObjectPathFromKey(RegistryKey))
+	{
+		return TryLoadAsset(*ObjectPath);
 	}
+
+	return nullptr;
+}
+
+const FSoftObjectPath* UMetaSoundAssetSubsystem::FindObjectPathFromKey(const Metasound::Frontend::FNodeRegistryKey& InRegistryKey) const
+{
+	return PathMap.Find(InRegistryKey);
 }
 
-void UMetaSoundAssetSubsystem::RemoveAsset(const FAssetData& InAssetData)
+FMetasoundAssetBase* UMetaSoundAssetSubsystem::TryLoadAsset(const FSoftObjectPath& InObjectPath) const
+{
+	return Metasound::IMetasoundUObjectRegistry::Get().GetObjectAsAssetBase(InObjectPath.TryLoad());
+}
+
+void UMetaSoundAssetSubsystem::RemoveAsset(UObject& InObject, bool bInUnregisterWithFrontend)
 {
 	using namespace Metasound;
-	using namespace Metasound::AssetSubsystemPrivate;
 	using namespace Metasound::Frontend;
 
-	FNodeClassInfo ClassInfo;
-	if (!GetAssetClassInfo(InAssetData, ClassInfo))
+	if (FMetasoundAssetBase* MetaSoundAsset = Metasound::IMetasoundUObjectRegistry::Get().GetObjectAsAssetBase(&InObject))
 	{
-		UObject* Object = nullptr;
+		const FNodeClassInfo ClassInfo = MetaSoundAsset->GetAssetClassInfo();
+		FNodeRegistryKey RegistryKey = FMetasoundFrontendRegistryContainer::Get()->GetRegistryKey(ClassInfo);
+		PathMap.Remove(RegistryKey);
 
-		FSoftObjectPath Path(InAssetData.ObjectPath);
-		if (InAssetData.IsAssetLoaded())
-		{
-			Object = Path.ResolveObject();
-		}
-		else
+		if (bInUnregisterWithFrontend)
 		{
-			Object = Path.TryLoad();
+			MetaSoundAsset->UnregisterGraphWithFrontend();
 		}
+	}
+}
 
-		if (ensure(Object))
+void UMetaSoundAssetSubsystem::RemoveAsset(const FAssetData& InAssetData, bool bInUnregisterWithFrontend)
+{
+	using namespace Metasound;
+	using namespace Metasound::Frontend;
+
+	FNodeClassInfo ClassInfo;
+	if (ensureAlways(AssetSubsystemPrivate::GetAssetClassInfo(InAssetData, ClassInfo)))
+	{
+		FNodeRegistryKey RegistryKey = FMetasoundFrontendRegistryContainer::Get()->GetRegistryKey(ClassInfo);
+		PathMap.Remove(RegistryKey);
+
+		UObject* Object = InAssetData.GetAsset();
+		if (Object && bInUnregisterWithFrontend)
 		{
 			FMetasoundAssetBase* MetaSoundAsset = Metasound::IMetasoundUObjectRegistry::Get().GetObjectAsAssetBase(Object);
 			check(MetaSoundAsset);
 
-			FDocumentHandle Document = MetaSoundAsset->GetDocumentHandle();
+			MetaSoundAsset->UnregisterGraphWithFrontend();
+		}
+	}
+}
 
-			// Must version to ensure registration uses the correct key,
-			// which must be based off of most up-to-date document model
-			// for safety.
-			const FName AssetName = Object->GetFName();
-			const FString AssetPath = Object->GetPathName();
-			FVersionDocument(AssetName, AssetPath).Transform(Document);
+void UMetaSoundAssetSubsystem::RenameAsset(const FAssetData& InAssetData, bool bInReregisterWithFrontend)
+{
+	RemoveAsset(InAssetData, bInReregisterWithFrontend);
+	SynchronizeAssetClassName(InAssetData);
+	AddOrUpdateAsset(InAssetData, bInReregisterWithFrontend);
+}
 
-			ClassInfo = MetaSoundAsset->GetAssetClassInfo();
-		}
+void UMetaSoundAssetSubsystem::SynchronizeAssetClassName(const FAssetData& InAssetData)
+{
+	UObject* Object = nullptr;
+	FSoftObjectPath Path(InAssetData.ObjectPath);
+	if (InAssetData.IsAssetLoaded())
+	{
+		Object = Path.ResolveObject();
+	}
+	else
+	{
+		Object = Path.TryLoad();
 	}
 
-	const FNodeRegistryKey RegistryKey = FMetasoundFrontendRegistryContainer::Get()->GetRegistryKey(ClassInfo);
-	FMetasoundFrontendRegistryContainer::Get()->UnregisterNode(RegistryKey);
+	FMetasoundAssetBase* MetaSoundAsset = Metasound::IMetasoundUObjectRegistry::Get().GetObjectAsAssetBase(Object);
+	check(MetaSoundAsset);
+	Metasound::Frontend::FSynchronizeAssetClassName(InAssetData.AssetName).Transform(MetaSoundAsset->GetDocumentHandle());
 }
 
 namespace Metasound
@@ -316,7 +371,7 @@ namespace Metasound
 					PackageToSaveTo = GetTransientPackage();
 				}
 
-				UObject* NewMetasoundObject = InEntry.NewObject(PackageToSaveTo, *InDocument.RootGraph.Metadata.ClassName.GetFullName().ToString());
+				UObject* NewMetasoundObject = InEntry.NewObject(PackageToSaveTo, *InDocument.RootGraph.Metadata.GetClassName().GetFullName().ToString());
 				FMetasoundAssetBase* NewAssetBase = InEntry.Cast(NewMetasoundObject);
 				if (ensure(nullptr != NewAssetBase))
 				{
diff --git a/Engine/Plugins/Runtime/Metasound/Source/MetasoundEngine/Public/Metasound.h b/Engine/Plugins/Runtime/Metasound/Source/MetasoundEngine/Public/Metasound.h
--- a/Engine/Plugins/Runtime/Metasound/Source/MetasoundEngine/Public/Metasound.h
+++ b/Engine/Plugins/Runtime/Metasound/Source/MetasoundEngine/Public/Metasound.h
@@ -9,6 +9,8 @@
 #include "MetasoundFrontendDocument.h"
 #include "MetasoundOperatorSettings.h"
 #include "MetasoundRouter.h"
+#include "MetasoundUObjectRegistry.h"
+#include "Serialization/Archive.h"
 
 #if WITH_EDITORONLY_DATA
 #include "Algo/Transform.h"
@@ -17,6 +19,21 @@
 #include "Metasound.generated.h"
 
 
+UCLASS()
+class METASOUNDENGINE_API UMetasoundEditorGraphBase : public UEdGraph
+{
+	GENERATED_BODY()
+
+public:
+	virtual bool IsEditorOnly() const override { return true; }
+	virtual bool NeedsLoadForEditorGame() const override { return false; }
+
+	virtual bool Synchronize() { return false; }
+	virtual bool Validate(bool bInAutoUpdate) { return false; }
+	virtual void RegisterGraphWithFrontend() { }
+};
+
+
 namespace Metasound
 {
 	namespace ConsoleVariables
@@ -24,11 +41,53 @@ namespace Metasound
 		static float BlockRate = 100.f;
 	} // namespace ConsoleVariables
 
+#if WITH_EDITOR
 	template <typename TMetaSoundObject>
 	void PostEditChangeProperty(TMetaSoundObject& InMetaSound, FPropertyChangedEvent& InEvent)
 	{
-		// TODO: Update registry info here if interface has changed,
-		// potentially enabling removal of PreSave registration
+	}
+
+	template <typename TMetaSoundObject>
+	void PostAssetUndo(TMetaSoundObject& InMetaSound)
+	{
+		if (UMetasoundEditorGraphBase* MetaSoundGraph = Cast<UMetasoundEditorGraphBase>(InMetaSound.GetGraph()))
+		{
+			if (MetaSoundGraph->Validate(false /* bAutoUpdate */))
+			{
+				MetaSoundGraph->RegisterGraphWithFrontend();
+			}
+		}
+	}
+#endif // WITH_EDITOR
+
+	template <typename TMetaSoundObject>
+	void PreSaveAsset(TMetaSoundObject& InMetaSound)
+	{
+#if WITH_EDITORONLY_DATA
+		if (UMetasoundEditorGraphBase* MetaSoundGraph = Cast<UMetasoundEditorGraphBase>(InMetaSound.GetGraph()))
+		{
+			const bool bAutoUpdate = false;
+			const bool bIsValid = MetaSoundGraph->Validate(bAutoUpdate /* bAutoUpdate */);
+			if (bIsValid)
+			{
+				MetaSoundGraph->RegisterGraphWithFrontend();
+			}
+		}
+#endif // WITH_EDITORONLY_DATA
+	}
+
+	template <typename TMetaSoundObject>
+	void SerializeToArchive(TMetaSoundObject& InMetaSound, FArchive& InArchive)
+	{
+		if (InArchive.IsLoading())
+		{
+			const bool bUpdateReferences = false;
+			const bool bMarkDirty = false;
+
+			UMetaSoundAssetSubsystem& AssetSubsystem = UMetaSoundAssetSubsystem::Get();
+			InMetaSound.VersionAsset(AssetSubsystem, bMarkDirty, bUpdateReferences);
+			InMetaSound.AutoUpdate(AssetSubsystem, bMarkDirty, bUpdateReferences);
+		}
 	}
 
 #if WITH_EDITORONLY_DATA
@@ -72,18 +131,6 @@ namespace Metasound
 #endif // WITH_EDITORONLY_DATA
 } // namespace Metasound
 
-UCLASS()
-class METASOUNDENGINE_API UMetasoundEditorGraphBase : public UEdGraph
-{
-	GENERATED_BODY()
-
-public:
-	virtual bool IsEditorOnly() const override { return true; }
-	virtual bool NeedsLoadForEditorGame() const override { return false; }
-
-	virtual void Synchronize() { }
-};
-
 
 /**
  * This asset type is used for Metasound assets that can only be used as nodes in other Metasound graphs.
@@ -106,6 +153,9 @@ protected:
 public:
 	UMetaSound(const FObjectInitializer& ObjectInitializer);
 
+	UPROPERTY()
+	TSet<FSoftObjectPath> ReferencedAssets;
+
 	UPROPERTY(AssetRegistrySearchable)
 	FGuid AssetClassID;
 
@@ -153,11 +203,24 @@ public:
 #endif // #if WITH_EDITORONLY_DATA
 
 #if WITH_EDITOR
-	virtual void PreSave(FObjectPreSaveContext SaveContext) override;
 	virtual void PostEditUndo() override;
 	virtual void PostEditChangeProperty(FPropertyChangedEvent& InEvent) override;
 #endif // WITH_EDITOR
 
+	virtual void BeginDestroy() override;
+	virtual void PreSave(FObjectPreSaveContext InSaveContext) override;
+	virtual void Serialize(FArchive& InArchive) override;
+
+	virtual TSet<FSoftObjectPath>& GetReferencedAssets() override
+	{
+		return ReferencedAssets;
+	}
+
+	virtual const TSet<FSoftObjectPath>& GetReferencedAssets() const override
+	{
+		return ReferencedAssets;
+	}
+
 	// Returns Asset Metadata associated with this MetaSound
 	virtual Metasound::Frontend::FNodeClassInfo GetAssetClassInfo() const override;
 
diff --git a/Engine/Plugins/Runtime/Metasound/Source/MetasoundEngine/Public/MetasoundSource.h b/Engine/Plugins/Runtime/Metasound/Source/MetasoundEngine/Public/MetasoundSource.h
--- a/Engine/Plugins/Runtime/Metasound/Source/MetasoundEngine/Public/MetasoundSource.h
+++ b/Engine/Plugins/Runtime/Metasound/Source/MetasoundEngine/Public/MetasoundSource.h
@@ -42,6 +42,9 @@ protected:
 	UPROPERTY(EditAnywhere, Category = CustomView)
 	FMetasoundFrontendDocument RootMetasoundDocument;
 
+	UPROPERTY()
+	TSet<FSoftObjectPath> ReferencedAssets;
+
 #if WITH_EDITORONLY_DATA
 	UPROPERTY()
 	UMetasoundEditorGraphBase* Graph;
@@ -101,7 +104,6 @@ public:
 #endif // #if WITH_EDITORONLY_DATA
 
 #if WITH_EDITOR
-	virtual void PreSave(FObjectPreSaveContext SaveContext) override;
 	virtual void PostEditUndo() override;
 
 	virtual bool GetRedrawThumbnail() const override
@@ -121,6 +123,20 @@ public:
 	virtual void PostEditChangeProperty(FPropertyChangedEvent& InEvent) override;
 #endif // WITH_EDITOR
 
+	virtual void BeginDestroy() override;
+	virtual void PreSave(FObjectPreSaveContext InSaveContext) override;
+	virtual void Serialize(FArchive& Ar) override;
+
+	virtual TSet<FSoftObjectPath>& GetReferencedAssets() override
+	{
+		return ReferencedAssets;
+	}
+
+	virtual const TSet<FSoftObjectPath>& GetReferencedAssets() const override
+	{
+		return ReferencedAssets;
+	}
+
 	// Returns Asset Metadata associated with this MetaSoundSource
 	virtual Metasound::Frontend::FNodeClassInfo GetAssetClassInfo() const override;
 
diff --git a/Engine/Plugins/Runtime/Metasound/Source/MetasoundEngine/Public/MetasoundUObjectRegistry.h b/Engine/Plugins/Runtime/Metasound/Source/MetasoundEngine/Public/MetasoundUObjectRegistry.h
--- a/Engine/Plugins/Runtime/Metasound/Source/MetasoundEngine/Public/MetasoundUObjectRegistry.h
+++ b/Engine/Plugins/Runtime/Metasound/Source/MetasoundEngine/Public/MetasoundUObjectRegistry.h
@@ -13,7 +13,7 @@
 
 /** The subsystem in charge of the MetaSound asset registry */
 UCLASS()
-class METASOUNDENGINE_API UMetaSoundAssetSubsystem : public UEngineSubsystem
+class METASOUNDENGINE_API UMetaSoundAssetSubsystem : public UEngineSubsystem, public IMetaSoundAssetInterface
 {
 	GENERATED_BODY()
 
@@ -27,12 +27,23 @@ public:
 		return *GEngine->GetEngineSubsystem<UMetaSoundAssetSubsystem>();
 	}
 
-	void AddOrUpdateAsset(const FAssetData& InAssetData);
-	void RemoveAsset(const FAssetData& InAssetData);
+	void AddOrUpdateAsset(UObject& InObject, bool bInRegisterWithFrontend = true);
+	void AddOrUpdateAsset(const FAssetData& InAssetData, bool bInRegisterWithFrontend = true);
+	void RemoveAsset(UObject& InObject, bool bInUnregisterWithFrontend = true);
+	void RemoveAsset(const FAssetData& InAssetData, bool bInUnregisterWithFrontend = true);
+	void RenameAsset(const FAssetData& InAssetData, bool bInReregisterWithFrontend = true);
+	void SynchronizeAssetClassName(const FAssetData& InAssetData);
+
+	virtual FMetasoundAssetBase* FindAssetFromKey(const Metasound::Frontend::FNodeRegistryKey& RegistryKey) const override;
+	virtual const FSoftObjectPath* FindObjectPathFromKey(const Metasound::Frontend::FNodeRegistryKey& RegistryKey) const override;
+	virtual FMetasoundAssetBase* TryLoadAsset(const FSoftObjectPath& InObjectPath) const override;
 
 protected:
 	void PostEngineInit();
 	void PostInitAssetScan();
+
+private:
+	TMap<Metasound::Frontend::FNodeRegistryKey, FSoftObjectPath> PathMap;
 };
 
 namespace Metasound
