First pass on MetaSound Composition/Preset Support

diff --git a/Engine/Plugins/Runtime/Metasound/Source/MetasoundFrontend/Private/MetasoundAssetBase.cpp b/Engine/Plugins/Runtime/Metasound/Source/MetasoundFrontend/Private/MetasoundAssetBase.cpp
--- a/Engine/Plugins/Runtime/Metasound/Source/MetasoundFrontend/Private/MetasoundAssetBase.cpp
+++ b/Engine/Plugins/Runtime/Metasound/Source/MetasoundFrontend/Private/MetasoundAssetBase.cpp
@@ -7,6 +7,7 @@
 #include "HAL/FileManager.h"
 #include "Internationalization/Text.h"
 #include "IStructSerializerBackend.h"
+#include "Logging/LogMacros.h"
 #include "MetasoundArchetype.h"
 #include "MetasoundFrontendArchetypeRegistry.h"
 #include "MetasoundFrontendController.h"
@@ -37,6 +38,28 @@ namespace Metasound
 		const FName RegistryOutputTypes = "RegistryOutputTypes";
 #endif // WITH_EDITORONLY_DATA
 	} // namespace AssetTags
+
+	namespace AssetBasePrivate
+	{
+		void DepthFirstTraversal(const FMetasoundAssetBase& InInitAsset, TFunctionRef<TSet<const FMetasoundAssetBase*>(const FMetasoundAssetBase&)> InVisitFunction)
+		{
+			// Non recursive depth first traversal.
+			TArray<const FMetasoundAssetBase*> Stack({ &InInitAsset });
+			TSet<const FMetasoundAssetBase*> Visited;
+
+			while (!Stack.IsEmpty())
+			{
+				const FMetasoundAssetBase* CurrentNode = Stack.Pop();
+				if (!Visited.Contains(CurrentNode))
+				{
+					TArray<const FMetasoundAssetBase*> Children = InVisitFunction(*CurrentNode).Array();
+					Stack.Append(Children);
+
+					Visited.Add(CurrentNode);
+				}
+			}
+		}
+	} // namespace AssetBasePrivate
 } // namespace Metasound
 
 const FString FMetasoundAssetBase::FileExtension(TEXT(".metasound"));
@@ -65,7 +88,7 @@ void FMetasoundAssetBase::RegisterGraphWithFrontend()
 		{
 			// Copy frontend class to preserve original document.
 			FrontendClass = Document.RootGraph;
-			FrontendClass.Metadata.Type = EMetasoundFrontendClassType::External;
+			FrontendClass.Metadata.SetType(EMetasoundFrontendClassType::External);
 			ClassInfo = FNodeClassInfo(Document.RootGraph, InAssetPath);
 		}
 
@@ -101,6 +124,11 @@ void FMetasoundAssetBase::RegisterGraphWithFrontend()
 			return MakeUnique<FNodeRegistryEntry>(Name, Document, ClassInfo.AssetPath);
 		}
 
+		virtual bool IsNative() const override
+		{
+			return false;
+		}
+
 	private:
 		
 		FString Name;
@@ -109,45 +137,38 @@ void FMetasoundAssetBase::RegisterGraphWithFrontend()
 		FNodeClassInfo ClassInfo;
 	};
 
-	FNodeClassInfo ClassInfo = GetAssetClassInfo();
-	FNodeRegistryKey RegistryKey = NodeRegistryKey::CreateKey(ClassInfo);
+	UnregisterGraphWithFrontend();
 
-	if (NodeRegistryKey::IsValid(RegistryKey))
+	FNodeClassInfo AssetClassInfo = GetAssetClassInfo();
+	if (const FMetasoundFrontendDocument* Doc = DocumentPtr.Get())
 	{
-		// Unregister prior version if it exists.
-		if (bHasRegistered)
-		{
-			ensure(FMetasoundFrontendRegistryContainer::Get()->UnregisterNode(RegistryKey));
-		}
+		RegistryKey = FMetasoundFrontendRegistryContainer::Get()->RegisterNode(MakeUnique<FNodeRegistryEntry>(AssetName, *Doc, AssetClassInfo.AssetPath));
 	}
-
-	if (const FMetasoundFrontendDocument* Doc = DocumentPtr.Get())
+	else
 	{
-		RegistryKey = FMetasoundFrontendRegistryContainer::Get()->RegisterNode(MakeUnique<FNodeRegistryEntry>(AssetName, *Doc, ClassInfo.AssetPath));
+		RegistryKey = FNodeRegistryKey();
 	}
 
 	if (NodeRegistryKey::IsValid(RegistryKey))
 	{
-		bHasRegistered = true;
-
 #if WITH_EDITORONLY_DATA
 		// Refresh Asset Registry Info if successfully registered with Frontend
-		FConstGraphHandle GraphHandle = GetDocumentHandle()->GetRootGraph();
-		const FMetasoundFrontendClassMetadata& ClassFrontendMetadata = GraphHandle->GetGraphMetadata();
-		ClassInfo.AssetClassID = FGuid(ClassFrontendMetadata.ClassName.Name.ToString());
-		FNodeClassName ClassName = ClassFrontendMetadata.ClassName.ToNodeClassName();
+		const FMetasoundFrontendGraphClass& DocumentClassGraph = GetDocumentHandle()->GetRootGraphClass();
+		const FMetasoundFrontendClassMetadata& DocumentClassMetadata = DocumentClassGraph.Metadata;
+		AssetClassInfo.AssetClassID = FGuid(DocumentClassMetadata.GetClassName().Name.ToString());
+		FNodeClassName ClassName = DocumentClassMetadata.GetClassName().ToNodeClassName();
 		FMetasoundFrontendClass GraphClass;
-		ensure(ISearchEngine::Get().FindClassWithMajorVersion(ClassName, ClassFrontendMetadata.Version.Major, GraphClass));
+		ensure(ISearchEngine::Get().FindClassWithMajorVersion(ClassName, DocumentClassMetadata.GetVersion().Major, GraphClass));
 
-		ClassInfo.Version = ClassFrontendMetadata.Version;
+		AssetClassInfo.Version = DocumentClassMetadata.GetVersion();
 
-		ClassInfo.InputTypes.Reset();
-		Algo::Transform(GraphClass.Interface.Inputs, ClassInfo.InputTypes, [] (const FMetasoundFrontendClassInput& Input) { return Input.TypeName; });
+		AssetClassInfo.InputTypes.Reset();
+		Algo::Transform(GraphClass.Interface.Inputs, AssetClassInfo.InputTypes, [] (const FMetasoundFrontendClassInput& Input) { return Input.TypeName; });
 
-		ClassInfo.OutputTypes.Reset();
-		Algo::Transform(GraphClass.Interface.Outputs, ClassInfo.OutputTypes, [](const FMetasoundFrontendClassOutput& Output) { return Output.TypeName; });
+		AssetClassInfo.OutputTypes.Reset();
+		Algo::Transform(GraphClass.Interface.Outputs, AssetClassInfo.OutputTypes, [](const FMetasoundFrontendClassOutput& Output) { return Output.TypeName; });
 
-		SetRegistryAssetClassInfo(MoveTemp(ClassInfo));
+		SetRegistryAssetClassInfo(MoveTemp(AssetClassInfo));
 #endif // WITH_EDITORONLY_DATA
 	}
 	else
@@ -156,15 +177,34 @@ void FMetasoundAssetBase::RegisterGraphWithFrontend()
 	}
 }
 
+void FMetasoundAssetBase::UnregisterGraphWithFrontend()
+{
+	using namespace Metasound::Frontend;
+
+	if (!NodeRegistryKey::IsValid(RegistryKey))
+	{
+		return;
+	}
+
+	const UObject* OwningAsset = GetOwningAsset();
+	if (!ensureAlways(OwningAsset))
+	{
+		return;
+	}
+
+	ensureAlways(FMetasoundFrontendRegistryContainer::Get()->UnregisterNode(RegistryKey));
+	RegistryKey = FNodeRegistryKey();
+}
+
 void FMetasoundAssetBase::SetMetadata(FMetasoundFrontendClassMetadata& InMetadata)
 {
 	FMetasoundFrontendDocument& Doc = GetDocumentChecked();
 	Doc.RootGraph.Metadata = InMetadata;
 
-	if (Doc.RootGraph.Metadata.Type != EMetasoundFrontendClassType::Graph)
+	if (Doc.RootGraph.Metadata.GetType() != EMetasoundFrontendClassType::Graph)
 	{
 		UE_LOG(LogMetaSound, Display, TEXT("Forcing class type to EMetasoundFrontendClassType::Graph on root graph metadata"));
-		Doc.RootGraph.Metadata.Type = EMetasoundFrontendClassType::Graph;
+		Doc.RootGraph.Metadata.SetType(EMetasoundFrontendClassType::Graph);
 	}
 
 	MarkMetasoundDocumentDirty();
@@ -242,7 +282,36 @@ void FMetasoundAssetBase::ConformDocumentToArchetype()
 	}
 }
 
-bool FMetasoundAssetBase::VersionAsset(bool bInMarkDirty)
+bool FMetasoundAssetBase::AutoUpdate(const IMetaSoundAssetInterface& InAssetInterface, bool bInMarkDirty, bool bInUpdateReferencedAssets)
+{
+	using namespace Metasound::Frontend;
+
+	bool bDidEdit = false;
+
+	if (bInUpdateReferencedAssets)
+	{
+		const TSet<FSoftObjectPath>& ReferencedAssets = GetReferencedAssets();
+		for (const FSoftObjectPath& ReferencedAsset : ReferencedAssets)
+		{
+			FMetasoundAssetBase* ReferencedMetaSound = InAssetInterface.TryLoadAsset(ReferencedAsset);
+			if (ReferencedMetaSound)
+			{
+				bDidEdit |= FAutoUpdateRootGraph(*ReferencedMetaSound, InAssetInterface).Transform(ReferencedMetaSound->GetDocumentHandle());
+			}
+		}
+	}
+
+	bDidEdit |= FAutoUpdateRootGraph(*this, InAssetInterface).Transform(GetDocumentHandle());
+
+	if (bDidEdit && bInMarkDirty)
+	{
+		MarkMetasoundDocumentDirty();
+	}
+
+	return bDidEdit;
+}
+
+bool FMetasoundAssetBase::VersionAsset(const IMetaSoundAssetInterface& InAssetInterface, bool bInMarkDirty, bool bInVersionReferencedAssets)
 {
 	using namespace Metasound;
 	using namespace Metasound::Frontend;
@@ -255,24 +324,36 @@ bool FMetasoundAssetBase::VersionAsset(bool bInMarkDirty)
 		AssetPath = OwningAsset->GetPathName();
 	}
 
+	bool bDidEdit = false;
+
+	if (bInVersionReferencedAssets)
+	{
+		const TSet<FSoftObjectPath>& ReferencedAssets = GetReferencedAssets();
+		for (const FSoftObjectPath& ReferencedAsset : ReferencedAssets)
+		{
+			FMetasoundAssetBase* ReferencedMetaSound = InAssetInterface.TryLoadAsset(ReferencedAsset);
+			if (ensure(ReferencedMetaSound))
+			{
+				bDidEdit |= ReferencedMetaSound->VersionAsset(InAssetInterface, bInMarkDirty);
+			}
+		}
+	}
+
 	FDocumentHandle DocumentHandle = GetDocumentHandle();
 
-	const bool bDidUpdateDocumentVersion = FVersionDocument(AssetName, AssetPath).Transform(DocumentHandle);
-	bool bDidMatchRootGraphToArchetype = false;
+	bDidEdit |= FVersionDocument(AssetName, AssetPath).Transform(DocumentHandle);
 	if (FMetasoundFrontendDocument* Doc = GetDocument().Get())
 	{
 		FMetasoundFrontendVersion ArchetypeVerison = GetPreferredArchetypeVersion(*Doc);
-		bDidMatchRootGraphToArchetype = FMatchRootGraphToArchetype(ArchetypeVerison).Transform(DocumentHandle);
+		bDidEdit |= FMatchRootGraphToArchetype(ArchetypeVerison).Transform(DocumentHandle);
 	}
 
-	const bool bDidUpdate = bDidUpdateDocumentVersion || bDidMatchRootGraphToArchetype;
-
-	if (bInMarkDirty && bDidUpdate)
+	if (bDidEdit && bInMarkDirty)
 	{
 		MarkMetasoundDocumentDirty();
 	}
 
-	return bDidUpdate;
+	return bDidEdit;
 }
 
 bool FMetasoundAssetBase::CopyDocumentAndInjectReceiveNodes(uint64 InInstanceID, const FMetasoundFrontendDocument& InSourceDoc, FMetasoundFrontendDocument& OutDestDoc) const
@@ -358,6 +439,68 @@ bool FMetasoundAssetBase::CopyDocumentAndInjectReceiveNodes(uint64 InInstanceID,
 	return true;
 }
 
+void FMetasoundAssetBase::RebuildReferencedAssets(const IMetaSoundAssetInterface& InAssetInterface)
+{
+	using namespace Metasound::Frontend;
+
+	TSet<FSoftObjectPath>& ReferencedAssets = GetReferencedAssets();
+	ReferencedAssets.Reset();
+
+	GetRootGraphHandle()->IterateConstNodes([AssetInterface = &InAssetInterface, RefAssets = &ReferencedAssets](FConstNodeHandle NodeHandle)
+	{
+		using namespace Metasound::Frontend;
+
+		const FMetasoundFrontendClassMetadata& Metadata = NodeHandle->GetClassMetadata();
+		const FNodeRegistryKey LocalRegistryKey = NodeRegistryKey::CreateKey(Metadata);
+		if (const FSoftObjectPath* Path = AssetInterface->FindObjectPathFromKey(LocalRegistryKey))
+		{
+			if (Path->IsValid())
+			{
+				RefAssets->Add(*Path);
+			}
+		}
+	}, EMetasoundFrontendClassType::External);
+}
+
+bool FMetasoundAssetBase::ContainReferenceLoop(const IMetaSoundAssetInterface& IMetaSoundAssetInterface) const
+{
+	return false;
+}
+
+bool FMetasoundAssetBase::AddingReferenceCausesLoop(const FSoftObjectPath& InReferencePath, const IMetaSoundAssetInterface& InAssetInterface) const
+{
+	const FMetasoundAssetBase* ReferenceAsset = InAssetInterface.TryLoadAsset(InReferencePath);
+	if (!ensureAlways(ReferenceAsset))
+	{
+		return false;
+	}
+
+	bool bCausesLoop = false;
+	const FMetasoundAssetBase* Parent = this;
+	Metasound::AssetBasePrivate::DepthFirstTraversal(*ReferenceAsset, [&](const FMetasoundAssetBase& ChildAsset)
+	{
+		TSet<const FMetasoundAssetBase*> Children;
+
+		if (Parent == &ChildAsset)
+		{
+			bCausesLoop = true;
+			return Children;
+		}
+
+		for (const FSoftObjectPath& Path : ChildAsset.GetReferencedAssets())
+		{
+			const FMetasoundAssetBase* ChildReference = InAssetInterface.TryLoadAsset(Path);
+			if (ensureAlways(ChildReference))
+			{
+				Children.Add(ChildReference);
+			}
+		}
+		return Children;
+	});
+
+	return bCausesLoop;
+}
+
 Metasound::FSendAddress FMetasoundAssetBase::CreateSendAddress(uint64 InInstanceID, const FString& InVertexName, const FName& InDataTypeName) const
 {
 	using namespace Metasound;
@@ -377,6 +520,10 @@ void FMetasoundAssetBase::ConvertFromPreset()
 	FMetasoundFrontendGraphStyle Style = GraphHandle->GetGraphStyle();
 	Style.bIsGraphEditable = true;
 	GraphHandle->SetGraphStyle(Style);
+
+	FMetasoundFrontendClassMetadata Metadata = GraphHandle->GetGraphMetadata();
+	Metadata.SetAutoUpdateManagesInterface(false);
+	GraphHandle->SetGraphMetadata(Metadata);
 }
 
 TArray<FMetasoundAssetBase::FSendInfoAndVertexName> FMetasoundAssetBase::GetSendInfos(uint64 InInstanceID) const
diff --git a/Engine/Plugins/Runtime/Metasound/Source/MetasoundFrontend/Private/MetasoundFrontend.cpp b/Engine/Plugins/Runtime/Metasound/Source/MetasoundFrontend/Private/MetasoundFrontend.cpp
--- a/Engine/Plugins/Runtime/Metasound/Source/MetasoundFrontend/Private/MetasoundFrontend.cpp
+++ b/Engine/Plugins/Runtime/Metasound/Source/MetasoundFrontend/Private/MetasoundFrontend.cpp
@@ -36,7 +36,7 @@ namespace Metasound
 			FMetasoundFrontendClass ClassDescription;
 
 			ClassDescription.Metadata = FMetasoundFrontendClassMetadata(InNodeMetadata);
-			ClassDescription.Metadata.Type = ClassType;
+			ClassDescription.Metadata.SetType(ClassType);
 
 			FMetasoundFrontendClassStyleDisplay DisplayStyle(InNodeMetadata.DisplayStyle);
 			ClassDescription.Style = FMetasoundFrontendClassStyle
@@ -46,10 +46,8 @@ namespace Metasound
 
 			FMetasoundFrontendClassInterface& ClassInterface = ClassDescription.Interface;
 
-			// External metasounds aren't dependent on any other nodes by definition, so all we need to do
-			// is populate the Input and Output sets.
 			const FInputVertexInterface& InputInterface = InNodeMetadata.DefaultInterface.GetInputInterface();
-			TArray<int>& InputSortOrder = ClassInterface.InputStyle.DefaultSortOrder;
+			FMetasoundFrontendInterfaceStyle InputStyle;
 			for (auto& InputTuple : InputInterface)
 			{
 				FMetasoundFrontendClassInput ClassInput;
@@ -75,13 +73,14 @@ namespace Metasound
 				{
 					OrderIndex += InputInterface.Num();
 				}
-				InputSortOrder.Add(OrderIndex);
+				InputStyle.DefaultSortOrder.Add(OrderIndex);
 
 				ClassInterface.Inputs.Add(ClassInput);
 			}
+			ClassInterface.SetInputStyle(InputStyle);
 
 			const FOutputVertexInterface& OutputInterface = InNodeMetadata.DefaultInterface.GetOutputInterface();
-			TArray<int32>& OutputSortOrder = ClassInterface.InputStyle.DefaultSortOrder;
+			FMetasoundFrontendInterfaceStyle OutputStyle;
 			for (auto& OutputTuple : OutputInterface)
 			{
 				FMetasoundFrontendClassOutput ClassOutput;
@@ -101,10 +100,11 @@ namespace Metasound
 				{
 					OrderIndex += OutputInterface.Num();
 				}
-				ClassInterface.OutputStyle.DefaultSortOrder.Add(OrderIndex);
+				OutputStyle.DefaultSortOrder.Add(OrderIndex);
 
-				ClassDescription.Interface.Outputs.Add(ClassOutput);
+				ClassInterface.Outputs.Add(ClassOutput);
 			}
+			ClassInterface.SetOutputStyle(OutputStyle);
 
 			for (auto& EnvTuple : InNodeMetadata.DefaultInterface.GetEnvironmentInterface())
 			{
@@ -115,7 +115,7 @@ namespace Metasound
 				EnvVar.Metadata.Description = EnvTuple.Value.GetDescription();
 				EnvVar.bIsRequired = true;
 
-				ClassDescription.Interface.Environment.Add(EnvVar);
+				ClassInterface.Environment.Add(EnvVar);
 			}
 
 			return ClassDescription;
diff --git a/Engine/Plugins/Runtime/Metasound/Source/MetasoundFrontend/Private/MetasoundFrontendDocument.cpp b/Engine/Plugins/Runtime/Metasound/Source/MetasoundFrontend/Private/MetasoundFrontendDocument.cpp
--- a/Engine/Plugins/Runtime/Metasound/Source/MetasoundFrontend/Private/MetasoundFrontendDocument.cpp
+++ b/Engine/Plugins/Runtime/Metasound/Source/MetasoundFrontend/Private/MetasoundFrontendDocument.cpp
@@ -2,6 +2,7 @@
 
 #include "MetasoundFrontendDocument.h"
 
+#include "Algo/Transform.h"
 #include "MetasoundFrontend.h"
 #include "MetasoundFrontendRegistries.h"
 #include "MetasoundLog.h"
@@ -9,6 +10,18 @@
 namespace Metasound
 {
 	const FGuid FrontendInvalidID = FGuid();
+
+	namespace Frontend
+	{
+		namespace DisplayStyle
+		{
+			namespace NodeLayout
+			{
+				const FVector2D DefaultOffsetX { 300.0f, 0.0f };
+				const FVector2D DefaultOffsetY { 0.0f, 80.0f };
+			} // namespace NodeLayout
+		} // namespace DisplayStyle
+	} // namespace Frontend
 }
 
 FMetasoundFrontendNodeInterface::FMetasoundFrontendNodeInterface(const FMetasoundFrontendClassInterface& InClassInterface)
@@ -35,7 +48,7 @@ FMetasoundFrontendNodeInterface::FMetasoundFrontendNodeInterface(const FMetasoun
 
 FMetasoundFrontendNode::FMetasoundFrontendNode(const FMetasoundFrontendClass& InClass)
 : ClassID(InClass.ID)
-, Name(InClass.Metadata.ClassName.Name.ToString())
+, Name(InClass.Metadata.GetClassName().Name.ToString())
 , Interface(InClass.Interface)
 {
 
@@ -62,7 +75,7 @@ bool FMetasoundFrontendVertex::IsFunctionalEquivalent(const FMetasoundFrontendVe
 	return (InLHS.Name == InRHS.Name) && (InLHS.TypeName == InRHS.TypeName);
 }
 
-bool FMetasoundFrontendClassVertex::IsFunctionalEquivalent(const FMetasoundFrontendClassVertex& InLHS, const FMetasoundFrontendClassVertex& InRHS) 
+bool FMetasoundFrontendClassVertex::IsFunctionalEquivalent(const FMetasoundFrontendClassVertex& InLHS, const FMetasoundFrontendClassVertex& InRHS)
 {
 	return FMetasoundFrontendVertex::IsFunctionalEquivalent(InLHS, InRHS);
 }
@@ -99,9 +112,14 @@ bool operator==(const FMetasoundFrontendClassName& InLHS, const FMetasoundFronte
 	return (InLHS.Namespace == InRHS.Namespace) && (InLHS.Name == InRHS.Name) && (InLHS.Variant == InRHS.Variant);
 }
 
+bool operator!=(const FMetasoundFrontendClassName& InLHS, const FMetasoundFrontendClassName& InRHS)
+{
+	return !(InLHS == InRHS);
+}
+
 FMetasoundFrontendClassMetadata::FMetasoundFrontendClassMetadata(const Metasound::FNodeClassMetadata& InNodeClassMetadata)
 : ClassName(InNodeClassMetadata.ClassName)
-, Version{InNodeClassMetadata.MajorVersion, InNodeClassMetadata.MinorVersion}
+, Version{ InNodeClassMetadata.MajorVersion, InNodeClassMetadata.MinorVersion }
 , Type(EMetasoundFrontendClassType::External)
 , DisplayName(InNodeClassMetadata.DisplayName)
 , Description(InNodeClassMetadata.Description)
@@ -130,13 +148,12 @@ FMetasoundFrontendClassVariable::FMetasoundFrontendClassVariable(const FMetasoun
 
 FMetasoundFrontendGraphClass::FMetasoundFrontendGraphClass()
 {
-	Metadata.Type = EMetasoundFrontendClassType::Graph;
+	Metadata.SetType(EMetasoundFrontendClassType::Graph);
 }
 
 FMetasoundFrontendDocument::FMetasoundFrontendDocument()
 {
 	RootGraph.ID = FGuid::NewGuid();
-	RootGraph.Metadata.Type = EMetasoundFrontendClassType::Graph;
+	RootGraph.Metadata.SetType(EMetasoundFrontendClassType::Graph);
 	ArchetypeVersion = FMetasoundFrontendVersion::GetInvalid();
 }
-
diff --git a/Engine/Plugins/Runtime/Metasound/Source/MetasoundFrontend/Private/MetasoundFrontendGraph.cpp b/Engine/Plugins/Runtime/Metasound/Source/MetasoundFrontend/Private/MetasoundFrontendGraph.cpp
--- a/Engine/Plugins/Runtime/Metasound/Source/MetasoundFrontend/Private/MetasoundFrontendGraph.cpp
+++ b/Engine/Plugins/Runtime/Metasound/Source/MetasoundFrontend/Private/MetasoundFrontendGraph.cpp
@@ -201,7 +201,7 @@ namespace Metasound
 
 	TUniquePtr<INode> FFrontendGraphBuilder::CreateOutputNode(const FMetasoundFrontendNode& InNode, const FMetasoundFrontendClass& InClass, FBuildGraphContext& InGraphContext)
 	{
-		check(InClass.Metadata.Type == EMetasoundFrontendClassType::Output);
+		check(InClass.Metadata.GetType() == EMetasoundFrontendClassType::Output);
 		check(InNode.ClassID == InClass.ID);
 
 		if (ensure(InNode.Interface.Outputs.Num() == 1))
@@ -232,7 +232,7 @@ namespace Metasound
 
 	TUniquePtr<INode> FFrontendGraphBuilder::CreateExternalNode(const FMetasoundFrontendNode& InNode, const FMetasoundFrontendClass& InClass, FBuildGraphContext& InGraphContext)
 	{
-		check(InClass.Metadata.Type == EMetasoundFrontendClassType::External);
+		check(InClass.Metadata.GetType() == EMetasoundFrontendClassType::External);
 		check(InNode.ClassID == InClass.ID);
 
 		const FNodeInitData InitData = FrontendGraphPrivate::CreateNodeInitData(InNode);
@@ -359,7 +359,7 @@ namespace Metasound
 
 			if (ensure(nullptr != NodeClass))
 			{
-				switch (NodeClass->Metadata.Type)
+				switch (NodeClass->Metadata.GetType())
 				{
 					case EMetasoundFrontendClassType::Input:
 					{
@@ -620,9 +620,9 @@ namespace Metasound
 		// All dependencies are external dependencies in a flat graph
 		auto IsClassExternal = [&](const FMetasoundFrontendClass& InDesc) 
 		{ 
-			bool bIsExternal = (InDesc.Metadata.Type == EMetasoundFrontendClassType::External) ||
-				(InDesc.Metadata.Type == EMetasoundFrontendClassType::Input) ||
-				(InDesc.Metadata.Type == EMetasoundFrontendClassType::Output);
+			bool bIsExternal = (InDesc.Metadata.GetType() == EMetasoundFrontendClassType::External) ||
+				(InDesc.Metadata.GetType() == EMetasoundFrontendClassType::Input) ||
+				(InDesc.Metadata.GetType() == EMetasoundFrontendClassType::Output);
 			return bIsExternal;
 		};
 		const bool bIsEveryDependencyExternal = Algo::AllOf(InDependencies, IsClassExternal);
@@ -696,7 +696,7 @@ namespace Metasound
 
 	TUniquePtr<FFrontendGraph> FFrontendGraphBuilder::CreateGraph(FBuildContext& InContext, const FMetasoundFrontendGraphClass& InGraphClass)
 	{
-		const FString GraphName = InGraphClass.Metadata.ClassName.GetFullName().ToString();
+		const FString GraphName = InGraphClass.Metadata.GetClassName().GetFullName().ToString();
 
 		FBuildGraphContext BuildGraphContext
 		{
@@ -746,7 +746,7 @@ namespace Metasound
 			TSharedPtr<const INode> Subgraph(CreateGraph(Context, *FrontendSubgraphPtr).Release());
 			if (!Subgraph.IsValid())
 			{
-				UE_LOG(LogMetaSound, Warning, TEXT("Failed to create subgraph [SubgraphName: %s]"), *FrontendSubgraphPtr->Metadata.ClassName.ToString());
+				UE_LOG(LogMetaSound, Warning, TEXT("Failed to create subgraph [SubgraphName: %s]"), *FrontendSubgraphPtr->Metadata.GetClassName().ToString());
 			}
 			else
 			{
diff --git a/Engine/Plugins/Runtime/Metasound/Source/MetasoundFrontend/Private/MetasoundFrontendInvalidController.h b/Engine/Plugins/Runtime/Metasound/Source/MetasoundFrontend/Private/MetasoundFrontendInvalidController.h
--- a/Engine/Plugins/Runtime/Metasound/Source/MetasoundFrontend/Private/MetasoundFrontendInvalidController.h
+++ b/Engine/Plugins/Runtime/Metasound/Source/MetasoundFrontend/Private/MetasoundFrontendInvalidController.h
@@ -149,6 +149,8 @@ namespace Metasound
 			const FMetasoundFrontendNodeStyle& GetNodeStyle() const override { static const FMetasoundFrontendNodeStyle Invalid; return Invalid; }
 			void SetNodeStyle(const FMetasoundFrontendNodeStyle& InNodeStyle) { };
 
+			FNodeHandle ReplaceWithVersion(const FMetasoundFrontendVersionNumber& InNewVersion) override { return FInvalidNodeController::GetInvalid(); }
+
 			bool CanAddInput(const FString& InVertexName) const override { return false; }
 			FInputHandle AddInput(const FString& InVertexName, const FMetasoundFrontendLiteral* InDefault) override { return FInvalidInputController::GetInvalid(); }
 			bool RemoveInput(FGuid InVertexID) override { return false; }
@@ -171,6 +173,10 @@ namespace Metasound
 
 			bool IsRequired() const override { return false; }
 
+			bool CanAutoUpdate() const override { return false; }
+			FMetasoundFrontendVersionNumber FindHighestVersionInRegistry() const override { return FMetasoundFrontendVersionNumber::GetInvalid(); }
+			FMetasoundFrontendVersionNumber FindHighestMinorVersionInRegistry() const override { return FMetasoundFrontendVersionNumber::GetInvalid(); }
+
 			TSharedRef<IGraphController> AsGraph() override;
 			TSharedRef<const IGraphController> AsGraph() const override;
 
@@ -238,6 +244,8 @@ namespace Metasound
 			const FMetasoundFrontendGraphStyle& GetGraphStyle() const override { return MetasoundFrontendInvalidControllerPrivate::GetInvalid<FMetasoundFrontendGraphStyle>(); }
 			void SetGraphStyle(const FMetasoundFrontendGraphStyle& InStyle) override { }
 
+			void ClearGraph() override { };
+
 			void IterateConstNodes(TUniqueFunction<void(FConstNodeHandle)> InFunction, EMetasoundFrontendClassType InClassType /* = EMetasoundFrontendClassType::Invalid */) const override { }
 			void IterateNodes(TUniqueFunction<void(FNodeHandle)> InFunction, EMetasoundFrontendClassType InClassType /* = EMetasoundFrontendClassType::Invalid */) override { }
 
@@ -308,6 +316,14 @@ namespace Metasound
 
 			TSharedRef<IDocumentController> GetOwningDocument() override;
 			TSharedRef<const IDocumentController> GetOwningDocument() const override;
+
+			void UpdateInterfaceChangeID() override { }
+
+			FNodeHandle AddInputNode(const FMetasoundFrontendClassInput& InClassInput, const FMetasoundFrontendLiteral* InDefaultValue, const FText* InDisplayName) override { return INodeController::GetInvalidHandle(); }
+			FNodeHandle AddInputNode(const FName InTypeName, const FText& InToolTip, const FMetasoundFrontendLiteral* InDefaultValue, const FText* InDisplayName) override { return INodeController::GetInvalidHandle(); }
+			FNodeHandle AddOutputNode(const FMetasoundFrontendClassOutput& InClassOutput, const FText* InDisplayName) override { return INodeController::GetInvalidHandle(); }
+			FNodeHandle AddOutputNode(const FName InTypeName, const FText& InToolTip, const FText* InDisplayName) override { return INodeController::GetInvalidHandle(); }
+
 		protected:
 			FDocumentAccess ShareAccess() override { return FDocumentAccess(); }
 			FConstDocumentAccess ShareAccess() const override { return FConstDocumentAccess(); }
@@ -331,9 +347,8 @@ namespace Metasound
 
 				bool IsValid() const override { return false; }
 
-				TArray<FMetasoundFrontendClass> GetDependencies() const override { return TArray<FMetasoundFrontendClass>(); }
-				TArray<FMetasoundFrontendGraphClass> GetSubgraphs() const override { return TArray<FMetasoundFrontendGraphClass>(); }
-				TArray<FMetasoundFrontendClass> GetClasses() const override { return TArray<FMetasoundFrontendClass>(); }
+				const TArray<FMetasoundFrontendClass>& GetDependencies() const override { return MetasoundFrontendInvalidControllerPrivate::GetInvalid<TArray<FMetasoundFrontendClass>>(); }
+				const TArray<FMetasoundFrontendGraphClass>& GetSubgraphs() const override { return MetasoundFrontendInvalidControllerPrivate::GetInvalid<TArray<FMetasoundFrontendGraphClass>>(); }
 				const FMetasoundFrontendGraphClass& GetRootGraphClass() const override
 				{
 					return MetasoundFrontendInvalidControllerPrivate::GetInvalid<FMetasoundFrontendGraphClass>();
@@ -349,12 +364,13 @@ namespace Metasound
 				FConstClassAccessPtr FindOrAddClass(const FMetasoundFrontendClassMetadata& InMetadata) override{ return FConstClassAccessPtr(); }
 				FGraphHandle AddDuplicateSubgraph(const IGraphController& InGraph) override { return FInvalidGraphController::GetInvalid(); }
 
-				virtual const FMetasoundFrontendVersion& GetArchetypeVersion() const override { return FMetasoundFrontendVersion::GetInvalid(); }
-				virtual void SetArchetypeVersion(const FMetasoundFrontendVersion& InVersion) override {}
+				const FMetasoundFrontendVersion& GetArchetypeVersion() const override { return FMetasoundFrontendVersion::GetInvalid(); }
+				void SetArchetypeVersion(const FMetasoundFrontendVersion& InVersion) override { }
 
 				void SetMetadata(const FMetasoundFrontendDocumentMetadata& InMetadata) override { }
 				const FMetasoundFrontendDocumentMetadata& GetMetadata() const override { return MetasoundFrontendInvalidControllerPrivate::GetInvalid<FMetasoundFrontendDocumentMetadata>(); }
 
+				const FMetasoundFrontendClass* SynchronizeDependency(const FNodeRegistryKey& InKey) override { return nullptr; }
 				void SynchronizeDependencies() override { }
 
 				TArray<FGraphHandle> GetSubgraphHandles() override { return TArray<FGraphHandle>(); }
diff --git a/Engine/Plugins/Runtime/Metasound/Source/MetasoundFrontend/Private/MetasoundFrontendQuerySteps.cpp b/Engine/Plugins/Runtime/Metasound/Source/MetasoundFrontend/Private/MetasoundFrontendQuerySteps.cpp
--- a/Engine/Plugins/Runtime/Metasound/Source/MetasoundFrontend/Private/MetasoundFrontendQuerySteps.cpp
+++ b/Engine/Plugins/Runtime/Metasound/Source/MetasoundFrontend/Private/MetasoundFrontendQuerySteps.cpp
@@ -140,7 +140,7 @@ namespace Metasound
 
 	bool FFilterClassesByClassName::Filter(const FFrontendQueryEntry& InEntry) const 
 	{
-		return InEntry.Value.Get<FMetasoundFrontendClass>().Metadata.ClassName == ClassName;
+		return InEntry.Value.Get<FMetasoundFrontendClass>().Metadata.GetClassName() == ClassName;
 	}
 
 	FFilterClassesByClassID::FFilterClassesByClassID(const FGuid InClassID)
@@ -156,7 +156,7 @@ namespace Metasound
 	FFrontendQueryEntry::FKey FMapToFullClassName::Map(const FFrontendQueryEntry& InEntry) const
 	{
 		const FMetasoundFrontendClass& FrontendClass = InEntry.Value.Get<FMetasoundFrontendClass>();
-		return FFrontendQueryEntry::FKey(FrontendClass.Metadata.ClassName.GetFullName());
+		return FFrontendQueryEntry::FKey(FrontendClass.Metadata.GetClassName().GetFullName());
 	}
 
 	void FReduceClassesToHighestVersion::Reduce(FFrontendQueryEntry::FKey InKey, TArrayView<FFrontendQueryEntry * const>& InEntries, FReduceOutputView& OutResult) const
@@ -166,7 +166,7 @@ namespace Metasound
 
 		for (FFrontendQueryEntry* Entry : InEntries)
 		{
-			const int32 EntryMajorVersion = Entry->Value.Get<FMetasoundFrontendClass>().Metadata.Version.Major;
+			const int32 EntryMajorVersion = Entry->Value.Get<FMetasoundFrontendClass>().Metadata.GetVersion().Major;
 
 			if (!HighestVersionEntry || HighestMajorVersion < EntryMajorVersion)
 			{
@@ -190,7 +190,7 @@ namespace Metasound
 	{
 		for (FFrontendQueryEntry* Entry : InEntries)
 		{
-			const int32 EntryMajorVersion = Entry->Value.Get<FMetasoundFrontendClass>().Metadata.Version.Major;
+			const int32 EntryMajorVersion = Entry->Value.Get<FMetasoundFrontendClass>().Metadata.GetVersion().Major;
 			if (MajorVersion == EntryMajorVersion)
 			{
 				OutResult.Add(*Entry);
@@ -201,8 +201,8 @@ namespace Metasound
 
 	bool FSortClassesByVersion::Sort(const FFrontendQueryEntry& InEntryLHS, const FFrontendQueryEntry& InEntryRHS) const
 	{
-		const FMetasoundFrontendVersionNumber& VersionLHS = InEntryLHS.Value.Get<FMetasoundFrontendClass>().Metadata.Version;
-		const FMetasoundFrontendVersionNumber& VersionRHS = InEntryRHS.Value.Get<FMetasoundFrontendClass>().Metadata.Version;
+		const FMetasoundFrontendVersionNumber& VersionLHS = InEntryLHS.Value.Get<FMetasoundFrontendClass>().Metadata.GetVersion();
+		const FMetasoundFrontendVersionNumber& VersionRHS = InEntryRHS.Value.Get<FMetasoundFrontendClass>().Metadata.GetVersion();
 		return VersionLHS > VersionRHS;
 	}
 }
diff --git a/Engine/Plugins/Runtime/Metasound/Source/MetasoundFrontend/Private/MetasoundFrontendRegistries.cpp b/Engine/Plugins/Runtime/Metasound/Source/MetasoundFrontend/Private/MetasoundFrontendRegistries.cpp
--- a/Engine/Plugins/Runtime/Metasound/Source/MetasoundFrontend/Private/MetasoundFrontendRegistries.cpp
+++ b/Engine/Plugins/Runtime/Metasound/Source/MetasoundFrontend/Private/MetasoundFrontendRegistries.cpp
@@ -208,6 +208,11 @@ namespace Metasound
 					return MakeUnique<FInputNodeRegistryEntry>(TUniquePtr<IDataTypeRegistryEntry>());
 				}
 
+				virtual bool IsNative() const override
+				{
+					return true;
+				}
+
 			private:
 				
 				TUniquePtr<IDataTypeRegistryEntry> DataTypeEntry;
@@ -271,6 +276,11 @@ namespace Metasound
 					return MakeUnique<FOutputNodeRegistryEntry>(TUniquePtr<IDataTypeRegistryEntry>());
 				}
 
+				virtual bool IsNative() const override
+				{
+					return true;
+				}
+
 			private:
 				
 				TUniquePtr<IDataTypeRegistryEntry> DataTypeEntry;
@@ -334,6 +344,11 @@ namespace Metasound
 					return MakeUnique<FVariableNodeRegistryEntry>(TUniquePtr<IDataTypeRegistryEntry>());
 				}
 
+				virtual bool IsNative() const override
+				{
+					return true;
+				}
+
 			private:
 				
 				TUniquePtr<IDataTypeRegistryEntry> DataTypeEntry;
@@ -380,6 +395,7 @@ namespace Metasound
 				virtual FNodeRegistryKey RegisterNode(TUniquePtr<Metasound::Frontend::INodeRegistryEntry>&&) override;
 				virtual bool UnregisterNode(const FNodeRegistryKey& InKey) override;
 				virtual bool IsNodeRegistered(const FNodeRegistryKey& InKey) const override;
+				virtual bool IsNodeNative(const FNodeRegistryKey& InKey) const override;
 
 				bool RegisterConversionNode(const FConverterNodeRegistryKey& InNodeKey, const FConverterNodeInfo& InNodeInfo) override;
 
@@ -806,7 +822,6 @@ namespace Metasound
 				return false;
 			}
 
-
 			bool FRegistryContainerImpl::RegisterConversionNode(const FConverterNodeRegistryKey& InNodeKey, const FConverterNodeInfo& InNodeInfo)
 			{
 				if (!ConverterNodeRegistry.Contains(InNodeKey))
@@ -833,6 +848,16 @@ namespace Metasound
 				return RegisteredNodes.Contains(InKey);
 			}
 
+			bool FRegistryContainerImpl::IsNodeNative(const FNodeRegistryKey& InKey) const
+			{
+				if (const INodeRegistryEntry* Entry = FindNodeEntry(InKey))
+				{
+					return Entry->IsNative();
+				}
+
+				return false;
+			}
+
 			TArray<FName> FRegistryContainerImpl::GetAllValidDataTypes()
 			{
 				TArray<FName> OutDataTypes;
@@ -1020,11 +1045,11 @@ namespace Metasound
 
 			bool IsEqual(const FMetasoundFrontendClassMetadata& InLHS, const FMetasoundFrontendClassMetadata& InRHS)
 			{
-				if (InLHS.ClassName == InRHS.ClassName)
+				if (InLHS.GetClassName() == InRHS.GetClassName())
 				{
-					if (InLHS.Type == InRHS.Type)
+					if (InLHS.GetType() == InRHS.GetType())
 					{
-						if (InLHS.Version == InRHS.Version)
+						if (InLHS.GetVersion() == InRHS.GetVersion())
 						{
 							return true;
 						}
@@ -1035,11 +1060,11 @@ namespace Metasound
 
 			bool IsEqual(const FNodeClassInfo& InLHS, const FMetasoundFrontendClassMetadata& InRHS)
 			{
-				if (InLHS.ClassName == InRHS.ClassName)
+				if (InLHS.ClassName == InRHS.GetClassName())
 				{
-					if (InLHS.Type == InRHS.Type)
+					if (InLHS.Type == InRHS.GetType())
 					{
-						if (InLHS.Version == InRHS.Version)
+						if (InLHS.Version == InRHS.GetVersion())
 						{
 							return true;
 						}
@@ -1055,7 +1080,7 @@ namespace Metasound
 
 			FNodeRegistryKey CreateKey(const FMetasoundFrontendClassMetadata& InNodeMetadata)
 			{
-				return CreateKey(InNodeMetadata.Type, InNodeMetadata.ClassName.GetFullName().ToString(), InNodeMetadata.Version.Major, InNodeMetadata.Version.Minor);
+				return CreateKey(InNodeMetadata.GetType(), InNodeMetadata.GetClassName().GetFullName().ToString(), InNodeMetadata.GetVersion().Major, InNodeMetadata.GetVersion().Minor);
 			}
 
 			FNodeRegistryKey CreateKey(const FNodeClassInfo& InClassInfo)
@@ -1070,18 +1095,18 @@ namespace Metasound
 		}
 
 		FNodeClassInfo::FNodeClassInfo(const FMetasoundFrontendClassMetadata& InMetadata)
-			: ClassName(InMetadata.ClassName)
-			, Type(InMetadata.Type)
-			, Version(InMetadata.Version)
+			: ClassName(InMetadata.GetClassName())
+			, Type(InMetadata.GetType())
+			, Version(InMetadata.GetVersion())
 		{
 		}
 
 		FNodeClassInfo::FNodeClassInfo(const FMetasoundFrontendGraphClass& InClass, FName InAssetPath)
-			: ClassName(InClass.Metadata.ClassName)
+			: ClassName(InClass.Metadata.GetClassName())
 			, Type(EMetasoundFrontendClassType::External) // Overridden as it is considered the same as an external class in the registry
 			, AssetClassID(FGuid(ClassName.Name.ToString()))
 			, AssetPath(InAssetPath)
-			, Version(InClass.Metadata.Version)
+			, Version(InClass.Metadata.GetVersion())
 		{
 			ensure(!AssetPath.IsNone());
 
diff --git a/Engine/Plugins/Runtime/Metasound/Source/MetasoundFrontend/Private/MetasoundFrontendStandardController.cpp b/Engine/Plugins/Runtime/Metasound/Source/MetasoundFrontend/Private/MetasoundFrontendStandardController.cpp
--- a/Engine/Plugins/Runtime/Metasound/Source/MetasoundFrontend/Private/MetasoundFrontendStandardController.cpp
+++ b/Engine/Plugins/Runtime/Metasound/Source/MetasoundFrontend/Private/MetasoundFrontendStandardController.cpp
@@ -13,6 +13,7 @@
 #include "MetasoundFrontendDocumentAccessPtr.h"
 #include "MetasoundFrontendGraph.h"
 #include "MetasoundFrontendRegistries.h"
+#include "MetasoundFrontendSearchEngine.h"
 #include "MetasoundFrontendSubgraphNodeController.h"
 #include "MetasoundJsonBackend.h"
 #include "MetasoundOperatorBuilder.h"
@@ -28,6 +29,12 @@ namespace Metasound
 	{
 		namespace FrontendControllerIntrinsics
 		{
+			namespace NodeLayout
+			{
+				static const FVector2D BufferX{ 250.0f, 0.0f };
+				static const FVector2D BufferY{ 0.0f, 100.0f };
+			}
+
 			// utility function for returning invalid values. If an invalid value type
 			// needs special construction, this template can be specialized. 
 			template<typename ValueType>
@@ -55,8 +62,7 @@ namespace Metasound
 				static const ValueType Value = GetInvalidValue<ValueType>();
 				return Value;
 			}
-
-		}
+		} // namespace FrontendControllerIntrinsics
 
 
 		FDocumentAccess IDocumentAccessor::GetSharedAccess(IDocumentAccessor& InDocumentAccessor)
@@ -349,7 +355,7 @@ namespace Metasound
 
 		//
 		// FOutputNodeOutputController
-		// 
+		//
 		FOutputNodeOutputController::FOutputNodeOutputController(const FOutputNodeOutputController::FInitParams& InParams)
 		: FBaseOutputController({InParams.ID, InParams.NodeVertexPtr, InParams.ClassOutputPtr, InParams.GraphPtr, InParams.OwningNode})
 		, OwningGraphClassOutputPtr(InParams.OwningGraphClassOutputPtr)
@@ -1041,22 +1047,20 @@ namespace Metasound
 		{
 			if (const FMetasoundFrontendClass* Class = ClassPtr.Get())
 			{
-				return Class->Interface.InputStyle;
+				return Class->Interface.GetInputStyle();
 			}
 
-			static const FMetasoundFrontendInterfaceStyle Invalid;
-			return Invalid;
+			return FrontendControllerIntrinsics::GetInvalidValueConstRef<FMetasoundFrontendInterfaceStyle>();
 		}
 
 		const FMetasoundFrontendInterfaceStyle& FBaseNodeController::GetOutputStyle() const
 		{
 			if (const FMetasoundFrontendClass* Class = ClassPtr.Get())
 			{
-				return Class->Interface.OutputStyle;
+				return Class->Interface.GetOutputStyle();
 			}
 
-			static const FMetasoundFrontendInterfaceStyle Invalid;
-			return Invalid;
+			return FrontendControllerIntrinsics::GetInvalidValueConstRef<FMetasoundFrontendInterfaceStyle>();
 		}
 
 		const FMetasoundFrontendClassStyle& FBaseNodeController::GetClassStyle() const
@@ -1093,7 +1097,7 @@ namespace Metasound
 		{
 			if (const FMetasoundFrontendClass* Class = ClassPtr.Get())
 			{
-				return Class->Metadata.Description;
+				return Class->Metadata.GetDescription();
 			}
 			return FrontendControllerIntrinsics::GetInvalidValueConstRef<FText>();
 		}
@@ -1257,7 +1261,7 @@ namespace Metasound
 		{
 			if (const FMetasoundFrontendClass* Class = ClassPtr.Get())
 			{
-				return Class->Metadata.DisplayName;
+				return Class->Metadata.GetDisplayName();
 			}
 
 			return FrontendControllerIntrinsics::GetInvalidValueConstRef<FText>();
@@ -1267,9 +1271,10 @@ namespace Metasound
 		{
 			if (const FMetasoundFrontendClass* Class = ClassPtr.Get())
 			{
-				return Class->Metadata.DisplayName;
+				return Class->Metadata.GetDisplayName();
 			}
-			return FText::GetEmpty();
+
+			return FrontendControllerIntrinsics::GetInvalidValueConstRef<FText>();
 		}
 
 		void FBaseNodeController::IterateConstInputs(TUniqueFunction<void(FConstInputHandle)> InFunction) const
@@ -1570,7 +1575,7 @@ namespace Metasound
 		FConstGraphHandle FBaseNodeController::AsGraph() const
 		{
 			// TODO: add support for graph owned in another asset.
-			// Will require lookup support for external subgraphs..
+			// Will require lookup support for external subgraphs.
 			if (const FMetasoundFrontendClass* Class = ClassPtr.Get())
 			{
 				return GetOwningGraph()->GetOwningDocument()->GetSubgraphWithClassID(Class->ID);
@@ -1579,6 +1584,194 @@ namespace Metasound
 			return FInvalidGraphController::GetInvalid();
 		}
 
+		FMetasoundFrontendVersionNumber FBaseNodeController::FindHighestMinorVersionInRegistry() const
+		{
+			const FMetasoundFrontendClassMetadata& Metadata = GetClassMetadata();
+			const FMetasoundFrontendVersionNumber& CurrentVersion = Metadata.GetVersion();
+			Metasound::FNodeClassName NodeClassName = Metadata.GetClassName().ToNodeClassName();
+
+			FMetasoundFrontendClass ClassWithMajorVersion;
+			if (ISearchEngine::Get().FindClassWithMajorVersion(NodeClassName, CurrentVersion.Major, ClassWithMajorVersion))
+			{
+				if (ClassWithMajorVersion.Metadata.GetVersion().Minor >= CurrentVersion.Minor)
+				{
+					return ClassWithMajorVersion.Metadata.GetVersion();
+				}
+			}
+
+			return FMetasoundFrontendVersionNumber::GetInvalid();
+		}
+
+		FMetasoundFrontendVersionNumber FBaseNodeController::FindHighestVersionInRegistry() const
+		{
+			const FMetasoundFrontendClassMetadata& Metadata = GetClassMetadata();
+			const FMetasoundFrontendVersionNumber& CurrentVersion = Metadata.GetVersion();
+			Metasound::FNodeClassName NodeClassName = Metadata.GetClassName().ToNodeClassName();
+
+			FMetasoundFrontendClass ClassWithHighestVersion;
+			if (ISearchEngine::Get().FindClassWithHighestVersion(NodeClassName, ClassWithHighestVersion))
+			{
+				if (ClassWithHighestVersion.Metadata.GetVersion().Major >= CurrentVersion.Major)
+				{
+					return ClassWithHighestVersion.Metadata.GetVersion();
+				}
+			}
+
+			return FMetasoundFrontendVersionNumber::GetInvalid();
+
+		}
+
+		FNodeHandle FBaseNodeController::ReplaceWithVersion(const FMetasoundFrontendVersionNumber& InNewVersion)
+		{
+			const FMetasoundFrontendClassMetadata& Metadata = GetClassMetadata();
+			const TArray<FMetasoundFrontendClass> Versions = ISearchEngine::Get().FindClassesWithName(Metadata.GetClassName().ToNodeClassName(), false /* bInSortByVersion */);
+
+			auto IsClassOfNewVersion = [InNewVersion](const FMetasoundFrontendClass& RegisteredClass)
+			{
+				return RegisteredClass.Metadata.GetVersion() == InNewVersion;
+			};
+
+			const FMetasoundFrontendClass* RegisteredClass = Versions.FindByPredicate(IsClassOfNewVersion);
+			if (!ensure(RegisteredClass))
+			{
+				return this->AsShared();
+			}
+
+			const FMetasoundFrontendNodeStyle Style = GetNodeStyle();
+
+			using FConnectionKey = TPair<FString, FName>;
+
+			struct FInputConnectionInfo
+			{
+				FOutputHandle ConnectedOutput;
+				FName DataType;
+				FMetasoundFrontendLiteral DefaultValue;
+			};
+
+			// Cache input/output connections by name to try so they can be
+			// hooked back up after swapping to the new class version.
+			TMap<FConnectionKey, FInputConnectionInfo> InputConnections;
+			IterateInputs([Connections = &InputConnections](FInputHandle InputHandle)
+			{
+				FMetasoundFrontendLiteral DefaultLiteral;
+				if (const FMetasoundFrontendLiteral* Literal = InputHandle->GetLiteral())
+				{
+					DefaultLiteral = *Literal;
+				}
+				else if (const FMetasoundFrontendLiteral* ClassLiteral = InputHandle->GetClassDefaultLiteral())
+				{
+					DefaultLiteral = *ClassLiteral;
+				}
+
+				const FConnectionKey ConnectionKey(InputHandle->GetName(), InputHandle->GetDataType());
+				Connections->Add(ConnectionKey, FInputConnectionInfo
+				{
+					InputHandle->GetConnectedOutput(),
+					InputHandle->GetDataType(),
+					MoveTemp(DefaultLiteral)
+				});
+			});
+
+			struct FOutputConnectionInfo
+			{
+				TArray<FInputHandle> ConnectedInputs;
+				FName DataType;
+			};
+
+			TMap<FConnectionKey, FOutputConnectionInfo> OutputConnections;
+			IterateOutputs([Connections = &OutputConnections](FOutputHandle OutputHandle)
+			{
+				const FConnectionKey ConnectionKey(OutputHandle->GetName(), OutputHandle->GetDataType());
+				Connections->Add(ConnectionKey, FOutputConnectionInfo
+				{
+					OutputHandle->GetConnectedInputs(),
+					OutputHandle->GetDataType(),
+				});
+			});
+
+			if (!ensureAlways(GetOwningGraph()->RemoveNode(*this)))
+			{
+				return this->AsShared();
+			}
+
+			// Make sure classes are up-to-date with registered versions of class.
+			// Note that this may break other nodes in the graph that have stale
+			// class API, but that's on the caller to fix-up or report invalid state.
+			const FNodeRegistryKey RegistryKey = FMetasoundFrontendRegistryContainer::Get()->GetRegistryKey(RegisteredClass->Metadata);
+			FDocumentHandle Document = GetOwningGraph()->GetOwningDocument();
+			ensureAlways(Document->SynchronizeDependency(RegistryKey) != nullptr);
+
+			FNodeHandle ReplacementNode = GetOwningGraph()->AddNode(RegisteredClass->Metadata);
+			if (!ensureAlways(ReplacementNode->IsValid()))
+			{
+				return this->AsShared();
+			}
+
+			ReplacementNode->SetNodeStyle(Style);
+
+			ReplacementNode->IterateInputs([Connections = &InputConnections](FInputHandle InputHandle)
+			{
+				const FConnectionKey ConnectionKey(InputHandle->GetName(), InputHandle->GetDataType());
+				if (FInputConnectionInfo* ConnectionInfo = Connections->Find(ConnectionKey))
+				{
+					InputHandle->SetLiteral(ConnectionInfo->DefaultValue);
+					if (ConnectionInfo->ConnectedOutput->IsValid())
+					{
+						ensure(InputHandle->Connect(*ConnectionInfo->ConnectedOutput));
+					}
+				}
+			});
+
+			ReplacementNode->IterateOutputs([Connections = &OutputConnections](FOutputHandle OutputHandle)
+			{
+				const FConnectionKey ConnectionKey(OutputHandle->GetName(), OutputHandle->GetDataType());
+				if (FOutputConnectionInfo* ConnectionInfo = Connections->Find(ConnectionKey))
+				{
+					for (FInputHandle InputHandle : ConnectionInfo->ConnectedInputs)
+					{
+						if (InputHandle->IsValid())
+						{
+							ensure(InputHandle->Connect(*OutputHandle));
+						}
+					}
+				}
+			});
+
+			return ReplacementNode;
+		}
+
+		bool FBaseNodeController::CanAutoUpdate() const
+		{
+			const FMetasoundFrontendClassMetadata& Metadata = GetClassMetadata();
+
+			const FMetasoundFrontendVersionNumber& CurrentVersion = Metadata.GetVersion();
+			Metasound::FNodeClassName NodeClassName = Metadata.GetClassName().ToNodeClassName();
+
+			// TODO: Natively defined member should probably be added to the FMetasoundFrontendClass
+			// instead of being an independent member on the INodeRegistryEntry.
+			const FNodeRegistryKey RegistryKey = FMetasoundFrontendRegistryContainer::Get()->GetRegistryKey(Metadata);
+			if (FMetasoundFrontendRegistryContainer::Get()->IsNodeNative(RegistryKey))
+			{
+				return false;
+			}
+
+			FMetasoundFrontendClass ClassWithMajorVersion;
+			if (!ISearchEngine::Get().FindClassWithMajorVersion(NodeClassName, CurrentVersion.Major, ClassWithMajorVersion))
+			{
+				return false;
+			}
+
+			if (ClassWithMajorVersion.Metadata.GetChangeID() == Metadata.GetChangeID())
+			{
+				const FGuid& ChangeID = GetClassInterface().GetChangeID();
+				if (ClassWithMajorVersion.Interface.GetChangeID() == ChangeID)
+				{
+					return false;
+				}
+			}
+
+			return true;
+		}
 
 		FDocumentAccess FBaseNodeController::ShareAccess() 
 		{
@@ -1702,7 +1895,7 @@ namespace Metasound
 			{
 				if (const FMetasoundFrontendClass* Class = InParams.ClassPtr.Get())
 				{
-					if (EMetasoundFrontendClassType::Output == Class->Metadata.Type)
+					if (EMetasoundFrontendClassType::Output == Class->Metadata.GetType())
 					{
 						if (Class->ID == Node->ClassID)
 						{
@@ -1749,6 +1942,7 @@ namespace Metasound
 			if (FMetasoundFrontendClassOutput* ClassOutput = ConstCastAccessPtr<FClassOutputAccessPtr>(OwningGraphClassOutputPtr).Get())
 			{
 				ClassOutput->Metadata.Description = InDescription;
+				OwningGraph->UpdateInterfaceChangeID();
 			}
 		}
 
@@ -1758,6 +1952,7 @@ namespace Metasound
 			if (FMetasoundFrontendClassOutput* ClassOutput = ConstCastAccessPtr<FClassOutputAccessPtr>(OwningGraphClassOutputPtr).Get())
 			{
 				ClassOutput->Metadata.DisplayName = InDisplayName;
+				OwningGraph->UpdateInterfaceChangeID();
 			}
 		}
 
@@ -1767,7 +1962,7 @@ namespace Metasound
 			{
 				if (const FMetasoundFrontendClass* Class = InParams.ClassPtr.Get())
 				{
-					if (EMetasoundFrontendClassType::Output == Class->Metadata.Type)
+					if (EMetasoundFrontendClassType::Output == Class->Metadata.GetType())
 					{
 						if (Class->ID == Node->ClassID)
 						{
@@ -1885,7 +2080,7 @@ namespace Metasound
 			{
 				if (const FMetasoundFrontendClass* Class = InParams.ClassPtr.Get())
 				{
-					if (EMetasoundFrontendClassType::Input == Class->Metadata.Type)
+					if (EMetasoundFrontendClassType::Input == Class->Metadata.GetType())
 					{
 						if (Class->ID == Node->ClassID)
 						{
@@ -1912,7 +2107,7 @@ namespace Metasound
 			{
 				if (const FMetasoundFrontendClass* Class = InParams.ClassPtr.Get())
 				{
-					if (EMetasoundFrontendClassType::Input == Class->Metadata.Type)
+					if (EMetasoundFrontendClassType::Input == Class->Metadata.GetType())
 					{
 						if (Class->ID == Node->ClassID)
 						{
@@ -2018,6 +2213,7 @@ namespace Metasound
 			if (FMetasoundFrontendClassInput* ClassInput = ConstCastAccessPtr<FClassInputAccessPtr>(OwningGraphClassInputPtr).Get())
 			{
 				ClassInput->Metadata.Description = InDescription;
+				OwningGraph->UpdateInterfaceChangeID();
 			}
 		}
 
@@ -2027,6 +2223,7 @@ namespace Metasound
 			if (FMetasoundFrontendClassInput* ClassInput = ConstCastAccessPtr<FClassInputAccessPtr>(OwningGraphClassInputPtr).Get())
 			{
 				ClassInput->Metadata.DisplayName = InDisplayName;
+				OwningGraph->UpdateInterfaceChangeID();
 			}
 		}
 
@@ -2064,7 +2261,7 @@ namespace Metasound
 		{
 			if (FMetasoundFrontendGraphClass* GraphClass = InParams.GraphClassPtr.Get())
 			{
-				if (GraphClass->Metadata.Type == EMetasoundFrontendClassType::Graph)
+				if (GraphClass->Metadata.GetType() == EMetasoundFrontendClassType::Graph)
 				{
 					return MakeShared<FGraphController>(EPrivateToken::Token, InParams);
 				}
@@ -2080,7 +2277,7 @@ namespace Metasound
 		{
 			if (const FMetasoundFrontendGraphClass* GraphClass = InParams.GraphClassPtr.Get())
 			{
-				if (GraphClass->Metadata.Type == EMetasoundFrontendClassType::Graph)
+				if (GraphClass->Metadata.GetType() == EMetasoundFrontendClassType::Graph)
 				{
 					return MakeShared<FGraphController>(EPrivateToken::Token, InParams);
 				}
@@ -2111,7 +2308,7 @@ namespace Metasound
 		{
 			if (const FMetasoundFrontendGraphClass* GraphClass = GraphClassPtr.Get())
 			{
-				return GraphClass->Metadata.DisplayName;
+				return GraphClass->Metadata.GetDisplayName();
 			}
 
 			return FrontendControllerIntrinsics::GetInvalidValueConstRef<FText>();
@@ -2225,20 +2422,32 @@ namespace Metasound
 		{
 			auto IsOutputNode = [](const FMetasoundFrontendClass& NodeClass, const FMetasoundFrontendNode& Node) 
 			{
-				return NodeClass.Metadata.Type == EMetasoundFrontendClassType::Output; 
+				return NodeClass.Metadata.GetType() == EMetasoundFrontendClassType::Output;
 			};
 			return GetNodesByPredicate(IsOutputNode);
 		}
 
 		TArray<FNodeHandle> FGraphController::GetInputNodes()
 		{
-			auto IsInputNode = [](const FMetasoundFrontendClass& NodeClass, const FMetasoundFrontendNode& Node) 
+			auto IsInputNode = [](const FMetasoundFrontendClass& NodeClass, const FMetasoundFrontendNode& Node)
 			{
-				return NodeClass.Metadata.Type == EMetasoundFrontendClassType::Input; 
+				return NodeClass.Metadata.GetType() == EMetasoundFrontendClassType::Input;
 			};
 			return GetNodesByPredicate(IsInputNode);
 		}
 
+		void FGraphController::ClearGraph()
+		{
+			if (FMetasoundFrontendGraphClass* GraphClass = GraphClassPtr.Get())
+			{
+				GraphClass->Graph.Nodes.Reset();
+				GraphClass->Graph.Edges.Reset();
+				GraphClass->Interface.Inputs.Reset();
+				GraphClass->Interface.Outputs.Reset();
+				OwningDocument->SynchronizeDependencies();
+			}
+		}
+
 		void FGraphController::IterateNodes(TUniqueFunction<void(FNodeHandle)> InFunction, EMetasoundFrontendClassType InClassType)
 		{
 			if (FMetasoundFrontendGraphClass* GraphClass = GraphClassPtr.Get())
@@ -2248,7 +2457,7 @@ namespace Metasound
 					FConstClassAccessPtr NodeClassPtr = OwningDocument->FindClassWithID(Node.ClassID);
 					if (const FMetasoundFrontendClass* NodeClass = NodeClassPtr.Get())
 					{
-						if (InClassType == EMetasoundFrontendClassType::Invalid || NodeClass->Metadata.Type == InClassType)
+						if (InClassType == EMetasoundFrontendClassType::Invalid || NodeClass->Metadata.GetType() == InClassType)
 						{
 							FNodeAccessPtr NodePtr = GraphClassPtr.GetNodeWithNodeID(Node.ID);
 							FNodeHandle NodeHandle = GetNodeHandle(FGraphController::FNodeAndClass{ NodePtr, NodeClassPtr });
@@ -2272,7 +2481,7 @@ namespace Metasound
 					FConstClassAccessPtr NodeClassPtr = OwningDocument->FindClassWithID(Node.ClassID);
 					if (const FMetasoundFrontendClass* NodeClass = NodeClassPtr.Get())
 					{
-						if (InClassType == EMetasoundFrontendClassType::Invalid || NodeClass->Metadata.Type == InClassType)
+						if (InClassType == EMetasoundFrontendClassType::Invalid || NodeClass->Metadata.GetType() == InClassType)
 						{
 							FConstNodeAccessPtr NodePtr = GraphClassPtr.GetNodeWithNodeID(Node.ID);
 							FConstNodeHandle NodeHandle = GetNodeHandle(FGraphController::FConstNodeAndClass{ NodePtr, NodeClassPtr });
@@ -2291,7 +2500,7 @@ namespace Metasound
 		{
 			auto IsOutputNode = [](const FMetasoundFrontendClass& NodeClass, const FMetasoundFrontendNode& Node) 
 			{
-				return NodeClass.Metadata.Type == EMetasoundFrontendClassType::Output; 
+				return NodeClass.Metadata.GetType() == EMetasoundFrontendClassType::Output;
 			};
 			return GetNodesByPredicate(IsOutputNode);
 		}
@@ -2300,7 +2509,7 @@ namespace Metasound
 		{
 			auto IsInputNode = [](const FMetasoundFrontendClass& NodeClass, const FMetasoundFrontendNode& Node) 
 			{
-				return NodeClass.Metadata.Type == EMetasoundFrontendClassType::Input; 
+				return NodeClass.Metadata.GetType() == EMetasoundFrontendClassType::Input;
 			};
 			return GetNodesByPredicate(IsInputNode);
 		}
@@ -2309,8 +2518,8 @@ namespace Metasound
 		{
 			if (const FMetasoundFrontendGraphClass* GraphClass = GraphClassPtr.Get())
 			{
-				auto IsOutputVertexWithSameName = [&](const FMetasoundFrontendClassOutput& ClassOutput) 
-				{ 
+				auto IsOutputVertexWithSameName = [&](const FMetasoundFrontendClassOutput& ClassOutput)
+				{
 					return ClassOutput.Name == InName;
 				};
 				return GraphClass->Interface.Outputs.ContainsByPredicate(IsOutputVertexWithSameName);
@@ -2335,7 +2544,7 @@ namespace Metasound
 		{
 			auto IsOutputNodeWithSameName = [&](const FMetasoundFrontendClass& NodeClass, const FMetasoundFrontendNode& Node) 
 			{ 
-				return (NodeClass.Metadata.Type == EMetasoundFrontendClassType::Output) && (Node.Name == InName); 
+				return (NodeClass.Metadata.GetType() == EMetasoundFrontendClassType::Output) && (Node.Name == InName);
 			};
 			return GetNodeByPredicate(IsOutputNodeWithSameName);
 		}
@@ -2344,7 +2553,7 @@ namespace Metasound
 		{
 			auto IsInputNodeWithSameName = [&](const FMetasoundFrontendClass& NodeClass, const FMetasoundFrontendNode& Node) 
 			{ 
-				return (NodeClass.Metadata.Type == EMetasoundFrontendClassType::Input) && (Node.Name == InName); 
+				return (NodeClass.Metadata.GetType() == EMetasoundFrontendClassType::Input) && (Node.Name == InName);
 			};
 			return GetNodeByPredicate(IsInputNodeWithSameName);
 		}
@@ -2353,16 +2562,16 @@ namespace Metasound
 		{
 			auto IsOutputNodeWithSameName = [&](const FMetasoundFrontendClass& NodeClass, const FMetasoundFrontendNode& Node) 
 			{ 
-				return (NodeClass.Metadata.Type == EMetasoundFrontendClassType::Output) && (Node.Name == InName); 
+				return (NodeClass.Metadata.GetType() == EMetasoundFrontendClassType::Output) && (Node.Name == InName);
 			};
 			return GetNodeByPredicate(IsOutputNodeWithSameName);
 		}
 
 		FNodeHandle FGraphController::GetInputNodeWithName(const FString& InName)
 		{
-			auto IsInputNodeWithSameName = [&](const FMetasoundFrontendClass& NodeClass, const FMetasoundFrontendNode& Node) 
-			{ 
-				return (NodeClass.Metadata.Type == EMetasoundFrontendClassType::Input) && (Node.Name == InName); 
+			auto IsInputNodeWithSameName = [&](const FMetasoundFrontendClass& NodeClass, const FMetasoundFrontendNode& Node)
+			{
+				return (NodeClass.Metadata.GetType() == EMetasoundFrontendClassType::Input) && (Node.Name == InName);
 			};
 			return GetNodeByPredicate(IsInputNodeWithSameName);
 		}
@@ -2407,6 +2616,8 @@ namespace Metasound
 							FMetasoundFrontendClassInput& NewInput = GraphClass->Interface.Inputs.Add_GetRef(InClassInput);
 							NewInput.NodeID = Node.ID;
 
+							GraphClass->Interface.UpdateChangeID();
+
 							FNodeAccessPtr NodePtr = GraphClassPtr.GetNodeWithNodeID(Node.ID);
 							return GetNodeHandle(FGraphController::FNodeAndClass{NodePtr, InputClassPtr});
 						}
@@ -2424,11 +2635,77 @@ namespace Metasound
 			return FInvalidNodeController::GetInvalid();
 		}
 
+		FNodeHandle FGraphController::AddInputNode(const FMetasoundFrontendClassInput& InClassInput, const FMetasoundFrontendLiteral* InDefaultValue, const FText* InDisplayName)
+		{
+			FNodeHandle NodeHandle = AddInputVertex(InClassInput);
+			if (ensure(NodeHandle->IsValid()))
+			{
+				if (InDisplayName)
+				{
+					SetInputDisplayName(InClassInput.Name, *InDisplayName);
+				}
+
+				if (InDefaultValue)
+				{
+					SetDefaultInput(InClassInput.VertexID, *InDefaultValue);
+				}
+				else
+				{
+					SetDefaultInputToDefaultLiteralOfType(InClassInput.VertexID);
+				}
+			}
+
+			return NodeHandle;
+		}
+
+		FNodeHandle FGraphController::AddInputNode(const FName InTypeName, const FText& InToolTip, const FMetasoundFrontendLiteral* InDefaultValue, const FText* InDisplayName)
+		{
+			const FGuid VertexID = FGuid::NewGuid();
+
+			FMetasoundFrontendClassInput Description;
+
+			Description.Name = VertexID.ToString();
+			Description.TypeName = InTypeName;
+			Description.Metadata.Description = InToolTip;
+			Description.VertexID = VertexID;
+
+			return AddInputNode(Description, InDefaultValue, InDisplayName);
+		}
+
+		FNodeHandle FGraphController::AddOutputNode(const FMetasoundFrontendClassOutput& InClassOutput, const FText* InDisplayName)
+		{
+			FNodeHandle NewOutput = AddOutputVertex(InClassOutput);
+			if (ensure(NewOutput->IsValid()))
+			{
+				if (InDisplayName)
+				{
+					SetOutputDisplayName(InClassOutput.Name, *InDisplayName);
+				}
+			}
+
+			return NewOutput;
+		}
+
+		FNodeHandle FGraphController::AddOutputNode(const FName InTypeName, const FText& InToolTip, const FText* InDisplayName)
+		{
+			const FGuid VertexID = FGuid::NewGuid();
+
+			FMetasoundFrontendClassOutput Description;
+
+			Description.Name = VertexID.ToString();
+			Description.TypeName = InTypeName;
+			Description.Metadata.Description = InToolTip;
+			Description.VertexID = VertexID;
+
+			return AddOutputNode(Description, InDisplayName);
+		}
+
+
 		bool FGraphController::RemoveInputVertex(const FString& InName)
 		{
-			auto IsInputNodeWithSameName = [&](const FMetasoundFrontendClass& InClass, const FMetasoundFrontendNode& InNode) 
-			{ 
-				return (InClass.Metadata.Type == EMetasoundFrontendClassType::Input) && (InNode.Name == InName); 
+			auto IsInputNodeWithSameName = [&](const FMetasoundFrontendClass& InClass, const FMetasoundFrontendNode& InNode)
+			{
+				return (InClass.Metadata.GetType() == EMetasoundFrontendClassType::Input) && (InNode.Name == InName);
 			};
 
 			for (const FNodeAndClass& NodeAndClass : GetNodesAndClassesByPredicate(IsInputNodeWithSameName))
@@ -2483,6 +2760,8 @@ namespace Metasound
 							FMetasoundFrontendClassOutput& NewOutput = GraphClass->Interface.Outputs.Add_GetRef(InClassOutput);
 							NewOutput.NodeID = Node.ID;
 
+							GraphClass->Interface.UpdateChangeID();
+
 							FNodeAccessPtr NodePtr = GraphClassPtr.GetNodeWithNodeID(Node.ID);
 							return GetNodeHandle(FGraphController::FNodeAndClass { NodePtr, OutputClassPtr });
 						}
@@ -2503,9 +2782,9 @@ namespace Metasound
 
 		bool FGraphController::RemoveOutputVertex(const FString& InName)
 		{
-			auto IsOutputNodeWithSameName = [&](const FMetasoundFrontendClass& InClass, const FMetasoundFrontendNode& InNode) 
-			{ 
-				return (InClass.Metadata.Type == EMetasoundFrontendClassType::Output) && (InNode.Name == InName); 
+			auto IsOutputNodeWithSameName = [&](const FMetasoundFrontendClass& InClass, const FMetasoundFrontendNode& InNode)
+			{
+				return (InClass.Metadata.GetType() == EMetasoundFrontendClassType::Output) && (InNode.Name == InName);
 			};
 
 			for (const FNodeAndClass& NodeAndClass : GetNodesAndClassesByPredicate(IsOutputNodeWithSameName))
@@ -2546,7 +2825,6 @@ namespace Metasound
 			{
 				return Desc->DefaultLiteral;
 			}
-
 			return FMetasoundFrontendLiteral{};
 		}
 
@@ -2574,7 +2852,6 @@ namespace Metasound
 			{
 				Metasound::FLiteral Literal = Frontend::GetDefaultParamForDataType(Desc->TypeName);
 				Desc->DefaultLiteral.SetFromLiteral(Literal);
-
 				return Desc->DefaultLiteral.IsValid();
 			}
 
@@ -2672,7 +2949,7 @@ namespace Metasound
 
 			FConstClassAccessPtr ClassPtr;
 
-			if (EMetasoundFrontendClassType::Graph == ClassMetadata.Type)
+			if (EMetasoundFrontendClassType::Graph == ClassMetadata.GetType())
 			{
 				// Add subgraph and dependencies if needed
 				ClassPtr = OwningDocument->FindClass(ClassMetadata);
@@ -2703,7 +2980,7 @@ namespace Metasound
 				auto IsNodeWithSameID = [&](const FMetasoundFrontendNode& InDesc) { return InDesc.ID == NodeID; };
 				if (const FMetasoundFrontendNode* Desc = GraphClass->Graph.Nodes.FindByPredicate(IsNodeWithSameID))
 				{
-					switch(InNode.GetClassMetadata().Type)
+					switch(InNode.GetClassMetadata().GetType())
 					{
 						case EMetasoundFrontendClassType::Input:
 						{
@@ -2753,17 +3030,17 @@ namespace Metasound
 			}
 		}
 
-		FNodeHandle FGraphController::CreateEmptySubgraph(const FMetasoundFrontendClassMetadata& InInfo)
+		FNodeHandle FGraphController::CreateEmptySubgraph(const FMetasoundFrontendClassMetadata& InMetadata)
 		{
-			if (InInfo.Type == EMetasoundFrontendClassType::Graph)
+			if (InMetadata.GetType() == EMetasoundFrontendClassType::Graph)
 			{
-				if (const FMetasoundFrontendClass* ExistingDependency = OwningDocument->FindClass(InInfo).Get())
+				if (const FMetasoundFrontendClass* ExistingDependency = OwningDocument->FindClass(InMetadata).Get())
 				{
-					UE_LOG(LogMetaSound, Error, TEXT("Cannot add new subgraph. Metasound class already exists with matching metadata Name: \"%s\", Version %d.%d"), *(ExistingDependency->Metadata.ClassName.GetFullName().ToString()), ExistingDependency->Metadata.Version.Major, ExistingDependency->Metadata.Version.Minor);
+					UE_LOG(LogMetaSound, Error, TEXT("Cannot add new subgraph. Metasound class already exists with matching metadata Name: \"%s\", Version %d.%d"), *(ExistingDependency->Metadata.GetClassName().GetFullName().ToString()), ExistingDependency->Metadata.GetVersion().Major, ExistingDependency->Metadata.GetVersion().Minor);
 				}
 				else 
 				{
-					return AddNode(OwningDocument->FindOrAddClass(InInfo));
+					return AddNode(OwningDocument->FindOrAddClass(InMetadata));
 				}
 			}
 			else
@@ -2782,8 +3059,8 @@ namespace Metasound
 				// TODO: Implement subgraph inflation step here.
 
 				// TODO: bubble up errors. 
-				TArray<FMetasoundFrontendGraphClass> Subgraphs = OwningDocument->GetSubgraphs();
-				TArray<FMetasoundFrontendClass> Dependencies = OwningDocument->GetDependencies();
+				const TArray<FMetasoundFrontendGraphClass>& Subgraphs = OwningDocument->GetSubgraphs();
+				const TArray<FMetasoundFrontendClass>& Dependencies = OwningDocument->GetDependencies();
 
 				TUniquePtr<FFrontendGraph> Graph = FFrontendGraphBuilder::CreateGraph(*GraphClass, Subgraphs, Dependencies);
 
@@ -2858,9 +3135,14 @@ namespace Metasound
 
 				int32 NumInputsRemoved = GraphClass->Interface.Inputs.RemoveAll(IsInputWithSameNodeID);
 
+				if (NumInputsRemoved > 0)
+				{
+					GraphClass->Interface.UpdateChangeID();
+				}
+
 				bool bDidRemoveNode = RemoveNode(InNode);
 
-				return (NumInputsRemoved > 0) || bDidRemoveNode;
+				return NumInputsRemoved > 0 || bDidRemoveNode;
 			}
 
 			return false;
@@ -2874,6 +3156,11 @@ namespace Metasound
 
 				int32 NumOutputsRemoved = GraphClass->Interface.Outputs.RemoveAll(IsOutputWithSameNodeID);
 
+				if (NumOutputsRemoved)
+				{
+					GraphClass->Interface.UpdateChangeID();
+				}
+
 				bool bDidRemoveNode = RemoveNode(InNode);
 
 				return (NumOutputsRemoved > 0) || bDidRemoveNode;
@@ -2882,6 +3169,14 @@ namespace Metasound
 			return false;
 		}
 
+		void FGraphController::UpdateInterfaceChangeID()
+		{
+			if (FMetasoundFrontendGraphClass* GraphClass = GraphClassPtr.Get())
+			{
+				GraphClass->Interface.UpdateChangeID();
+			}
+		}
+
 		bool FGraphController::ContainsNodesAndClassesByPredicate(TFunctionRef<bool (const FMetasoundFrontendClass&, const FMetasoundFrontendNode&)> InPredicate) const
 		{
 			if (FMetasoundFrontendGraphClass* GraphClass = GraphClassPtr.Get())
@@ -3090,7 +3385,7 @@ namespace Metasound
 				FGraphHandle OwningGraph = this->AsShared();
 				FGraphAccessPtr GraphPtr = GraphClassPtr.GetGraph();
 
-				switch (NodeClass->Metadata.Type)
+				switch (NodeClass->Metadata.GetType())
 				{
 					case EMetasoundFrontendClassType::Input:
 						{
@@ -3198,7 +3493,7 @@ namespace Metasound
 				FConstGraphHandle OwningGraph = this->AsShared();
 				FConstGraphAccessPtr GraphPtr = GraphClassPtr.GetGraph();
 
-				switch (NodeClass->Metadata.Type)
+				switch (NodeClass->Metadata.GetType())
 				{
 					case EMetasoundFrontendClassType::Input:
 						{
@@ -3262,7 +3557,14 @@ namespace Metasound
 		{
 			if (FMetasoundFrontendGraphClass* GraphClass = GraphClassPtr.Get())
 			{
-				return GraphClass->Interface.Inputs.FindByPredicate([&](const FMetasoundFrontendClassInput& Desc) { return Desc.Name == InName; });
+				FMetasoundFrontendClassInput* ClassInput = GraphClass->Interface.Inputs.FindByPredicate([&](const FMetasoundFrontendClassInput& Desc) { return Desc.Name == InName; });
+				if (ClassInput)
+				{
+					// TODO: This assumes the class input is being mutated due to the adjacent const correct call not being utilized.
+					// Make this more explicit rather than risking whether or not the caller is using proper const correctness.
+					GraphClass->Interface.UpdateChangeID();
+					return ClassInput;
+				}
 			}
 			return nullptr;
 		}
@@ -3280,7 +3582,14 @@ namespace Metasound
 		{
 			if (FMetasoundFrontendGraphClass* GraphClass = GraphClassPtr.Get())
 			{
-				return GraphClass->Interface.Outputs.FindByPredicate([&](const FMetasoundFrontendClassOutput& Desc) { return Desc.Name == InName; });
+				FMetasoundFrontendClassOutput* ClassOutput = GraphClass->Interface.Outputs.FindByPredicate([&](const FMetasoundFrontendClassOutput& Desc) { return Desc.Name == InName; });
+				if (ClassOutput)
+				{
+					// TODO: This assumes the class input is being mutated due to the adjacent const correct call not being utilized.
+					// Make this more explicit rather than risking whether or not the caller is using proper const correctness.
+					GraphClass->Interface.UpdateChangeID();
+					return ClassOutput;
+				}
 			}
 			return nullptr;
 		}
@@ -3298,7 +3607,14 @@ namespace Metasound
 		{
 			if (FMetasoundFrontendGraphClass* GraphClass = GraphClassPtr.Get())
 			{
-				return GraphClass->Interface.Inputs.FindByPredicate([&](const FMetasoundFrontendClassInput& Desc) { return Desc.VertexID == InVertexID; });
+				FMetasoundFrontendClassInput* ClassInput = GraphClass->Interface.Inputs.FindByPredicate([&](const FMetasoundFrontendClassInput& Desc) { return Desc.VertexID == InVertexID; });
+				if (ClassInput)
+				{
+					// TODO: This assumes the class input is being mutated due to the adjacent const correct call not being utilized.
+					// Make this more explicit rather than risking whether or not the caller is using proper const correctness.
+					GraphClass->Interface.UpdateChangeID();
+					return ClassInput;
+				}
 			}
 			return nullptr;
 		}
@@ -3316,7 +3632,14 @@ namespace Metasound
 		{
 			if (FMetasoundFrontendGraphClass* GraphClass = GraphClassPtr.Get())
 			{
-				return GraphClass->Interface.Outputs.FindByPredicate([&](const FMetasoundFrontendClassOutput& Desc) { return Desc.VertexID == InVertexID; });
+				FMetasoundFrontendClassOutput* ClassOutput = GraphClass->Interface.Outputs.FindByPredicate([&](const FMetasoundFrontendClassOutput& Desc) { return Desc.VertexID == InVertexID; });
+				if (ClassOutput)
+				{
+					// TODO: This assumes the class input is being mutated due to the adjacent const correct call not being utilized.
+					// Make this more explicit rather than risking whether or not the caller is using proper const correctness.
+					GraphClass->Interface.UpdateChangeID();
+					return ClassOutput;
+				}
 			}
 			return nullptr;
 		}
@@ -3383,29 +3706,23 @@ namespace Metasound
 			return (nullptr != DocumentPtr.Get());
 		}
 
-		TArray<FMetasoundFrontendClass> FDocumentController::GetDependencies() const
+		const TArray<FMetasoundFrontendClass>& FDocumentController::GetDependencies() const
 		{
 			if (const FMetasoundFrontendDocument* Document = DocumentPtr.Get())
 			{
 				return Document->Dependencies;
 			}
-			return TArray<FMetasoundFrontendClass>();
+
+			return FrontendControllerIntrinsics::GetInvalidValueConstRef<TArray<FMetasoundFrontendClass>>();
 		}
 
-		TArray<FMetasoundFrontendGraphClass> FDocumentController::GetSubgraphs() const
+		const TArray<FMetasoundFrontendGraphClass>& FDocumentController::GetSubgraphs() const
 		{
 			if (const FMetasoundFrontendDocument* Document = DocumentPtr.Get())
 			{
 				return Document->Subgraphs;
 			}
-			return TArray<FMetasoundFrontendGraphClass>();
-		}
-
-		TArray<FMetasoundFrontendClass> FDocumentController::GetClasses() const 
-		{
-			TArray<FMetasoundFrontendClass> Classes = GetDependencies();
-			Classes.Append(GetSubgraphs());
-			return Classes;
+			return FrontendControllerIntrinsics::GetInvalidValueConstRef<TArray<FMetasoundFrontendGraphClass>>();
 		}
 
 		const FMetasoundFrontendGraphClass& FDocumentController::GetRootGraphClass() const
@@ -3437,7 +3754,7 @@ namespace Metasound
 					if (const FMetasoundFrontendClass* OriginalNodeClass = InOtherDocument.Dependencies.FindByPredicate(IsClassWithClassID))
 					{
 						// Should not be a graph class since it's in the dependencies list
-						check(EMetasoundFrontendClassType::Graph != OriginalNodeClass->Metadata.Type);
+						check(EMetasoundFrontendClassType::Graph != OriginalNodeClass->Metadata.GetType());
 
 						if (const FMetasoundFrontendClass* NodeClass = FindOrAddClass(OriginalNodeClass->Metadata).Get())
 						{
@@ -3448,7 +3765,7 @@ namespace Metasound
 						}
 						else
 						{
-							UE_LOG(LogMetaSound, Error, TEXT("Failed to add subgraph dependency [Class:%s]"), *OriginalNodeClass->Metadata.ClassName.ToString());
+							UE_LOG(LogMetaSound, Error, TEXT("Failed to add subgraph dependency [Class:%s]"), *OriginalNodeClass->Metadata.GetClassName().ToString());
 							bSuccess = false;
 						}
 					}
@@ -3587,7 +3904,7 @@ namespace Metasound
 
 					// Cannot add a subgraph using this method because dependencies
 					// of external graph are not added in this method.
-					check(EMetasoundFrontendClassType::Graph != NewClassDescription.Metadata.Type);
+					check(EMetasoundFrontendClassType::Graph != NewClassDescription.Metadata.GetType());
 					NewClassDescription.ID = FGuid::NewGuid();
 
 					Document->Dependencies.Add(MoveTemp(NewClassDescription));
@@ -3599,13 +3916,13 @@ namespace Metasound
 				if (const FMetasoundFrontendClass* MetasoundClass = ClassPtr.Get())
 				{
 					// External node classes must match version to return shared definition.
-					if (MetasoundClass->Metadata.Type == EMetasoundFrontendClassType::External)
+					if (MetasoundClass->Metadata.GetType() == EMetasoundFrontendClassType::External)
 					{
 						// TODO: Assuming we want to recheck classes when they add another
 						// node, this should be replace with a call to synchronize a 
 						// single class.
 						FMetasoundFrontendClass NewClass = GenerateClassDescription(InKey);
-						if (NewClass.Metadata.Version.Major != MetasoundClass->Metadata.Version.Major)
+						if (NewClass.Metadata.GetVersion().Major != MetasoundClass->Metadata.GetVersion().Major)
 						{
 							return AddClass(MoveTemp(NewClass));
 						}
@@ -3633,9 +3950,9 @@ namespace Metasound
 			if (const FMetasoundFrontendClass* Class = ClassPtr.Get())
 			{
 				// External node classes must match major version to return shared definition.
-				if (EMetasoundFrontendClassType::External == InMetadata.Type)
+				if (EMetasoundFrontendClassType::External == InMetadata.GetType())
 				{
-					if (InMetadata.Version.Major != Class->Metadata.Version.Major)
+					if (InMetadata.GetVersion().Major != Class->Metadata.GetVersion().Major)
 					{
 						// Mismatched major version. Reset class pointer to null.
 						ClassPtr = FConstClassAccessPtr();
@@ -3650,7 +3967,7 @@ namespace Metasound
 				// If no matching class found, attempt to add a class matching the metadata.
 				if (FMetasoundFrontendDocument* Document = DocumentPtr.Get())
 				{
-					switch (InMetadata.Type)
+					switch (InMetadata.GetType())
 					{
 						case EMetasoundFrontendClassType::External:
 						case EMetasoundFrontendClassType::Input:
@@ -3669,10 +3986,10 @@ namespace Metasound
 								UE_LOG(LogMetaSound, Error,
 									TEXT("Cannot add external dependency. No Metasound class found with matching registry key [Key:%s, Name:%s, Version:%s]. Suggested solution \"%s\" by %s."),
 									*Key,
-									*InMetadata.ClassName.GetFullName().ToString(),
-									*InMetadata.Version.ToString(),
-									*InMetadata.PromptIfMissing.ToString(),
-									*InMetadata.Author.ToString());
+									*InMetadata.GetClassName().GetFullName().ToString(),
+									*InMetadata.GetVersion().ToString(),
+									*InMetadata.GetPromptIfMissing().ToString(),
+									*InMetadata.GetAuthor().ToString());
 							}
 						}
 						break;
@@ -3691,8 +4008,8 @@ namespace Metasound
 						{
 							UE_LOG(LogMetaSound, Error, TEXT(
 								"Unsupported metasound class type for node: \"%s\" (%s)."),
-								*InMetadata.ClassName.GetFullName().ToString(),
-								*InMetadata.Version.ToString());
+								*InMetadata.GetClassName().GetFullName().ToString(),
+								*InMetadata.GetVersion().ToString());
 							checkNoEntry();
 						}
 					}
@@ -3704,6 +4021,32 @@ namespace Metasound
 			return ClassPtr;
 		}
 
+		const FMetasoundFrontendClass* FDocumentController::SynchronizeDependency(const FNodeRegistryKey& InKey)
+		{
+			FMetasoundFrontendDocument* Document = DocumentPtr.Get();
+			if (!ensure(Document))
+			{
+				return nullptr;
+			}
+
+			for (FMetasoundFrontendClass& Class : Document->Dependencies)
+			{
+				const FNodeRegistryKey DependencyRegistryKey = FMetasoundFrontendRegistryContainer::Get()->GetRegistryKey(Class.Metadata);
+				if (InKey == DependencyRegistryKey)
+				{
+					const FGuid ClassID = Class.ID;
+					FMetasoundFrontendClass RegisteredClass = GenerateClassDescription(InKey);
+					RegisteredClass.ID = ClassID;
+					Class = RegisteredClass;
+					return &Class;
+				}
+			}
+
+			FMetasoundFrontendClass NewClass = GenerateClassDescription(InKey);
+			NewClass.ID = FGuid::NewGuid();
+			return &Document->Dependencies.Add_GetRef(NewClass);
+		}
+
 		void FDocumentController::SynchronizeDependencies()
 		{
 			if (FMetasoundFrontendDocument* Document = DocumentPtr.Get())
diff --git a/Engine/Plugins/Runtime/Metasound/Source/MetasoundFrontend/Private/MetasoundFrontendStandardController.h b/Engine/Plugins/Runtime/Metasound/Source/MetasoundFrontend/Private/MetasoundFrontendStandardController.h
--- a/Engine/Plugins/Runtime/Metasound/Source/MetasoundFrontend/Private/MetasoundFrontendStandardController.h
+++ b/Engine/Plugins/Runtime/Metasound/Source/MetasoundFrontend/Private/MetasoundFrontendStandardController.h
@@ -325,12 +325,11 @@ namespace Metasound
 			// Connection controls.
 			bool ConnectWithConverterNode(IOutputController& InController, const FConverterNodeInfo& InNodeClassName) override;
 
-
 		private:
-
 			FConstClassInputAccessPtr OwningGraphClassInputPtr;
 		};
 
+
 		/** FBaseNodeController provides common functionality for multiple derived
 		 * node controllers.
 		 */
@@ -376,6 +375,11 @@ namespace Metasound
 			const FMetasoundFrontendNodeStyle& GetNodeStyle() const override;
 			void SetNodeStyle(const FMetasoundFrontendNodeStyle& InStyle) override;
 
+			bool CanAutoUpdate() const override;
+			FNodeHandle ReplaceWithVersion(const FMetasoundFrontendVersionNumber& InNewVersion) override;
+			FMetasoundFrontendVersionNumber FindHighestVersionInRegistry() const override;
+			FMetasoundFrontendVersionNumber FindHighestMinorVersionInRegistry() const override;
+
 			const FString& GetNodeName() const override;
 
 			/** Returns the human-readable display name of the given node. */
@@ -492,7 +496,7 @@ namespace Metasound
 			virtual FOutputHandle CreateOutputController(FGuid InVertexID, FConstVertexAccessPtr InNodeVertexPtr, FConstClassOutputAccessPtr InClassOutputPtr, FNodeHandle InOwningNode) const = 0;
 		};
 
-		/** FNodeController represents a external or subgraph node. */
+		/** Represents an external node (defined in either code or by an asset's root graph). */
 		class FNodeController : public FBaseNodeController
 		{
 			using FRegistry = FMetasoundFrontendRegistryContainer;
@@ -732,6 +736,8 @@ namespace Metasound
 			TArray<FNodeHandle> GetInputNodes() override;
 			TArray<FConstNodeHandle> GetConstInputNodes() const override;
 
+			void ClearGraph() override;
+
 			void IterateNodes(TUniqueFunction<void(FNodeHandle)> InFunction, EMetasoundFrontendClassType InClassType) override;
 			void IterateConstNodes(TUniqueFunction<void(FConstNodeHandle)> InFunction, EMetasoundFrontendClassType InClassType) const override;
 
@@ -756,6 +762,8 @@ namespace Metasound
 			FNodeHandle AddOutputVertex(const FMetasoundFrontendClassOutput& InDescription) override;
 			bool RemoveOutputVertex(const FString& InName) override;
 
+			void UpdateInterfaceChangeID() override;
+
 			// This can be used to determine what kind of property editor we should use for the data type of a given input.
 			// Will return Invalid if the input couldn't be found, or if the input doesn't support any kind of literals.
 			ELiteralType GetPreferredLiteralTypeForInputVertex(const FString& InInputName) const override;
@@ -812,6 +820,11 @@ namespace Metasound
 			FDocumentHandle GetOwningDocument() override;
 			FConstDocumentHandle GetOwningDocument() const override;
 
+			FNodeHandle AddInputNode(const FMetasoundFrontendClassInput& InClassInput, const FMetasoundFrontendLiteral* InDefaultValue, const FText* InDisplayName) override;
+			FNodeHandle AddInputNode(const FName InTypeName, const FText& InToolTip, const FMetasoundFrontendLiteral* InDefaultValue, const FText* InDisplayName) override;
+			FNodeHandle AddOutputNode(const FMetasoundFrontendClassOutput& InClassOutput, const FText* InDisplayName) override;
+			FNodeHandle AddOutputNode(const FName InTypeName, const FText& InToolTip, const FText* InDisplayName) override;
+
 		protected:
 
 			FDocumentAccess ShareAccess() override;
@@ -906,9 +919,8 @@ namespace Metasound
 
 			bool IsValid() const override;
 
-			TArray<FMetasoundFrontendClass> GetDependencies() const override;
-			TArray<FMetasoundFrontendGraphClass> GetSubgraphs() const override;
-			TArray<FMetasoundFrontendClass> GetClasses() const override;
+			const TArray<FMetasoundFrontendClass>& GetDependencies() const override;
+			const TArray<FMetasoundFrontendGraphClass>& GetSubgraphs() const override;
 			const FMetasoundFrontendGraphClass& GetRootGraphClass() const override;
 
 			FConstClassAccessPtr FindDependencyWithID(FGuid InClassID) const override;
@@ -921,14 +933,15 @@ namespace Metasound
 			FConstClassAccessPtr FindOrAddClass(const FNodeRegistryKey& InKey) override;
 			FConstClassAccessPtr FindOrAddClass(const FMetasoundFrontendClassMetadata& InMetadata) override;
 
-			virtual FGraphHandle AddDuplicateSubgraph(const IGraphController& InGraph) override;
+			FGraphHandle AddDuplicateSubgraph(const IGraphController& InGraph) override;
 
-			virtual const FMetasoundFrontendVersion& GetArchetypeVersion() const override;
-			virtual void SetArchetypeVersion(const FMetasoundFrontendVersion& InVersion) override;
+			const FMetasoundFrontendVersion& GetArchetypeVersion() const override;
+			void SetArchetypeVersion(const FMetasoundFrontendVersion& InVersion) override;
 
 			void SetMetadata(const FMetasoundFrontendDocumentMetadata& InMetadata) override;
 			const FMetasoundFrontendDocumentMetadata& GetMetadata() const override;
 
+			const FMetasoundFrontendClass* SynchronizeDependency(const FNodeRegistryKey& InKey) override;
 			void SynchronizeDependencies() override;
 
 			FGraphHandle GetRootGraph() override;
diff --git a/Engine/Plugins/Runtime/Metasound/Source/MetasoundFrontend/Private/MetasoundFrontendTransform.cpp b/Engine/Plugins/Runtime/Metasound/Source/MetasoundFrontend/Private/MetasoundFrontendTransform.cpp
--- a/Engine/Plugins/Runtime/Metasound/Source/MetasoundFrontend/Private/MetasoundFrontendTransform.cpp
+++ b/Engine/Plugins/Runtime/Metasound/Source/MetasoundFrontend/Private/MetasoundFrontendTransform.cpp
@@ -4,8 +4,9 @@
 #include "Algo/Transform.h"
 #include "CoreMinimal.h"
 #include "MetasoundArchetype.h"
+#include "MetasoundAssetBase.h"
 #include "MetasoundFrontendArchetypeRegistry.h"
-#include "MetasoundFrontendDocument.h"
+#include "MetasoundFrontendRegistries.h"
 #include "MetasoundFrontendSearchEngine.h"
 #include "MetasoundLog.h"
 #include "Misc/App.h"
@@ -330,6 +331,313 @@ namespace Metasound
 			return bDidEdit;
 		}
 
+		bool FAutoUpdateRootGraph::Transform(FDocumentHandle InDocument) const
+		{
+			check(Asset);
+			check(AssetInterface);
+
+			bool bDidEdit = false;
+
+			FDocumentHandle PresetReferenceDocument = IDocumentController::GetInvalidHandle();
+			TArray<TPair<FNodeHandle, FMetasoundFrontendVersionNumber>> NodesToUpdate;
+
+			FGraphHandle RootGraph = InDocument->GetRootGraph();
+			RootGraph->IterateNodes([&](FNodeHandle NodeHandle)
+			{
+				using namespace Metasound::Frontend;
+
+				const FMetasoundFrontendClassMetadata& ClassMetadata = NodeHandle->GetClassMetadata();
+
+				if (!NodeHandle->CanAutoUpdate())
+				{
+					return;
+				}
+
+				FMetasoundFrontendVersionNumber UpdateVersion = NodeHandle->FindHighestMinorVersionInRegistry();
+				if (UpdateVersion.IsValid() && UpdateVersion > ClassMetadata.GetVersion())
+				{
+					UE_LOG(LogMetaSound, Display, TEXT("Auto-Updating node class '%s': Newer minor version '%s' found."), *ClassMetadata.GetDisplayName().ToString(), *UpdateVersion.ToString());
+				}
+				else
+				{
+					UpdateVersion = ClassMetadata.GetVersion();
+					UE_LOG(LogMetaSound, Display, TEXT("Auto-Updating node with class '%s (%s)': Interface change detected."), *ClassMetadata.GetDisplayName().ToString(), *UpdateVersion.ToString());
+				}
+
+				const bool bInterfaceManaged = RootGraph->GetGraphMetadata().GetAutoUpdateManagesInterface();
+				if (bInterfaceManaged)
+				{
+					if (ensure(ClassMetadata.GetType() == EMetasoundFrontendClassType::External))
+					{
+						const FNodeRegistryKey RegistryKey = FMetasoundFrontendRegistryContainer::Get()->GetRegistryKey(ClassMetadata);
+						FMetasoundAssetBase* ReferencedMetaSoundAsset = AssetInterface->FindAssetFromKey(RegistryKey);
+						if (ReferencedMetaSoundAsset)
+						{
+							PresetReferenceDocument = ReferencedMetaSoundAsset->GetDocumentHandle();
+							return;
+						}
+					}
+				}
+
+				NodesToUpdate.Add(TPair<FNodeHandle, FMetasoundFrontendVersionNumber>(NodeHandle, UpdateVersion));
+			}, EMetasoundFrontendClassType::External);
+
+			if (PresetReferenceDocument->IsValid())
+			{
+				bDidEdit |= FRebuildPresetRootGraph(PresetReferenceDocument).Transform(InDocument);
+			}
+			else
+			{
+				for (const TPair<FNodeHandle, FMetasoundFrontendVersionNumber>& Pair : NodesToUpdate)
+				{
+					FNodeHandle ExistingNode = Pair.Key;
+					FNodeHandle NewNode = ExistingNode->ReplaceWithVersion(Pair.Value);
+					bDidEdit |= NewNode->GetID() != ExistingNode->GetID();
+				}
+			}
+
+			return bDidEdit;
+		}
+
+		bool FRebuildPresetRootGraph::Transform(FDocumentHandle InDocument) const
+		{
+			FGraphHandle RootGraphHandle = InDocument->GetRootGraph();
+			if (!ensure(RootGraphHandle->IsValid()))
+			{
+				return false;
+			}
+
+			if (!RootGraphHandle->GetGraphMetadata().GetAutoUpdateManagesInterface())
+			{
+				return false;
+			}
+
+			FConstGraphHandle ReferencedGraphHandle = ReferencedDocument->GetRootGraph();
+			if (!ensure(ReferencedGraphHandle->IsValid()))
+			{
+				return false;
+			}
+
+			using FDefaultKey = TPair<FString, FName>;
+			struct FDefaultInputValue
+			{
+				FGuid VertexID;
+				FString VertexName;
+				FMetasoundFrontendLiteral Literal;
+			};
+			struct FDefaultOutputValue
+			{
+				FGuid VertexID;
+				FString VertexName;
+			};
+			TMap<FDefaultKey, FDefaultInputValue> InitInputDefaults;
+			TMap<FDefaultKey, FDefaultOutputValue> InitOutputDefaults;
+
+			// 1. Prep by caching existing output guids, input literal values, clearing graph & ensuring ref archetype is supported
+			RootGraphHandle->IterateConstNodes([&](FConstNodeHandle InputNode)
+			{
+				FGuid InputID = InputNode->GetID();
+				if (ensure(InputNode->GetNumInputs() == 1))
+				{
+					InputNode->IterateConstInputs([&](FConstInputHandle Input)
+					{
+						const FGuid InputVertexID = RootGraphHandle->GetVertexIDForInputVertex(InputNode->GetNodeName());
+						FMetasoundFrontendLiteral Literal = RootGraphHandle->GetDefaultInput(InputVertexID);
+						FDefaultKey Key(InputNode->GetDisplayName().ToString(), Input->GetDataType());
+						FDefaultInputValue Value = { MoveTemp(InputID), InputNode->GetNodeName(), MoveTemp(Literal), };
+						InitInputDefaults.Emplace(MoveTemp(Key), MoveTemp(Value));
+					});
+				}
+			}, EMetasoundFrontendClassType::Input);
+
+			RootGraphHandle->IterateConstNodes([&](FConstNodeHandle OutputNode)
+			{
+				FGuid OutputID = OutputNode->GetID();
+				if (ensure(OutputNode->GetNumOutputs() == 1))
+				{
+					OutputNode->IterateConstOutputs([&](FConstOutputHandle Output)
+					{
+						FDefaultKey Key(OutputNode->GetDisplayName().ToString(), Output->GetDataType());
+						FDefaultOutputValue Value = { MoveTemp(OutputID), OutputNode->GetNodeName() };
+						InitOutputDefaults.Emplace(MoveTemp(Key), MoveTemp(Value));
+					});
+				}
+			}, EMetasoundFrontendClassType::Output);
+
+			RootGraphHandle->ClearGraph();
+			const FMetasoundFrontendVersion& RefArchetype = ReferencedDocument->GetArchetypeVersion();
+// 			if (!ensure(IsArchetypeSupported(RefArchetype)))
+// 			{
+// 				return false;
+// 			}
+
+			// 2. Set initial locations of 3 columns (inputs, reference node, and outputs)
+			const TArray<FConstNodeHandle> ReferencedInputs = ReferencedGraphHandle->GetConstInputNodes();
+			const TArray<FConstNodeHandle> ReferencedOutputs = ReferencedGraphHandle->GetConstOutputNodes();
+
+			FVector2D InputNodeLocation = FVector2D::ZeroVector;
+
+			const float ReferenceNodeYOffset = FMath::Max(ReferencedInputs.Num(), ReferencedOutputs.Num()) * 0.5f;
+			const FVector2D ReferenceNodeLocation = InputNodeLocation + FVector2D(DisplayStyle::NodeLayout::DefaultOffsetX.X, ReferenceNodeYOffset);
+			FVector2D OutputNodeLocation = InputNodeLocation + (2 * DisplayStyle::NodeLayout::DefaultOffsetX);
+
+			TMap<FString, FString> PresetInputToReferenceInputMap;
+			TMap<FString, FString> PresetOutputToReferenceOutputMap;
+
+			// 2. Generate preset's inputs & outputs
+			for (const FConstNodeHandle& RefGraphInputNode : ReferencedInputs)
+			{
+				const FText& DisplayName = RefGraphInputNode->GetDisplayName();
+				const FText& Description = RefGraphInputNode->GetDescription();
+
+				// Should only ever be one
+				ensure(RefGraphInputNode->GetNumInputs() == 1);
+				RefGraphInputNode->IterateConstInputs([&](FConstInputHandle RefGraphInputHandle)
+				{
+					using namespace Metasound::Frontend;
+					FNodeHandle NodeHandle = INodeController::GetInvalidHandle();
+					const FName DataTypeName = RefGraphInputHandle->GetDataType();
+					const FDefaultKey Key(DisplayName.ToString(), DataTypeName);
+					if (const FDefaultInputValue* InputValue = InitInputDefaults.Find(Key))
+					{
+						FMetasoundFrontendClassInput ClassInput;
+						ClassInput.Name = InputValue->VertexName;
+						ClassInput.TypeName = DataTypeName;
+						ClassInput.Metadata.Description = Description;
+						ClassInput.VertexID = InputValue->VertexID;
+
+						NodeHandle = RootGraphHandle->AddInputNode(ClassInput, &InputValue->Literal, &DisplayName);
+					}
+					else
+					{
+						FGuid InputID = ReferencedGraphHandle->GetVertexIDForInputVertex(RefGraphInputHandle->GetName());
+						const FMetasoundFrontendLiteral DefaultLiteral = ReferencedGraphHandle->GetDefaultInput(InputID);
+						NodeHandle = RootGraphHandle->AddInputNode(DataTypeName, Description, &DefaultLiteral, &DisplayName);
+					}
+
+					if (ensure(NodeHandle->IsValid()))
+					{
+						FMetasoundFrontendNodeStyle NodeStyle;
+						NodeStyle.Display.Locations.Add(FGuid().NewGuid(), InputNodeLocation);
+						NodeHandle->SetNodeStyle(NodeStyle);
+
+						PresetInputToReferenceInputMap.Add(NodeHandle->GetNodeName(), RefGraphInputHandle->GetName());
+						InputNodeLocation += DisplayStyle::NodeLayout::DefaultOffsetY;
+					}
+				});
+			}
+
+			for (const FConstNodeHandle& RefGraphOutputNode : ReferencedOutputs)
+			{
+				const FText& DisplayName = RefGraphOutputNode->GetDisplayName();
+				const FText& Description = RefGraphOutputNode->GetDescription();
+
+				// Should only ever be one
+				ensure(RefGraphOutputNode->GetNumOutputs() == 1);
+				RefGraphOutputNode->IterateConstOutputs([&](FConstOutputHandle RefGraphOutputHandle)
+				{
+					using namespace Metasound::Frontend;
+
+					FNodeHandle NodeHandle = INodeController::GetInvalidHandle();
+					const FName DataTypeName = RefGraphOutputHandle->GetDataType();
+					const FDefaultKey Key(DisplayName.ToString(), DataTypeName);
+					if (FDefaultOutputValue* OutputValue = InitOutputDefaults.Find(Key))
+					{
+						FMetasoundFrontendClassOutput ClassOutput;
+						ClassOutput.Name = OutputValue->VertexName;
+						ClassOutput.TypeName = DataTypeName;
+						ClassOutput.Metadata.Description = Description;
+						ClassOutput.VertexID = OutputValue->VertexID;
+
+						NodeHandle = RootGraphHandle->AddOutputNode(ClassOutput, &DisplayName);
+					}
+					else
+					{
+						NodeHandle = RootGraphHandle->AddOutputNode(DataTypeName, Description, &DisplayName);
+					}
+
+					if (ensure(NodeHandle->IsValid()))
+					{
+						FMetasoundFrontendNodeStyle NodeStyle;
+						NodeStyle.Display.Locations.Add(FGuid().NewGuid(), OutputNodeLocation);
+						NodeHandle->SetNodeStyle(NodeStyle);
+
+						PresetOutputToReferenceOutputMap.Add(NodeHandle->GetNodeName(), RefGraphOutputHandle->GetName());
+						OutputNodeLocation += DisplayStyle::NodeLayout::DefaultOffsetY;
+					}
+				});
+			}
+
+			// 3. Generate a referencing node to the given referenced MetaSound.
+			FMetasoundFrontendClassMetadata ReferencedClassMetadata = ReferencedGraphHandle->GetOwningDocument()->GetRootGraphClass().Metadata;
+
+			// Swap type on look-up as it will be referenced as an externally defined class relative to the new Preset asset
+			ReferencedClassMetadata.SetType(EMetasoundFrontendClassType::External);
+
+			FNodeHandle ReferencedNodeHandle = RootGraphHandle->AddNode(ReferencedClassMetadata);
+			ensure (ReferencedNodeHandle->IsValid());
+			FMetasoundFrontendNodeStyle RefNodeStyle;
+			RefNodeStyle.Display.Locations.Add(FGuid::NewGuid(), ReferenceNodeLocation);
+			ReferencedNodeHandle->SetNodeStyle(RefNodeStyle);
+
+			// 4. Connect preset's respective inputs & outputs to generated referencing node.
+			RootGraphHandle->IterateNodes([&](FNodeHandle InputNode)
+			{
+				const FString& Name = InputNode->GetNodeName();
+
+				// Should only ever be one
+				ensure(InputNode->GetNumOutputs() == 1);
+				InputNode->IterateOutputs([&](FOutputHandle OutputHandle)
+				{
+					const FString& InputName = PresetInputToReferenceInputMap.FindChecked(Name);
+					const TArray<FInputHandle> ReferenceInputs = ReferencedNodeHandle->GetInputsWithVertexName(InputName);
+					if (ensure(ReferenceInputs.Num() == 1))
+					{
+						ensure(ReferenceInputs[0]->Connect(*OutputHandle));
+					}
+				});
+			},
+			EMetasoundFrontendClassType::Input);
+
+			RootGraphHandle->IterateNodes([&](FNodeHandle OutputNode)
+			{
+				const FString& Name = OutputNode->GetNodeName();
+
+				// Should only ever be one
+				ensure(OutputNode->GetNumInputs() == 1);
+				OutputNode->IterateInputs([&](FInputHandle InputHandle)
+				{
+					const FString& OutputName = PresetOutputToReferenceOutputMap.FindChecked(Name);
+					const TArray<FOutputHandle> ReferenceInputs = ReferencedNodeHandle->GetOutputsWithVertexName(OutputName);
+					if (ensure(ReferenceInputs.Num() == 1))
+					{
+						ensure(InputHandle->Connect(*ReferenceInputs[0]));
+					}
+				});
+			},
+			EMetasoundFrontendClassType::Output);
+
+			// 5. Run transform to set to match archetype (marking inputs/outputs as required that were copied over)
+			FMatchRootGraphToArchetype(RefArchetype).Transform(InDocument);
+
+			return true;
+		}
+
+		bool FSynchronizeAssetClassName::Transform(FDocumentHandle InDocument) const
+		{
+			const FMetasoundFrontendClassMetadata& Metadata = InDocument->GetRootGraphClass().Metadata;
+			const FText NewAssetName = FText::FromString(AssetName.ToString());
+
+			if (Metadata.GetDisplayName().CompareTo(NewAssetName) != 0)
+			{
+				FMetasoundFrontendClassMetadata NewMetadata = Metadata;
+				NewMetadata.SetDisplayName(NewAssetName);
+				InDocument->GetRootGraph()->SetGraphMetadata(NewMetadata);
+			}
+
+			return true;
+		}
+
 		// Versioning Transforms
 		class FVersionDocumentTransform : public IDocumentTransform
 		{
@@ -378,7 +686,7 @@ namespace Metasound
 				for (FNodeHandle& NodeHandle : FrontendNodes)
 				{
 					const bool bIsHiddenNode = NodeHandle->GetNodeStyle().Display.Visibility == EMetasoundFrontendNodeStyleDisplayVisibility::Hidden;
-					const bool bIsInputNode = EMetasoundFrontendClassType::Input == NodeHandle->GetClassMetadata().Type;
+					const bool bIsInputNode = EMetasoundFrontendClassType::Input == NodeHandle->GetClassMetadata().GetType();
 					const bool bIsHiddenInputNode = bIsHiddenNode && bIsInputNode;
 
 					if (bIsHiddenInputNode)
@@ -437,12 +745,12 @@ namespace Metasound
 
 			void TransformInternal(FDocumentHandle InDocument) const override
 			{
-				FGraphHandle GraphHandle = InDocument->GetRootGraph();
-				FMetasoundFrontendClassMetadata Metadata = GraphHandle->GetGraphMetadata();
+				const FMetasoundFrontendGraphClass& GraphClass = InDocument->GetRootGraphClass();
+				FMetasoundFrontendClassMetadata Metadata = GraphClass.Metadata;
 
-				Metadata.ClassName = FMetasoundFrontendClassName { "GraphAsset", Name, *Path };
-				Metadata.DisplayName = FText::FromString(Name.ToString());
-				GraphHandle->SetGraphMetadata(Metadata);
+				Metadata.SetClassName({ "GraphAsset", Name, *Path });
+				Metadata.SetDisplayName(FText::FromString(Name.ToString()));
+				InDocument->GetRootGraph()->SetGraphMetadata(Metadata);
 			}
 		};
 
@@ -461,11 +769,11 @@ namespace Metasound
 
 			void TransformInternal(FDocumentHandle InDocument) const override
 			{
-				FGraphHandle GraphHandle = InDocument->GetRootGraph();
-				FMetasoundFrontendClassMetadata Metadata = GraphHandle->GetGraphMetadata();
+				const FMetasoundFrontendGraphClass& GraphClass = InDocument->GetRootGraphClass();
+				FMetasoundFrontendClassMetadata Metadata = GraphClass.Metadata;
 
-				Metadata.ClassName = FMetasoundFrontendClassName { FName(), *FGuid::NewGuid().ToString(), FName() };
-				GraphHandle->SetGraphMetadata(Metadata);
+				Metadata.SetClassName(FMetasoundFrontendClassName { FName(), *FGuid::NewGuid().ToString(), FName() });
+				InDocument->GetRootGraph()->SetGraphMetadata(Metadata);
 			}
 		};
 
@@ -505,18 +813,43 @@ namespace Metasound
 
 					if (const FMetasoundFrontendArchetype* Arch = FindMostSimilarArchetypeSupportingEnvironment(RootGraph, Dependencies, Subgraphs, AllArchetypes))
 					{
-						UE_LOG(LogMetaSound, Display, TEXT("Assigned archetype [ArchetypeVersion:%s] to document [RootGraphClassName:%s]"), *Arch->Version.ToString(), *RootGraph.Metadata.ClassName.ToString());
+						UE_LOG(LogMetaSound, Display, TEXT("Assigned archetype [ArchetypeVersion:%s] to document [RootGraphClassName:%s]"),
+							*Arch->Version.ToString(), *RootGraph.Metadata.GetClassName().ToString());
 
 						InDocument->SetArchetypeVersion(Arch->Version);
 					}
 					else
 					{
-						UE_LOG(LogMetaSound, Warning, TEXT("Failed to find archetype for document [RootGraphClassName:%s]"), *RootGraph.Metadata.ClassName.ToString());
+						UE_LOG(LogMetaSound, Warning, TEXT("Failed to find archetype for document [RootGraphClassName:%s]"),
+							*RootGraph.Metadata.GetClassName().ToString());
 					}
 				}
 			}
 		};
 
+		/** Versions document from 1.4 to 1.5. */
+		class FVersionDocument_1_5 : public FVersionDocumentTransform
+		{
+		public:
+			FVersionDocument_1_5(FName InAssetName)
+				: AssetName(InAssetName)
+			{
+			}
+
+			FMetasoundFrontendVersionNumber GetTargetVersion() const override
+			{
+				return { 1, 5 };
+			}
+
+			void TransformInternal(FDocumentHandle InDocument) const override
+			{
+				FSynchronizeAssetClassName(AssetName).Transform(InDocument);
+			}
+
+		private:
+			FName AssetName;
+		};
+
 		FVersionDocument::FVersionDocument(FName InName, const FString& InPath)
 			: Name(InName)
 			, Path(InPath)
@@ -537,6 +870,7 @@ namespace Metasound
 			bWasUpdated |= FVersionDocument_1_2(Name, Path).Transform(InDocument);
 			bWasUpdated |= FVersionDocument_1_3().Transform(InDocument);
 			bWasUpdated |= FVersionDocument_1_4().Transform(InDocument);
+			bWasUpdated |= FVersionDocument_1_5(Name).Transform(InDocument);
 
 			return bWasUpdated;
 		}
