First pass on MetaSound Composition/Preset Support

diff --git a/Engine/Plugins/Runtime/Metasound/Source/MetasoundFrontend/Public/MetasoundAssetBase.h b/Engine/Plugins/Runtime/Metasound/Source/MetasoundFrontend/Public/MetasoundAssetBase.h
--- a/Engine/Plugins/Runtime/Metasound/Source/MetasoundFrontend/Public/MetasoundAssetBase.h
+++ b/Engine/Plugins/Runtime/Metasound/Source/MetasoundFrontend/Public/MetasoundAssetBase.h
@@ -6,15 +6,20 @@
 #include "MetasoundAccessPtr.h"
 #include "MetasoundFrontend.h"
 #include "MetasoundFrontendController.h"
+#include "MetasoundFrontendDocument.h"
 #include "MetasoundFrontendDocumentAccessPtr.h"
 #include "MetasoundGraph.h"
 #include "MetasoundInstanceTransmitter.h"
 #include "MetasoundLog.h"
+#include "UObject/SoftObjectPath.h"
 #include "UObject/WeakObjectPtrTemplates.h"
 
 
+// Forward Declarations
+class FMetasoundAssetBase;
 class UEdGraph;
 
+
 namespace Metasound
 {
 	namespace AssetTags
@@ -32,10 +37,20 @@ namespace Metasound
 	} // namespace AssetTags
 } // namespace Metasound
 
+class METASOUNDFRONTEND_API IMetaSoundAssetInterface
+{
+public:
+	virtual FMetasoundAssetBase* FindAssetFromKey(const Metasound::Frontend::FNodeRegistryKey& InRegistryKey) const = 0;
+	virtual const FSoftObjectPath* FindObjectPathFromKey(const Metasound::Frontend::FNodeRegistryKey& InRegistryKey) const = 0;
+
+	// Attempts to load an FMetasoundAssetBase from the given path, or returns it if its already loaded
+	virtual FMetasoundAssetBase* TryLoadAsset(const FSoftObjectPath& InObjectPath) const = 0;
+};
+
 /** FMetasoundAssetBase is intended to be a mix-in subclass for UObjects which utilize
  * Metasound assets.  It provides consistent access to FMetasoundFrontendDocuments, control
  * over the FMetasoundFrontendArchetype of the FMetasoundFrontendDocument.  It also enables the UObject
- * to be utilized by a host of other engine tools built to support Metasounds.
+ * to be utilized by a host of other engine tools built to support MetaSounds.
  */
 class METASOUNDFRONTEND_API FMetasoundAssetBase
 {
@@ -69,6 +84,9 @@ public:
 	// Registers the root graph of the given asset with the MetaSound Frontend.
 	void RegisterGraphWithFrontend();
 
+	// Unregisters the root graph of the given asset with the MetaSound Frontend.
+	void UnregisterGraphWithFrontend();
+
 	bool CopyDocumentAndInjectReceiveNodes(uint64 InInstanceID, const FMetasoundFrontendDocument& InSourceDoc, FMetasoundFrontendDocument& OutDestDoc) const;
 
 	// Sets/overwrites the root class metadata
@@ -93,8 +111,17 @@ public:
 	// Gets the asset class info.
 	virtual Metasound::Frontend::FNodeClassInfo GetAssetClassInfo() const = 0;
 
+
+	virtual TSet<FSoftObjectPath>& GetReferencedAssets() = 0;
+	virtual const TSet<FSoftObjectPath>& GetReferencedAssets() const = 0;
+
 	void ConvertFromPreset();
 
+	void RebuildReferencedAssets(const IMetaSoundAssetInterface& InAssetInterface);
+
+	bool AddingReferenceCausesLoop(const FSoftObjectPath& InReferencePath, const IMetaSoundAssetInterface& InAssetInterface) const;
+	bool ContainReferenceLoop(const IMetaSoundAssetInterface& IMetaSoundAssetInterface) const;
+
 	// Imports data from a JSON string directly
 	bool ImportFromJSON(const FString& InJSON);
 
@@ -121,9 +148,11 @@ public:
 	// This must be called on UObject::PostLoad, as well as in this asset's UFactory, to fix up the root document based on the most recent version of the archetype.
 	void ConformDocumentToArchetype();
 
-	bool VersionAsset(bool bInMarkDirty = true);
+	bool AutoUpdate(const IMetaSoundAssetInterface& InAssetInterface, bool bInMarkDirty, bool bInUpdateReferencedAssets = true);
+
+	bool VersionAsset(const IMetaSoundAssetInterface& InAssetInterface, bool bInMarkDirty, bool bInVersionReferencedAssets = true);
 
-	// Calls the outermost package and marks it dirty. 
+	// Calls the outermost package and marks it dirty.
 	bool MarkMetasoundDocumentDirty() const;
 
 protected:
@@ -152,9 +181,8 @@ protected:
 	// Returns the owning asset responsible for transactions applied to metasound
 	virtual const UObject* GetOwningAsset() const = 0;
 
-
 private:
-	bool bHasRegistered = false;
+	Metasound::Frontend::FNodeRegistryKey RegistryKey;
 
 	bool GetReceiveNodeMetadataForDataType(const FName& InTypeName, FMetasoundFrontendClassMetadata& OutMetadata) const;
 	TArray<FString> GetTransmittableInputVertexNames() const;
diff --git a/Engine/Plugins/Runtime/Metasound/Source/MetasoundFrontend/Public/MetasoundFrontendController.h b/Engine/Plugins/Runtime/Metasound/Source/MetasoundFrontend/Public/MetasoundFrontendController.h
--- a/Engine/Plugins/Runtime/Metasound/Source/MetasoundFrontend/Public/MetasoundFrontendController.h
+++ b/Engine/Plugins/Runtime/Metasound/Source/MetasoundFrontend/Public/MetasoundFrontendController.h
@@ -335,6 +335,19 @@ namespace Metasound
 			/** Returns true if node is required to satisfy the document archetype. */
 			virtual bool IsRequired() const = 0;
 
+			/** Returns the highest minor version number available from the class registry that shares this node's name & major version. */
+			virtual FMetasoundFrontendVersionNumber FindHighestVersionInRegistry() const = 0;
+
+			/** Returns the highest version number available from the class registry that shares this node's name. */
+			virtual FMetasoundFrontendVersionNumber FindHighestMinorVersionInRegistry() const = 0;
+
+			/**
+			  * Replaces this node with a new node of the provided version number, and attempts to 
+			  * rebuild edges where possible with matching vertex names that share the same DataType.
+			  * Returns a node handle to the new node.  If operation fails, returns a handle to this node.
+			  */
+			virtual FNodeHandle ReplaceWithVersion(const FMetasoundFrontendVersionNumber& InNewVersion) = 0;
+
 			/** Returns an input with the given id.
 			 *
 			 * If the input does not exist, an invalid handle is returned.
@@ -384,6 +397,11 @@ namespace Metasound
 			virtual const FMetasoundFrontendInterfaceStyle& GetInputStyle() const = 0;
 			virtual const FMetasoundFrontendClassStyle& GetClassStyle() const = 0;
 
+			/** Returns whether the node is eligible for auto-updating (i.e.
+			  * has undergone minor revision or the interface has changed, but
+			  * no higher major revision is available.) */
+			virtual bool CanAutoUpdate() const = 0;
+
 			/** Description of the given node. */
 			virtual const FText& GetDescription() const = 0;
 
@@ -476,6 +494,9 @@ namespace Metasound
 			/** Returns all input nodes in the graph. */
 			virtual TArray<FConstNodeHandle> GetConstInputNodes() const = 0;
 
+			/** Clears the graph, its associated interface, and synchronizes removed dependencies with the owning graph. */
+			virtual void ClearGraph() = 0;
+
 			/** Iterate over all input nodes with the given function. If ClassType is specified, only iterate over given type. */
 			virtual void IterateNodes(TUniqueFunction<void(FNodeHandle)> InFunction, EMetasoundFrontendClassType InClassType = EMetasoundFrontendClassType::Invalid) = 0;
 
@@ -603,6 +624,13 @@ namespace Metasound
 			/** Set the description for the output with the given name. */
 			virtual void SetOutputDescription(const FString& InName, const FText& InDescription) = 0;
 
+			/**
+			  * Updates the ChangeID for the class interface, which signals AutoUpdate to
+			  * attempt to patch class references at runtime even if the graph class has not
+			  * been versioned. TODO: Remove this once underlying architecture no longer requires it.
+			  */
+			virtual void UpdateInterfaceChangeID() = 0;
+
 			/** Clear the current literal for a given input.
 			 *
 			 * @return True on success, false on failure.
@@ -666,6 +694,18 @@ namespace Metasound
 
 			/** Returns a handle to the document owning this graph. */
 			virtual FConstDocumentHandle GetOwningDocument() const = 0;
+
+			/** Creates an input node */
+			virtual FNodeHandle AddInputNode(const FMetasoundFrontendClassInput& InClassInput, const FMetasoundFrontendLiteral* InDefaultValue, const FText* InDisplayName) = 0;
+
+			/** Creates an input node */
+			virtual FNodeHandle AddInputNode(const FName InTypeName, const FText& InToolTip, const FMetasoundFrontendLiteral* InDefaultValue, const FText* InDisplayName) = 0;
+
+			/** Creates an output node */
+			virtual FNodeHandle AddOutputNode(const FMetasoundFrontendClassOutput& InClassOutput, const FText* InDisplayName) = 0;
+
+			/** Creates an output node */
+			virtual FNodeHandle AddOutputNode(const FName InTypeName, const FText& InToolTip, const FText* InDisplayName) = 0;
 		};
 
 		/* An IDocumentController provides methods for querying and manipulating a Metasound document. */
@@ -693,9 +733,8 @@ namespace Metasound
 			// TODO: perhaps functions returning access pointers could be removed from main interface and only exist in FDocumentController.
 			
 			/** Returns an array of all class dependencies for this document. */
-			virtual TArray<FMetasoundFrontendClass> GetDependencies() const = 0;
-			virtual TArray<FMetasoundFrontendGraphClass> GetSubgraphs() const = 0;
-			virtual TArray<FMetasoundFrontendClass> GetClasses() const = 0;
+			virtual const TArray<FMetasoundFrontendClass>& GetDependencies() const = 0;
+			virtual const TArray<FMetasoundFrontendGraphClass>& GetSubgraphs() const = 0;
 			virtual const FMetasoundFrontendGraphClass& GetRootGraphClass() const = 0;
 
 			virtual FConstClassAccessPtr FindDependencyWithID(FGuid InClassID) const = 0;
@@ -749,8 +788,14 @@ namespace Metasound
 			virtual const FMetasoundFrontendVersion& GetArchetypeVersion() const = 0;
 			virtual void SetArchetypeVersion(const FMetasoundFrontendVersion& InVersion) = 0;
 
+			/** Updates dependency with given registry key if found with that in the active registry
+			  *
+			  * @return True if found (and subsequently updated), false if not found.
+			  */
+			virtual const FMetasoundFrontendClass* SynchronizeDependency(const FNodeRegistryKey& InKey) = 0;
+
 			/** Removes all dependencies which are no longer referenced by any graphs within this document
-			  * and updates dependency Metadata where necessary with that found in the registry.  */
+			  */
 			virtual void SynchronizeDependencies() = 0;
 
 			/** Returns an array of all subgraphs for this document. */
diff --git a/Engine/Plugins/Runtime/Metasound/Source/MetasoundFrontend/Public/MetasoundFrontendDocument.h b/Engine/Plugins/Runtime/Metasound/Source/MetasoundFrontend/Public/MetasoundFrontendDocument.h
--- a/Engine/Plugins/Runtime/Metasound/Source/MetasoundFrontend/Public/MetasoundFrontendDocument.h
+++ b/Engine/Plugins/Runtime/Metasound/Source/MetasoundFrontend/Public/MetasoundFrontendDocument.h
@@ -20,6 +20,18 @@ namespace Metasound
 	struct FLiteral;
 
 	extern const FGuid METASOUNDFRONTEND_API FrontendInvalidID;
+
+	namespace Frontend
+	{
+		namespace DisplayStyle
+		{
+			namespace NodeLayout
+			{
+				extern const FVector2D METASOUNDFRONTEND_API DefaultOffsetX;
+				extern const FVector2D METASOUNDFRONTEND_API DefaultOffsetY;
+			} // namespace NodeLayout
+		} // namespace DisplayStyle
+	} // namespace Frontend
 } // namespace Metasound
 
 
@@ -396,6 +408,7 @@ struct FMetasoundFrontendGraphStyle
 {
 	GENERATED_BODY()
 
+	// Whether or not the graph is editable by a user
 	UPROPERTY()
 	bool bIsGraphEditable = true;
 
@@ -651,10 +664,11 @@ struct FMetasoundFrontendInterfaceStyle
 
 
 USTRUCT()
-struct FMetasoundFrontendClassInterface
+struct METASOUNDFRONTEND_API FMetasoundFrontendClassInterface
 {
 	GENERATED_BODY()
 
+private:
 	// Style info for inputs.
 	UPROPERTY()
 	FMetasoundFrontendInterfaceStyle InputStyle;
@@ -663,17 +677,59 @@ struct FMetasoundFrontendClassInterface
 	UPROPERTY()
 	FMetasoundFrontendInterfaceStyle OutputStyle;
 
+public:
 	// Description of class inputs.
-	UPROPERTY(EditAnywhere, Category = CustomView)
+	UPROPERTY(VisibleAnywhere, Category = CustomView)
 	TArray<FMetasoundFrontendClassInput> Inputs;
 
 	// Description of class outputs.
-	UPROPERTY(EditAnywhere, Category = CustomView)
+	UPROPERTY(VisibleAnywhere, Category = CustomView)
 	TArray<FMetasoundFrontendClassOutput> Outputs;
 
 	// Description of class environment variables.
-	UPROPERTY(EditAnywhere, Category = CustomView)
+	UPROPERTY(VisibleAnywhere, Category = CustomView)
 	TArray<FMetasoundFrontendClassEnvironmentVariable> Environment;
+
+private:
+	UPROPERTY()
+	FGuid ChangeID;
+
+public:
+	const FMetasoundFrontendInterfaceStyle& GetInputStyle() const
+	{
+		return InputStyle;
+	}
+
+	void SetInputStyle(const FMetasoundFrontendInterfaceStyle& InInputStyle)
+	{
+		InputStyle = InInputStyle;
+		ChangeID = FGuid::NewGuid();
+	}
+
+	const FMetasoundFrontendInterfaceStyle& GetOutputStyle() const
+	{
+		return OutputStyle;
+	}
+
+	void SetOutputStyle(const FMetasoundFrontendInterfaceStyle& InOutputStyle)
+	{
+		OutputStyle = InOutputStyle;
+		ChangeID = FGuid::NewGuid();
+	}
+
+	const FGuid& GetChangeID() const
+	{
+		return ChangeID;
+	}
+
+	// TODO: This is unfortunately required to be manually managed and executed anytime the input/output/environment arrays
+	// are mutated due to the design of the controller system obscuring away read/write permissions
+	// when querying.  Need to add accessors and refactor so that this isn't as error prone and
+	// remove manual execution at the call sites when mutating aforementioned UPROPERTIES.
+	void UpdateChangeID()
+	{
+		ChangeID = FGuid::NewGuid();
+	}
 };
 
 
@@ -717,6 +773,8 @@ struct METASOUNDFRONTEND_API FMetasoundFrontendClassName
 	FString ToString() const;
 
 	METASOUNDFRONTEND_API friend bool operator==(const FMetasoundFrontendClassName& InLHS, const FMetasoundFrontendClassName& InRHS);
+
+	METASOUNDFRONTEND_API friend bool operator!=(const FMetasoundFrontendClassName& InLHS, const FMetasoundFrontendClassName& InRHS);
 };
 
 
@@ -726,35 +784,185 @@ struct METASOUNDFRONTEND_API FMetasoundFrontendClassMetadata
 	GENERATED_BODY()
 
 	FMetasoundFrontendClassMetadata() = default;
-
 	FMetasoundFrontendClassMetadata(const Metasound::FNodeClassMetadata& InNodeClassMetadata);
 
-	UPROPERTY(EditAnywhere, Category = Metasound)
-	FMetasoundFrontendClassName ClassName;
+	private:
+		UPROPERTY(VisibleAnywhere, Category = Metasound)
+		FMetasoundFrontendClassName ClassName;
 
-	UPROPERTY(EditAnywhere, Category = Metasound)
-	FMetasoundFrontendVersionNumber Version;
+		UPROPERTY(VisibleAnywhere, Category = Metasound)
+		FMetasoundFrontendVersionNumber Version;
 
-	UPROPERTY(VisibleAnywhere, Category = Metasound)
-	EMetasoundFrontendClassType Type = EMetasoundFrontendClassType::Invalid;
+		UPROPERTY(VisibleAnywhere, Category = Metasound)
+		EMetasoundFrontendClassType Type = EMetasoundFrontendClassType::Invalid;
 
-	UPROPERTY(EditAnywhere, Category = Metasound)
-	FText DisplayName;
+		UPROPERTY(VisibleAnywhere, Category = Metasound)
+		FText DisplayName;
 
-	UPROPERTY(EditAnywhere, Category = Metasound)
-	FText Description;
+		UPROPERTY(VisibleAnywhere, Category = Metasound)
+		FText Description;
+
+		UPROPERTY(VisibleAnywhere, Category = Metasound)
+		FText PromptIfMissing;
+
+		UPROPERTY(VisibleAnywhere, Category = Metasound)
+		FText Author;
+
+		UPROPERTY(VisibleAnywhere, Category = Metasound)
+		TArray<FName> Keywords;
+
+		UPROPERTY(VisibleAnywhere, Category = Metasound)
+		TArray<FText> CategoryHierarchy;
+
+		// If true, auto-update will manage (add and remove)
+		// inputs/outputs associated with internally connected
+		// nodes when the interface of the given node is auto-updated.
+		UPROPERTY()
+		bool bAutoUpdateManagesInterface = false;
+
+		// ID used to identify if any of the above have been modified,
+		// to determine if the parent class should be auto-updated.
+		UPROPERTY()
+		FGuid ChangeID;
+
+	public:
+		static FName GetAuthorPropertyName()
+		{
+			return GET_MEMBER_NAME_CHECKED(FMetasoundFrontendClassMetadata, Author);
+		}
+
+		static FName GetClassNamePropertyName()
+		{
+			return GET_MEMBER_NAME_CHECKED(FMetasoundFrontendClassMetadata, ClassName);
+		}
+
+		static FName GetDescriptionPropertyName()
+		{
+			return GET_MEMBER_NAME_CHECKED(FMetasoundFrontendClassMetadata, Description);
+		}
+
+		static FName GetVersionPropertyName()
+		{
+			return GET_MEMBER_NAME_CHECKED(FMetasoundFrontendClassMetadata, Version);
+		}
+
+		const FMetasoundFrontendClassName& GetClassName() const
+		{
+			return ClassName;
+		}
+
+		bool GetAutoUpdateManagesInterface() const
+		{
+			return bAutoUpdateManagesInterface;
+		}
+
+		void SetClassName(const FMetasoundFrontendClassName& InClassName)
+		{
+			ClassName = InClassName;
+			ChangeID = FGuid::NewGuid();
+		}
+
+		EMetasoundFrontendClassType GetType() const
+		{
+			return Type;
+		}
+
+		void SetAutoUpdateManagesInterface(bool bInAutoUpdateManagesInterface)
+		{
+			bAutoUpdateManagesInterface = bInAutoUpdateManagesInterface;
+			ChangeID = FGuid::NewGuid();
+		}
+
+		void SetType(const EMetasoundFrontendClassType InType)
+		{
+			Type = InType;
+			// TODO: Type is modified while querying and swapped between
+			// to be external, so don't modify the ChangeID in this case.
+			// External/Internal should probably be a separate field.
+			// ChangeID = FGuid::NewGuid();
+		}
+
+		const FMetasoundFrontendVersionNumber& GetVersion() const
+		{
+			return Version;
+		}
+
+		void SetVersion(const FMetasoundFrontendVersionNumber& InVersion)
+		{
+			Version = InVersion;
+			ChangeID = FGuid::NewGuid();
+		}
 
-	UPROPERTY(EditAnywhere, Category = Metasound)
-	FText PromptIfMissing;
+		const FText& GetDisplayName() const
+		{
+			return DisplayName;
+		}
 
-	UPROPERTY(EditAnywhere, Category = Metasound)
-	FText Author;
+		void SetDisplayName(const FText& InDisplayName)
+		{
+			DisplayName = InDisplayName;
+			ChangeID = FGuid::NewGuid();
+		}
 
-	UPROPERTY(EditAnywhere, Category = Metasound)
-	TArray<FName> Keywords;
+		const FText& GetDescription() const
+		{
+			return Description;
+		}
 
-	UPROPERTY(EditAnywhere, Category = Metasound)
-	TArray<FText> CategoryHierarchy;
+		void SetDescription(const FText& InDescription)
+		{
+			Description = InDescription;
+			ChangeID = FGuid::NewGuid();
+		}
+
+		const FText& GetPromptIfMissing() const
+		{
+			return PromptIfMissing;
+		}
+
+		void SetPromptIfMissing(const FText& InPromptIfMissing)
+		{
+			PromptIfMissing = InPromptIfMissing;
+			ChangeID = FGuid::NewGuid();
+		}
+
+		const FText& GetAuthor() const
+		{
+			return Author;
+		}
+
+		void SetAuthor(const FText& InAuthor)
+		{
+			Author = InAuthor;
+			ChangeID = FGuid::NewGuid();
+		}
+
+		const TArray<FName>& GetKeywords() const
+		{
+			return Keywords;
+		}
+
+		void SetKeywords(const TArray<FName>& InKeywords)
+		{
+			Keywords = InKeywords;
+			ChangeID = FGuid::NewGuid();
+		}
+
+		const TArray<FText>& GetCategoryHierarchy() const
+		{
+			return CategoryHierarchy;
+		}
+
+		void SetCategoryHierarchy(const TArray<FText>& InCategoryHierarchy)
+		{
+			CategoryHierarchy = InCategoryHierarchy;
+			ChangeID = FGuid::NewGuid();
+		}
+
+		const FGuid& GetChangeID() const
+		{
+			return ChangeID;
+		}
 };
 
 USTRUCT()
@@ -764,6 +972,9 @@ struct FMetasoundFrontendClassStyle
 
 	UPROPERTY()
 	FMetasoundFrontendClassStyleDisplay Display;
+
+	UPROPERTY()
+	bool bIsUserGenerated = false;
 };
 
 USTRUCT()
@@ -812,7 +1023,6 @@ struct METASOUNDFRONTEND_API FMetasoundFrontendGraphClass : public FMetasoundFro
 
 	UPROPERTY()
 	FMetasoundFrontendGraph Graph;
-
 };
 
 USTRUCT()
@@ -830,6 +1040,9 @@ struct FMetasoundFrontendArchetypeInterface
 {
 	GENERATED_BODY()
 
+	UPROPERTY()
+	FGuid TransactionID;
+
 	UPROPERTY()
 	TArray<FMetasoundFrontendClassVertex> Inputs;
 
@@ -884,5 +1097,3 @@ struct METASOUNDFRONTEND_API FMetasoundFrontendDocument
 	UPROPERTY()
 	TArray<FMetasoundFrontendClass> Dependencies;
 };
-
-
diff --git a/Engine/Plugins/Runtime/Metasound/Source/MetasoundFrontend/Public/MetasoundFrontendTransform.h b/Engine/Plugins/Runtime/Metasound/Source/MetasoundFrontend/Public/MetasoundFrontendTransform.h
--- a/Engine/Plugins/Runtime/Metasound/Source/MetasoundFrontend/Public/MetasoundFrontendTransform.h
+++ b/Engine/Plugins/Runtime/Metasound/Source/MetasoundFrontend/Public/MetasoundFrontendTransform.h
@@ -1,6 +1,7 @@
 // Copyright Epic Games, Inc. All Rights Reserved.
 #pragma once
 
+#include "MetasoundAssetBase.h"
 #include "MetasoundFrontendArchetypeRegistry.h"
 #include "MetasoundFrontendController.h"
 #include "MetasoundFrontendDocument.h"
@@ -66,6 +67,50 @@ namespace Metasound
 			FMetasoundFrontendVersion ArchetypeVersion;
 		};
 
+		class METASOUNDFRONTEND_API FRebuildPresetRootGraph : public IDocumentTransform
+		{
+		public:
+			FRebuildPresetRootGraph(FConstDocumentHandle InReferencedDocument)
+				: ReferencedDocument(InReferencedDocument)
+			{
+			}
+
+			bool Transform(FDocumentHandle InDocument) const override;
+
+		private:
+			FConstDocumentHandle ReferencedDocument;
+		};
+
+		class METASOUNDFRONTEND_API FAutoUpdateRootGraph : public IDocumentTransform
+		{
+		public:
+			FAutoUpdateRootGraph(FMetasoundAssetBase& InAsset, const IMetaSoundAssetInterface& InAssetInterface)
+				: Asset(&InAsset)
+				, AssetInterface(&InAssetInterface)
+			{
+			}
+
+			bool Transform(FDocumentHandle InDocument) const override;
+
+		private:
+			FMetasoundAssetBase* Asset = nullptr;
+			const IMetaSoundAssetInterface* AssetInterface = nullptr;
+		};
+
+		class METASOUNDFRONTEND_API FSynchronizeAssetClassName : public IDocumentTransform
+		{
+		public:
+			FSynchronizeAssetClassName(FName InAssetName)
+				: AssetName(InAssetName)
+			{
+			}
+
+			bool Transform(FDocumentHandle InDocument) const override;
+
+		private:
+			FName AssetName;
+		};
+
 		/** Base class for versioning a document. */
 		class METASOUNDFRONTEND_API FVersionDocument : public IDocumentTransform
 		{
@@ -75,7 +120,7 @@ namespace Metasound
 		public:
 			static FMetasoundFrontendVersionNumber GetMaxVersion()
 			{
-				return FMetasoundFrontendVersionNumber { 1, 4 };
+				return FMetasoundFrontendVersionNumber { 1, 5 };
 			}
 
 			FVersionDocument(FName InName, const FString& InPath);
diff --git a/Engine/Plugins/Runtime/Metasound/Source/MetasoundFrontend/Public/MetasoundNodeRegistrationMacro.h b/Engine/Plugins/Runtime/Metasound/Source/MetasoundFrontend/Public/MetasoundNodeRegistrationMacro.h
--- a/Engine/Plugins/Runtime/Metasound/Source/MetasoundFrontend/Public/MetasoundNodeRegistrationMacro.h
+++ b/Engine/Plugins/Runtime/Metasound/Source/MetasoundFrontend/Public/MetasoundNodeRegistrationMacro.h
@@ -77,11 +77,16 @@ namespace Metasound
 				return FrontendClass;
 			}
 
-			virtual TUniquePtr<INodeRegistryEntry> Clone() const 
+			virtual TUniquePtr<INodeRegistryEntry> Clone() const override
 			{
 				return MakeUnique<FNodeRegistryEntry>(*this);
 			}
 
+			virtual bool IsNative() const override
+			{
+				return true;
+			}
+
 		private:
 
 			Frontend::FNodeClassInfo ClassInfo;
