Attempt to unblock editor start-up while investigating duplicate registry issue in MetaSounds specific to branch's primary project

diff --git a/Engine/Plugins/Runtime/Metasound/Source/MetasoundFrontend/Private/MetasoundFrontendRegistries.cpp b/Engine/Plugins/Runtime/Metasound/Source/MetasoundFrontend/Private/MetasoundFrontendRegistries.cpp
--- a/Engine/Plugins/Runtime/Metasound/Source/MetasoundFrontend/Private/MetasoundFrontendRegistries.cpp
+++ b/Engine/Plugins/Runtime/Metasound/Source/MetasoundFrontend/Private/MetasoundFrontendRegistries.cpp
@@ -788,12 +788,12 @@ namespace Metasound
 					Key = NodeRegistryKey::CreateKey(InEntry->GetClassInfo());
 
 					// check to see if an identical node was already registered, and log
-					ensureAlwaysMsgf(
-						!RegisteredNodes.Contains(Key),
-						TEXT("Node with registry key already registered. "
-							"The previously registered node will be overwritten. "
-							"This can happen if two classes share the same name or if METASOUND_REGISTER_NODE is defined in a public header."),
-						*Key);
+// 					ensureAlwaysMsgf(
+// 						!RegisteredNodes.Contains(Key),
+// 						TEXT("Node with registry key already registered. "
+// 							"The previously registered node will be overwritten. "
+// 							"This can happen if two classes share the same name or if METASOUND_REGISTER_NODE is defined in a public header."),
+// 						*Key);
 
 					// Store update to newly registered node in history so nodes
 					// can be queried by transaction ID
