Stop-gap fix for duplicated MetaSound assets sharing registration node guid/class name.
Static Analysis warning fix

diff --git a/Engine/Plugins/Runtime/Metasound/Source/MetasoundEditor/Private/MetasoundEditor.cpp b/Engine/Plugins/Runtime/Metasound/Source/MetasoundEditor/Private/MetasoundEditor.cpp
--- a/Engine/Plugins/Runtime/Metasound/Source/MetasoundEditor/Private/MetasoundEditor.cpp
+++ b/Engine/Plugins/Runtime/Metasound/Source/MetasoundEditor/Private/MetasoundEditor.cpp
@@ -451,6 +451,16 @@ namespace Metasound
 			FSlateApplication::Get().DismissAllMenus();
 		}
 
+		void FEditor::NotifyNodePasteFailure_ReferenceLoop()
+		{
+			FNotificationInfo Info(LOCTEXT("NodePasteFailed_ReferenceLoop", "Node(s) not pasted: Nodes would create asset reference cycle."));
+			Info.bFireAndForget = true;
+			Info.bUseSuccessFailIcons = false;
+			Info.ExpireDuration = 5.0f;
+
+			MetasoundGraphEditor->AddNotification(Info, false /* bSuccess */);
+		}
+
 		void FEditor::NotifyUserModifiedBySync()
 		{
 			FNotificationInfo Info(LOCTEXT("SynchronizationRequired", "Operation modified pin(s), connection(s), and/or node(s).  Please refer to graph."));
@@ -1563,6 +1573,9 @@ namespace Metasound
 				Location = MetasoundGraphEditor->GetPasteLocation();
 			}
 
+			FMetasoundAssetBase* MetasoundAsset = IMetasoundUObjectRegistry::Get().GetObjectAsAssetBase(Metasound);
+			check(MetasoundAsset);
+
 			UMetasoundEditorGraph& Graph = GetMetaSoundGraphChecked();
 
 			const FScopedTransaction Transaction(InTransactionText);
@@ -1577,16 +1590,37 @@ namespace Metasound
 
 			NodeTextToPaste.Empty();
 
+			bool bNotifyReferenceLoop = false;
+
 			TArray<UEdGraphNode*> NodesToRemove;
 			for (UEdGraphNode* GraphNode : PastedGraphNodes)
 			{
 				GraphNode->CreateNewGuid();
 				if (UMetasoundEditorGraphExternalNode* ExternalNode = Cast<UMetasoundEditorGraphExternalNode>(GraphNode))
 				{
-					FNodeHandle NewHandle = FGraphBuilder::AddNodeHandle(*Metasound, *ExternalNode);
-					if (!NewHandle->IsValid())
+					FMetasoundFrontendClassMetadata LookupMetadata;
+					LookupMetadata.SetClassName(ExternalNode->GetClassName());
+					LookupMetadata.SetType(EMetasoundFrontendClassType::External);
+					const FNodeRegistryKey PastedRegistryKey = NodeRegistryKey::CreateKey(LookupMetadata);
+					if (const FSoftObjectPath* AssetPath = UMetaSoundAssetSubsystem::Get().FindObjectPathFromKey(PastedRegistryKey))
 					{
-						NodesToRemove.Add(GraphNode);
+						if (MetasoundAsset->AddingReferenceCausesLoop(*AssetPath, UMetaSoundAssetSubsystem::Get()))
+						{
+							FMetasoundFrontendClass MetaSoundClass;
+							FMetasoundFrontendRegistryContainer::Get()->FindFrontendClassFromRegistered(PastedRegistryKey, MetaSoundClass);
+							const FString FriendlyClassName = MetaSoundClass.Metadata.GetDisplayName().ToString();
+							UE_LOG(LogMetaSound, Warning, TEXT("Failed to paste node with class '%s'.  Class would introduce cyclic asset dependency."), *FriendlyClassName);
+							bNotifyReferenceLoop = true;
+							NodesToRemove.Add(GraphNode);
+						}
+					}
+					else
+					{
+						FNodeHandle NewHandle = FGraphBuilder::AddNodeHandle(*Metasound, *ExternalNode);
+						if (!NewHandle->IsValid())
+						{
+							NodesToRemove.Add(GraphNode);
+						}
 					}
 				}
 				else if (UMetasoundEditorGraphInputNode* InputNode = Cast<UMetasoundEditorGraphInputNode>(GraphNode))
@@ -1712,6 +1746,14 @@ namespace Metasound
 				NotifyUserModifiedBySync();
 			}
 
+			const bool bAutoUpdate = true;
+			FGraphBuilder::ValidateGraph(*Metasound, bAutoUpdate);
+
+			if (bNotifyReferenceLoop)
+			{
+				NotifyNodePasteFailure_ReferenceLoop();
+			}
+
 			MetasoundGraphEditor->NotifyGraphChanged();
 		}
 
diff --git a/Engine/Plugins/Runtime/Metasound/Source/MetasoundEditor/Private/MetasoundEditor.h b/Engine/Plugins/Runtime/Metasound/Source/MetasoundEditor/Private/MetasoundEditor.h
--- a/Engine/Plugins/Runtime/Metasound/Source/MetasoundEditor/Private/MetasoundEditor.h
+++ b/Engine/Plugins/Runtime/Metasound/Source/MetasoundEditor/Private/MetasoundEditor.h
@@ -266,6 +266,7 @@ namespace Metasound
 		private:
 			void EditObjectSettings();
 
+			void NotifyNodePasteFailure_ReferenceLoop();
 			void NotifyUserModifiedBySync();
 
 			bool IsPlaying() const;
diff --git a/Engine/Plugins/Runtime/Metasound/Source/MetasoundEditor/Private/MetasoundEditorGraphBuilder.cpp b/Engine/Plugins/Runtime/Metasound/Source/MetasoundEditor/Private/MetasoundEditorGraphBuilder.cpp
--- a/Engine/Plugins/Runtime/Metasound/Source/MetasoundEditor/Private/MetasoundEditorGraphBuilder.cpp
+++ b/Engine/Plugins/Runtime/Metasound/Source/MetasoundEditor/Private/MetasoundEditorGraphBuilder.cpp
@@ -1033,12 +1033,15 @@ namespace Metasound
 
 			MetaSoundAsset->RegisterGraphWithFrontend();
 
-			TArray<UObject*> EditedAssets = GEditor->GetEditorSubsystem<UAssetEditorSubsystem>()->GetAllEditedAssets();
-			for (UObject* Asset : EditedAssets)
+			if (GEditor)
 			{
-				if (Asset != &InMetaSound && IMetasoundUObjectRegistry::Get().IsRegisteredClass(Asset))
+				TArray<UObject*> EditedAssets = GEditor->GetEditorSubsystem<UAssetEditorSubsystem>()->GetAllEditedAssets();
+				for (UObject* Asset : EditedAssets)
 				{
-					ValidateGraph(*Asset, true /* bAutoUpdate */);
+					if (Asset != &InMetaSound && IMetasoundUObjectRegistry::Get().IsRegisteredClass(Asset))
+					{
+						ValidateGraph(*Asset, true /* bAutoUpdate */);
+					}
 				}
 			}
 		}
@@ -1050,17 +1053,20 @@ namespace Metasound
 
 			MetaSoundAsset->UnregisterGraphWithFrontend();
 
-			TArray<UObject*> EditedAssets = GEditor->GetEditorSubsystem<UAssetEditorSubsystem>()->GetAllEditedAssets();
-			for (UObject* Asset : EditedAssets)
+			if (GEditor)
 			{
-				if (Asset != &InMetaSound && IMetasoundUObjectRegistry::Get().IsRegisteredClass(Asset))
+				TArray<UObject*> EditedAssets = GEditor->GetEditorSubsystem<UAssetEditorSubsystem>()->GetAllEditedAssets();
+				for (UObject* Asset : EditedAssets)
 				{
-					ValidateGraph(*Asset, true /* bAutoUpdate */);
+					if (Asset != &InMetaSound && IMetasoundUObjectRegistry::Get().IsRegisteredClass(Asset))
+					{
+						ValidateGraph(*Asset, true /* bAutoUpdate */);
+					}
 				}
 			}
 		}
 
-		bool FGraphBuilder::IsMatchingInputHandleAndPin(const Frontend::FInputHandle& InInputHandle, const UEdGraphPin& InEditorPin)
+		bool FGraphBuilder::IsMatchingInputHandleAndPin(const Frontend::FConstInputHandle& InInputHandle, const UEdGraphPin& InEditorPin)
 		{
 			if (InEditorPin.Direction != EGPD_Input)
 			{
@@ -1076,7 +1082,7 @@ namespace Metasound
 			return false;
 		}
 
-		bool FGraphBuilder::IsMatchingOutputHandleAndPin(const Frontend::FOutputHandle& InOutputHandle, const UEdGraphPin& InEditorPin)
+		bool FGraphBuilder::IsMatchingOutputHandleAndPin(const Frontend::FConstOutputHandle& InOutputHandle, const UEdGraphPin& InEditorPin)
 		{
 			if (InEditorPin.Direction != EGPD_Output)
 			{
@@ -1114,7 +1120,7 @@ namespace Metasound
 			}
 		}
 
-		UEdGraphPin* FGraphBuilder::AddPinToNode(UMetasoundEditorGraphNode& InEditorNode, Frontend::FInputHandle InInputHandle)
+		UEdGraphPin* FGraphBuilder::AddPinToNode(UMetasoundEditorGraphNode& InEditorNode, Frontend::FConstInputHandle InInputHandle)
 		{
 			using namespace Frontend;
 
@@ -1125,15 +1131,13 @@ namespace Metasound
 			if (ensure(NewPin))
 			{
 				RefreshPinMetadata(*NewPin, InInputHandle->GetMetadata());
-
-				FNodeHandle NodeHandle = InInputHandle->GetOwningNode();
 				SynchronizePinLiteral(*NewPin);
 			}
 
 			return NewPin;
 		}
 
-		UEdGraphPin* FGraphBuilder::AddPinToNode(UMetasoundEditorGraphNode& InEditorNode, Frontend::FOutputHandle InOutputHandle)
+		UEdGraphPin* FGraphBuilder::AddPinToNode(UMetasoundEditorGraphNode& InEditorNode, Frontend::FConstOutputHandle InOutputHandle)
 		{
 			IMetasoundEditorModule& EditorModule = FModuleManager::GetModuleChecked<IMetasoundEditorModule>("MetaSoundEditor");
 			FEdGraphPinType	PinType = EditorModule.FindDataType(InOutputHandle->GetDataType()).PinType;
@@ -1421,12 +1425,12 @@ namespace Metasound
 			return bIsEditorGraphDirty;
 		}
 
-		bool FGraphBuilder::SynchronizeNodePins(UMetasoundEditorGraphNode& InEditorNode, Frontend::FNodeHandle InNode, bool bRemoveUnusedPins, bool bLogChanges)
+		bool FGraphBuilder::SynchronizeNodePins(UMetasoundEditorGraphNode& InEditorNode, Frontend::FConstNodeHandle InNode, bool bRemoveUnusedPins, bool bLogChanges)
 		{
 			bool bIsNodeDirty = false;
 
-			TArray<Frontend::FInputHandle> InputHandles = InNode->GetInputs();
-			TArray<Frontend::FOutputHandle> OutputHandles = InNode->GetOutputs();
+			TArray<Frontend::FConstInputHandle> InputHandles = InNode->GetConstInputs();
+			TArray<Frontend::FConstOutputHandle> OutputHandles = InNode->GetConstOutputs();
 			TArray<UEdGraphPin*> EditorPins = InEditorNode.Pins;
 
 			// Filter out pins which are not paired.
@@ -1434,12 +1438,12 @@ namespace Metasound
 			{
 				UEdGraphPin* Pin = EditorPins[i];
 
-				auto IsMatchingInputHandle = [&](const Frontend::FInputHandle& InputHandle) -> bool
+				auto IsMatchingInputHandle = [&](const Frontend::FConstInputHandle& InputHandle) -> bool
 				{
 					return IsMatchingInputHandleAndPin(InputHandle, *Pin);
 				};
 
-				auto IsMatchingOutputHandle = [&](const Frontend::FOutputHandle& OutputHandle) -> bool
+				auto IsMatchingOutputHandle = [&](const Frontend::FConstOutputHandle& OutputHandle) -> bool
 				{
 					return IsMatchingOutputHandleAndPin(OutputHandle, *Pin);
 				};
@@ -1495,7 +1499,7 @@ namespace Metasound
 				{
 					bIsNodeDirty = true;
 					InputHandles = InNode->GetInputStyle().SortDefaults(InputHandles);
-					for (Frontend::FInputHandle& InputHandle : InputHandles)
+					for (Frontend::FConstInputHandle& InputHandle : InputHandles)
 					{
 						if (bLogChanges)
 						{
@@ -1515,7 +1519,7 @@ namespace Metasound
 				{
 					bIsNodeDirty = true;
 					OutputHandles = InNode->GetOutputStyle().SortDefaults(OutputHandles);
-					for (Frontend::FOutputHandle& OutputHandle : OutputHandles)
+					for (Frontend::FConstOutputHandle& OutputHandle : OutputHandles)
 					{
 						if (bLogChanges)
 						{
diff --git a/Engine/Plugins/Runtime/Metasound/Source/MetasoundEditor/Private/MetasoundEditorGraphBuilder.h b/Engine/Plugins/Runtime/Metasound/Source/MetasoundEditor/Private/MetasoundEditorGraphBuilder.h
--- a/Engine/Plugins/Runtime/Metasound/Source/MetasoundEditor/Private/MetasoundEditorGraphBuilder.h
+++ b/Engine/Plugins/Runtime/Metasound/Source/MetasoundEditor/Private/MetasoundEditorGraphBuilder.h
@@ -155,10 +155,10 @@ namespace Metasound
 			static bool DeleteNode(UEdGraphNode& InNode);
 
 			// Adds an Input UEdGraphPin to a UMetasoundEditorGraphNode
-			static UEdGraphPin* AddPinToNode(UMetasoundEditorGraphNode& InEditorNode, Frontend::FInputHandle InInputHandle);
+			static UEdGraphPin* AddPinToNode(UMetasoundEditorGraphNode& InEditorNode, Frontend::FConstInputHandle InInputHandle);
 
 			// Adds an Output UEdGraphPin to a UMetasoundEditorGraphNode
-			static UEdGraphPin* AddPinToNode(UMetasoundEditorGraphNode& InEditorNode, Frontend::FOutputHandle InOutputHandle);
+			static UEdGraphPin* AddPinToNode(UMetasoundEditorGraphNode& InEditorNode, Frontend::FConstOutputHandle InOutputHandle);
 
 			// Refreshes pin state from class FrontendClassVertexMetadata
 			static void RefreshPinMetadata(UEdGraphPin& InPin, const FMetasoundFrontendVertexMetadata& InMetadata);
@@ -171,7 +171,7 @@ namespace Metasound
 			// Adds and removes pins so that the UMetasoundEditorGraphNode matches the InNode.
 			//
 			// @return True if the UMetasoundEditorGraphNode was altered. False otherwise.
-			static bool SynchronizeNodePins(UMetasoundEditorGraphNode& InEditorNode, Frontend::FNodeHandle InNode, bool bRemoveUnusedPins = true, bool bLogChanges = true);
+			static bool SynchronizeNodePins(UMetasoundEditorGraphNode& InEditorNode, Frontend::FConstNodeHandle InNode, bool bRemoveUnusedPins = true, bool bLogChanges = true);
 
 			// Adds and removes connections so that the UEdGraph of the MetaSound has the same
 			// connections as the FMetasoundFrontendDocument graph.
@@ -188,10 +188,10 @@ namespace Metasound
 			static bool SynchronizeVariables(UObject& InMetaSound);
 
 			// Returns true if the FInputHandle and UEdGraphPin match each other.
-			static bool IsMatchingInputHandleAndPin(const Frontend::FInputHandle& InInputHandle, const UEdGraphPin& InEditorPin);
+			static bool IsMatchingInputHandleAndPin(const Frontend::FConstInputHandle& InInputHandle, const UEdGraphPin& InEditorPin);
 
 			// Returns true if the FOutputHandle and UEdGraphPin match each other.
-			static bool IsMatchingOutputHandleAndPin(const Frontend::FOutputHandle& InOutputHandle, const UEdGraphPin& InEditorPin);
+			static bool IsMatchingOutputHandleAndPin(const Frontend::FConstOutputHandle& InOutputHandle, const UEdGraphPin& InEditorPin);
 
 			// Function signature for visiting a node doing depth first traversal.
 			//
diff --git a/Engine/Plugins/Runtime/Metasound/Source/MetasoundEditor/Private/MetasoundEditorGraphNode.cpp b/Engine/Plugins/Runtime/Metasound/Source/MetasoundEditor/Private/MetasoundEditorGraphNode.cpp
--- a/Engine/Plugins/Runtime/Metasound/Source/MetasoundEditor/Private/MetasoundEditorGraphNode.cpp
+++ b/Engine/Plugins/Runtime/Metasound/Source/MetasoundEditor/Private/MetasoundEditorGraphNode.cpp
@@ -167,7 +167,8 @@ void UMetasoundEditorGraphNode::AllocateDefaultPins()
 
 void UMetasoundEditorGraphNode::ReconstructNode()
 {
-	using namespace Metasound;
+	using namespace Metasound::Editor;
+	using namespace Metasound::Frontend;
 
 	// Don't remove unused pins here. Reconstruction can occur while duplicating or pasting nodes,
 	// and subsequent steps clean-up unused pins.  This can be called mid-copy, which means the node
@@ -175,7 +176,11 @@ void UMetasoundEditorGraphNode::ReconstructNode()
 	// are lost.
 	// TODO: User will want to see dead pins as well for node definition changes. Label and color-code dead
 	// pins (ex. red), and leave dead connections for reference like BP.
-	Editor::FGraphBuilder::SynchronizeNodePins(*this, GetNodeHandle(), false /* bRemoveUnusedPins */, false /* bLogChanges */);
+	FConstNodeHandle NodeHandle = GetNodeHandle();
+	if (NodeHandle->IsValid())
+	{
+		FGraphBuilder::SynchronizeNodePins(*this, NodeHandle, false /* bRemoveUnusedPins */, false /* bLogChanges */);
+	}
 }
 
 void UMetasoundEditorGraphNode::AutowireNewNode(UEdGraphPin* FromPin)
@@ -463,7 +468,7 @@ FMetasoundFrontendClassName UMetasoundEditorGraphOutputNode::GetClassName() cons
 
 FGuid UMetasoundEditorGraphOutputNode::GetNodeID() const
 {
-	if (ensure(Output))
+	if (Output)
 	{
 		return Output->NodeID;
 	}
@@ -599,11 +604,12 @@ UMetasoundEditorGraphExternalNode* UMetasoundEditorGraphExternalNode::UpdateToVe
 	UObject& MetaSound = GetMetasoundChecked();
 	UMetasoundEditorGraphExternalNode* ReplacementEdNode = FGraphBuilder::AddExternalNode(MetaSound, ReplacementNode, NodeLocation, false /* bInSelectNewNode */);
 
-	if (ensure(ReplacementEdNode))
+	if (!ensure(ReplacementEdNode))
 	{
-		GetGraph()->RemoveNode(this);
+		return this;
 	}
 
+	GetGraph()->RemoveNode(this);
 	MetaSound.Modify();
 
 	// TODO: Pare down to SynchronizeNode/NodeConnections call to avoid having to synchronize the entire graph for all version calls
diff --git a/Engine/Plugins/Runtime/Metasound/Source/MetasoundEditor/Private/MetasoundEditorModule.cpp b/Engine/Plugins/Runtime/Metasound/Source/MetasoundEditor/Private/MetasoundEditorModule.cpp
--- a/Engine/Plugins/Runtime/Metasound/Source/MetasoundEditor/Private/MetasoundEditorModule.cpp
+++ b/Engine/Plugins/Runtime/Metasound/Source/MetasoundEditor/Private/MetasoundEditorModule.cpp
@@ -26,6 +26,7 @@
 #include "MetasoundEditorSettings.h"
 #include "MetasoundFrontendDocument.h"
 #include "MetasoundFrontendRegistries.h"
+#include "MetasoundFrontendTransform.h"
 #include "MetasoundNodeDetailCustomization.h"
 #include "MetasoundSource.h"
 #include "MetasoundTime.h"
@@ -146,7 +147,32 @@ namespace Metasound
 
 		class FModule : public IMetasoundEditorModule
 		{
-			void AddOrUpdateClassRegistryAsset(const FAssetData& InAssetData)
+			void AddClassRegistryAsset(const FAssetData& InAssetData)
+			{
+				using namespace Metasound;
+				using namespace Metasound::Frontend;
+
+				if (IsMetaSoundAssetClass(InAssetData.AssetClass))
+				{
+					// Must reset guid prior to registration as asset could've been duplicated from another asset
+					UObject* AssetObject = InAssetData.GetAsset();
+					FMetasoundAssetBase* MetasoundAsset = IMetasoundUObjectRegistry::Get().GetObjectAsAssetBase(AssetObject);
+					if (ensure(MetasoundAsset))
+					{
+						FRegenerateAssetClassName().Transform(MetasoundAsset->GetDocumentHandle());
+					}
+
+					// Use the editor version of UnregisterWithFrontend so it refreshes any open MetaSound editors
+					UMetaSoundAssetSubsystem::Get().AddOrUpdateAsset(InAssetData, false /* bRegisterWithFrontend */);
+
+					if (ensure(AssetObject))
+					{
+						FGraphBuilder::RegisterGraphWithFrontend(*AssetObject);
+					}
+				}
+			}
+
+			void UpdateClassRegistryAsset(const FAssetData& InAssetData)
 			{
 				using namespace Metasound;
 				using namespace Metasound::Frontend;
@@ -208,9 +234,9 @@ namespace Metasound
 				Filter.ClassNames = MetaSoundClassNames;
 
 				FAssetRegistryModule& AssetRegistryModule = FModuleManager::LoadModuleChecked<FAssetRegistryModule>("AssetRegistry");
-				AssetRegistryModule.Get().EnumerateAssets(Filter, [this](const FAssetData& AssetData) { AddOrUpdateClassRegistryAsset(AssetData); return true; });
-				AssetRegistryModule.Get().OnAssetAdded().AddRaw(this, &FModule::AddOrUpdateClassRegistryAsset);
-				AssetRegistryModule.Get().OnAssetUpdated().AddRaw(this, &FModule::AddOrUpdateClassRegistryAsset);
+				AssetRegistryModule.Get().EnumerateAssets(Filter, [this](const FAssetData& AssetData) { UpdateClassRegistryAsset(AssetData); return true; });
+				AssetRegistryModule.Get().OnAssetAdded().AddRaw(this, &FModule::AddClassRegistryAsset);
+				AssetRegistryModule.Get().OnAssetUpdated().AddRaw(this, &FModule::UpdateClassRegistryAsset);
 				AssetRegistryModule.Get().OnAssetRemoved().AddRaw(this, &FModule::RemoveAssetFromClassRegistry);
 				AssetRegistryModule.Get().OnAssetRenamed().AddRaw(this, &FModule::RenameAssetInClassRegistry);
 
diff --git a/Engine/Plugins/Runtime/Metasound/Source/MetasoundEditor/Public/MetasoundEditorGraphNode.h b/Engine/Plugins/Runtime/Metasound/Source/MetasoundEditor/Public/MetasoundEditorGraphNode.h
--- a/Engine/Plugins/Runtime/Metasound/Source/MetasoundEditor/Public/MetasoundEditorGraphNode.h
+++ b/Engine/Plugins/Runtime/Metasound/Source/MetasoundEditor/Public/MetasoundEditorGraphNode.h
@@ -139,8 +139,8 @@ protected:
 	// Whether or not the referenced class is natively defined
 	// (false if defined in another asset). Cached from node
 	// implementation for fast access when validated.
-	UPROPERTY(Transient)
-	bool bIsClassNative = false;
+	UPROPERTY()
+	bool bIsClassNative = true;
 
 public:
 	virtual void PostLoad() override;
diff --git a/Engine/Plugins/Runtime/Metasound/Source/MetasoundEngine/Private/MetasoundUObjectRegistry.cpp b/Engine/Plugins/Runtime/Metasound/Source/MetasoundEngine/Private/MetasoundUObjectRegistry.cpp
--- a/Engine/Plugins/Runtime/Metasound/Source/MetasoundEngine/Private/MetasoundUObjectRegistry.cpp
+++ b/Engine/Plugins/Runtime/Metasound/Source/MetasoundEngine/Private/MetasoundUObjectRegistry.cpp
@@ -242,11 +242,11 @@ void UMetaSoundAssetSubsystem::RemoveAsset(const FAssetData& InAssetData, bool b
 void UMetaSoundAssetSubsystem::RenameAsset(const FAssetData& InAssetData, bool bInReregisterWithFrontend)
 {
 	RemoveAsset(InAssetData, bInReregisterWithFrontend);
-	SynchronizeAssetClassName(InAssetData);
+	SynchronizeAssetClassDisplayName(InAssetData);
 	AddOrUpdateAsset(InAssetData, bInReregisterWithFrontend);
 }
 
-void UMetaSoundAssetSubsystem::SynchronizeAssetClassName(const FAssetData& InAssetData)
+void UMetaSoundAssetSubsystem::SynchronizeAssetClassDisplayName(const FAssetData& InAssetData)
 {
 	UObject* Object = nullptr;
 	FSoftObjectPath Path(InAssetData.ObjectPath);
@@ -261,7 +261,7 @@ void UMetaSoundAssetSubsystem::SynchronizeAssetClassName(const FAssetData& InAss
 
 	FMetasoundAssetBase* MetaSoundAsset = Metasound::IMetasoundUObjectRegistry::Get().GetObjectAsAssetBase(Object);
 	check(MetaSoundAsset);
-	Metasound::Frontend::FSynchronizeAssetClassName(InAssetData.AssetName).Transform(MetaSoundAsset->GetDocumentHandle());
+	Metasound::Frontend::FSynchronizeAssetClassDisplayName(InAssetData.AssetName).Transform(MetaSoundAsset->GetDocumentHandle());
 }
 
 namespace Metasound
diff --git a/Engine/Plugins/Runtime/Metasound/Source/MetasoundEngine/Public/MetasoundUObjectRegistry.h b/Engine/Plugins/Runtime/Metasound/Source/MetasoundEngine/Public/MetasoundUObjectRegistry.h
--- a/Engine/Plugins/Runtime/Metasound/Source/MetasoundEngine/Public/MetasoundUObjectRegistry.h
+++ b/Engine/Plugins/Runtime/Metasound/Source/MetasoundEngine/Public/MetasoundUObjectRegistry.h
@@ -32,7 +32,7 @@ public:
 	void RemoveAsset(UObject& InObject, bool bInUnregisterWithFrontend = true);
 	void RemoveAsset(const FAssetData& InAssetData, bool bInUnregisterWithFrontend = true);
 	void RenameAsset(const FAssetData& InAssetData, bool bInReregisterWithFrontend = true);
-	void SynchronizeAssetClassName(const FAssetData& InAssetData);
+	void SynchronizeAssetClassDisplayName(const FAssetData& InAssetData);
 
 	virtual FMetasoundAssetBase* FindAssetFromKey(const Metasound::Frontend::FNodeRegistryKey& RegistryKey) const override;
 	virtual const FSoftObjectPath* FindObjectPathFromKey(const Metasound::Frontend::FNodeRegistryKey& RegistryKey) const override;
diff --git a/Engine/Plugins/Runtime/Metasound/Source/MetasoundFrontend/Private/MetasoundFrontendRegistries.cpp b/Engine/Plugins/Runtime/Metasound/Source/MetasoundFrontend/Private/MetasoundFrontendRegistries.cpp
--- a/Engine/Plugins/Runtime/Metasound/Source/MetasoundFrontend/Private/MetasoundFrontendRegistries.cpp
+++ b/Engine/Plugins/Runtime/Metasound/Source/MetasoundFrontend/Private/MetasoundFrontendRegistries.cpp
@@ -788,12 +788,12 @@ namespace Metasound
 					Key = NodeRegistryKey::CreateKey(InEntry->GetClassInfo());
 
 					// check to see if an identical node was already registered, and log
-// 					ensureAlwaysMsgf(
-// 						!RegisteredNodes.Contains(Key),
-// 						TEXT("Node with registry key already registered. "
-// 							"The previously registered node will be overwritten. "
-// 							"This can happen if two classes share the same name or if METASOUND_REGISTER_NODE is defined in a public header."),
-// 						*Key);
+					ensureAlwaysMsgf(
+						!RegisteredNodes.Contains(Key),
+						TEXT("Node with registry key '%s' already registered. "
+							"The previously registered node will be overwritten. "
+							"This can happen if two classes share the same name or if METASOUND_REGISTER_NODE is defined in a public header."),
+						*Key);
 
 					// Store update to newly registered node in history so nodes
 					// can be queried by transaction ID
diff --git a/Engine/Plugins/Runtime/Metasound/Source/MetasoundFrontend/Private/MetasoundFrontendTransform.cpp b/Engine/Plugins/Runtime/Metasound/Source/MetasoundFrontend/Private/MetasoundFrontendTransform.cpp
--- a/Engine/Plugins/Runtime/Metasound/Source/MetasoundFrontend/Private/MetasoundFrontendTransform.cpp
+++ b/Engine/Plugins/Runtime/Metasound/Source/MetasoundFrontend/Private/MetasoundFrontendTransform.cpp
@@ -623,7 +623,7 @@ namespace Metasound
 			return true;
 		}
 
-		bool FSynchronizeAssetClassName::Transform(FDocumentHandle InDocument) const
+		bool FSynchronizeAssetClassDisplayName::Transform(FDocumentHandle InDocument) const
 		{
 			const FMetasoundFrontendClassMetadata& Metadata = InDocument->GetRootGraphClass().Metadata;
 			const FText NewAssetName = FText::FromString(AssetName.ToString());
@@ -633,8 +633,19 @@ namespace Metasound
 				FMetasoundFrontendClassMetadata NewMetadata = Metadata;
 				NewMetadata.SetDisplayName(NewAssetName);
 				InDocument->GetRootGraph()->SetGraphMetadata(NewMetadata);
+				return true;
 			}
 
+			return false;
+		}
+
+		bool FRegenerateAssetClassName::Transform(FDocumentHandle InDocument) const
+		{
+			FMetasoundFrontendClassMetadata Metadata = InDocument->GetRootGraph()->GetGraphMetadata();
+			FMetasoundFrontendClassName NewName = Metadata.GetClassName();
+			NewName.Name = *FGuid::NewGuid().ToString();
+			Metadata.SetClassName(NewName);
+			InDocument->GetRootGraph()->SetGraphMetadata(Metadata);
 			return true;
 		}
 
@@ -660,7 +671,6 @@ namespace Metasound
 					NewMetadata.Version.Number = GetTargetVersion();
 					InDocument->SetMetadata(NewMetadata);
 
-					UE_LOG(LogMetaSound, Display, TEXT("Versioned MetaSound Document to %s"), *Metadata.Version.Number.ToString());
 					return true;
 				}
 		};
@@ -810,7 +820,6 @@ namespace Metasound
 					const TArray<FMetasoundFrontendClass>& Dependencies = InDocument->GetDependencies();
 					const TArray<FMetasoundFrontendGraphClass>& Subgraphs = InDocument->GetSubgraphs();
 
-
 					if (const FMetasoundFrontendArchetype* Arch = FindMostSimilarArchetypeSupportingEnvironment(RootGraph, Dependencies, Subgraphs, AllArchetypes))
 					{
 						UE_LOG(LogMetaSound, Display, TEXT("Assigned archetype [ArchetypeVersion:%s] to document [RootGraphClassName:%s]"),
@@ -843,13 +852,32 @@ namespace Metasound
 
 			void TransformInternal(FDocumentHandle InDocument) const override
 			{
-				FSynchronizeAssetClassName(AssetName).Transform(InDocument);
+				FSynchronizeAssetClassDisplayName(AssetName).Transform(InDocument);
 			}
 
 		private:
 			FName AssetName;
 		};
 
+		/** Versions document from 1.5 to 1.6. */
+		class FVersionDocument_1_6 : public FVersionDocumentTransform
+		{
+		public:
+			FVersionDocument_1_6()
+			{
+			}
+
+			FMetasoundFrontendVersionNumber GetTargetVersion() const override
+			{
+				return { 1, 6 };
+			}
+
+			void TransformInternal(FDocumentHandle InDocument) const override
+			{
+				FRegenerateAssetClassName().Transform(InDocument);
+			}
+		};
+
 		FVersionDocument::FVersionDocument(FName InName, const FString& InPath)
 			: Name(InName)
 			, Path(InPath)
@@ -865,12 +893,22 @@ namespace Metasound
 
 			bool bWasUpdated = false;
 
+			const FMetasoundFrontendVersionNumber InitVersionNumber = InDocument->GetMetadata().Version.Number;
+
 			// Add additional transforms here after defining them above, example below.
 			bWasUpdated |= FVersionDocument_1_1().Transform(InDocument);
 			bWasUpdated |= FVersionDocument_1_2(Name, Path).Transform(InDocument);
 			bWasUpdated |= FVersionDocument_1_3().Transform(InDocument);
 			bWasUpdated |= FVersionDocument_1_4().Transform(InDocument);
 			bWasUpdated |= FVersionDocument_1_5(Name).Transform(InDocument);
+			bWasUpdated |= FVersionDocument_1_6().Transform(InDocument);
+
+			if (bWasUpdated)
+			{
+				const FText& DisplayName = InDocument->GetRootGraph()->GetGraphMetadata().GetDisplayName();
+				const FMetasoundFrontendVersionNumber NewVersionNumber = InDocument->GetMetadata().Version.Number;
+				UE_LOG(LogMetaSound, Display, TEXT("MetaSound Graph '%s' Parent Document Versioned: '%s' --> '%s'"), *DisplayName.ToString(), *InitVersionNumber.ToString(), *NewVersionNumber.ToString());
+			}
 
 			return bWasUpdated;
 		}
diff --git a/Engine/Plugins/Runtime/Metasound/Source/MetasoundFrontend/Public/MetasoundFrontendTransform.h b/Engine/Plugins/Runtime/Metasound/Source/MetasoundFrontend/Public/MetasoundFrontendTransform.h
--- a/Engine/Plugins/Runtime/Metasound/Source/MetasoundFrontend/Public/MetasoundFrontendTransform.h
+++ b/Engine/Plugins/Runtime/Metasound/Source/MetasoundFrontend/Public/MetasoundFrontendTransform.h
@@ -97,10 +97,11 @@ namespace Metasound
 			const IMetaSoundAssetInterface* AssetInterface = nullptr;
 		};
 
-		class METASOUNDFRONTEND_API FSynchronizeAssetClassName : public IDocumentTransform
+		/** Synchronizes the document's root graph's display name with that of the asset. */
+		class METASOUNDFRONTEND_API FSynchronizeAssetClassDisplayName : public IDocumentTransform
 		{
 		public:
-			FSynchronizeAssetClassName(FName InAssetName)
+			FSynchronizeAssetClassDisplayName(FName InAssetName)
 				: AssetName(InAssetName)
 			{
 			}
@@ -111,6 +112,16 @@ namespace Metasound
 			FName AssetName;
 		};
 
+		/** Regenerates the class' name, effectively causing the class to get registered as
+		  * a new class (useful for when an asset is duplicated). */
+		class METASOUNDFRONTEND_API FRegenerateAssetClassName : public IDocumentTransform
+		{
+		public:
+			FRegenerateAssetClassName() = default;
+
+			bool Transform(FDocumentHandle InDocument) const override;
+		};
+
 		/** Base class for versioning a document. */
 		class METASOUNDFRONTEND_API FVersionDocument : public IDocumentTransform
 		{
@@ -120,7 +131,7 @@ namespace Metasound
 		public:
 			static FMetasoundFrontendVersionNumber GetMaxVersion()
 			{
-				return FMetasoundFrontendVersionNumber { 1, 5 };
+				return FMetasoundFrontendVersionNumber { 1, 6 };
 			}
 
 			FVersionDocument(FName InName, const FString& InPath);
