- Fix for asset class names/guids rebuilding for all MetaSounds on editor start-up
- Now, only refreshing asset class name/guid when duplicating object

diff --git a/Engine/Plugins/Runtime/Metasound/Source/MetasoundEditor/Private/MetasoundEditorModule.cpp b/Engine/Plugins/Runtime/Metasound/Source/MetasoundEditor/Private/MetasoundEditorModule.cpp
--- a/Engine/Plugins/Runtime/Metasound/Source/MetasoundEditor/Private/MetasoundEditorModule.cpp
+++ b/Engine/Plugins/Runtime/Metasound/Source/MetasoundEditor/Private/MetasoundEditorModule.cpp
@@ -147,40 +147,15 @@ namespace Metasound
 
 		class FModule : public IMetasoundEditorModule
 		{
-			void AddClassRegistryAsset(const FAssetData& InAssetData)
+			void AddOrUpdateClassRegistryAsset(const FAssetData& InAssetData)
 			{
 				using namespace Metasound;
 				using namespace Metasound::Frontend;
 
 				if (IsMetaSoundAssetClass(InAssetData.AssetClass))
 				{
-					// Must reset guid prior to registration as asset could've been duplicated from another asset
-					UObject* AssetObject = InAssetData.GetAsset();
-					FMetasoundAssetBase* MetasoundAsset = IMetasoundUObjectRegistry::Get().GetObjectAsAssetBase(AssetObject);
-					if (ensure(MetasoundAsset))
-					{
-						FRegenerateAssetClassName().Transform(MetasoundAsset->GetDocumentHandle());
-					}
-
 					// Use the editor version of UnregisterWithFrontend so it refreshes any open MetaSound editors
 					UMetaSoundAssetSubsystem::Get().AddOrUpdateAsset(InAssetData, false /* bRegisterWithFrontend */);
-
-					if (ensure(AssetObject))
-					{
-						FGraphBuilder::RegisterGraphWithFrontend(*AssetObject);
-					}
-				}
-			}
-
-			void UpdateClassRegistryAsset(const FAssetData& InAssetData)
-			{
-				using namespace Metasound;
-				using namespace Metasound::Frontend;
-
-				if (IsMetaSoundAssetClass(InAssetData.AssetClass))
-				{
-					// Use the editor version of UnregisterWithFrontend so it refreshes any open MetaSound editors
-					UMetaSoundAssetSubsystem::Get().AddOrUpdateAsset(InAssetData, false /* bUnregisterWithFrontend */);
 					if (UObject* AssetObject = InAssetData.GetAsset())
 					{
 						FGraphBuilder::RegisterGraphWithFrontend(*AssetObject);
@@ -228,18 +203,20 @@ namespace Metasound
 				}
 			}
 
-			void InitializeAssetClassRegistry()
+			void OnAssetScanFinished()
 			{
 				FARFilter Filter;
 				Filter.ClassNames = MetaSoundClassNames;
 
 				FAssetRegistryModule& AssetRegistryModule = FModuleManager::LoadModuleChecked<FAssetRegistryModule>("AssetRegistry");
-				AssetRegistryModule.Get().EnumerateAssets(Filter, [this](const FAssetData& AssetData) { UpdateClassRegistryAsset(AssetData); return true; });
-				AssetRegistryModule.Get().OnAssetAdded().AddRaw(this, &FModule::AddClassRegistryAsset);
-				AssetRegistryModule.Get().OnAssetUpdated().AddRaw(this, &FModule::UpdateClassRegistryAsset);
+				AssetRegistryModule.Get().EnumerateAssets(Filter, [this](const FAssetData& AssetData) { AddOrUpdateClassRegistryAsset(AssetData); return true; });
+				AssetRegistryModule.Get().OnAssetAdded().AddRaw(this, &FModule::AddOrUpdateClassRegistryAsset);
+				AssetRegistryModule.Get().OnAssetUpdated().AddRaw(this, &FModule::AddOrUpdateClassRegistryAsset);
 				AssetRegistryModule.Get().OnAssetRemoved().AddRaw(this, &FModule::RemoveAssetFromClassRegistry);
 				AssetRegistryModule.Get().OnAssetRenamed().AddRaw(this, &FModule::RenameAssetInClassRegistry);
 
+				AssetRegistryModule.Get().OnFilesLoaded().RemoveAll(this);
+
 				FCoreUObjectDelegates::OnPackageReloaded.AddRaw(this, &FModule::OnPackageReloaded);
 			}
 
@@ -408,6 +385,7 @@ namespace Metasound
 					AssetRegistryModule->Get().OnAssetUpdated().RemoveAll(this);
 					AssetRegistryModule->Get().OnAssetRemoved().RemoveAll(this);
 					AssetRegistryModule->Get().OnAssetRenamed().RemoveAll(this);
+					AssetRegistryModule->Get().OnFilesLoaded().RemoveAll(this);
 
 					FCoreUObjectDelegates::OnPackageReloaded.RemoveAll(this);
 				}
@@ -515,7 +493,9 @@ namespace Metasound
 
 				FAssetTypeActions_MetaSound::RegisterMenuActions();
 				FAssetTypeActions_MetaSoundSource::RegisterMenuActions();
-				InitializeAssetClassRegistry();
+
+				FAssetRegistryModule& AssetRegistryModule = FModuleManager::LoadModuleChecked<FAssetRegistryModule>("AssetRegistry");
+				AssetRegistryModule.Get().OnFilesLoaded().AddRaw(this, &FModule::OnAssetScanFinished);
 			}
 
 			virtual void ShutdownModule() override
diff --git a/Engine/Plugins/Runtime/Metasound/Source/MetasoundEngine/Private/Metasound.cpp b/Engine/Plugins/Runtime/Metasound/Source/MetasoundEngine/Private/Metasound.cpp
--- a/Engine/Plugins/Runtime/Metasound/Source/MetasoundEngine/Private/Metasound.cpp
+++ b/Engine/Plugins/Runtime/Metasound/Source/MetasoundEngine/Private/Metasound.cpp
@@ -37,6 +37,12 @@ UMetaSound::UMetaSound(const FObjectInitializer& ObjectInitializer)
 }
 
 #if WITH_EDITOR
+void UMetaSound::PostDuplicate(EDuplicateMode::Type InDuplicateMode)
+{
+	Super::PostDuplicate(InDuplicateMode);
+	Metasound::PostDuplicate(*this, InDuplicateMode);
+}
+
 void UMetaSound::PostEditUndo()
 {
 	Super::PostEditUndo();
diff --git a/Engine/Plugins/Runtime/Metasound/Source/MetasoundEngine/Private/MetasoundSource.cpp b/Engine/Plugins/Runtime/Metasound/Source/MetasoundEngine/Private/MetasoundSource.cpp
--- a/Engine/Plugins/Runtime/Metasound/Source/MetasoundEngine/Private/MetasoundSource.cpp
+++ b/Engine/Plugins/Runtime/Metasound/Source/MetasoundEngine/Private/MetasoundSource.cpp
@@ -59,6 +59,12 @@ void UMetaSoundSource::PostEditUndo()
 	Metasound::PostAssetUndo(*this);
 }
 
+void UMetaSoundSource::PostDuplicate(EDuplicateMode::Type InDuplicateMode)
+{
+	Super::PostDuplicate(InDuplicateMode);
+	Metasound::PostDuplicate(*this, InDuplicateMode);
+}
+
 void UMetaSoundSource::PostEditChangeProperty(FPropertyChangedEvent& InEvent)
 {
 	using namespace Metasound;
diff --git a/Engine/Plugins/Runtime/Metasound/Source/MetasoundEngine/Public/Metasound.h b/Engine/Plugins/Runtime/Metasound/Source/MetasoundEngine/Public/Metasound.h
--- a/Engine/Plugins/Runtime/Metasound/Source/MetasoundEngine/Public/Metasound.h
+++ b/Engine/Plugins/Runtime/Metasound/Source/MetasoundEngine/Public/Metasound.h
@@ -7,6 +7,7 @@
 #include "MetasoundAssetBase.h"
 #include "MetasoundFrontend.h"
 #include "MetasoundFrontendDocument.h"
+#include "MetasoundFrontendTransform.h"
 #include "MetasoundOperatorSettings.h"
 #include "MetasoundRouter.h"
 #include "MetasoundUObjectRegistry.h"
@@ -58,6 +59,21 @@ namespace Metasound
 			}
 		}
 	}
+
+	template <typename TMetaSoundObject>
+	void PostDuplicate(TMetaSoundObject& InMetaSound, EDuplicateMode::Type InDuplicateMode)
+	{
+		if (InDuplicateMode == EDuplicateMode::Normal)
+		{
+			// Guid is reset as asset may share implementation from
+			// asset duplicated from but should not be registered as such.
+			FMetasoundAssetBase* MetasoundAsset = IMetasoundUObjectRegistry::Get().GetObjectAsAssetBase(&InMetaSound);
+			if (ensure(MetasoundAsset))
+			{
+				Metasound::Frontend::FRegenerateAssetClassName().Transform(MetasoundAsset->GetDocumentHandle());
+			}
+		}
+	}
 #endif // WITH_EDITOR
 
 	template <typename TMetaSoundObject>
@@ -203,6 +219,7 @@ public:
 #endif // #if WITH_EDITORONLY_DATA
 
 #if WITH_EDITOR
+	virtual void PostDuplicate(EDuplicateMode::Type DuplicateMode) override;
 	virtual void PostEditUndo() override;
 	virtual void PostEditChangeProperty(FPropertyChangedEvent& InEvent) override;
 #endif // WITH_EDITOR
diff --git a/Engine/Plugins/Runtime/Metasound/Source/MetasoundEngine/Public/MetasoundSource.h b/Engine/Plugins/Runtime/Metasound/Source/MetasoundEngine/Public/MetasoundSource.h
--- a/Engine/Plugins/Runtime/Metasound/Source/MetasoundEngine/Public/MetasoundSource.h
+++ b/Engine/Plugins/Runtime/Metasound/Source/MetasoundEngine/Public/MetasoundSource.h
@@ -120,6 +120,8 @@ public:
 		return false;
 	}
 
+	virtual void PostDuplicate(EDuplicateMode::Type InDuplicateMode) override;
+
 	virtual void PostEditChangeProperty(FPropertyChangedEvent& InEvent) override;
 #endif // WITH_EDITOR
 
