Remove dependency on metasound search engine for injecting BP API
- Added MetasoundFrontendInjectReceiveNodes.h/cpp which removes dependency on FMetasoundAssetBase::CopyDocumentAndInjectReceiveNodes.
- Significant speed improvments on MetaSound "Play Sound"

diff --git a/Engine/Plugins/Runtime/Metasound/Source/MetasoundEngine/Private/Metasound.cpp b/Engine/Plugins/Runtime/Metasound/Source/MetasoundEngine/Private/Metasound.cpp
--- a/Engine/Plugins/Runtime/Metasound/Source/MetasoundEngine/Private/Metasound.cpp
+++ b/Engine/Plugins/Runtime/Metasound/Source/MetasoundEngine/Private/Metasound.cpp
@@ -120,24 +120,6 @@ const FMetasoundFrontendVersion& UMetaSound::GetDefaultArchetypeVersion() const
 	return DefaultArchetypeVersion;
 }
 
-Metasound::FOperatorSettings UMetaSound::GetOperatorSettings(Metasound::FSampleRate InSampleRate) const
-{
-	const float BlockRate = FMath::Clamp(Metasound::ConsoleVariables::BlockRate, 1.0f, 1000.0f);
-	return Metasound::FOperatorSettings(InSampleRate, BlockRate);
-}
-
-Metasound::FSendAddress UMetaSound::CreateSendAddress(uint64 InInstanceID, const FString& InVertexName, const FName& InDataTypeName) const
-{
-	using namespace Metasound;
-
-	FSendAddress Address;
-
-	Address.Subsystem = GetSubsystemNameForSendScope(ETransmissionScope::Global);
-	Address.ChannelName = FName(FString::Printf(TEXT("%d:%s:%s"), InInstanceID, *InVertexName, *InDataTypeName.ToString()));
-
-	return Address;
-}
-
 const TArray<FMetasoundFrontendVersion>& UMetaSound::GetSupportedArchetypeVersions() const
 {
 	static const TArray<FMetasoundFrontendVersion> Supported
@@ -148,10 +130,4 @@ const TArray<FMetasoundFrontendVersion>& UMetaSound::GetSupportedArchetypeVersio
 	return Supported;
 }
 
-const FString& UMetaSound::GetAudioDeviceHandleVariableName()
-{
-	static const FString AudioDeviceHandleVarName = TEXT("AudioDeviceHandle");
-	return AudioDeviceHandleVarName;
-}
-
 #undef LOCTEXT_NAMESPACE // MetaSound
diff --git a/Engine/Plugins/Runtime/Metasound/Source/MetasoundEngine/Private/MetasoundEngineArchetypes.cpp b/Engine/Plugins/Runtime/Metasound/Source/MetasoundEngine/Private/MetasoundEngineArchetypes.cpp
--- a/Engine/Plugins/Runtime/Metasound/Source/MetasoundEngine/Private/MetasoundEngineArchetypes.cpp
+++ b/Engine/Plugins/Runtime/Metasound/Source/MetasoundEngine/Private/MetasoundEngineArchetypes.cpp
@@ -296,6 +296,12 @@ namespace Metasound
 				return SoundIsPreviewSoundName;
 			}
 
+			const FString& GetInstanceIDName()
+			{
+				static const FString& Name = FMetasoundInstanceTransmitter::GetInstanceIDEnvironmentVariableName();
+				return Name;
+			}
+
 			FMetasoundFrontendClassVertex GetClassAudioOutput()
 			{
 				FMetasoundFrontendClassVertex GeneratedAudio;
@@ -462,6 +468,12 @@ namespace Metasound
 				return SoundIsPreviewSoundName;
 			}
 
+			const FString& GetInstanceIDName()
+			{
+				static const FString& Name = FMetasoundInstanceTransmitter::GetInstanceIDEnvironmentVariableName();
+				return Name;
+			}
+
 			FMetasoundFrontendClassVertex GetClassLeftAudioOutput()
 			{
 				FMetasoundFrontendClassVertex GeneratedLeftAudio;
@@ -646,6 +658,12 @@ namespace Metasound
 				check(CurrentMonoVersion::GetIsPreviewSoundName() == CurrentStereoVersion::GetIsPreviewSoundName());
 				return CurrentMonoVersion::GetIsPreviewSoundName();
 			}
+
+			const FString& GetInstanceIDName()
+			{
+				check(CurrentMonoVersion::GetInstanceIDName() == CurrentStereoVersion::GetInstanceIDName());
+				return CurrentMonoVersion::GetInstanceIDName();
+			}
 		}
 
 		namespace MetasoundSourceMono
diff --git a/Engine/Plugins/Runtime/Metasound/Source/MetasoundEngine/Private/MetasoundEngineArchetypes.h b/Engine/Plugins/Runtime/Metasound/Source/MetasoundEngine/Private/MetasoundEngineArchetypes.h
--- a/Engine/Plugins/Runtime/Metasound/Source/MetasoundEngine/Private/MetasoundEngineArchetypes.h
+++ b/Engine/Plugins/Runtime/Metasound/Source/MetasoundEngine/Private/MetasoundEngineArchetypes.h
@@ -49,6 +49,7 @@ namespace Metasound
 			const FString& GetAudioDeviceHandleVariableName();
 			const FString& GetSoundUniqueIdName();
 			const FString& GetIsPreviewSoundName();
+			const FString& GetInstanceIDName();
 		}
 
 		// V1.1 of a Metasound Stereo source. Uses two FAudioBuffers as outputs.
@@ -62,6 +63,7 @@ namespace Metasound
 			const FString& GetAudioDeviceHandleVariableName();
 			const FString& GetSoundUniqueIdName();
 			const FString& GetIsPreviewSoundName();
+			const FString& GetInstanceIDName();
 		}
 
 		// Current version of Metasound Source.
@@ -72,6 +74,7 @@ namespace Metasound
 			const FString& GetAudioDeviceHandleVariableName();
 			const FString& GetSoundUniqueIdName();
 			const FString& GetIsPreviewSoundName();
+			const FString& GetInstanceIDName();
 		}
 
 		// Current version of MetasoundSourceMono
diff --git a/Engine/Plugins/Runtime/Metasound/Source/MetasoundEngine/Private/MetasoundSource.cpp b/Engine/Plugins/Runtime/Metasound/Source/MetasoundEngine/Private/MetasoundSource.cpp
--- a/Engine/Plugins/Runtime/Metasound/Source/MetasoundEngine/Private/MetasoundSource.cpp
+++ b/Engine/Plugins/Runtime/Metasound/Source/MetasoundEngine/Private/MetasoundSource.cpp
@@ -10,9 +10,9 @@
 #include "MetasoundEngineEnvironment.h"
 #include "MetasoundEnvironment.h"
 #include "MetasoundFrontendController.h"
+#include "MetasoundFrontendInjectReceiveNodes.h"
 #include "MetasoundFrontendQuery.h"
 #include "MetasoundFrontendQuerySteps.h"
-#include "MetasoundFrontendSearchEngine.h"
 #include "MetasoundFrontendTransform.h"
 #include "MetasoundGenerator.h"
 #include "MetasoundInstanceTransmitter.h"
@@ -30,6 +30,35 @@
 
 #define LOCTEXT_NAMESPACE "MetaSoundSource"
 
+namespace MetaSoundSourcePrivate
+{
+	using FFormatOutputVertexKeyMap = TMap<EMetasoundSourceAudioFormat, TArray<Metasound::FVertexKey>>;
+
+	const FFormatOutputVertexKeyMap& GetFormatOutputVertexKeys()
+	{
+		using namespace Metasound;
+		using namespace Metasound::Engine;
+
+		static const FFormatOutputVertexKeyMap Map
+		(
+			{
+				{
+					EMetasoundSourceAudioFormat::Mono, 
+					{ MetasoundSourceMono::GetAudioOutputName() }
+				},
+				{
+					EMetasoundSourceAudioFormat::Stereo, 
+					{ 
+						MetasoundSourceStereo::GetLeftAudioOutputName(), 
+						MetasoundSourceStereo::GetRightAudioOutputName() 
+					}
+				}
+			}
+		);
+		return Map;
+	}
+}
+
 
 FAutoConsoleVariableRef CVarMetaSoundBlockRate(
 	TEXT("au.MetaSound.BlockRate"),
@@ -204,74 +233,27 @@ ISoundGeneratorPtr UMetaSoundSource::CreateSoundGenerator(const FSoundGeneratorI
 
 	SampleRate = InParams.SampleRate;
 	FOperatorSettings InSettings = GetOperatorSettings(static_cast<FSampleRate>(SampleRate));
-	FMetasoundEnvironment Environment;
-
-	// Add audio device ID to environment.
-	FAudioDeviceHandle DeviceHandle;
-	if (UWorld* World = GetWorld())
-	{
-		DeviceHandle = World->GetAudioDevice();
-	}
-	if (!DeviceHandle.IsValid())
-	{
-		if (FAudioDeviceManager* DeviceManager = FAudioDeviceManager::Get())
-		{
-			DeviceHandle = DeviceManager->GetMainAudioDeviceHandle();
-		}
-	}
-	Environment.SetValue<FAudioDeviceHandle>(MetasoundSource::GetAudioDeviceHandleVariableName(), DeviceHandle);
-	Environment.SetValue<uint32>(MetasoundSource::GetSoundUniqueIdName(), GetUniqueID());
-	Environment.SetValue<bool>(MetasoundSource::GetIsPreviewSoundName(), InParams.bIsPreviewSound);
+	FMetasoundEnvironment Environment = CreateEnvironment(InParams);
 
-	const FMetasoundFrontendDocument* OriginalDoc = GetDocument().Get();
-	if (nullptr == OriginalDoc)
+	// TODO: cache graph to avoid having to create it every call to `CreateSoundGenerator(...)`
+	TUniquePtr<IGraph> MetasoundGraph = BuildMetasoundDocument();
+	if (!MetasoundGraph.IsValid())
 	{
-		UE_LOG(LogMetaSound, Error, TEXT("Cannot create sound generator. Null Metasound document in UMetaSoundSource [Name:%s]"), *GetName());
+		UE_LOG(LogMetaSound, Error, TEXT("Cannot create UMetaSoundSource SoundGenerator [Name:%s]. Failed to create MetaSound Graph"), *GetName());
 		return ISoundGeneratorPtr(nullptr);
 	}
 
-	// Inject receive nodes for unused transmittable inputs. Perform edits on a copy
-	// of the document to avoid altering document.
-	// TODO: Use a light wrapper of the document instead of a copy.
-	FMetasoundFrontendDocument DocumentWithInjectedReceives;
-	ensure(CopyDocumentAndInjectReceiveNodes(InParams.InstanceID, *OriginalDoc, DocumentWithInjectedReceives));
-
-	// Create handles for new root graph
-	FConstDocumentHandle NewDocumentHandle = IDocumentController::CreateDocumentHandle(MakeAccessPtr<FConstDocumentAccessPtr>(DocumentWithInjectedReceives.AccessPoint, DocumentWithInjectedReceives));
-
 	FMetasoundGeneratorInitParams InitParams = 
 	{
 		InSettings,
-		MoveTemp(DocumentWithInjectedReceives),
+		MoveTemp(MetasoundGraph),
 		Environment,
 		GetName(),
-		{}, // Channel names still need to be set.
+		GetAudioOutputVertexKeys(),
 		MetasoundSource::GetOnPlayInputName(),
 		MetasoundSource::GetIsFinishedOutputName()
 	};
 
-	// Set init params dependent upon output audio format.
-	switch (OutputFormat)
-	{
-		case EMetasoundSourceAudioFormat::Mono:
-		{
-			InitParams.AudioOutputNames = { MetasoundSourceMono::GetAudioOutputName() };
-		}
-		break;
-
-		case EMetasoundSourceAudioFormat::Stereo:
-		{
-			InitParams.AudioOutputNames = { MetasoundSourceStereo::GetLeftAudioOutputName(), MetasoundSourceStereo::GetRightAudioOutputName() };
-		}
-		break;
-
-		default:
-		{
-			static_assert(static_cast<int32>(EMetasoundSourceAudioFormat::COUNT) == 2, "Possible missing format switch case coverage.");
-		}
-		break;
-	}
-
 	return ISoundGeneratorPtr(new FMetasoundGenerator(MoveTemp(InitParams)));
 }
 
@@ -304,4 +286,71 @@ const TArray<FMetasoundFrontendVersion>& UMetaSoundSource::GetSupportedArchetype
 	return Supported;
 }
 
+
+Metasound::FMetasoundEnvironment UMetaSoundSource::CreateEnvironment() const
+{
+	using namespace Metasound;
+	using namespace Metasound::Engine;
+
+	FMetasoundEnvironment Environment;
+	// Add audio device ID to environment.
+	FAudioDeviceHandle DeviceHandle;
+	if (UWorld* World = GetWorld())
+	{
+		DeviceHandle = World->GetAudioDevice();
+	}
+	if (!DeviceHandle.IsValid())
+	{
+		if (FAudioDeviceManager* DeviceManager = FAudioDeviceManager::Get())
+		{
+			DeviceHandle = DeviceManager->GetMainAudioDeviceHandle();
+		}
+	}
+
+	Environment.SetValue<FAudioDeviceHandle>(MetasoundSource::GetAudioDeviceHandleVariableName(), DeviceHandle);
+	Environment.SetValue<uint32>(MetasoundSource::GetSoundUniqueIdName(), GetUniqueID());
+
+	return Environment;
+}
+
+Metasound::FMetasoundEnvironment UMetaSoundSource::CreateEnvironment(const FSoundGeneratorInitParams& InParams) const
+{
+	using namespace Metasound;
+	using namespace Metasound::Engine;
+
+	FMetasoundEnvironment Environment = CreateEnvironment();
+	Environment.SetValue<bool>(MetasoundSource::GetIsPreviewSoundName(), InParams.bIsPreviewSound);
+	Environment.SetValue<uint64>(MetasoundSource::GetInstanceIDName(), InParams.InstanceID);
+
+	return Environment;
+}
+
+Metasound::FMetasoundEnvironment UMetaSoundSource::CreateEnvironment(const FAudioInstanceTransmitterInitParams& InParams) const
+{
+	using namespace Metasound;
+	using namespace Metasound::Engine;
+
+	FMetasoundEnvironment Environment = CreateEnvironment();
+	Environment.SetValue<uint64>(MetasoundSource::GetInstanceIDName(), InParams.InstanceID);
+
+	return Environment;
+}
+
+const TArray<Metasound::FVertexKey>& UMetaSoundSource::GetAudioOutputVertexKeys() const
+{
+	using namespace MetaSoundSourcePrivate;
+
+	if (const TArray<Metasound::FVertexKey>* ArrayKeys = GetFormatOutputVertexKeys().Find(OutputFormat))
+	{
+		return *ArrayKeys;
+	}
+	else
+	{
+		// Unhandled audio format. Need to update audio output format vertex key map.
+		checkNoEntry();
+		static const TArray<Metasound::FVertexKey> Empty;
+		return Empty;
+	}
+}
+
 #undef LOCTEXT_NAMESPACE // MetaSoundSource
diff --git a/Engine/Plugins/Runtime/Metasound/Source/MetasoundEngine/Public/Metasound.h b/Engine/Plugins/Runtime/Metasound/Source/MetasoundEngine/Public/Metasound.h
--- a/Engine/Plugins/Runtime/Metasound/Source/MetasoundEngine/Public/Metasound.h
+++ b/Engine/Plugins/Runtime/Metasound/Source/MetasoundEngine/Public/Metasound.h
@@ -253,8 +253,6 @@ public:
 		return this;
 	}
 
-	TUniquePtr<IAudioInstanceTransmitter> CreateInstanceTransmitter(const FAudioInstanceTransmitterInitParams& InParams) const;
-
 	// Get the most up to date archetype for metasound sources.
 	const TArray<FMetasoundFrontendVersion>& GetSupportedArchetypeVersions() const override;
 
@@ -275,13 +273,4 @@ protected:
 		// object is no longer valid.
 		return MakeAccessPtr<FConstDocumentAccessPtr>(RootMetaSoundDocument.AccessPoint, RootMetaSoundDocument);
 	}
-
-	bool CopyDocumentAndInjectReceiveNodes(uint64 InInstanceID, const FMetasoundFrontendDocument& InSourceDoc, FMetasoundFrontendDocument& OutDestDoc) const;	
-
-private:
-	TArray<FString> GetTransmittableInputVertexNames() const;
-	Metasound::FOperatorSettings GetOperatorSettings(Metasound::FSampleRate InSampleRate) const;
-	Metasound::FSendAddress CreateSendAddress(uint64 InInstanceID, const FString& InVertexName, const FName& InDataTypeName) const;
-	static const FString& GetAudioDeviceHandleVariableName();
-	static const FMetasoundFrontendArchetype& GetBaseArchetype();
 };
diff --git a/Engine/Plugins/Runtime/Metasound/Source/MetasoundEngine/Public/MetasoundSource.h b/Engine/Plugins/Runtime/Metasound/Source/MetasoundEngine/Public/MetasoundSource.h
--- a/Engine/Plugins/Runtime/Metasound/Source/MetasoundEngine/Public/MetasoundSource.h
+++ b/Engine/Plugins/Runtime/Metasound/Source/MetasoundEngine/Public/MetasoundSource.h
@@ -183,4 +183,8 @@ protected:
 
 private:
 	Metasound::FOperatorSettings GetOperatorSettings(Metasound::FSampleRate InSampleRate) const;
+	Metasound::FMetasoundEnvironment CreateEnvironment() const;
+	Metasound::FMetasoundEnvironment CreateEnvironment(const FSoundGeneratorInitParams& InParams) const;
+	Metasound::FMetasoundEnvironment CreateEnvironment(const FAudioInstanceTransmitterInitParams& InParams) const;
+	const TArray<Metasound::FVertexKey>& GetAudioOutputVertexKeys() const;
 };
diff --git a/Engine/Plugins/Runtime/Metasound/Source/MetasoundFrontend/Private/MetasoundAssetBase.cpp b/Engine/Plugins/Runtime/Metasound/Source/MetasoundFrontend/Private/MetasoundAssetBase.cpp
--- a/Engine/Plugins/Runtime/Metasound/Source/MetasoundFrontend/Private/MetasoundAssetBase.cpp
+++ b/Engine/Plugins/Runtime/Metasound/Source/MetasoundFrontend/Private/MetasoundAssetBase.cpp
@@ -13,8 +13,10 @@
 #include "MetasoundFrontendController.h"
 #include "MetasoundFrontendDocument.h"
 #include "MetasoundFrontendGraph.h"
+#include "MetasoundFrontendInjectReceiveNodes.h"
 #include "MetasoundFrontendSearchEngine.h"
 #include "MetasoundFrontendTransform.h"
+#include "MetasoundInstanceTransmitter.h"
 #include "MetasoundJsonBackend.h"
 #include "MetasoundLog.h"
 #include "MetasoundReceiveNode.h"
@@ -259,6 +261,28 @@ FMetasoundFrontendVersion FMetasoundAssetBase::GetPreferredArchetypeVersion(cons
 	return GetDefaultArchetypeVersion();
 }
 
+bool FMetasoundAssetBase::GetArchetype(FMetasoundFrontendArchetype& OutArchetype) const
+{
+	using namespace Metasound;
+	using namespace Metasound::Frontend;
+
+	bool bFoundArchetype = false;
+
+	const FMetasoundFrontendDocument* Document = GetDocument().Get();
+	if (nullptr != Document)
+	{
+		FArchetypeRegistryKey ArchetypeRegistryKey = GetArchetypeRegistryKey(Document->ArchetypeVersion);
+		bFoundArchetype = IArchetypeRegistry::Get().FindArchetype(ArchetypeRegistryKey, OutArchetype);
+
+		if (!bFoundArchetype)
+		{
+			UE_LOG(LogMetaSound, Warning, TEXT("No registered archetype matching archetype version on document [ArchetypeVersion:%s]"), *Document->ArchetypeVersion.ToString());
+		}
+	}
+
+	return bFoundArchetype;
+}
+
 void FMetasoundAssetBase::SetDocument(const FMetasoundFrontendDocument& InDocument)
 {
 	FMetasoundFrontendDocument& Document = GetDocumentChecked();
@@ -356,87 +380,38 @@ bool FMetasoundAssetBase::VersionAsset(const IMetaSoundAssetInterface& InAssetIn
 	return bDidEdit;
 }
 
-bool FMetasoundAssetBase::CopyDocumentAndInjectReceiveNodes(uint64 InInstanceID, const FMetasoundFrontendDocument& InSourceDoc, FMetasoundFrontendDocument& OutDestDoc) const
+
+TUniquePtr<Metasound::IGraph> FMetasoundAssetBase::BuildMetasoundDocument() const
 {
 	using namespace Metasound;
 	using namespace Metasound::Frontend;
 
-	OutDestDoc = InSourceDoc;
-
-	FDocumentHandle Document = IDocumentController::CreateDocumentHandle(MakeAccessPtr<FDocumentAccessPtr>(OutDestDoc.AccessPoint, OutDestDoc));
-	FGraphHandle RootGraph = Document->GetRootGraph();
-
-	TArray<FSendInfoAndVertexName> SendInfoAndVertexes = GetSendInfos(InInstanceID);
-
-	// Inject receive nodes for each transmittable input
-	for (const FSendInfoAndVertexName& InfoAndVertexName : SendInfoAndVertexes)
+	const FMetasoundFrontendDocument* Doc = GetDocument().Get();
+	if (nullptr == Doc)
 	{
+		UE_LOG(LogMetaSound, Error, TEXT("Cannot create graph. Null MetaSound document in MetaSound asset [Name:%s]"), *GetOwningAssetName());
+		return TUniquePtr<IGraph>(nullptr);
+	}
 
-		// Add receive node to graph
-		FMetasoundFrontendClassMetadata ReceiveNodeMetadata;
-		bool bSuccess = GetReceiveNodeMetadataForDataType(InfoAndVertexName.SendInfo.TypeName, ReceiveNodeMetadata);
-		if (!bSuccess)
-		{
-			// TODO: log warning
-			continue;
-		}
-		FNodeHandle ReceiveNode = RootGraph->AddNode(ReceiveNodeMetadata);
-
-		// Add receive node address to graph
-		FNodeHandle AddressNode = AddInputPinForSendAddress(InfoAndVertexName.SendInfo, RootGraph);
-		TArray<FOutputHandle> AddressNodeOutputs = AddressNode->GetOutputs();
-		if (AddressNodeOutputs.Num() != 1)
-		{
-			// TODO: log warning
-			continue;
-		}
-
-		FOutputHandle AddressOutput = AddressNodeOutputs[0];
-		TArray<FInputHandle> ReceiveAddressInput = ReceiveNode->GetInputsWithVertexName(Metasound::FReceiveNodeNames::GetAddressInputName());
-		if (ReceiveAddressInput.Num() != 1)
-		{
-			// TODO: log error
-			continue;
-		}
-
-		ensure(ReceiveAddressInput[0]->Connect(*AddressOutput));
-
-		// Swap input node connections with receive node connections
-		FNodeHandle InputNode = RootGraph->GetInputNodeWithName(InfoAndVertexName.VertexName);
-		if (!ensure(InputNode->GetOutputs().Num() == 1))
-		{
-			// TODO: handle input node with varying number of outputs or varying output types.
-			continue;
-		}
+	// Create graph which can spawn instances. TODO: cache graph.
+	TUniquePtr<FFrontendGraph> FrontendGraph = FFrontendGraphBuilder::CreateGraph(*Doc);
+	if (FrontendGraph.IsValid())
+	{
+		auto GetVertexName = [](const FMetasoundFrontendClassVertex& InVertex) { return InVertex.Name; };
 
-		FOutputHandle InputNodeOutput = InputNode->GetOutputs()[0];
+		TSet<FVertexKey> VerticesToSkip;
+		FMetasoundFrontendArchetype Archetype;
+		GetArchetype(Archetype);
+		Algo::Transform(Archetype.Interface.Inputs, VerticesToSkip, GetVertexName);
 
-		if (ensure(ReceiveNode->IsValid()))
+		bool bSuccessfullyInjectedReceiveNodes = InjectReceiveNodes(*FrontendGraph, FMetasoundInstanceTransmitter::CreateSendAddressFromEnvironment, VerticesToSkip);
+		if (!bSuccessfullyInjectedReceiveNodes)
 		{
-			TArray<FOutputHandle> ReceiveNodeOutputs = ReceiveNode->GetOutputs();
-			if (!ensure(ReceiveNodeOutputs.Num() == 1))
-			{
-				// TODO: handle array outputs and receive nodes of varying formats.
-				continue;
-			}
-
-			TArray<FInputHandle> ReceiveDefaultInputs = ReceiveNode->GetInputsWithVertexName(Metasound::FReceiveNodeNames::GetDefaultDataInputName());
-			if (ensure(ReceiveDefaultInputs.Num() == 1))
-			{
-				FOutputHandle ReceiverNodeOutput = ReceiveNodeOutputs[0];
-				for (FInputHandle NodeInput : InputNodeOutput->GetConnectedInputs())
-				{
-					// Swap connections to receiver node
-					ensure(InputNodeOutput->Disconnect(*NodeInput));
-					ensure(ReceiverNodeOutput->Connect(*NodeInput));
-				}
-
-				ReceiveDefaultInputs[0]->Connect(*InputNodeOutput);
-			}
+			UE_LOG(LogMetaSound, Error, TEXT("Error while injecting async communication hooks. Instance communication may not function properly [Name:%s]."), *GetOwningAssetName());
 		}
 	}
 
-	return true;
+	return MoveTemp(FrontendGraph);
 }
 
 void FMetasoundAssetBase::RebuildReferencedAssets(const IMetaSoundAssetInterface& InAssetInterface)
@@ -502,6 +477,7 @@ bool FMetasoundAssetBase::AddingReferenceCausesLoop(const FSoftObjectPath& InRef
 }
 
 Metasound::FSendAddress FMetasoundAssetBase::CreateSendAddress(uint64 InInstanceID, const FString& InVertexName, const FName& InDataTypeName) const
+
 {
 	using namespace Metasound;
 
@@ -528,8 +504,9 @@ void FMetasoundAssetBase::ConvertFromPreset()
 
 TArray<FMetasoundAssetBase::FSendInfoAndVertexName> FMetasoundAssetBase::GetSendInfos(uint64 InInstanceID) const
 {
-	using FSendInfo = Metasound::FMetasoundInstanceTransmitter::FSendInfo;
+	using namespace Metasound;
 	using namespace Metasound::Frontend;
+	using FSendInfo = FMetasoundInstanceTransmitter::FSendInfo;
 
 	TArray<FSendInfoAndVertexName> SendInfos;
 
@@ -545,7 +522,7 @@ TArray<FMetasoundAssetBase::FSendInfoAndVertexName> FMetasoundAssetBase::GetSend
 
 			// TODO: incorporate VertexID into address. But need to ensure that VertexID
 			// will be maintained after injecting Receive nodes. 
-			Info.SendInfo.Address = CreateSendAddress(InInstanceID, InputHandle->GetName(), InputHandle->GetDataType());
+			Info.SendInfo.Address = FMetasoundInstanceTransmitter::CreateSendAddressFromInstanceID(InInstanceID, InputHandle->GetName(), InputHandle->GetDataType());
 			Info.SendInfo.ParameterName = FName(InputHandle->GetDisplayName().ToString()); // TODO: display name hack. Need to have naming consistent in editor for inputs
 			//Info.SendInfo.ParameterName = FName(*InputHandle->GetName()); // TODO: this is the expected parameter name.
 			Info.SendInfo.TypeName = InputHandle->GetDataType();
@@ -645,23 +622,6 @@ FText FMetasoundAssetBase::GetDisplayName(FString&& InTypeName) const
 }
 #endif // WITH_EDITORONLY_DATA
 
-bool FMetasoundAssetBase::GetReceiveNodeMetadataForDataType(const FName& InTypeName, FMetasoundFrontendClassMetadata& OutMetadata) const
-{
-	using namespace Metasound;
-	using namespace Metasound::Frontend;
-
-	const FNodeClassName ClassName = FReceiveNodeNames::GetClassNameForDataType(InTypeName);
-	TArray<FMetasoundFrontendClass> ReceiverNodeClasses = ISearchEngine::Get().FindClassesWithName(ClassName, true /* bInSortByVersion */);
-
-	if (ReceiverNodeClasses.IsEmpty())
-	{
-		return false;
-	}
-
-	OutMetadata = ReceiverNodeClasses[0].Metadata;
-	return true;
-}
-
 bool FMetasoundAssetBase::MarkMetasoundDocumentDirty() const
 {
 	if (const UObject* OwningAsset = GetOwningAsset())
@@ -739,4 +699,14 @@ const FMetasoundFrontendDocument& FMetasoundAssetBase::GetDocumentChecked() cons
 	check(nullptr != Document);
 	return *Document;
 }
+
+FString FMetasoundAssetBase::GetOwningAssetName() const
+{
+	if (const UObject* OwningAsset = GetOwningAsset())
+	{
+		return OwningAsset->GetName();
+	}
+	return FString();
+}
+
 #undef LOCTEXT_NAMESPACE // "MetaSound"
diff --git a/Engine/Plugins/Runtime/Metasound/Source/MetasoundFrontend/Private/MetasoundFrontendGraph.cpp b/Engine/Plugins/Runtime/Metasound/Source/MetasoundFrontend/Private/MetasoundFrontendGraph.cpp
--- a/Engine/Plugins/Runtime/Metasound/Source/MetasoundFrontend/Private/MetasoundFrontendGraph.cpp
+++ b/Engine/Plugins/Runtime/Metasound/Source/MetasoundFrontend/Private/MetasoundFrontendGraph.cpp
@@ -522,15 +522,16 @@ namespace Metasound
 
 	void FFrontendGraphBuilder::AddDefaultInputVariables(FBuildGraphContext& InGraphContext)
 	{
-		for (const TPair<FNodeIDVertexID, FDefaultVariableData>& Pair : InGraphContext.DefaultInputs)
+		using FIterator = FDefaultInputByIDMap::TIterator;
+
+		for (FDefaultInputByIDMap::ElementType& Pair : InGraphContext.DefaultInputs)
 		{
-			const FDefaultVariableData& VariableData = Pair.Value;
+			FDefaultVariableData& VariableData = Pair.Value;
 			const FGuid VariableNodeID = FGuid::NewGuid();
 
 			// 1. Construct and add the default variable to the graph
 			{
-				FVariableNodeConstructorParams InitParams = VariableData.InitParams.Clone();
-				TUniquePtr<const INode> DefaultVariable = FMetasoundFrontendRegistryContainer::Get()->CreateVariableNode(VariableData.TypeName, MoveTemp(InitParams));
+				TUniquePtr<const INode> DefaultVariable = FMetasoundFrontendRegistryContainer::Get()->CreateVariableNode(VariableData.TypeName, MoveTemp(VariableData.InitParams));
 				InGraphContext.Graph->AddNode(VariableNodeID, TSharedPtr<const INode>(DefaultVariable.Release()));
 			}
 
@@ -555,6 +556,9 @@ namespace Metasound
 				UE_LOG(LogMetaSound, Error, TEXT("Failed to connect default variable edge: from '%s' to '%s'"), *FromVertexKey, *ToVertexKey);
 			}
 		}
+
+		// Clear default inputs because literals have been moved out of default input map.
+		InGraphContext.DefaultInputs.Reset();
 	}
 
 	TArray<FFrontendGraphBuilder::FDefaultVariableData> FFrontendGraphBuilder::GetInputDefaultVariableData(const FMetasoundFrontendNode& InNode, const FNodeInitData& InInitData)
diff --git a/Engine/Plugins/Runtime/Metasound/Source/MetasoundFrontend/Private/MetasoundFrontendGraph.h b/Engine/Plugins/Runtime/Metasound/Source/MetasoundFrontend/Private/MetasoundFrontendGraph.h
deleted file mode 100644
diff --git a/Engine/Plugins/Runtime/Metasound/Source/MetasoundFrontend/Private/MetasoundFrontendInjectReceiveNodes.cpp b/Engine/Plugins/Runtime/Metasound/Source/MetasoundFrontend/Private/MetasoundFrontendInjectReceiveNodes.cpp
new file mode 100644
--- /dev/null
+++ b/Engine/Plugins/Runtime/Metasound/Source/MetasoundFrontend/Private/MetasoundFrontendInjectReceiveNodes.cpp
@@ -0,0 +1,223 @@
+// Copyright Epic Games, Inc. All Rights Reserved.
+
+#include "MetasoundFrontendInjectReceiveNodes.h"
+
+#include "CoreMinimal.h"
+#include "MetasoundFacade.h"
+#include "MetasoundFrontendGraph.h"
+#include "MetasoundFrontendRegistries.h"
+#include "MetasoundReceiveNode.h"
+#include "MetasoundVertex.h"
+
+namespace Metasound
+{
+	namespace Frontend 
+	{
+		namespace MetasoundFrontendInjectReceiveNodesPrivate
+		{
+			// Metasound Operator which returns the transmission address for the injected receive node.
+			class FAddressOperator : public FNoOpOperator
+			{
+			public:
+				static const TCHAR* GetAddressVertexKey();
+				static TUniquePtr<IOperator> CreateOperator(const FCreateOperatorParams&, TArray<TUniquePtr<IOperatorBuildError>>&);
+				static const FNodeClassMetadata& GetNodeInfo();
+
+				FAddressOperator(TDataReadReference<FSendAddress> InAddress)
+				: Address(InAddress)
+				{
+				}
+
+				virtual ~FAddressOperator() = default;
+
+				virtual FDataReferenceCollection GetOutputs() const override
+				{
+					FDataReferenceCollection Outputs;
+					Outputs.AddDataReadReference(GetAddressVertexKey(), Address);
+					return Outputs;
+				}
+
+			private:
+				TDataReadReference<FSendAddress> Address;
+			};
+
+			// Metasound Node which returns the transmission address for the injected receive node.
+			class FAddressNode : public FNodeFacade
+			{
+			public:
+				FAddressNode(const FGuid& InID, const FVertexKey& InVertexKey, const FName& InTypeName, FReceiveNodeAddressFunction InAddressFunction)
+				: FNodeFacade(FString::Format(TEXT("ReceiveAddressInject_{0}"), { InVertexKey }), InID, TFacadeOperatorClass<FAddressOperator>())
+				, VertexKey(InVertexKey)
+				, TypeName(InTypeName)
+				, AddressFunction(InAddressFunction)
+				{
+				}
+
+				FSendAddress GetAddress(const FMetasoundEnvironment& InEnvironment) const
+				{
+					return AddressFunction(InEnvironment, VertexKey, TypeName);
+				}
+
+			private:
+				FVertexKey VertexKey;
+				FName TypeName;
+				FReceiveNodeAddressFunction AddressFunction;
+			};
+
+			const TCHAR* FAddressOperator::GetAddressVertexKey() { return TEXT("Address"); }
+
+			TUniquePtr<IOperator> FAddressOperator::CreateOperator(const FCreateOperatorParams& InParams, TArray<TUniquePtr<IOperatorBuildError>>& OutBuildErrors)
+			{
+				const FAddressNode& AddressNode = static_cast<const FAddressNode&>(InParams.Node);
+				FSendAddress Address = AddressNode.GetAddress(InParams.Environment);
+				return MakeUnique<FAddressOperator>(TDataReadReference<FSendAddress>::CreateNew(Address));
+			}
+
+			const FNodeClassMetadata& FAddressOperator::GetNodeInfo()
+			{
+				auto CreateVertexInterface = []() -> FVertexInterface
+				{
+					FVertexInterface Interface
+					{
+						FInputVertexInterface{},
+						FOutputVertexInterface
+						{
+							TOutputDataVertexModel<FSendAddress>(FAddressOperator::GetAddressVertexKey(), FText::GetEmpty())
+						}
+					};
+
+					return Interface;
+				};
+
+				auto CreateNodeClassMetadata = [&CreateVertexInterface]() -> FNodeClassMetadata
+				{
+					FNodeClassMetadata Metadata
+					{
+						FNodeClassName{"MetasoundFrontendInjectReceiveNodes", "ReceiveNodeAddress", ""},
+						1, // MajorVersion
+						0, // MinorVersion
+						FText::GetEmpty(), // DisplayName
+						FText::GetEmpty(), // Description
+						FText::GetEmpty(), // Author
+						FText::GetEmpty(), // Prompt If Missing
+						CreateVertexInterface(), // DefaultInterface
+						TArray<FText>(), // CategoryHierachy
+						TArray<FName>(), // Keywoards
+						FNodeDisplayStyle {} // NodeDisplayStyle
+					}; 
+
+					return Metadata;
+				};
+
+				static const FNodeClassMetadata Metadata = CreateNodeClassMetadata();
+				return Metadata;
+			}
+
+			TUniquePtr<INode> CreateReceiveNodeForDataType(const FGuid& InID, const FVertexKey& InVertexKey, const FName& InDataType)
+			{
+				const EMetasoundFrontendClassType Type = EMetasoundFrontendClassType::External;
+				const FNodeClassName& ClassName = ReceiveNodeInfo::GetClassNameForDataType(InDataType);
+				const int32 MajorVersion = ReceiveNodeInfo::GetCurrentMajorVersion();
+				const int32 MinorVersion = ReceiveNodeInfo::GetCurrentMinorVersion();
+
+				FNodeRegistryKey ReceiveNodeRegistryKey = NodeRegistryKey::CreateKey(Type, ClassName.GetFullName().ToString(), MajorVersion, MinorVersion);
+				FNodeInitData ReceiveNodeInitData 
+				{ 
+					FString::Format(TEXT("ReceiveInject_{0}"), { InVertexKey }),
+					InID
+				};
+				return FMetasoundFrontendRegistryContainer::Get()->CreateNode(ReceiveNodeRegistryKey, MoveTemp(ReceiveNodeInitData));
+			}
+		}
+
+		bool InjectReceiveNode(FFrontendGraph& InGraph, const FReceiveNodeAddressFunction& InAddressFunction, const FInputDataDestination& InputDestination)
+		{
+			using namespace MetasoundFrontendInjectReceiveNodesPrivate;
+			// should never contain null nodes for input destination.
+			check(InputDestination.Node != nullptr);
+
+			const FVertexKey& VertexKey = InputDestination.Vertex.GetVertexName();
+			const FName& DataType = InputDestination.Vertex.GetDataTypeName();
+
+			// Create a receive node.
+			const FGuid ReceiveNodeID = FGuid::NewGuid();
+			TSharedPtr<INode> ReceiveNode(CreateReceiveNodeForDataType(ReceiveNodeID, VertexKey, DataType).Release());
+
+			if (!ReceiveNode.IsValid())
+			{
+				UE_LOG(LogMetaSound, Warning, TEXT("Failed to create receive node while injecting receive node for graph input [VertexKey:%s, VertexDescription:%s, DataTypeName:%s]"), *VertexKey, *InputDestination.Vertex.GetMetadata().Description.ToString(), *DataType.ToString());
+				return false;
+			}
+
+			// Create a node for address of receive node.
+			const FGuid AddressNodeID = FGuid::NewGuid();
+			TSharedPtr<INode> AddressNode = MakeShared<FAddressNode>(AddressNodeID, VertexKey, DataType, InAddressFunction);
+
+			if (!AddressNode.IsValid())
+			{
+				UE_LOG(LogMetaSound, Warning, TEXT("Failed to create address node while injecting receive node for graph input [VertexKey:%s, VertexDescription:%s, DataTypeName:%s]"), *VertexKey, *InputDestination.Vertex.GetMetadata().Description.ToString(), *DataType.ToString());
+				return false;
+			}
+
+			// Add receive node and value node to graph.
+			InGraph.AddNode(ReceiveNodeID, ReceiveNode);
+			InGraph.AddNode(AddressNodeID, AddressNode);
+			check(InGraph.AddDataEdge(*AddressNode, FAddressOperator::GetAddressVertexKey(), *ReceiveNode, ReceiveNodeInfo::GetAddressInputName()));
+			
+			auto IsEdgeConnectedToCurrentInput = [&InputDestination](const FDataEdge& InEdge)
+			{
+				return (InEdge.From.Node == InputDestination.Node) && (InEdge.From.Vertex.GetVertexName() == InputDestination.Vertex.GetVertexName());
+			};
+
+			// Cache previous connections from input node.
+			TArray<FDataEdge> EdgesFromInput = InGraph.GetDataEdges().FilterByPredicate(IsEdgeConnectedToCurrentInput);
+			// Remove existing connections from input node.
+			InGraph.RemoveDataEdgeByPredicate(IsEdgeConnectedToCurrentInput);
+
+			// Connect previous connections to receive node output.
+			FOutputDataSource ReceiveOutputSource { *ReceiveNode, ReceiveNode->GetVertexInterface().GetOutputVertex(ReceiveNodeInfo::GetOutputName()) };
+			for (const FDataEdge& Edge : EdgesFromInput)
+			{
+				FDataEdge NewEdge { ReceiveOutputSource, Edge.To };
+				InGraph.AddDataEdge(FDataEdge{ ReceiveOutputSource, Edge.To });
+			}
+
+			// Connect input node to receive node.
+			check(InGraph.AddDataEdge(*InputDestination.Node, VertexKey, *ReceiveNode, ReceiveNodeInfo::GetDefaultDataInputName()));
+
+			return true;
+		}
+
+		bool InjectReceiveNodes(FFrontendGraph& InGraph, const FReceiveNodeAddressFunction& InAddressFunction, const TSet<FVertexKey>& InInputVerticesToSkip)
+		{
+			using namespace Metasound::Frontend;
+			bool bSuccess = true;
+
+			const FInputDataDestinationCollection& InputDestinations = InGraph.GetInputDataDestinations();
+
+			for (const FInputDataDestinationCollection::ElementType& Pair : InputDestinations)
+			{
+				const FInputDataDestination& InputDestination = Pair.Value;
+				if (!ensure(InputDestination.Node != nullptr))
+				{
+					// should never contain null nodes for input destination.
+					bSuccess = false;
+					continue;
+				}
+
+				const FVertexKey& VertexKey = InputDestination.Vertex.GetVertexName();
+				if (InInputVerticesToSkip.Contains(VertexKey))
+				{
+					// Skip this key as requested by caller.
+					continue;
+				}
+
+				const bool bInjectionSuccess = InjectReceiveNode(InGraph, InAddressFunction, InputDestination);
+				bSuccess = bSuccess || bInjectionSuccess;
+			}
+
+			return bSuccess;
+		}
+	}
+}
+
diff --git a/Engine/Plugins/Runtime/Metasound/Source/MetasoundFrontend/Private/MetasoundInstanceTransmitter.cpp b/Engine/Plugins/Runtime/Metasound/Source/MetasoundFrontend/Private/MetasoundInstanceTransmitter.cpp
--- a/Engine/Plugins/Runtime/Metasound/Source/MetasoundFrontend/Private/MetasoundInstanceTransmitter.cpp
+++ b/Engine/Plugins/Runtime/Metasound/Source/MetasoundFrontend/Private/MetasoundInstanceTransmitter.cpp
@@ -4,6 +4,34 @@
 
 namespace Metasound
 {
+	FString FMetasoundInstanceTransmitter::GetInstanceIDEnvironmentVariableName()
+	{
+		return TEXT("TransmitterInstanceID");
+	}
+
+	FSendAddress FMetasoundInstanceTransmitter::CreateSendAddressFromEnvironment(const FMetasoundEnvironment& InEnvironment, const FVertexKey& InVertexKey, const FName& InTypeName)
+	{
+		const FString IDVarName = GetInstanceIDEnvironmentVariableName();
+		uint64 InstanceID = -1;
+
+		if (ensure(InEnvironment.Contains<uint64>(IDVarName)))
+		{
+			InstanceID = InEnvironment.GetValue<uint64>(IDVarName);
+		}
+
+		return CreateSendAddressFromInstanceID(InstanceID, InVertexKey, InTypeName);
+	}
+
+	FSendAddress FMetasoundInstanceTransmitter::CreateSendAddressFromInstanceID(uint64 InInstanceID, const FVertexKey& InVertexKey, const FName& InTypeName)
+	{
+		FSendAddress Address;
+
+		Address.Subsystem = GetSubsystemNameForSendScope(ETransmissionScope::Global);
+		Address.ChannelName = FName(FString::Format(TEXT("{0}:{1}:{2}"), { InInstanceID, InVertexKey, InTypeName.ToString() }));
+
+		return Address;
+	}
+
 	FMetasoundInstanceTransmitter::FMetasoundInstanceTransmitter(const FMetasoundInstanceTransmitter::FInitParams& InInitParams)
 	: SendInfos(InInitParams.Infos)
 	, OperatorSettings(InInitParams.OperatorSettings)
diff --git a/Engine/Plugins/Runtime/Metasound/Source/MetasoundFrontend/Private/MetasoundReceiveNode.cpp b/Engine/Plugins/Runtime/Metasound/Source/MetasoundFrontend/Private/MetasoundReceiveNode.cpp
new file mode 100644
--- /dev/null
+++ b/Engine/Plugins/Runtime/Metasound/Source/MetasoundFrontend/Private/MetasoundReceiveNode.cpp
@@ -0,0 +1,36 @@
+// Copyright Epic Games, Inc. All Rights Reserved.
+
+#include "MetasoundReceiveNode.h"
+
+namespace Metasound
+{
+	namespace ReceiveNodeInfo
+	{
+		const FString& GetAddressInputName()
+		{
+			static const FString InputName = FString(TEXT("Address"));
+			return InputName;
+		}
+
+		const FString& GetDefaultDataInputName()
+		{
+			static const FString DefaultDataName = FString(TEXT("Default"));
+			return DefaultDataName;
+		}
+
+		const FString& GetOutputName()
+		{
+			static const FString OutputName = FString(TEXT("Out"));
+			return OutputName;
+		}
+
+		FNodeClassName GetClassNameForDataType(const FName& InDataTypeName)
+		{
+			return FNodeClassName{TEXT("Receive"), InDataTypeName, TEXT("")};
+		}
+
+		int32 GetCurrentMajorVersion() { return 1; }
+
+		int32 GetCurrentMinorVersion() { return 0; }
+	}
+}
diff --git a/Engine/Plugins/Runtime/Metasound/Source/MetasoundFrontend/Public/MetasoundAssetBase.h b/Engine/Plugins/Runtime/Metasound/Source/MetasoundFrontend/Public/MetasoundAssetBase.h
--- a/Engine/Plugins/Runtime/Metasound/Source/MetasoundFrontend/Public/MetasoundAssetBase.h
+++ b/Engine/Plugins/Runtime/Metasound/Source/MetasoundFrontend/Public/MetasoundAssetBase.h
@@ -87,7 +87,6 @@ public:
 	// Unregisters the root graph of the given asset with the MetaSound Frontend.
 	void UnregisterGraphWithFrontend();
 
-	bool CopyDocumentAndInjectReceiveNodes(uint64 InInstanceID, const FMetasoundFrontendDocument& InSourceDoc, FMetasoundFrontendDocument& OutDestDoc) const;
 
 	// Sets/overwrites the root class metadata
 	virtual void SetMetadata(FMetasoundFrontendClassMetadata& InMetadata);
@@ -103,10 +102,7 @@ public:
 
 	// Returns the preferred archetype for the given document.
 	virtual FMetasoundFrontendVersion GetPreferredArchetypeVersion(const FMetasoundFrontendDocument& InDocument) const;
-
-	// TODO:
-	//virtual void OnMetaSoundDependencyAdded(const FMetasoundFrontendClassMetadata& InMetadata) = 0;
-	//virtual void OnMetaSoundDependencyRemoved(const FMetasoundFrontendClassMetadata& InMetadata) = 0;
+	bool GetArchetype(FMetasoundFrontendArchetype& OutArchetype) const;
 
 	// Gets the asset class info.
 	virtual Metasound::Frontend::FNodeClassInfo GetAssetClassInfo() const = 0;
@@ -155,13 +151,17 @@ public:
 	// Calls the outermost package and marks it dirty.
 	bool MarkMetasoundDocumentDirty() const;
 
-protected:
 	struct FSendInfoAndVertexName
 	{
 		Metasound::FMetasoundInstanceTransmitter::FSendInfo SendInfo;
 		FString VertexName;
 	};
 
+	// Builds the Metasound Document returned by `GetDocument() const`.
+	virtual TUniquePtr<Metasound::IGraph> BuildMetasoundDocument() const;
+
+protected:
+
 	TArray<FSendInfoAndVertexName> GetSendInfos(uint64 InInstanceID) const;
 
 #if WITH_EDITORONLY_DATA
@@ -181,10 +181,12 @@ protected:
 	// Returns the owning asset responsible for transactions applied to metasound
 	virtual const UObject* GetOwningAsset() const = 0;
 
+	FString GetOwningAssetName() const;
+
+
 private:
 	Metasound::Frontend::FNodeRegistryKey RegistryKey;
 
-	bool GetReceiveNodeMetadataForDataType(const FName& InTypeName, FMetasoundFrontendClassMetadata& OutMetadata) const;
 	TArray<FString> GetTransmittableInputVertexNames() const;
 	Metasound::FSendAddress CreateSendAddress(uint64 InInstanceID, const FString& InVertexName, const FName& InDataTypeName) const;
 	Metasound::Frontend::FNodeHandle AddInputPinForSendAddress(const Metasound::FMetasoundInstanceTransmitter::FSendInfo& InSendInfo, Metasound::Frontend::FGraphHandle InGraph) const;
diff --git a/Engine/Plugins/Runtime/Metasound/Source/MetasoundFrontend/Public/MetasoundFrontendGraph.h b/Engine/Plugins/Runtime/Metasound/Source/MetasoundFrontend/Public/MetasoundFrontendGraph.h
new file mode 100644
--- /dev/null
+++ b/Engine/Plugins/Runtime/Metasound/Source/MetasoundFrontend/Public/MetasoundFrontendGraph.h
@@ -0,0 +1,170 @@
+// Copyright Epic Games, Inc. All Rights Reserved.
+
+#pragma once
+
+#include "CoreMinimal.h"
+
+#include "MetasoundFrontend.h"
+#include "MetasoundGraph.h"
+#include "MetasoundNodeInterface.h"
+
+namespace Metasound
+{
+	/** FFrontendGraph is a utility graph for use in the frontend. It can own nodes
+	 * that live within the graph and provides query interfaces for finding nodes
+	 * by dependency ID or input/output index. 
+	 */
+	class METASOUNDFRONTEND_API FFrontendGraph : public FGraph
+	{
+		public:
+			/** FFrontendGraph constructor.
+			 *
+			 * @parma InInstanceName - Name of this graph.
+			 * @parma InInstanceID - ID of this graph.
+			 */
+			FFrontendGraph(const FString& InInstanceName, const FGuid& InInstanceID);
+
+			virtual ~FFrontendGraph() = default;
+
+			/** Add an input node to this graph.
+			 *
+			 * @param InNodeID - The NodeID related to the parent FMetasoundFrontendClass.
+			 * @param InIndex - The positional index for the input.
+			 * @param InVertexKey - The key for the graph input vertex.
+			 * @param InNode - A shared pointer to an input node. 
+			 */
+			void AddInputNode(FGuid InNodeID, int32 InIndex, const FVertexKey& InVertexKey, TSharedPtr<const INode> InNode);
+
+			/** Add an output node to this graph.
+			 *
+			 * @param InNodeID - The NodeID related to the parent FMetasoundFrontendClass.
+			 * @param InIndex - The positional index for the output.
+			 * @param InVertexKey - The key for the graph output vertex.
+			 * @param InNode - A shared pointer to an output node. 
+			 */
+			void AddOutputNode(FGuid InNodeID, int32 InIndex, const FVertexKey& InVertexKey, TSharedPtr<const INode> InNode);
+
+			/** Store a node on this graph. 
+			 *
+			 * @param InNodeID - The NodeID related to the parent FMetasoundFrontendClass.
+			 * @param InNode - A shared pointer to a node. 
+			 */
+			void AddNode(FGuid InNodeID, TSharedPtr<const INode> InNode);
+
+			/** Retrieve node by node ID.
+			 *
+			 * @param InNodeID - The NodeID of the requested Node.
+			 *
+			 * @return Pointer to the Node if it is stored on this graph. nullptr otherwise. 
+			 */
+			const INode* FindNode(FGuid InNodeID) const;
+
+			/** Retrieve node by input index.
+			 *
+			 * @param InIndex - The index of the requested input.
+			 *
+			 * @return Pointer to the Node if it is stored on this graph. nullptr otherwise. 
+			 */
+			const INode* FindInputNode(int32 InIndex) const;
+
+			/** Retrieve node by output index.
+			 *
+			 * @param InIndex - The index of the requested output.
+			 *
+			 * @return Pointer to the Node if it is stored on this graph. nullptr otherwise. 
+			 */
+			const INode* FindOutputNode(int32 InIndex) const;
+
+			/** Returns true if all edges, destinations and sources refer to 
+			 * nodes stored in this graph. */
+			bool OwnsAllReferencedNodes() const;
+
+		private:
+
+			void StoreNode(TSharedPtr<const INode> InNode);
+
+			TMap<int32, const INode*> InputNodes;
+			TMap<int32, const INode*> OutputNodes;
+
+			TMap<FGuid, const INode*> NodeMap;
+			TSet<const INode*> StoredNodes;
+			TArray<TSharedPtr<const INode>> NodeStorage;
+	};
+
+	/** FFrontendGraphBuilder builds a FFrontendGraph from a FMetasoundDoucment
+	 * or FMetasoundFrontendClass.
+	 */
+	class METASOUNDFRONTEND_API FFrontendGraphBuilder
+	{
+	public:
+
+		/** Check that all dependencies are C++ class dependencies. 
+		 * 
+		 * @param InDocument - Document containing dependencies.
+		 *
+		 * @return True if all dependencies are C++ classes. False otherwise.
+		 */
+		static bool IsFlat(const FMetasoundFrontendDocument& InDocument);
+
+		static bool IsFlat(const FMetasoundFrontendGraphClass& InRoot, const TArray<FMetasoundFrontendClass>& InDependencies);
+
+		/* Metasound document should be in order to create this graph. */
+		static TUniquePtr<FFrontendGraph> CreateGraph(const FMetasoundFrontendDocument& InDocument);
+
+		/* Metasound document should be in order to create this graph. */
+		static TUniquePtr<FFrontendGraph> CreateGraph(const FMetasoundFrontendGraphClass& InGraphClass, const TArray<FMetasoundFrontendGraphClass>& InSubgraphs, const TArray<FMetasoundFrontendClass>& InDependencies);
+
+	private:
+		struct FDefaultVariableData
+		{
+			FGuid DestinationNodeID;
+			FGuid DestinationVertexID;
+			FVertexKey DestinationVertexKey;
+			FName TypeName;
+			FVariableNodeConstructorParams InitParams;
+		};
+
+		// Map of Input VertexID to variable data required to construct and connect default variable
+		using FNodeIDVertexID = TTuple<FGuid, FGuid>;
+		using FDependencyByIDMap = TMap<FGuid, const FMetasoundFrontendClass*>;
+		using FSharedNodeByIDMap = TMap<FGuid, TSharedPtr<const INode>>;
+		using FDefaultInputByIDMap = TMap<FNodeIDVertexID, FDefaultVariableData>;
+
+		// Context used throughout entire graph build process
+		// (for both a root and nested subgraphs)
+		struct FBuildContext
+		{
+			FDependencyByIDMap FrontendClasses;
+			FSharedNodeByIDMap Graphs;
+		};
+
+		// Transient context used for building a specific graph
+		struct FBuildGraphContext
+		{
+			TUniquePtr<FFrontendGraph> Graph;
+			const FMetasoundFrontendGraphClass& GraphClass;
+			FBuildContext& BuildContext;
+
+			FDefaultInputByIDMap DefaultInputs;
+		};
+
+		static TArray<FDefaultVariableData> GetInputDefaultVariableData(const FMetasoundFrontendNode& InNode, const FNodeInitData& InInitData);
+
+		static bool SortSubgraphDependencies(TArray<const FMetasoundFrontendGraphClass*>& Subgraphs);
+
+		static TUniquePtr<FFrontendGraph> CreateGraph(FBuildContext& InContext, const FMetasoundFrontendGraphClass& InSubgraph);
+
+		static const FMetasoundFrontendClassInput* FindClassInputForInputNode(const FMetasoundFrontendGraphClass& InOwningGraph, const FMetasoundFrontendNode& InInputNode, int32& OutClassInputIndex);
+		static const FMetasoundFrontendClassOutput* FindClassOutputForOutputNode(const FMetasoundFrontendGraphClass& InOwningGraph, const FMetasoundFrontendNode& InOutputNode, int32& OutClassOutputIndex);
+		static const FMetasoundFrontendLiteral* FindInputLiteralForInputNode(const FMetasoundFrontendNode& InInputNode, const FMetasoundFrontendClass& InInputNodeClass, const FMetasoundFrontendClassInput& InOwningGraphClassInput);
+
+		static TUniquePtr<INode> CreateInputNode(const FMetasoundFrontendNode& InNode, const FMetasoundFrontendClass& InClass, const FMetasoundFrontendClassInput& InOwningGraphClassInput);
+		static TUniquePtr<INode> CreateOutputNode(const FMetasoundFrontendNode& InNode, const FMetasoundFrontendClass& InClass, FBuildGraphContext& InGraphContext);
+		static TUniquePtr<INode> CreateExternalNode(const FMetasoundFrontendNode& InNode, const FMetasoundFrontendClass& InClass, FBuildGraphContext& InGraphContext);
+
+		// TODO: add errors here. Most will be a "PromptIfMissing"...
+		static void AddNodesToGraph(FBuildGraphContext& InGraphContext);
+		static void AddEdgesToGraph(FBuildGraphContext& InGraphContext);
+		static void AddDefaultInputVariables(FBuildGraphContext& InGraphContext);
+	};
+}
diff --git a/Engine/Plugins/Runtime/Metasound/Source/MetasoundFrontend/Public/MetasoundFrontendInjectReceiveNodes.h b/Engine/Plugins/Runtime/Metasound/Source/MetasoundFrontend/Public/MetasoundFrontendInjectReceiveNodes.h
new file mode 100644
--- /dev/null
+++ b/Engine/Plugins/Runtime/Metasound/Source/MetasoundFrontend/Public/MetasoundFrontendInjectReceiveNodes.h
@@ -0,0 +1,47 @@
+// Copyright Epic Games, Inc. All Rights Reserved.
+
+#pragma once
+
+#include "CoreMinimal.h"
+#include "MetasoundFrontendGraph.h"
+#include "MetasoundVertex.h"
+
+namespace Metasound
+{
+	namespace Frontend 
+	{
+		/** FReceiveNodeAddressFunction defines the function signature for callbacks
+		 * to create send addresses. Different systems may rely on different data
+		 * in the Metasound Environment to create a unique address.
+		 *
+		 * @param InEnv - The current environment of the graph hosting the receive node.
+		 * @param InVertexKey - The input vertex key of the graph input which will have a 
+		 *                      receive node injected. 
+		 * @param InTypeName - The data type of the receive node. 
+		 *
+		 * @return FSendAddress The resulting transmission address for the receive node. 
+		 */
+		using FReceiveNodeAddressFunction = TFunction<FSendAddress(const FMetasoundEnvironment& InEnv, const FVertexKey& InVertexKey, const FName& InTypeName)>;
+
+		/** Injects a receive node between an graph input and the connected internal nodes.
+		 *
+		 * @param InGraph - The graph to manipulate.
+		 * @parma InAddressPolicy - A function which creates unique transmission addresses for each input.
+		 * @param InInputDestination - The input destination on the graph where a receive node should be injected.
+		 *
+		 * @return True on success, false on failure. 
+		 */
+		METASOUNDFRONTEND_API bool InjectReceiveNode(FFrontendGraph& InGraph, const FReceiveNodeAddressFunction& InAddressPolicy, const FInputDataDestination& InInputDestination);
+
+		/** Injects receive nodes between an graph inputs and the connected internal nodes.
+		 *
+		 * @param InGraph - The graph to manipulate.
+		 * @parma InAddressPolicy - A function which creates unique transmission addresses for each input.
+		 * @param InInputVerticesToSkip - A set of input vertices which should *not* have receive nodes injected.
+		 *
+		 * @return True on success, false on failure. 
+		 */
+		METASOUNDFRONTEND_API bool InjectReceiveNodes(FFrontendGraph& InGraph, const FReceiveNodeAddressFunction& InAddressPolicy, const TSet<FVertexKey>& InInputVerticesToSkip);
+	}
+}
+
diff --git a/Engine/Plugins/Runtime/Metasound/Source/MetasoundFrontend/Public/MetasoundFrontendRegistries.h b/Engine/Plugins/Runtime/Metasound/Source/MetasoundFrontend/Public/MetasoundFrontendRegistries.h
--- a/Engine/Plugins/Runtime/Metasound/Source/MetasoundFrontend/Public/MetasoundFrontendRegistries.h
+++ b/Engine/Plugins/Runtime/Metasound/Source/MetasoundFrontend/Public/MetasoundFrontendRegistries.h
@@ -31,18 +31,6 @@ namespace Metasound
 		FGuid InstanceID;
 		FString VertexName;
 		FLiteral InitParam = FLiteral::CreateInvalid();
-
-		// TODO: Remove Clone(), move variable instead. 
-		FDefaultLiteralNodeConstructorParams Clone() const
-		{
-			return
-			{
-				NodeName,
-				InstanceID,
-				VertexName,
-				InitParam.Clone()
-			};
-		}
 	};
 
 	using FOutputNodeConstructorParams = FDefaultNodeConstructorParams;
@@ -381,6 +369,7 @@ namespace Metasound
 			// Returns true if the class info and class metadata represent the same entry in the node registry.
 			METASOUNDFRONTEND_API bool IsEqual(const FNodeClassInfo& InLHS, const FMetasoundFrontendClassMetadata& InRHS);
 
+			METASOUNDFRONTEND_API FNodeRegistryKey CreateKey(EMetasoundFrontendClassType InType, const FString& InFullClassName, int32 InMajorVersion, int32 InMinorVersion);
 			METASOUNDFRONTEND_API FNodeRegistryKey CreateKey(const FNodeClassMetadata& InNodeMetadata);
 			METASOUNDFRONTEND_API FNodeRegistryKey CreateKey(const FMetasoundFrontendClassMetadata& InNodeMetadata);
 			METASOUNDFRONTEND_API FNodeRegistryKey CreateKey(const FNodeClassInfo& ClassInfo);
diff --git a/Engine/Plugins/Runtime/Metasound/Source/MetasoundFrontend/Public/MetasoundInstanceTransmitter.h b/Engine/Plugins/Runtime/Metasound/Source/MetasoundFrontend/Public/MetasoundInstanceTransmitter.h
--- a/Engine/Plugins/Runtime/Metasound/Source/MetasoundFrontend/Public/MetasoundInstanceTransmitter.h
+++ b/Engine/Plugins/Runtime/Metasound/Source/MetasoundFrontend/Public/MetasoundInstanceTransmitter.h
@@ -38,6 +38,7 @@ namespace Metasound
 			FName TypeName;
 		};
 
+
 		/** Initialization parameters for a FMetasoundInstanceTransmitter. */
 		struct FInitParams
 		{
@@ -60,6 +61,15 @@ namespace Metasound
 
 		};
 
+		/** Returns the MetaSound environment variable name which contains the instance ID. */
+		static FString GetInstanceIDEnvironmentVariableName();
+
+		/** Creates a unique send address using the given MetaSound environment. */
+		static FSendAddress CreateSendAddressFromEnvironment(const FMetasoundEnvironment& InEnvironment, const FVertexKey& InVertexkey, const FName& InTypeName);
+		
+		/** Creates a unique send address using the given InstanceID. */
+		static FSendAddress CreateSendAddressFromInstanceID(uint64 InInstanceID, const FVertexKey& InVertexKey, const FName& InTypeName);
+
 		FMetasoundInstanceTransmitter(const FMetasoundInstanceTransmitter::FInitParams& InInitParams);
 		virtual ~FMetasoundInstanceTransmitter() = default;
 
diff --git a/Engine/Plugins/Runtime/Metasound/Source/MetasoundFrontend/Public/MetasoundReceiveNode.h b/Engine/Plugins/Runtime/Metasound/Source/MetasoundFrontend/Public/MetasoundReceiveNode.h
--- a/Engine/Plugins/Runtime/Metasound/Source/MetasoundFrontend/Public/MetasoundReceiveNode.h
+++ b/Engine/Plugins/Runtime/Metasound/Source/MetasoundFrontend/Public/MetasoundReceiveNode.h
@@ -18,30 +18,14 @@
 
 namespace Metasound
 {
-	struct FReceiveNodeNames
+	namespace ReceiveNodeInfo
 	{
-		static const FString& GetAddressInputName()
-		{
-			static const FString InputName = FString(TEXT("Address"));
-			return InputName;
-		}
-
-		static const FString& GetDefaultDataInputName()
-		{
-			static const FString DefaultDataName = FString(TEXT("Default"));
-			return DefaultDataName;
-		}
-
-		static const FString& GetOutputName()
-		{
-			static const FString OutputName = FString(TEXT("Out"));
-			return OutputName;
-		}
-
-		static FNodeClassName GetClassNameForDataType(const FName& InDataTypeName)
-		{
-			return FNodeClassName{TEXT("Receive"), InDataTypeName, TEXT("")};
-		}
+		METASOUNDFRONTEND_API const FString& GetAddressInputName();
+		METASOUNDFRONTEND_API const FString& GetDefaultDataInputName();
+		METASOUNDFRONTEND_API const FString& GetOutputName();
+		METASOUNDFRONTEND_API FNodeClassName GetClassNameForDataType(const FName& InDataTypeName);
+		METASOUNDFRONTEND_API int32 GetCurrentMajorVersion();
+		METASOUNDFRONTEND_API int32 GetCurrentMinorVersion();
 	};
 
 	template<typename TDataType>
@@ -52,11 +36,11 @@ namespace Metasound
 		{
 			return FVertexInterface(
 				FInputVertexInterface(
-					TInputDataVertexModel<FSendAddress>(FReceiveNodeNames::GetAddressInputName(), FText::GetEmpty()),
-					TInputDataVertexModel<TDataType>(FReceiveNodeNames::GetDefaultDataInputName(), FText::GetEmpty())
+					TInputDataVertexModel<FSendAddress>(ReceiveNodeInfo::GetAddressInputName(), FText::GetEmpty()),
+					TInputDataVertexModel<TDataType>(ReceiveNodeInfo::GetDefaultDataInputName(), FText::GetEmpty())
 				),
 				FOutputVertexInterface(
-					TOutputDataVertexModel<TDataType>(FReceiveNodeNames::GetOutputName(), FText::GetEmpty())
+					TOutputDataVertexModel<TDataType>(ReceiveNodeInfo::GetOutputName(), FText::GetEmpty())
 				)
 			);
 		}
@@ -66,9 +50,9 @@ namespace Metasound
 			auto InitNodeInfo = []() -> FNodeClassMetadata
 			{
 				FNodeClassMetadata Info;
-				Info.ClassName = FReceiveNodeNames::GetClassNameForDataType(GetMetasoundDataTypeName<TDataType>());
-				Info.MajorVersion = 1;
-				Info.MinorVersion = 0;
+				Info.ClassName = ReceiveNodeInfo::GetClassNameForDataType(GetMetasoundDataTypeName<TDataType>());
+				Info.MajorVersion = ReceiveNodeInfo::GetCurrentMajorVersion();
+				Info.MinorVersion = ReceiveNodeInfo::GetCurrentMinorVersion();
 				Info.DisplayName = FText::Format(LOCTEXT("Metasound_ReceiveNodeDisplayNameFormat", "Receive {0}"), FText::FromName(GetMetasoundDataTypeName<TDataType>()));
 				Info.Description = LOCTEXT("Metasound_ReceiveNodeDescription", "Receives data from a send node with the same name.");
 				Info.Author = PluginAuthor;
@@ -111,8 +95,8 @@ namespace Metasound
 				{
 					FDataReferenceCollection Inputs;
 
-					Inputs.AddDataReadReference<TDataType>(FReceiveNodeNames::GetDefaultDataInputName(), DefaultData);
-					Inputs.AddDataReadReference<FSendAddress>(FReceiveNodeNames::GetAddressInputName(), SendAddress);
+					Inputs.AddDataReadReference<TDataType>(ReceiveNodeInfo::GetDefaultDataInputName(), DefaultData);
+					Inputs.AddDataReadReference<FSendAddress>(ReceiveNodeInfo::GetAddressInputName(), SendAddress);
 
 					return Inputs;
 				}
@@ -120,7 +104,7 @@ namespace Metasound
 				virtual FDataReferenceCollection GetOutputs() const override
 				{
 					FDataReferenceCollection Outputs;
-					Outputs.AddDataReadReference<TDataType>(FReceiveNodeNames::GetOutputName(), TDataReadReference<TDataType>(OutputData));
+					Outputs.AddDataReadReference<TDataType>(ReceiveNodeInfo::GetOutputName(), TDataReadReference<TDataType>(OutputData));
 					return Outputs;
 				}
 
@@ -185,15 +169,15 @@ namespace Metasound
 				{
 					TDataReadReference<TDataType> DefaultReadRef = TDataReadReferenceFactory<TDataType>::CreateAny(InParams.OperatorSettings);
 
-					if (InParams.InputDataReferences.ContainsDataReadReference<TDataType>(FReceiveNodeNames::GetDefaultDataInputName()))
+					if (InParams.InputDataReferences.ContainsDataReadReference<TDataType>(ReceiveNodeInfo::GetDefaultDataInputName()))
 					{
-						DefaultReadRef = InParams.InputDataReferences.GetDataReadReference<TDataType>(FReceiveNodeNames::GetDefaultDataInputName());
+						DefaultReadRef = InParams.InputDataReferences.GetDataReadReference<TDataType>(ReceiveNodeInfo::GetDefaultDataInputName());
 					}
 
 					return MakeUnique<TReceiverOperator>(
 						DefaultReadRef,
 						TDataWriteReferenceFactory<TDataType>::CreateAny(InParams.OperatorSettings, *DefaultReadRef),
-						InParams.InputDataReferences.GetDataReadReferenceOrConstruct<FSendAddress>(FReceiveNodeNames::GetAddressInputName()),
+						InParams.InputDataReferences.GetDataReadReferenceOrConstruct<FSendAddress>(ReceiveNodeInfo::GetAddressInputName()),
 						InParams.OperatorSettings
 						);
 				}
@@ -201,7 +185,6 @@ namespace Metasound
 
 		public:
 
-			// TODO: default value of received object.
 			TReceiveNode(const FNodeInitData& InInitData)
 				: FNode(InInitData.InstanceName, InInitData.InstanceID, GetNodeInfo())
 				, Interface(DeclareVertexInterface())
diff --git a/Engine/Plugins/Runtime/Metasound/Source/MetasoundFrontend/Public/MetasoundVariableNode.h b/Engine/Plugins/Runtime/Metasound/Source/MetasoundFrontend/Public/MetasoundVariableNode.h
--- a/Engine/Plugins/Runtime/Metasound/Source/MetasoundFrontend/Public/MetasoundVariableNode.h
+++ b/Engine/Plugins/Runtime/Metasound/Source/MetasoundFrontend/Public/MetasoundVariableNode.h
@@ -12,6 +12,7 @@
 #define LOCTEXT_NAMESPACE "MetasoundFrontend"
 
 
+// TODO: Consider rename to TLiteralNode. 
 namespace Metasound
 {
 	/** A writable variable (if non-const) and a readable output that is only accessible within a given graph. */
@@ -26,12 +27,11 @@ namespace Metasound
 
 			virtual ~TGetVariableOperator() = default;
 
+			/* TODO: No inputs should be return for a "GetVariable"
 			virtual FDataReferenceCollection GetInputs() const override
 			{
-				FDataReferenceCollection Inputs;
-				Inputs.AddDataWriteReference<DataType>(TInputOperator<DataType>::DataReferenceName, TInputOperator<DataType>::InputValue);
-				return Inputs;
 			}
+			*/
 	};
 
 	/** TGetVariableOperatorLiteralFactory creates an input by passing it a literal. */
@@ -93,8 +93,15 @@ namespace Metasound
 
 			static FVertexInterface DeclareVertexInterface(const FVertexKey& InVertexName)
 			{
+				// TODO: Do not need to copy vertex name for these nodes. Should
+				// probably just have them be named "Value" or something generic
+				// to avoid the trouble. Should make similar changes for TInputNode, TOutputNode
+				//
 				return FVertexInterface(
 					FInputVertexInterface(
+						// TODO: Once node versioning is possible, remove this.
+						// Variable nodes do not need writable inputs. They should
+						// constant and renamed to TLiteralNode
 						TInputDataVertexModel<DataType>(InVertexName, FText::GetEmpty())
 					),
 					FOutputVertexInterface(
diff --git a/Engine/Plugins/Runtime/Metasound/Source/MetasoundGenerator/Private/MetasoundGenerator.cpp b/Engine/Plugins/Runtime/Metasound/Source/MetasoundGenerator/Private/MetasoundGenerator.cpp
--- a/Engine/Plugins/Runtime/Metasound/Source/MetasoundGenerator/Private/MetasoundGenerator.cpp
+++ b/Engine/Plugins/Runtime/Metasound/Source/MetasoundGenerator/Private/MetasoundGenerator.cpp
@@ -22,14 +22,13 @@ namespace Metasound
 
 	void FAsyncMetaSoundBuilder::DoWork()
 	{
-		using namespace Frontend;
-		// Create handles for new root graph
-		FConstDocumentHandle NewDocumentHandle = IDocumentController::CreateDocumentHandle(MakeAccessPtr<FConstDocumentAccessPtr>(InitParams.DocumentCopy.AccessPoint, InitParams.DocumentCopy));
-		FConstGraphHandle RootGraph = NewDocumentHandle->GetRootGraph();
-		ensureAlways(RootGraph->IsValid());
-
+		// Create an instance of the new graph
+		FOperatorBuilder OperatorBuilder(FOperatorBuilderSettings::GetDefaultSettings());
+		FBuildGraphParams BuildParams{*InitParams.Graph, InitParams.OperatorSettings, FDataReferenceCollection{}, InitParams.Environment};
 		TArray<IOperatorBuilder::FBuildErrorPtr> BuildErrors;
-		TUniquePtr<IOperator> GraphOperator = RootGraph->BuildOperator(InitParams.OperatorSettings, InitParams.Environment, BuildErrors);
+
+		TUniquePtr<IOperator> GraphOperator = OperatorBuilder.BuildGraphOperator(BuildParams, BuildErrors);
+
 
 		// Log build errors
 		for (const IOperatorBuilder::FBuildErrorPtr& Error : BuildErrors)
diff --git a/Engine/Plugins/Runtime/Metasound/Source/MetasoundGenerator/Public/MetasoundGenerator.h b/Engine/Plugins/Runtime/Metasound/Source/MetasoundGenerator/Public/MetasoundGenerator.h
--- a/Engine/Plugins/Runtime/Metasound/Source/MetasoundGenerator/Public/MetasoundGenerator.h
+++ b/Engine/Plugins/Runtime/Metasound/Source/MetasoundGenerator/Public/MetasoundGenerator.h
@@ -19,7 +19,7 @@ namespace Metasound
 	struct FMetasoundGeneratorInitParams
 	{
 		FOperatorSettings OperatorSettings;
-		FMetasoundFrontendDocument DocumentCopy;
+		TUniquePtr<IGraph> Graph;
 		FMetasoundEnvironment Environment;
 		FString MetaSoundName;
 		TArray<FString> AudioOutputNames;
diff --git a/Engine/Plugins/Runtime/Metasound/Source/MetasoundGraphCore/Private/MetasoundLiteral.cpp b/Engine/Plugins/Runtime/Metasound/Source/MetasoundGraphCore/Private/MetasoundLiteral.cpp
--- a/Engine/Plugins/Runtime/Metasound/Source/MetasoundGraphCore/Private/MetasoundLiteral.cpp
+++ b/Engine/Plugins/Runtime/Metasound/Source/MetasoundGraphCore/Private/MetasoundLiteral.cpp
@@ -122,75 +122,6 @@ namespace Metasound
 				return ELiteralType::Invalid;
 		}
 	}
-
-	FLiteral FLiteral::Clone() const
-	{
-		// TODO: If Clone() is supported after reworking vertices,
-		// wrap the Audio::IProxyData in a object that manages virtual
-		// copy construction.  Then replace this switch with simple
-		// FLiteral copy constructor.
-		switch (GetType())
-		{
-		case ELiteralType::Invalid:
-			return CreateInvalid();
-
-		case ELiteralType::None:
-			return FLiteral(FNone());
-
-		case ELiteralType::Boolean:
-			return FLiteral(Value.Get<bool>());
-
-		case ELiteralType::Float:
-			return FLiteral(Value.Get<float>());
-
-		case ELiteralType::Integer:
-			return FLiteral(Value.Get<int32>());
-
-		case ELiteralType::String:
-			return FLiteral(Value.Get<FString>());
-
-		case ELiteralType::UObjectProxy:
-			if (const Audio::IProxyDataPtr& ProxyPtr = Value.Get<Audio::IProxyDataPtr>())
-			{
-				return FLiteral(ProxyPtr->Clone());
-			}
-			return CreateInvalid();
-
-		case ELiteralType::NoneArray:
-			return FLiteral(Value.Get<TArray<FNone>>());
-
-		case ELiteralType::BooleanArray:
-			return FLiteral(Value.Get<TArray<bool>>());
-
-		case ELiteralType::FloatArray:
-			return FLiteral(Value.Get<TArray<float>>());
-
-		case ELiteralType::IntegerArray:
-			return FLiteral(Value.Get<TArray<int32>>());
-
-		case ELiteralType::StringArray:
-			return FLiteral(Value.Get<TArray<FString>>());
-
-		case ELiteralType::UObjectProxyArray:
-			{
-				TArray<Audio::IProxyDataPtr> ProxyPtrArrayCopy;
-				for (const Audio::IProxyDataPtr& Ptr : Value.Get<TArray<Audio::IProxyDataPtr>>())
-				{
-					if (Ptr.IsValid())
-					{
-						ProxyPtrArrayCopy.Add(Ptr->Clone());
-					}
-				}
-				return FLiteral(MoveTemp(ProxyPtrArrayCopy));
-			}
-			return CreateInvalid();
-
-		default:
-			checkNoEntry();
-			return CreateInvalid();
-		}
-	}
-
 }
 
 FString LexToString(const Metasound::FLiteral& InLiteral)
diff --git a/Engine/Plugins/Runtime/Metasound/Source/MetasoundGraphCore/Public/MetasoundGraph.h b/Engine/Plugins/Runtime/Metasound/Source/MetasoundGraphCore/Public/MetasoundGraph.h
--- a/Engine/Plugins/Runtime/Metasound/Source/MetasoundGraphCore/Public/MetasoundGraph.h
+++ b/Engine/Plugins/Runtime/Metasound/Source/MetasoundGraphCore/Public/MetasoundGraph.h
@@ -72,6 +72,23 @@ namespace Metasound
 			 */
 			bool AddDataEdge(const INode& FromNode, const FVertexKey& FromVertexKey, const INode& ToNode, const FVertexKey& ToVertexKey);
 
+			/** Remove the given data edge. 
+			 *
+			 * @return True on success, false on failure.
+			 */
+			bool RemoveDataEdge(const INode& FromNode, const FVertexKey& FromVertexKey, const INode& ToNode, const FVertexKey& ToVertexKey);
+
+			/** Removes all edges for which that predicate returns true.
+			 *
+			 * @param Predicate - A callable object which accepts an FDataEdge and returns true if
+			 *                    the edge should be removed.
+			 */
+			template<typename PredicateType>
+			void RemoveDataEdgeByPredicate(const PredicateType& Predicate)
+			{
+				Edges.RemoveAllSwap(Predicate);
+			}
+
 			/** Add an input data destination to describe how data provided 
 			 * outside this graph should be routed internally.
 			 *
diff --git a/Engine/Plugins/Runtime/Metasound/Source/MetasoundGraphCore/Public/MetasoundLiteral.h b/Engine/Plugins/Runtime/Metasound/Source/MetasoundGraphCore/Public/MetasoundLiteral.h
--- a/Engine/Plugins/Runtime/Metasound/Source/MetasoundGraphCore/Public/MetasoundLiteral.h
+++ b/Engine/Plugins/Runtime/Metasound/Source/MetasoundGraphCore/Public/MetasoundLiteral.h
@@ -108,10 +108,6 @@ namespace Metasound
 #endif
 		}
 
-		//FLiteral(const FLiteral& Other)
-		//{
-		//	*this = Other.Clone();
-		//}
 		FLiteral(FLiteral&& Other) = default;
 		FLiteral& operator=(FLiteral&& Other) = default;
 
@@ -125,8 +121,6 @@ namespace Metasound
 
 		ELiteralType GetType() const;
 
-		FLiteral Clone() const;
-
 		template<typename ArgType>
 		void Set(ArgType&& Arg)
 		{
diff --git a/Engine/Plugins/Runtime/Metasound/Source/MetasoundGraphCore/Public/MetasoundNodeInterface.h b/Engine/Plugins/Runtime/Metasound/Source/MetasoundGraphCore/Public/MetasoundNodeInterface.h
--- a/Engine/Plugins/Runtime/Metasound/Source/MetasoundGraphCore/Public/MetasoundNodeInterface.h
+++ b/Engine/Plugins/Runtime/Metasound/Source/MetasoundGraphCore/Public/MetasoundNodeInterface.h
@@ -250,7 +250,7 @@ namespace Metasound
 	};
 
 	/** Key type for an FOutputDataSource or FInputDataDestination. */
-	typedef TTuple<const INode*, FString> FNodeDataVertexKey;
+	typedef TTuple<const INode*, FVertexKey> FNodeDataVertexKey;
 
 	/** FOutputDataSourceCollection contains multiple FOutputDataSources mapped 
 	 * by FNodeDataVertexKeys.
