Disable enforced VirtualizationMode for MetaSounds of 'PlayWhenSilent'

diff --git a/Engine/Plugins/Runtime/Metasound/Source/MetasoundEngine/Private/MetasoundSource.cpp b/Engine/Plugins/Runtime/Metasound/Source/MetasoundEngine/Private/MetasoundSource.cpp
--- a/Engine/Plugins/Runtime/Metasound/Source/MetasoundEngine/Private/MetasoundSource.cpp
+++ b/Engine/Plugins/Runtime/Metasound/Source/MetasoundEngine/Private/MetasoundSource.cpp
@@ -229,7 +229,6 @@ ISoundGeneratorPtr UMetaSoundSource::CreateSoundGenerator(const FSoundGeneratorI
 
 	Duration = INDEFINITELY_LOOPING_DURATION;
 	bLooping = true;
-	VirtualizationMode = EVirtualizationMode::PlayWhenSilent;
 
 	SampleRate = InParams.SampleRate;
 	FOperatorSettings InSettings = GetOperatorSettings(static_cast<FSampleRate>(SampleRate));
