Add 'Promote To Output' In MetaSound Editor & clean-up GenerateVariableUniqueName logic to be concise/useful/in line with usual UE-ish default naming

diff --git a/Engine/Plugins/Runtime/Metasound/Source/MetasoundEditor/Private/MetasoundEditorGraphBuilder.cpp b/Engine/Plugins/Runtime/Metasound/Source/MetasoundEditor/Private/MetasoundEditorGraphBuilder.cpp
--- a/Engine/Plugins/Runtime/Metasound/Source/MetasoundEditor/Private/MetasoundEditorGraphBuilder.cpp
+++ b/Engine/Plugins/Runtime/Metasound/Source/MetasoundEditor/Private/MetasoundEditorGraphBuilder.cpp
@@ -257,22 +257,29 @@ namespace Metasound
 			});
 		}
 
-		FText FGraphBuilder::GenerateUniqueNameByFilter(const UObject& InMetaSound, const FText& InBaseText, TUniqueFunction<bool(const Frontend::FConstGraphHandle&, const FText&)> InIsValidNameFilter)
+		FText FGraphBuilder::GenerateUniqueNameByFilter(const UObject& InMetaSound, const FText& InBaseText, TFunctionRef<bool(const Frontend::FConstGraphHandle&, const FText&)> InNameIsValidFilter)
 		{
 			const FMetasoundAssetBase* MetaSoundAsset = IMetasoundUObjectRegistry::Get().GetObjectAsAssetBase(&InMetaSound);
 			check(MetaSoundAsset);
 
 			const Frontend::FConstGraphHandle GraphHandle = MetaSoundAsset->GetRootGraphHandle();
 
-			static const FText VariableUniqueNameFormat = LOCTEXT("VariableUniqueNameFormat", "{0}{1}");
+			if (InNameIsValidFilter(GraphHandle, InBaseText))
+			{
+				return InBaseText;
+			}
 
-			int32 i = 1;
-			FText NewName = FText::Format(VariableUniqueNameFormat, InBaseText, i);
+			FNumberFormattingOptions NumberFormatOptions;
+			NumberFormatOptions.MinimumIntegralDigits = 2;
 
-			while (!InIsValidNameFilter(GraphHandle, NewName))
+			int32 i = 1;
+			FText NewName;
+			do
 			{
-				NewName = FText::Format(VariableUniqueNameFormat, InBaseText, ++i);
+				static const FText VariableUniqueNameFormat = LOCTEXT("GenerateUniqueVariableNameFormat", "{0} {1}");
+				NewName = FText::Format(VariableUniqueNameFormat, InBaseText, FText::AsNumber(++i, &NumberFormatOptions));
 			}
+			while (!InNameIsValidFilter(GraphHandle, NewName));
 
 			return NewName;
 		}
diff --git a/Engine/Plugins/Runtime/Metasound/Source/MetasoundEditor/Private/MetasoundEditorGraphBuilder.h b/Engine/Plugins/Runtime/Metasound/Source/MetasoundEditor/Private/MetasoundEditorGraphBuilder.h
--- a/Engine/Plugins/Runtime/Metasound/Source/MetasoundEditor/Private/MetasoundEditorGraphBuilder.h
+++ b/Engine/Plugins/Runtime/Metasound/Source/MetasoundEditor/Private/MetasoundEditorGraphBuilder.h
@@ -117,7 +117,7 @@ namespace Metasound
 			static FText GenerateUniqueOutputDisplayName(const UObject& InMetaSound, const FText* InBaseName = nullptr);
 
 			// Generates a unique name by retrieving the provided MetaSound's associated GraphHandle and filtering by the provided function.
-			static FText GenerateUniqueNameByFilter(const UObject& InMetaSound, const FText& InBaseText, TUniqueFunction<bool(const Frontend::FConstGraphHandle&, const FText&)> InIsValidNameFilter);
+			static FText GenerateUniqueNameByFilter(const UObject& InMetaSound, const FText& InBaseText, TFunctionRef<bool(const Frontend::FConstGraphHandle&, const FText&)> InNameIsValidFilter);
 
 			static TArray<FString> GetDataTypeNameCategories(const FName& InDataTypeName);
 
diff --git a/Engine/Plugins/Runtime/Metasound/Source/MetasoundEditor/Private/MetasoundEditorGraphSchema.cpp b/Engine/Plugins/Runtime/Metasound/Source/MetasoundEditor/Private/MetasoundEditorGraphSchema.cpp
--- a/Engine/Plugins/Runtime/Metasound/Source/MetasoundEditor/Private/MetasoundEditorGraphSchema.cpp
+++ b/Engine/Plugins/Runtime/Metasound/Source/MetasoundEditor/Private/MetasoundEditorGraphSchema.cpp
@@ -259,17 +259,14 @@ UEdGraphNode* FMetasoundGraphSchemaAction_PromoteToInput::PerformAction(UEdGraph
 	UObject& ParentMetasound = MetasoundGraph->GetMetasoundChecked();
 	ParentMetasound.Modify();
 
-	const FText& InputName = InputHandle->GetOwningNode()->GetDisplayName();
-	const FText& InputNodeName = InputHandle->GetDisplayName();
-
 	FMetasoundFrontendLiteral DefaultValue;
 	FGraphBuilder::GetPinLiteral(*FromPin, DefaultValue);
 
 	FNodeHandle NodeHandle = FGraphBuilder::AddInputNodeHandle(ParentMetasound, InputHandle->GetDataType(), FText::GetEmpty(), &DefaultValue);
 	if (ensure(NodeHandle->IsValid()))
 	{
-		const FText InputNameBase = FText::Format(LOCTEXT("PromoteInputNameBaseFormat", "{0} {1}"), InputName, InputNodeName);
-		const FText NewNodeName = FGraphBuilder::GenerateUniqueInputDisplayName(ParentMetasound, &InputNameBase);
+		const FText& InputName = InputHandle->GetDisplayName();
+		const FText NewNodeName = FGraphBuilder::GenerateUniqueInputDisplayName(ParentMetasound, &InputName);
 		NodeHandle->SetDisplayName(NewNodeName);
 		UMetasoundEditorGraphInput* Input = MetasoundGraph->FindOrAddInput(NodeHandle);
 		if (ensure(Input))
@@ -336,6 +333,58 @@ UEdGraphNode* FMetasoundGraphSchemaAction_NewOutput::PerformAction(UEdGraph* Par
 	return nullptr;
 }
 
+FMetasoundGraphSchemaAction_PromoteToOutput::FMetasoundGraphSchemaAction_PromoteToOutput(FText InNodeCategory, FText InDisplayName, FText InToolTip, const int32 InGrouping)
+	: FEdGraphSchemaAction(MoveTemp(InNodeCategory), MoveTemp(InDisplayName), MoveTemp(InToolTip), InGrouping)
+{
+}
+
+UEdGraphNode* FMetasoundGraphSchemaAction_PromoteToOutput::PerformAction(UEdGraph* ParentGraph, UEdGraphPin* FromPin, const FVector2D InLocation, bool bSelectNewNode /* = true */)
+{
+	using namespace Metasound::Editor;
+	using namespace Metasound::Frontend;
+
+	FOutputHandle OutputHandle = FGraphBuilder::GetOutputHandleFromPin(FromPin);
+	if (!ensure(OutputHandle->IsValid()))
+	{
+		return nullptr;
+	}
+
+	const FScopedTransaction Transaction(LOCTEXT("PromoteNodeOutputToGraphOutput", "Promote MetaSound Node Output to Graph Output"));
+	UMetasoundEditorGraph* MetasoundGraph = CastChecked<UMetasoundEditorGraph>(ParentGraph);
+	UObject& ParentMetasound = MetasoundGraph->GetMetasoundChecked();
+	ParentMetasound.Modify();
+
+	FNodeHandle NodeHandle = FGraphBuilder::AddOutputNodeHandle(ParentMetasound, OutputHandle->GetDataType(), FText::GetEmpty());
+	if (ensure(NodeHandle->IsValid()))
+	{
+		const FText& OutputName = OutputHandle->GetDisplayName();
+		const FText NewNodeName = FGraphBuilder::GenerateUniqueOutputDisplayName(ParentMetasound, &OutputName);
+		NodeHandle->SetDisplayName(NewNodeName);
+		UMetasoundEditorGraphOutput* Output = MetasoundGraph->FindOrAddOutput(NodeHandle);
+		if (ensure(Output))
+		{
+			if (UMetasoundEditorGraphOutputNode* NewGraphNode = FGraphBuilder::AddOutputNode(ParentMetasound, NodeHandle, InLocation))
+			{
+				UEdGraphNode* EdGraphNode = CastChecked<UEdGraphNode>(NewGraphNode);
+
+				if (ensure(SchemaPrivate::TryConnectNewNodeToPin(*EdGraphNode, FromPin)))
+				{
+					TSharedPtr<FEditor> MetasoundEditor = FGraphBuilder::GetEditorForGraph(*ParentGraph);
+					if (MetasoundEditor.IsValid())
+					{
+						MetasoundEditor->OnOutputNameChanged(NewGraphNode->GetNodeID());
+					}
+
+					FGraphBuilder::RegisterGraphWithFrontend(ParentMetasound);
+					return EdGraphNode;
+				}
+			}
+		}
+	}
+
+	return nullptr;
+}
+
 UEdGraphNode* FMetasoundGraphSchemaAction_NewFromSelected::PerformAction(UEdGraph* ParentGraph, UEdGraphPin* FromPin, const FVector2D Location, bool bSelectNewNode/* = true*/)
 {
 	// TODO: Implement
@@ -492,6 +541,14 @@ void UMetasoundEditorGraphSchema::GetGraphContextActions(FGraphContextMenuBuilde
 				});
 				return bHasInputOfType;
 			});
+
+			TSharedPtr<FMetasoundGraphSchemaAction_PromoteToOutput> NewNodeAction = MakeShared<FMetasoundGraphSchemaAction_PromoteToOutput>(
+				SchemaPrivate::GetContextGroupDisplayName(SchemaPrivate::EPrimaryContextGroup::Inputs),
+				LOCTEXT("PromoteToOutputName", "Promote To Graph Output"),
+				LOCTEXT("PromoteToOutputTooltip", "Promotes node output to graph output"),
+				static_cast<int32>(SchemaPrivate::EPrimaryContextGroup::Inputs));
+
+			static_cast<FGraphActionMenuBuilder&>(ContextMenuBuilder).AddAction(NewNodeAction);
 		}
 	}
 	else
diff --git a/Engine/Plugins/Runtime/Metasound/Source/MetasoundEditor/Public/MetasoundEditorGraphSchema.h b/Engine/Plugins/Runtime/Metasound/Source/MetasoundEditor/Public/MetasoundEditorGraphSchema.h
--- a/Engine/Plugins/Runtime/Metasound/Source/MetasoundEditor/Public/MetasoundEditorGraphSchema.h
+++ b/Engine/Plugins/Runtime/Metasound/Source/MetasoundEditor/Public/MetasoundEditorGraphSchema.h
@@ -58,7 +58,7 @@ struct METASOUNDEDITOR_API FMetasoundGraphSchemaAction_NewInput : public FEdGrap
 	//~ End FEdGraphSchemaAction Interface
 };
 
-/** Action to promote a literal to a graph input */
+/** Promotes an input to a graph input, using its respective literal value as the default value */
 USTRUCT()
 struct METASOUNDEDITOR_API FMetasoundGraphSchemaAction_PromoteToInput : public FEdGraphSchemaAction
 {
@@ -75,7 +75,7 @@ struct METASOUNDEDITOR_API FMetasoundGraphSchemaAction_PromoteToInput : public F
 	//~ End FEdGraphSchemaAction Interface
 };
 
-/** Action to add an output to the graph */
+/** Adds an output to the graph */
 USTRUCT()
 struct METASOUNDEDITOR_API FMetasoundGraphSchemaAction_NewOutput : public FEdGraphSchemaAction
 {
@@ -95,6 +95,23 @@ struct METASOUNDEDITOR_API FMetasoundGraphSchemaAction_NewOutput : public FEdGra
 	//~ End FEdGraphSchemaAction Interface
 };
 
+/** Promotes a node output to a graph output */
+USTRUCT()
+struct METASOUNDEDITOR_API FMetasoundGraphSchemaAction_PromoteToOutput : public FEdGraphSchemaAction
+{
+	GENERATED_USTRUCT_BODY();
+
+	FMetasoundGraphSchemaAction_PromoteToOutput()
+		: FEdGraphSchemaAction()
+	{}
+
+	FMetasoundGraphSchemaAction_PromoteToOutput(FText InNodeCategory, FText InDisplayName, FText InToolTip, const int32 InGrouping);
+
+	//~ Begin FEdGraphSchemaAction Interface
+	virtual UEdGraphNode* PerformAction(UEdGraph* ParentGraph, UEdGraphPin* FromPin, const FVector2D Location, bool bSelectNewNode = true) override;
+	//~ End FEdGraphSchemaAction Interface
+};
+
 /** Action to add a node to the graph */
 USTRUCT()
 struct METASOUNDEDITOR_API FMetasoundGraphSchemaAction_NewNode : public FEdGraphSchemaAction
