- Move check to ensure as stopgap fixing redirectors not crashing editor
- Fix for ensure when promoting MetaSound object array to input
#rnx

diff --git a/Engine/Plugins/Runtime/Metasound/Source/MetasoundEditor/Private/MetasoundEditorGraphBuilder.cpp b/Engine/Plugins/Runtime/Metasound/Source/MetasoundEditor/Private/MetasoundEditorGraphBuilder.cpp
--- a/Engine/Plugins/Runtime/Metasound/Source/MetasoundEditor/Private/MetasoundEditorGraphBuilder.cpp
+++ b/Engine/Plugins/Runtime/Metasound/Source/MetasoundEditor/Private/MetasoundEditorGraphBuilder.cpp
@@ -1056,7 +1056,10 @@ namespace Metasound
 		void FGraphBuilder::UnregisterGraphWithFrontend(UObject& InMetaSound)
 		{
 			FMetasoundAssetBase* MetaSoundAsset = IMetasoundUObjectRegistry::Get().GetObjectAsAssetBase(&InMetaSound);
-			check(MetaSoundAsset);
+			if (!ensure(MetaSoundAsset))
+			{
+				return;
+			}
 
 			MetaSoundAsset->UnregisterGraphWithFrontend();
 
diff --git a/Engine/Plugins/Runtime/Metasound/Source/MetasoundFrontend/Private/MetasoundFrontendLiteral.cpp b/Engine/Plugins/Runtime/Metasound/Source/MetasoundFrontend/Private/MetasoundFrontendLiteral.cpp
--- a/Engine/Plugins/Runtime/Metasound/Source/MetasoundFrontend/Private/MetasoundFrontendLiteral.cpp
+++ b/Engine/Plugins/Runtime/Metasound/Source/MetasoundFrontend/Private/MetasoundFrontendLiteral.cpp
@@ -468,7 +468,7 @@ bool FMetasoundFrontendLiteral::TryGet(UObject*& OutValue) const
 
 bool FMetasoundFrontendLiteral::TryGet(TArray<UObject*>& OutValue) const
 {
-	if (Type == EMetasoundFrontendLiteralType::UObject)
+	if (Type == EMetasoundFrontendLiteralType::UObjectArray)
 	{
 		OutValue = AsUObject;
 		return true;
