Submitting for Phil: Fix for checks being optimized away causing injected receive nodes to not be connected properly
#rnx

diff --git a/Engine/Plugins/Runtime/Metasound/Source/MetasoundFrontend/Private/MetasoundFrontendInjectReceiveNodes.cpp b/Engine/Plugins/Runtime/Metasound/Source/MetasoundFrontend/Private/MetasoundFrontendInjectReceiveNodes.cpp
--- a/Engine/Plugins/Runtime/Metasound/Source/MetasoundFrontend/Private/MetasoundFrontendInjectReceiveNodes.cpp
+++ b/Engine/Plugins/Runtime/Metasound/Source/MetasoundFrontend/Private/MetasoundFrontendInjectReceiveNodes.cpp
@@ -162,8 +162,10 @@ namespace Metasound
 			// Add receive node and value node to graph.
 			InGraph.AddNode(ReceiveNodeID, ReceiveNode);
 			InGraph.AddNode(AddressNodeID, AddressNode);
-			check(InGraph.AddDataEdge(*AddressNode, FAddressOperator::GetAddressVertexKey(), *ReceiveNode, ReceiveNodeInfo::GetAddressInputName()));
-			
+
+			bool bDataEdgeAdded = InGraph.AddDataEdge(*AddressNode, FAddressOperator::GetAddressVertexKey(), *ReceiveNode, ReceiveNodeInfo::GetAddressInputName());
+			ensureAlways(bDataEdgeAdded);
+
 			auto IsEdgeConnectedToCurrentInput = [&InputDestination](const FDataEdge& InEdge)
 			{
 				return (InEdge.From.Node == InputDestination.Node) && (InEdge.From.Vertex.GetVertexName() == InputDestination.Vertex.GetVertexName());
@@ -183,7 +185,8 @@ namespace Metasound
 			}
 
 			// Connect input node to receive node.
-			check(InGraph.AddDataEdge(*InputDestination.Node, VertexKey, *ReceiveNode, ReceiveNodeInfo::GetDefaultDataInputName()));
+			bDataEdgeAdded = InGraph.AddDataEdge(*InputDestination.Node, VertexKey, *ReceiveNode, ReceiveNodeInfo::GetDefaultDataInputName());
+			ensureAlways(bDataEdgeAdded);
 
 			return true;
 		}
