Add Compressor Node to Metasounds

diff --git a/Engine/Plugins/Runtime/Metasound/Source/MetasoundStandardNodes/Private/MetasoundCompressorNode.cpp b/Engine/Plugins/Runtime/Metasound/Source/MetasoundStandardNodes/Private/MetasoundCompressorNode.cpp
new file mode 100644
--- /dev/null
+++ b/Engine/Plugins/Runtime/Metasound/Source/MetasoundStandardNodes/Private/MetasoundCompressorNode.cpp
@@ -0,0 +1,339 @@
+// Copyright Epic Games, Inc. All Rights Reserved.
+
+#include "Internationalization/Text.h"
+#include "MetasoundExecutableOperator.h"
+#include "MetasoundEnumRegistrationMacro.h"
+#include "MetasoundNodeRegistrationMacro.h"
+#include "MetasoundPrimitives.h"
+#include "MetasoundStandardNodesNames.h"
+#include "MetasoundAudioBuffer.h"
+#include "MetasoundStandardNodesCategories.h"
+#include "MetasoundFacade.h"
+#include "MetasoundParamHelper.h"
+#include "DSP/Dsp.h"
+#include "DSP/DynamicsProcessor.h"
+
+#define LOCTEXT_NAMESPACE "MetasoundStandardNodes_CompressorNode"
+
+namespace Metasound
+{
+	/* Mid-Side Encoder */
+	namespace CompressorVertexNames
+	{
+		METASOUND_PARAM(InputAudio, "Audio", "Incoming audio signal to compress.");
+		METASOUND_PARAM(InputRatio, "Ratio", "Amount of gain reduction. 1 = no reduction, higher = more reduction.");
+		METASOUND_PARAM(InputThreshold, "Threshold", "Amplitude threshold (dB) above which gain will be reduced.");
+		METASOUND_PARAM(InputAttackTime, "Attack Time", "How long it takes for audio above the threshold to reach its compressed volume level.");
+		METASOUND_PARAM(InputReleaseTime, "Release Time", "How long it takes for audio below the threshold to return to its original volume level.");
+		METASOUND_PARAM(InputKnee, "Knee", "How hard or soft the gain reduction blends from no gain reduction to gain reduction. 0 dB = no blending.");
+		METASOUND_PARAM(InputSidechain, "Sidechain", "(Optional) external audio signal to control the compressor with. If empty, uses the input audio signal.");
+		METASOUND_PARAM(InputEnvelopeMode, "Envelope Mode", "The envelope-following method the compressor will use for gain detection.");
+		METASOUND_PARAM(InputIsAnalog, "Analog Mode", "Enable Analog Mode for the compressor's envelope follower.");
+		METASOUND_PARAM(InputWetDryMix, "Wet/Dry", "Ratio between the processed/wet signal and the unprocessed/dry signal. 0 is full dry, 1 is full wet, and 0.5 is 50/50.");
+
+		METASOUND_PARAM(OutputAudio, "Audio", "The output audio signal.");
+		METASOUND_PARAM(OutputEnvelope, "Gain Envelope", "The compressor's gain being applied to the signal.");
+	}
+
+	enum class EEnvelopePeakMode
+	{
+		Peak = 0,
+		MeanSquared,
+		RootMeanSquared,
+	};
+
+	DECLARE_METASOUND_ENUM(EEnvelopePeakMode, EEnvelopePeakMode::Peak, METASOUNDSTANDARDNODES_API,
+	FEnumEnvelopePeakMode, FEnumEnvelopePeakModeInfo, FEnvelopePeakModeReadRef, FEnumEnvelopePeakModeWriteRef);
+
+	DEFINE_METASOUND_ENUM_BEGIN(EEnvelopePeakMode, FEnumEnvelopePeakMode, "EnvelopePeakMode")
+		DEFINE_METASOUND_ENUM_ENTRY(EEnvelopePeakMode::MeanSquared, LOCTEXT("EnvelopePeakModeMSDescription", "MS"), LOCTEXT("EnvelopePeakModeMSDescriptionTT", "Envelope follows a running Mean Squared of the audio signal.")),
+		DEFINE_METASOUND_ENUM_ENTRY(EEnvelopePeakMode::RootMeanSquared, LOCTEXT("EnvelopePeakModeRMSDescription", "RMS"), LOCTEXT("EnvelopePeakModeRMSDescriptionTT", "Envelope follows a running Root Mean Squared of the audio signal.")),
+		DEFINE_METASOUND_ENUM_ENTRY(EEnvelopePeakMode::Peak, LOCTEXT("EnvelopePeakModePeakDescription", "Peak"), LOCTEXT("EnvelopePeakModePeakDescriptionTT", "Envelope follows the peaks in the audio signal.")),
+	DEFINE_METASOUND_ENUM_END()
+
+		// Operator Class
+	class FCompressorOperator : public TExecutableOperator<FCompressorOperator>
+	{
+	public:
+
+		FCompressorOperator(const FOperatorSettings& InSettings,
+			const FAudioBufferReadRef& InAudio,
+			const FFloatReadRef& InRatio,
+			const FFloatReadRef& InThreshold,
+			const FTimeReadRef& InAttackTime,
+			const FTimeReadRef& InReleaseTime,
+			const FFloatReadRef& InKnee,
+			const bool& bInUseSidechain,
+			const FAudioBufferReadRef& InSidechain,
+			const FEnvelopePeakModeReadRef& InEnvelopeMode,
+			const FBoolReadRef& bInIsAnalog,
+			const FFloatReadRef& InWetDryMix)
+			: AudioInput(InAudio)
+			, RatioInput(InRatio)
+			, ThresholdInput(InThreshold)
+			, AttackTimeInput(InAttackTime)
+			, ReleaseTimeInput(InReleaseTime)
+			, KneeInput(InKnee)
+			, SidechainInput(InSidechain)
+			, EnvelopeModeInput(InEnvelopeMode)
+			, bIsAnalogInput(bInIsAnalog)
+			, WetDryMixInput(InWetDryMix)
+			, AudioOutput(FAudioBufferWriteRef::CreateNew(InSettings))
+			, EnvelopeOutput(FAudioBufferWriteRef::CreateNew(InSettings))
+			, Compressor()
+			, bUseSidechain(bInUseSidechain)
+			, PrevRatio(FMath::Max(*InRatio, 1.0f))
+			, PrevThreshold(*InThreshold)
+			, PrevAttackTime(FMath::Max(FTime::ToMilliseconds(*InAttackTime), 0.0))
+			, PrevReleaseTime(FMath::Max(FTime::ToMilliseconds(*InReleaseTime), 0.0))
+			, PrevKnee(*InKnee)
+			, PrevEnvMode(*InEnvelopeMode)
+			, PrevIsAnalog(*bInIsAnalog)
+		{
+			Compressor.Init(InSettings.GetSampleRate(), 1);
+			Compressor.SetKeyNumChannels(1);
+			Compressor.SetRatio(PrevRatio);
+			Compressor.SetThreshold(*ThresholdInput);
+			Compressor.SetAttackTime(PrevAttackTime);
+			Compressor.SetReleaseTime(PrevReleaseTime);
+			Compressor.SetKneeBandwidth(*KneeInput);
+
+			switch (*EnvelopeModeInput)
+			{
+			default:
+			case EEnvelopePeakMode::MeanSquared:
+				Compressor.SetPeakMode(Audio::EPeakMode::MeanSquared);
+				break;
+			case EEnvelopePeakMode::RootMeanSquared:
+				Compressor.SetPeakMode(Audio::EPeakMode::RootMeanSquared);
+				break;
+			case EEnvelopePeakMode::Peak:
+				Compressor.SetPeakMode(Audio::EPeakMode::Peak);
+				break;
+			}
+
+		}
+
+		static const FNodeClassMetadata& GetNodeInfo()
+		{
+			auto CreateNodeClassMetadata = []() -> FNodeClassMetadata
+			{
+				FVertexInterface NodeInterface = DeclareVertexInterface();
+
+				FNodeClassMetadata Metadata
+				{
+					FNodeClassName{StandardNodes::Namespace, TEXT("Compressor"), StandardNodes::AudioVariant},
+					1, // Major Version
+					0, // Minor Version
+					LOCTEXT("CompressorDisplayName", "Compressor"),
+					LOCTEXT("CompressorDesc", "Lowers the dynamic range of a signal."),
+					PluginAuthor,
+					PluginNodeMissingPrompt,
+					NodeInterface,
+					{StandardNodes::Dynamics},
+					{TEXT("Compressor")},
+					FNodeDisplayStyle{}
+				};
+
+				return Metadata;
+			};
+
+			static const FNodeClassMetadata Metadata = CreateNodeClassMetadata();
+			return Metadata;
+		}
+
+		static const FVertexInterface& DeclareVertexInterface()
+		{
+			using namespace CompressorVertexNames;
+
+			static const FVertexInterface Interface(
+				FInputVertexInterface(
+					TInputDataVertexModel<FAudioBuffer>(METASOUND_GET_PARAM_NAME_AND_TT(InputAudio)),
+					TInputDataVertexModel<float>(METASOUND_GET_PARAM_NAME_AND_TT(InputRatio), 1.5f),
+					TInputDataVertexModel<float>(METASOUND_GET_PARAM_NAME_AND_TT(InputThreshold), -6.0f),
+					TInputDataVertexModel<FTime>(METASOUND_GET_PARAM_NAME_AND_TT(InputAttackTime), 0.01f),
+					TInputDataVertexModel<FTime>(METASOUND_GET_PARAM_NAME_AND_TT(InputReleaseTime), 0.1f),
+					TInputDataVertexModel<float>(METASOUND_GET_PARAM_NAME_AND_TT(InputKnee), 10.0f),
+					TInputDataVertexModel<FAudioBuffer>(METASOUND_GET_PARAM_NAME_AND_TT(InputSidechain)),
+					TInputDataVertexModel<FEnumEnvelopePeakMode>(METASOUND_GET_PARAM_NAME_AND_TT(InputEnvelopeMode)),
+					TInputDataVertexModel<bool>(METASOUND_GET_PARAM_NAME_AND_TT(InputIsAnalog), true),
+					TInputDataVertexModel<float>(METASOUND_GET_PARAM_NAME_AND_TT(InputWetDryMix), 1.0f)
+				),
+				FOutputVertexInterface(
+					TOutputDataVertexModel<FAudioBuffer>(METASOUND_GET_PARAM_NAME_AND_TT(OutputAudio)),
+					TOutputDataVertexModel<FAudioBuffer>(METASOUND_GET_PARAM_NAME_AND_TT(OutputEnvelope))
+				)
+			);
+
+			return Interface;
+		}
+
+		virtual FDataReferenceCollection GetInputs() const override
+		{
+			using namespace CompressorVertexNames;
+
+			FDataReferenceCollection InputDataReferences;
+
+			InputDataReferences.AddDataReadReference(METASOUND_GET_PARAM_NAME(InputAudio), AudioInput);
+			InputDataReferences.AddDataReadReference(METASOUND_GET_PARAM_NAME(InputRatio), RatioInput);
+			InputDataReferences.AddDataReadReference(METASOUND_GET_PARAM_NAME(InputThreshold), ThresholdInput);
+			InputDataReferences.AddDataReadReference(METASOUND_GET_PARAM_NAME(InputAttackTime), AttackTimeInput);
+			InputDataReferences.AddDataReadReference(METASOUND_GET_PARAM_NAME(InputReleaseTime), ReleaseTimeInput);
+			InputDataReferences.AddDataReadReference(METASOUND_GET_PARAM_NAME(InputKnee), KneeInput);
+			InputDataReferences.AddDataReadReference(METASOUND_GET_PARAM_NAME(InputSidechain), SidechainInput);
+			InputDataReferences.AddDataReadReference(METASOUND_GET_PARAM_NAME(InputEnvelopeMode), EnvelopeModeInput);
+			InputDataReferences.AddDataReadReference(METASOUND_GET_PARAM_NAME(InputIsAnalog), bIsAnalogInput);
+			InputDataReferences.AddDataReadReference(METASOUND_GET_PARAM_NAME(InputWetDryMix), WetDryMixInput);
+
+			return InputDataReferences;
+
+		}
+
+		virtual FDataReferenceCollection GetOutputs() const override
+		{
+			using namespace CompressorVertexNames;
+
+			FDataReferenceCollection OutputDataReferences;
+
+			OutputDataReferences.AddDataReadReference(METASOUND_GET_PARAM_NAME(OutputAudio), AudioOutput);
+			OutputDataReferences.AddDataReadReference(METASOUND_GET_PARAM_NAME(OutputEnvelope), EnvelopeOutput);
+
+			return OutputDataReferences;
+		}
+
+		static TUniquePtr<IOperator> CreateOperator(const FCreateOperatorParams& InParams, FBuildErrorArray& OutErrors)
+		{
+			using namespace CompressorVertexNames;
+			const FDataReferenceCollection& Inputs = InParams.InputDataReferences;
+			const FInputVertexInterface& InputInterface = DeclareVertexInterface().GetInputInterface();
+
+			FAudioBufferReadRef AudioIn = Inputs.GetDataReadReferenceOrConstruct<FAudioBuffer>(METASOUND_GET_PARAM_NAME(InputAudio), InParams.OperatorSettings);
+			FFloatReadRef RatioIn = Inputs.GetDataReadReferenceOrConstructWithVertexDefault<float>(InputInterface, METASOUND_GET_PARAM_NAME(InputRatio), InParams.OperatorSettings);
+			FFloatReadRef ThresholdIn = Inputs.GetDataReadReferenceOrConstructWithVertexDefault<float>(InputInterface, METASOUND_GET_PARAM_NAME(InputThreshold), InParams.OperatorSettings);
+			FTimeReadRef AttackTimeIn = Inputs.GetDataReadReferenceOrConstructWithVertexDefault<FTime>(InputInterface, METASOUND_GET_PARAM_NAME(InputAttackTime), InParams.OperatorSettings);
+			FTimeReadRef ReleaseTimeIn = Inputs.GetDataReadReferenceOrConstructWithVertexDefault<FTime>(InputInterface, METASOUND_GET_PARAM_NAME(InputReleaseTime), InParams.OperatorSettings);
+			FFloatReadRef KneeIn = Inputs.GetDataReadReferenceOrConstructWithVertexDefault<float>(InputInterface, METASOUND_GET_PARAM_NAME(InputKnee), InParams.OperatorSettings);
+			FAudioBufferReadRef SidechainIn = Inputs.GetDataReadReferenceOrConstruct<FAudioBuffer>(METASOUND_GET_PARAM_NAME(InputSidechain), InParams.OperatorSettings);
+			FEnvelopePeakModeReadRef EnvelopeModeIn = Inputs.GetDataReadReferenceOrConstruct<FEnumEnvelopePeakMode>(METASOUND_GET_PARAM_NAME(InputEnvelopeMode));
+			FBoolReadRef bIsAnalogIn = Inputs.GetDataReadReferenceOrConstructWithVertexDefault<bool>(InputInterface, METASOUND_GET_PARAM_NAME(InputIsAnalog), InParams.OperatorSettings);
+			FFloatReadRef WetDryMixIn = Inputs.GetDataReadReferenceOrConstructWithVertexDefault<float>(InputInterface, METASOUND_GET_PARAM_NAME(InputWetDryMix), InParams.OperatorSettings);
+
+			bool bIsSidechainConnected = Inputs.ContainsDataReadReference<FAudioBuffer>(METASOUND_GET_PARAM_NAME(InputSidechain));
+
+			return MakeUnique<FCompressorOperator>(InParams.OperatorSettings, AudioIn, RatioIn, ThresholdIn, AttackTimeIn, ReleaseTimeIn, KneeIn, bIsSidechainConnected, SidechainIn, EnvelopeModeIn, bIsAnalogIn, WetDryMixIn);
+		}
+
+		void Execute()
+		{
+			/* Update parameters */
+			
+			// For a compressor, ratio values should be 1 or greater
+			Compressor.SetRatio(FMath::Max(*RatioInput, 1.0f));
+
+			Compressor.SetThreshold(*ThresholdInput);
+			
+			// Attack time cannot be negative
+			double CurrAttack = FMath::Max(FTime::ToMilliseconds(*AttackTimeInput), 0.0f);
+			Compressor.SetAttackTime(CurrAttack);
+
+			// Release time cannot be negative
+			double CurrRelease = FMath::Max(FTime::ToMilliseconds(*ReleaseTimeInput), 0.0f);
+			if (!FMath::IsNearlyEqual(CurrRelease, PrevReleaseTime))
+			{
+				Compressor.SetReleaseTime(CurrRelease);
+				PrevReleaseTime = CurrRelease;
+			}
+			if (!FMath::IsNearlyEqual(*KneeInput, PrevKnee))
+			{
+				Compressor.SetKneeBandwidth(*KneeInput);
+				PrevKnee = *KneeInput;
+			}
+			if (PrevEnvMode != *EnvelopeModeInput)
+			{
+				switch (*EnvelopeModeInput)
+				{
+				default:
+				case EEnvelopePeakMode::MeanSquared:
+					Compressor.SetPeakMode(Audio::EPeakMode::MeanSquared);
+					break;
+
+				case EEnvelopePeakMode::RootMeanSquared:
+					Compressor.SetPeakMode(Audio::EPeakMode::RootMeanSquared);
+					break;
+
+				case EEnvelopePeakMode::Peak:
+					Compressor.SetPeakMode(Audio::EPeakMode::Peak);
+					break;
+				}
+
+				PrevEnvMode = *EnvelopeModeInput;
+			}
+
+
+			if (bUseSidechain)
+			{
+				Compressor.ProcessAudio(AudioInput->GetData(), AudioInput->Num(), AudioOutput->GetData(), EnvelopeOutput->GetData(), SidechainInput->GetData());
+			}
+			else
+			{
+				Compressor.ProcessAudio(AudioInput->GetData(), AudioInput->Num(), AudioOutput->GetData(), EnvelopeOutput->GetData());
+			}
+
+			// Calculate Wet/Dry mix
+			float NewWetDryMix = FMath::Clamp(*WetDryMixInput, 0.0f, 1.0f);
+			// Wet signal
+			Audio::ArrayMultiplyByConstantInPlace(*AudioOutput, NewWetDryMix);
+			// Add Dry signal
+			Audio::ArrayMultiplyAddInPlace(*AudioInput, 1.0f - NewWetDryMix, *AudioOutput);
+			
+		}
+
+	private:
+		// Audio input and output
+		FAudioBufferReadRef AudioInput;
+		FFloatReadRef RatioInput;
+		FFloatReadRef ThresholdInput;
+		FTimeReadRef AttackTimeInput;
+		FTimeReadRef ReleaseTimeInput;
+		FFloatReadRef KneeInput;
+		FAudioBufferReadRef SidechainInput;
+		FEnvelopePeakModeReadRef EnvelopeModeInput;
+		FBoolReadRef bIsAnalogInput;
+		FFloatReadRef WetDryMixInput;
+
+		FAudioBufferWriteRef AudioOutput;
+		// The gain being applied to the input buffer
+		FAudioBufferWriteRef EnvelopeOutput;
+
+		// Internal DSP Compressor
+		Audio::FDynamicsProcessor Compressor;
+
+		// Whether or not to use sidechain input (is false if no input pin is connected to sidechain input)
+		bool bUseSidechain;
+
+		// Cached variables
+		float PrevRatio;
+		float PrevThreshold;
+		double PrevAttackTime;
+		double PrevReleaseTime;
+		float PrevKnee;
+		EEnvelopePeakMode PrevEnvMode;
+		bool PrevIsAnalog;
+	};
+
+	// Node Class
+	class FCompressorNode : public FNodeFacade
+	{
+	public:
+		FCompressorNode(const FNodeInitData& InitData)
+			: FNodeFacade(InitData.InstanceName, InitData.InstanceID, TFacadeOperatorClass<FCompressorOperator>())
+		{
+		}
+	};
+
+	// Register node
+	METASOUND_REGISTER_NODE(FCompressorNode)
+
+}
+
+#undef LOCTEXT_NAMESPACE
diff --git a/Engine/Source/Runtime/SignalProcessing/Private/DynamicsProcessor.cpp b/Engine/Source/Runtime/SignalProcessing/Private/DynamicsProcessor.cpp
--- a/Engine/Source/Runtime/SignalProcessing/Private/DynamicsProcessor.cpp
+++ b/Engine/Source/Runtime/SignalProcessing/Private/DynamicsProcessor.cpp
@@ -259,6 +259,29 @@ namespace Audio
 		}
 	}
 
+	void FDynamicsProcessor::ProcessAudioFrame(const float* InFrame, float* OutFrame, const float* InKeyFrame, float* OutGain)
+	{
+		check(OutGain != nullptr);
+
+		const bool bKeyIsInput = InFrame == InKeyFrame;
+		if (ProcessKeyFrame(InKeyFrame, OutFrame, bKeyIsInput))
+		{
+			const int32 NumChannels = GetNumChannels();
+			for (int32 Channel = 0; Channel < NumChannels; ++Channel)
+			{
+				// Write and read into the look ahead delay line.
+				// We apply the compression output of the direct input to the output of this delay line
+				// This way sharp transients can be "caught" with the gain.
+				float LookaheadOutput = LookaheadDelay[Channel].ProcessAudioSample(InFrame[Channel]);
+
+				// Write into the output with the computed gain value
+				OutFrame[Channel] = Gain[Channel] * LookaheadOutput * OutputGain * InputGain;
+				// Also write the output gain value
+				OutGain[Channel] = Gain[Channel];
+			}
+		}
+	}
+
 	void FDynamicsProcessor::ProcessAudio(const float* InBuffer, const int32 InNumSamples, float* OutBuffer, const float* InKeyBuffer)
 	{
 		const int32 NumChannels = GetNumChannels();
@@ -284,6 +307,33 @@ namespace Audio
 		}
 	}
 
+	void FDynamicsProcessor::ProcessAudio(const float* InBuffer, const int32 InNumSamples, float* OutBuffer, float* OutEnvelope, const float* InKeyBuffer)
+	{
+		const int32 NumChannels = GetNumChannels();
+		const int32 KeyNumChannels = GetKeyNumChannels();
+
+		check(OutEnvelope != nullptr);
+
+		if (InKeyBuffer)
+		{
+			int32 KeySampleIndex = 0;
+			for (int32 SampleIndex = 0; SampleIndex < InNumSamples; SampleIndex += NumChannels)
+			{
+				const float* KeyFrame = &InKeyBuffer[KeySampleIndex];
+				ProcessAudioFrame(&InBuffer[SampleIndex], &OutBuffer[SampleIndex], KeyFrame, &OutEnvelope[SampleIndex]);
+				KeySampleIndex += KeyNumChannels;
+			}
+		}
+		else
+		{
+			for (int32 SampleIndex = 0; SampleIndex < InNumSamples; SampleIndex += NumChannels)
+			{
+				const float* KeyFrame = &InBuffer[SampleIndex];
+				ProcessAudioFrame(&InBuffer[SampleIndex], &OutBuffer[SampleIndex], KeyFrame, &OutEnvelope[SampleIndex]);
+			}
+		}
+	}
+
 	bool FDynamicsProcessor::ProcessKeyFrame(const float* InKeyFrame, float* OutFrame, bool bKeyIsInput)
 	{
 		// Get detector outputs
diff --git a/Engine/Source/Runtime/SignalProcessing/Public/DSP/DynamicsProcessor.h b/Engine/Source/Runtime/SignalProcessing/Public/DSP/DynamicsProcessor.h
--- a/Engine/Source/Runtime/SignalProcessing/Public/DSP/DynamicsProcessor.h
+++ b/Engine/Source/Runtime/SignalProcessing/Public/DSP/DynamicsProcessor.h
@@ -67,8 +67,10 @@ namespace Audio
 		void SetProcessingMode(const EDynamicsProcessingMode::Type ProcessingMode);
 
 		void ProcessAudioFrame(const float* InFrame, float* OutFrame, const float* InKeyFrame);
+		void ProcessAudioFrame(const float* InFrame, float* OutFrame, const float* InKeyFrame, float* OutGain);
 		void ProcessAudio(const float* InBuffer, const int32 InNumSamples, float* OutBuffer, const float* InKeyBuffer = nullptr);
-
+		void ProcessAudio(const float* InBuffer, const int32 InNumSamples, float* OutBuffer, float* OutEnvelope, const float* InKeyBuffer = nullptr);
+		 
 
 	protected:
 
