Fixed a few bugs in the Ladder Filter Node
- Initial values for cutoff and resonance were incorrect
- Execute() method was not using clamped input values
- Changed default resonance value to 1.0 (after fixing the first bug, 6.0 as the default was an extremely unpleasant sound).

diff --git a/Engine/Plugins/Runtime/Metasound/Source/MetasoundStandardNodes/Private/MetasoundBasicFilters.cpp b/Engine/Plugins/Runtime/Metasound/Source/MetasoundStandardNodes/Private/MetasoundBasicFilters.cpp
--- a/Engine/Plugins/Runtime/Metasound/Source/MetasoundStandardNodes/Private/MetasoundBasicFilters.cpp
+++ b/Engine/Plugins/Runtime/Metasound/Source/MetasoundStandardNodes/Private/MetasoundBasicFilters.cpp
@@ -184,7 +184,7 @@ namespace Metasound
 			FInputVertexInterface(
 				TInputDataVertexModel<FAudioBuffer>(METASOUND_GET_PARAM_NAME(ParamAudioInput), METASOUND_GET_PARAM_TT(ParamAudioInput)),
 				TInputDataVertexModel<float>(METASOUND_GET_PARAM_NAME(ParamCutoffFrequency), METASOUND_GET_PARAM_TT(ParamCutoffFrequency), 20000.f),
-				TInputDataVertexModel<float>(METASOUND_GET_PARAM_NAME(ParamResonance), METASOUND_GET_PARAM_TT(ParamResonance), 6.f)
+				TInputDataVertexModel<float>(METASOUND_GET_PARAM_NAME(ParamResonance), METASOUND_GET_PARAM_TT(ParamResonance), 1.f)
 			),
 			FOutputVertexInterface(
 				TOutputDataVertexModel<FAudioBuffer>(METASOUND_GET_PARAM_NAME(ParamAudioOutput), METASOUND_GET_PARAM_TT(ParamAudioOutput))
@@ -198,11 +198,13 @@ namespace Metasound
 	{
 
 		const FDataReferenceCollection& InputDataRefs = InParams.InputDataReferences;
+		FVertexInterface Interface = DeclareVertexInterface();
+		const FInputVertexInterface& InputInterface = Interface.GetInputInterface();
 
 		// inputs
 		FAudioBufferReadRef AudioIn = InputDataRefs.GetDataReadReferenceOrConstruct<FAudioBuffer>(METASOUND_GET_PARAM_NAME(ParamAudioInput), InParams.OperatorSettings);
-		FFloatReadRef FrequencyIn = InputDataRefs.GetDataReadReferenceOrConstruct<float>(METASOUND_GET_PARAM_NAME(ParamCutoffFrequency));
-		FFloatReadRef ResonanceIn = InputDataRefs.GetDataReadReferenceOrConstruct<float>(METASOUND_GET_PARAM_NAME(ParamResonance));
+		FFloatReadRef FrequencyIn = InputDataRefs.GetDataReadReferenceOrConstructWithVertexDefault<float>(InputInterface, METASOUND_GET_PARAM_NAME(ParamCutoffFrequency), InParams.OperatorSettings);
+		FFloatReadRef ResonanceIn = InputDataRefs.GetDataReadReferenceOrConstructWithVertexDefault<float>(InputInterface, METASOUND_GET_PARAM_NAME(ParamResonance), InParams.OperatorSettings);
 
 		return MakeUnique<FLadderFilterOperator>(
 			InParams.OperatorSettings
@@ -223,13 +225,13 @@ namespace Metasound
 
 		if (bNeedsUpdate)
 		{
-			LadderFilter.SetQ(*Resonance);
-			LadderFilter.SetFrequency(*Frequency);
+			LadderFilter.SetQ(CurrentResonance);
+			LadderFilter.SetFrequency(CurrentFrequency);
 
 			LadderFilter.Update();
 
-			PreviousFrequency = *Frequency;
-			PreviousResonance = *Resonance;
+			PreviousFrequency = CurrentFrequency;
+			PreviousResonance = CurrentResonance;
 		}
 
 		LadderFilter.ProcessAudio(AudioInput->GetData(), AudioInput->Num(), AudioOutput->GetData());
