Place new nodes near existing nodes using DataType as indicator of association & static node offset values when swapping archetypes
- Minor: Clean-up dependencies post autoupdate

diff --git a/Engine/Plugins/Runtime/Metasound/Source/MetasoundFrontend/Private/MetasoundFrontendTransform.cpp b/Engine/Plugins/Runtime/Metasound/Source/MetasoundFrontend/Private/MetasoundFrontendTransform.cpp
--- a/Engine/Plugins/Runtime/Metasound/Source/MetasoundFrontend/Private/MetasoundFrontendTransform.cpp
+++ b/Engine/Plugins/Runtime/Metasound/Source/MetasoundFrontend/Private/MetasoundFrontendTransform.cpp
@@ -6,6 +6,7 @@
 #include "MetasoundArchetype.h"
 #include "MetasoundAssetBase.h"
 #include "MetasoundFrontendArchetypeRegistry.h"
+#include "MetasoundFrontendDocument.h"
 #include "MetasoundFrontendRegistries.h"
 #include "MetasoundFrontendSearchEngine.h"
 #include "MetasoundLog.h"
@@ -59,7 +60,7 @@ namespace Metasound
 			}
 		}
 
-		bool FSwapGraphArchetype::Transform(FGraphHandle InGraph) const 
+		bool FSwapGraphArchetype::Transform(FGraphHandle InGraph) const
 		{
 			bool bDidEdit = false;
 
@@ -89,11 +90,64 @@ namespace Metasound
 				}
 			}
 
+			auto InputDataTypeCompareFilter = [](FConstNodeHandle NodeHandle, FName DataType)
+			{
+				bool bMatchesDataType = false;
+				NodeHandle->IterateConstOutputs([&](FConstOutputHandle OutputHandle)
+				{
+					if (OutputHandle->GetDataType() == DataType)
+					{
+						bMatchesDataType = true;
+					}
+				});
+
+				return bMatchesDataType;
+			};
+
+			auto OutputDataTypeCompareFilter = [](FConstNodeHandle NodeHandle, FName DataType)
+			{
+				bool bMatchesDataType = false;
+				NodeHandle->IterateConstInputs([&](FConstInputHandle InputHandle)
+				{
+					if (InputHandle->GetDataType() == DataType)
+					{
+						bMatchesDataType = true;
+					}
+				});
+
+				return bMatchesDataType;
+			};
+
+			auto FindLowestNodeLocationOfClassType = [](EMetasoundFrontendClassType ClassType, FGraphHandle Graph, FName DataType, TFunctionRef<bool(FConstNodeHandle, FName)> NodeDataTypeFilter)
+			{
+				FVector2D LowestLocation;
+				Graph->IterateConstNodes([&](FConstNodeHandle NodeHandle)
+				{
+					for (const TPair<FGuid, FVector2D>& Pair : NodeHandle->GetNodeStyle().Display.Locations)
+					{
+						if (Pair.Value.Y > LowestLocation.Y)
+						{
+							if (NodeDataTypeFilter(NodeHandle, DataType))
+							{
+								LowestLocation = Pair.Value;
+							}
+						}
+					}
+				}, ClassType);
+
+				return LowestLocation;
+			};
+
 			// Add missing inputs
 			for (const FMetasoundFrontendClassVertex& InputToAdd : InputsToAdd)
 			{
 				bDidEdit = true;
-				InGraph->AddInputVertex(InputToAdd);
+				FNodeHandle NewInputNode = InGraph->AddInputVertex(InputToAdd);
+
+				FMetasoundFrontendNodeStyle Style = NewInputNode->GetNodeStyle();
+				const FVector2D LastOutputLocation = FindLowestNodeLocationOfClassType(EMetasoundFrontendClassType::Input, InGraph, InputToAdd.TypeName, InputDataTypeCompareFilter);
+				Style.Display.Locations.Add(FGuid(), LastOutputLocation + DisplayStyle::NodeLayout::DefaultOffsetY);
+				NewInputNode->SetNodeStyle(Style);
 			}
 
 			// Add missing outputs
@@ -101,13 +155,10 @@ namespace Metasound
 			{
 				bDidEdit = true;
 				FNodeHandle NewOutputNode = InGraph->AddOutputVertex(OutputToAdd);
+
 				FMetasoundFrontendNodeStyle Style = NewOutputNode->GetNodeStyle();
-				// TODO: Create generic "FindGoodLocation()" for nodes. 
-				// Inputs are on the left,
-				// Outputs are on the right,
-				// Place near nodes with similar pin types.
-				// Requires a "NodeSize" approximation
-				Style.Display.Locations.Add(FGuid(), FVector2D{ 0.f, 0.f });
+				const FVector2D LastOutputLocation = FindLowestNodeLocationOfClassType(EMetasoundFrontendClassType::Output, InGraph, OutputToAdd.TypeName, OutputDataTypeCompareFilter);
+				Style.Display.Locations.Add(FGuid(), LastOutputLocation + DisplayStyle::NodeLayout::DefaultOffsetY);
 				NewOutputNode->SetNodeStyle(Style);
 			}
 
@@ -396,6 +447,8 @@ namespace Metasound
 				}
 			}
 
+			InDocument->SynchronizeDependencies();
+
 			return bDidEdit;
 		}
 
